Deprecate FGraphActionListBuilderBase::ActionGroup, as it is redundant with FEdGraphSchemaAction. It was introduced to handle 'multi edit' but was never fully realized and contains redundant categorization fields that are inconsistently used in scoring of menu entries. What little categorization was being done on the action group is now properly done on the FEdGraphSchemaAction

diff --git a/Engine/Plugins/Animation/DeformerGraph/Source/OptimusEditor/Private/Widgets/SOptimusEditorGraphExplorer.cpp b/Engine/Plugins/Animation/DeformerGraph/Source/OptimusEditor/Private/Widgets/SOptimusEditorGraphExplorer.cpp
--- a/Engine/Plugins/Animation/DeformerGraph/Source/OptimusEditor/Private/Widgets/SOptimusEditorGraphExplorer.cpp
+++ b/Engine/Plugins/Animation/DeformerGraph/Source/OptimusEditor/Private/Widgets/SOptimusEditorGraphExplorer.cpp
@@ -565,10 +565,7 @@ bool SOptimusEditorGraphExplorer::CanRequestRenameOnActionNode(TWeakPtr<FGraphAc
 	{
 		if (SelectedNode->IsActionNode())
 		{
-			if (ensure(!SelectedNode->Actions.IsEmpty()))
-			{
-				return CanRenameAction(SelectedNode->Actions[0]);
-			}
+			return CanRenameAction(SelectedNode->Action);
 		}
 	}
 
diff --git a/Engine/Plugins/Experimental/TextureGraph/Source/TextureGraphEditor/Private/Widgets/ActionMenuWidgets/STG_ActionMenuTileView.cpp b/Engine/Plugins/Experimental/TextureGraph/Source/TextureGraphEditor/Private/Widgets/ActionMenuWidgets/STG_ActionMenuTileView.cpp
--- a/Engine/Plugins/Experimental/TextureGraph/Source/TextureGraphEditor/Private/Widgets/ActionMenuWidgets/STG_ActionMenuTileView.cpp
+++ b/Engine/Plugins/Experimental/TextureGraph/Source/TextureGraphEditor/Private/Widgets/ActionMenuWidgets/STG_ActionMenuTileView.cpp
@@ -455,7 +455,7 @@ void STG_ActionMenuTileView::GetSelectedActions(TArray< TSharedPtr<FEdGraphSchem
 	{
 		for (int32 NodeIndex = 0; NodeIndex < SelectedNodes.Num(); NodeIndex++)
 		{
-			OutSelectedActions.Append(SelectedNodes[NodeIndex]->Actions);
+			OutSelectedActions.Add(SelectedNodes[NodeIndex]->Action);
 		}
 	}
 }
@@ -509,10 +509,7 @@ void STG_ActionMenuTileView::GetCategorySubActions(TWeakPtr<FGraphActionNode> In
 
 			if (CurrentChild.IsValid() && CurrentChild->IsActionNode())
 			{
-				for (int32 ActionIndex = 0; ActionIndex != CurrentChild->Actions.Num(); ActionIndex++)
-				{
-					OutActions.Add(CurrentChild->Actions[ActionIndex]);
-				}
+				OutActions.Add(CurrentChild->Action);
 			}
 		}
 	}
@@ -560,10 +557,8 @@ bool STG_ActionMenuTileView::SelectItemByName(const FName& ItemName, ESelectInfo
 			{
 				TSharedPtr<FGraphActionNode> CurrentChildNode = CurrentGraphNode->Children[ChildIdx];
 
-				for (int32 ActionIndex = 0; ActionIndex < CurrentChildNode->Actions.Num(); ActionIndex++)
+				if(FEdGraphSchemaAction* ChildGraphAction = CurrentChildNode->Action.Get())
 				{
-					FEdGraphSchemaAction* ChildGraphAction = CurrentChildNode->Actions[ActionIndex].Get();
-
 					// If the user is attempting to select a category, make sure it's a category
 					if (CurrentChildNode->IsCategoryNode() == bIsCategory)
 					{
@@ -1102,7 +1097,7 @@ void STG_ActionMenuTileView::OnItemDoubleClicked(TSharedPtr< FGraphActionNode >
 	{
 		if (InClickedItem->IsActionNode())
 		{
-			OnActionDoubleClicked.ExecuteIfBound(InClickedItem->Actions);
+			OnActionDoubleClicked.ExecuteIfBound({InClickedItem->Action});
 		}
 		else if (InClickedItem->Children.Num())
 		{
@@ -1256,12 +1251,8 @@ void STG_ActionMenuTileView::AddReferencedObjects(FReferenceCollector& Collector
 {
 	for (int32 CurTypeIndex = 0; CurTypeIndex < AllActions->GetNumActions(); ++CurTypeIndex)
 	{
-		FGraphActionListBuilderBase::ActionGroup& Action = AllActions->GetAction(CurTypeIndex);
-
-		for (int32 ActionIndex = 0; ActionIndex < Action.Actions.Num(); ActionIndex++)
-		{
-			Action.Actions[ActionIndex]->AddReferencedObjects(Collector);
-		}
+		TSharedPtr<FEdGraphSchemaAction>& Action = AllActions->GetSchemaAction(CurTypeIndex);
+		Action->AddReferencedObjects(Collector);
 	}
 }
 
@@ -1277,7 +1268,7 @@ bool STG_ActionMenuTileView::HandleSelection(TSharedPtr< FGraphActionNode >& InS
 	{
 		if (InSelectedItem.IsValid() && InSelectedItem->IsActionNode())
 		{
-			OnActionSelected.Execute(InSelectedItem->Actions, InSelectionType);
+			OnActionSelected.Execute({InSelectedItem->Action}, InSelectionType);
 			bResult = true;
 		}
 		else
@@ -1340,13 +1331,13 @@ STG_ActionMenuTileView::FScoreResults STG_ActionMenuTileView::ScoreAndAddActions
 	const int32 NumActions = AllActions->GetNumActions();
 	for (int32 CurTypeIndex = bIsPartialBuild ? StartingIndex : 0; CurTypeIndex < NumActions; ++CurTypeIndex)
 	{
-		FGraphActionListBuilderBase::ActionGroup& CurrentAction = AllActions->GetAction(CurTypeIndex);
+		TSharedPtr<FEdGraphSchemaAction> CurrentAction = AllActions->GetSchemaAction(CurTypeIndex);
 
 		// If we're filtering, search check to see if we need to show this action
 		bool bShowAction = true;
 		float EachWeight = TNumericLimits<float>::Lowest();
 
-		const FString& SearchText = CurrentAction.GetSearchTextForFirstAction();
+		const FString& SearchText = CurrentAction->GetFullSearchText();
 		for (int32 FilterIndex = 0; (FilterIndex < FilterTerms.Num()) && bShowAction; ++FilterIndex)
 		{
 			const bool bMatchesTerm = (SearchText.Contains(FilterTerms[FilterIndex], ESearchCase::CaseSensitive) || (SearchText.Contains(SanitizedFilterTerms[FilterIndex], ESearchCase::CaseSensitive) == true));
@@ -1361,7 +1352,7 @@ STG_ActionMenuTileView::FScoreResults STG_ActionMenuTileView::ScoreAndAddActions
 		if (bRequiresFiltering)
 		{
 			// Get the 'weight' of this in relation to the filter
-			EachWeight = ActionSchema->GetActionFilteredWeight(CurrentAction, FilterTerms, SanitizedFilterTerms, DraggedFromPins);
+			EachWeight = ActionSchema->GetActionFilteredWeight(*CurrentAction, FilterTerms, SanitizedFilterTerms, DraggedFromPins);
 			// If this action has a greater relevance than others, cache its index.
 			if (EachWeight > BestMatchCount)
 			{
@@ -1394,12 +1385,12 @@ void STG_ActionMenuTileView::UpdateActiveSelection(STG_ActionMenuTileView::FScor
 	// If theres a BestMatchIndex find it in the actions nodes and select it (maybe this should check the current selected suggestion first ?)
 	if (BestMatchIndex != INDEX_NONE)
 	{
-		FGraphActionListBuilderBase::ActionGroup& FilterSelectAction = AllActions->GetAction(BestMatchIndex);
-		if (FilterSelectAction.Actions[0].IsValid() == true)
+		TSharedPtr<FEdGraphSchemaAction>& Action = AllActions->GetSchemaAction(BestMatchIndex);
+		if (Action.IsValid() == true)
 		{
 			for (int32 iNode = 0; iNode < FilteredActionNodes.Num(); iNode++)
 			{
-				if (FilteredActionNodes[iNode].Get()->GetPrimaryAction() == FilterSelectAction.Actions[0])
+				if (FilteredActionNodes[iNode].Get()->GetPrimaryAction() == Action)
 				{
 					SelectedSuggestion = iNode;
 					SelectedSuggestionScore = BestMatchCount;
diff --git a/Engine/Source/Editor/BlueprintGraph/Classes/EdGraphSchema_K2.h b/Engine/Source/Editor/BlueprintGraph/Classes/EdGraphSchema_K2.h
--- a/Engine/Source/Editor/BlueprintGraph/Classes/EdGraphSchema_K2.h
+++ b/Engine/Source/Editor/BlueprintGraph/Classes/EdGraphSchema_K2.h
@@ -576,7 +576,11 @@ public:
 	virtual bool CanShowDataTooltipForPin(const UEdGraphPin& Pin) const override;
 
 #if WITH_EDITORONLY_DATA
+	virtual float GetActionFilteredWeight(const FEdGraphSchemaAction& InCurrentAction, const TArray<FString>& InFilterTerms, const TArray<FString>& InSanitizedFilterTerms, const TArray<UEdGraphPin*>& DraggedFromPins) const override;
+	PRAGMA_DISABLE_DEPRECATION_WARNINGS
+	UE_DEPRECATED(5.5, "FGraphActionListBuilderBase::ActionGroup weight calculation was always performed on the first action, provide that if scoring a FGraphActionListBuilderBase::ActionGroup, or just provide the FEdGraphSchemaAction")
 	virtual float GetActionFilteredWeight(const FGraphActionListBuilderBase::ActionGroup& InCurrentAction, const TArray<FString>& InFilterTerms, const TArray<FString>& InSanitizedFilterTerms, const TArray<UEdGraphPin*>& DraggedFromPins) const override;
+	PRAGMA_ENABLE_DEPRECATION_WARNINGS
 	virtual FGraphSchemaSearchWeightModifiers GetSearchWeightModifiers() const override;
 #endif // WITH_EDITORONLY_DATA	
 
diff --git a/Engine/Source/Editor/BlueprintGraph/Private/EdGraphSchema_K2.cpp b/Engine/Source/Editor/BlueprintGraph/Private/EdGraphSchema_K2.cpp
--- a/Engine/Source/Editor/BlueprintGraph/Private/EdGraphSchema_K2.cpp
+++ b/Engine/Source/Editor/BlueprintGraph/Private/EdGraphSchema_K2.cpp
@@ -7877,12 +7877,12 @@ struct FBPContextMenuWeightDebugInfo : public FGraphSchemaSearchTextDebugInfo
 	/**
 	* Print out the debug info about this weight info to the console
 	*/
-	virtual void Print(const TArray<FString>& SearchForKeywords, const FGraphActionListBuilderBase::ActionGroup& Action) const override
+	virtual void Print(const TArray<FString>& SearchForKeywords, const FEdGraphSchemaAction& Action) const override
 	{
 		// Combine the actions string, separate with \n so terms don't run into each other, and remove the spaces (incase the user is searching for a variable)
 		// In the case of groups containing multiple actions, they will have been created and added at the same place in the code, using the same description
 		// and keywords, so we only need to use the first one for filtering.
-		const FString& SearchText = Action.GetSearchTextForFirstAction();
+		const FString& SearchText = Action.GetFullSearchText();
 
 		UE_LOG(LogTemp, Warning, TEXT("[Weight for %s] \
 TotalWeight: %-8.2f | PercentageMatchWeight: %-8.2f | PercMatch: %-8.2f | ShorterWeight: %-8.2f | CategoryBonusWeight: %-8.2f | KeywordArrayWeight: %-8.2f | DescriptionWeight: %-8.2f | NodeTitleWeight: %-8.2f | CategoryWeight: %-8.2f | Fav. Bonus:%-8.2f\n"),
@@ -7890,7 +7890,7 @@ TotalWeight: %-8.2f | PercentageMatchWeight: %-8.2f | PercMatch: %-8.2f | Shorte
 	}
 };
 
-float UEdGraphSchema_K2::GetActionFilteredWeight(const FGraphActionListBuilderBase::ActionGroup& InCurrentAction, const TArray<FString>& InFilterTerms, const TArray<FString>& InSanitizedFilterTerms, const TArray<UEdGraphPin*>& DraggedFromPins) const
+float UEdGraphSchema_K2::GetActionFilteredWeight(const FEdGraphSchemaAction& InCurrentAction, const TArray<FString>& InFilterTerms, const TArray<FString>& InSanitizedFilterTerms, const TArray<UEdGraphPin*>& DraggedFromPins) const
 {
 	// The overall 'weight' of this action 
 	float TotalWeight = 0.0f;
@@ -7901,193 +7901,200 @@ float UEdGraphSchema_K2::GetActionFilteredWeight(const FGraphActionListBuilderBa
 
 	const bool bIsFromDrag = (DraggedFromPins.Num() > 0);
 
-	int32 Action = 0;
-	if (InCurrentAction.Actions[Action].IsValid() == true)
-	{
-		TSharedPtr<FEdGraphSchemaAction> CurrentAction = InCurrentAction.Actions[Action];
+	FGraphSchemaSearchWeightModifiers WeightModifiers = GetSearchWeightModifiers();
+	// If there are no keywords, bump the weight on description to compensate
+	const TArray<FString>& LocKeywords = InCurrentAction.GetLocalizedSearchKeywordsArray();
+	WeightModifiers.DescriptionWeight = LocKeywords.Num() > 0 ? WeightModifiers.DescriptionWeight : WeightModifiers.DescriptionWeight * 2.0f;
 
-		FGraphSchemaSearchWeightModifiers WeightModifiers = GetSearchWeightModifiers();
-		// If there are no keywords, bump the weight on description to compensate
-		const TArray<FString>& LocKeywords = InCurrentAction.GetLocalizedSearchKeywordsArrayForFirstAction();
-		WeightModifiers.DescriptionWeight = LocKeywords.Num() > 0 ? WeightModifiers.DescriptionWeight : WeightModifiers.DescriptionWeight * 2.0f;
+	CollectSearchTextWeightInfo(InCurrentAction, WeightModifiers, WeightedArrayList, &OutDebugInfo);
 
-		CollectSearchTextWeightInfo(InCurrentAction, WeightModifiers, WeightedArrayList, &OutDebugInfo);
+	// Give a weight bonus to actions whose category matches what was dragged off of
+	if (bIsFromDrag)
+	{
+		const TArray<FString>& InActionCategories = InCurrentAction.GetCategoryChain();
+		bool bAddMatchBonus = false;
 
-		// Give a weight bonus to actions whose category matches what was dragged off of
-		if (bIsFromDrag)
+		/** Get a string reference for an EPinContainerType */
+		auto GetContainerTypeString = [](const EPinContainerType Type) -> const FString&
 		{
-			const TArray<FString>& InActionCategories = InCurrentAction.GetCategoryChain();
-			bool bAddMatchBonus = false;
+			static const FString ArrayName = TEXT("Array");
+			static const FString MapName = TEXT("Map");
+			static const FString SetName = TEXT("Set");
+			static const FString InvalidName = TEXT("INVALID");
 
-			/** Get a string reference for an EPinContainerType */
-			auto GetContainerTypeString = [](const EPinContainerType Type) -> const FString&
+			switch (Type)
 			{
-				static const FString ArrayName = TEXT("Array");
-				static const FString MapName = TEXT("Map");
-				static const FString SetName = TEXT("Set");
-				static const FString InvalidName = TEXT("INVALID");
-
-				switch (Type)
-				{
-					case EPinContainerType::Array:
-						return ArrayName;
-					case EPinContainerType::Map:
-						return MapName;
-					case EPinContainerType::Set:
-						return SetName;
-					default:
-						return InvalidName;
-				}
-			};
+				case EPinContainerType::Array:
+					return ArrayName;
+				case EPinContainerType::Map:
+					return MapName;
+				case EPinContainerType::Set:
+					return SetName;
+				default:
+					return InvalidName;
+			}
+		};
 
-			bool bAddedContainerPreferenceBonus = false;
+		bool bAddedContainerPreferenceBonus = false;
 
-			for (const FString& InActionCategory : InActionCategories)
+		for (const FString& InActionCategory : InActionCategories)
+		{
+			for (UEdGraphPin* const FromPin : DraggedFromPins)
 			{
-				for (UEdGraphPin* const FromPin : DraggedFromPins)
-				{
-					check(FromPin != nullptr);
+				check(FromPin != nullptr);
 
-					// For containers, add a preference for functions that are marked in their category
-					if (!bAddedContainerPreferenceBonus && FromPin->PinType.IsContainer() && InActionCategory == GetContainerTypeString(FromPin->PinType.ContainerType))
-					{
-						TotalWeight += BPContextMenuConsoleVariables::ContainerBonus;
-						bAddedContainerPreferenceBonus = true;
-					}
+				// For containers, add a preference for functions that are marked in their category
+				if (!bAddedContainerPreferenceBonus && FromPin->PinType.IsContainer() && InActionCategory == GetContainerTypeString(FromPin->PinType.ContainerType))
+				{
+					TotalWeight += BPContextMenuConsoleVariables::ContainerBonus;
+					bAddedContainerPreferenceBonus = true;
+				}
 
-					// Check the subcategory of the object to cover more more complex struct types (LinearColor, date time, etc)
-					if (UObject* const SubCatObj = FromPin->PinType.PinSubCategoryObject.Get())
+				// Check the subcategory of the object to cover more more complex struct types (LinearColor, date time, etc)
+				if (UObject* const SubCatObj = FromPin->PinType.PinSubCategoryObject.Get())
+				{
+					const FString& SubCatObjName = SubCatObj->GetPathName();
+					// The pin SubObjectCategory names don't have any spaces, so split up the category
+					TArray<FString> DelimitedArray;
+					InActionCategory.ParseIntoArray(DelimitedArray, TEXT(" "), true);
+					for (const FString& DelimetedCat : DelimitedArray)
 					{
-						const FString& SubCatObjName = SubCatObj->GetPathName();
-						// The pin SubObjectCategory names don't have any spaces, so split up the category
-						TArray<FString> DelimitedArray;
-						InActionCategory.ParseIntoArray(DelimitedArray, TEXT(" "), true);
-						for (const FString& DelimetedCat : DelimitedArray)
+						if (SubCatObjName.Contains(DelimetedCat))
 						{
-							if (SubCatObjName.Contains(DelimetedCat))
-							{
-								bAddMatchBonus = true;
-								break;
-							}
+							bAddMatchBonus = true;
+							break;
 						}
 					}
-					// Check the category of the pin, this works for basic math types (int, float, byte, etc)
-					else if (InActionCategory.Contains(FromPin->PinType.PinCategory.ToString()))
-					{
-						bAddMatchBonus = true;
-					}
+				}
+				// Check the category of the pin, this works for basic math types (int, float, byte, etc)
+				else if (InActionCategory.Contains(FromPin->PinType.PinCategory.ToString()))
+				{
+					bAddMatchBonus = true;
+				}
 
-					// If we found match in any cases above then add the weight bonus and stop looking
-					if (bAddMatchBonus)
-					{
-						TotalWeight += BPContextMenuConsoleVariables::MatchingFromPinCategory;
-						OutDebugInfo.CategoryBonusWeight += BPContextMenuConsoleVariables::MatchingFromPinCategory;
+				// If we found match in any cases above then add the weight bonus and stop looking
+				if (bAddMatchBonus)
+				{
+					TotalWeight += BPContextMenuConsoleVariables::MatchingFromPinCategory;
+					OutDebugInfo.CategoryBonusWeight += BPContextMenuConsoleVariables::MatchingFromPinCategory;
 
-						// Break out of the loop so that we don't give any extra bonuses
-						break;
-					}
+					// Break out of the loop so that we don't give any extra bonuses
+					break;
 				}
 			}
 		}
+	}
 
-		// If the user has favorite this action, then give it a hefty bonus
-		const UEditorPerProjectUserSettings& EditorSettings = *GetDefault<UEditorPerProjectUserSettings>();
-		if (UBlueprintPaletteFavorites* BlueprintFavorites = EditorSettings.BlueprintFavorites)
+	// If the user has favorite this action, then give it a hefty bonus
+	const UEditorPerProjectUserSettings& EditorSettings = *GetDefault<UEditorPerProjectUserSettings>();
+	if (UBlueprintPaletteFavorites* BlueprintFavorites = EditorSettings.BlueprintFavorites)
+	{
+		if (BlueprintFavorites->IsFavorited(InCurrentAction))
 		{
-			if (BlueprintFavorites->IsFavorited(CurrentAction))
-			{
-				TotalWeight += BPContextMenuConsoleVariables::FavoriteBonus;
-				OutDebugInfo.FavoriteBonusWeight += BPContextMenuConsoleVariables::FavoriteBonus;
-			}
+			TotalWeight += BPContextMenuConsoleVariables::FavoriteBonus;
+			OutDebugInfo.FavoriteBonusWeight += BPContextMenuConsoleVariables::FavoriteBonus;
 		}
+	}
 
-		// Now iterate through all the filter terms and calculate a 'weight' using the values and multipliers
-		const FString* EachTerm = nullptr;
-		const FString* EachTermSanitized = nullptr;
+	// Now iterate through all the filter terms and calculate a 'weight' using the values and multipliers
+	const FString* EachTerm = nullptr;
+	const FString* EachTermSanitized = nullptr;
+
+	// For every filter item the user has typed in (the text in the search bar, seperated by spaces)
+	for (int32 FilterIndex = 0; FilterIndex < InFilterTerms.Num(); ++FilterIndex)
+	{
+		EachTerm = &InFilterTerms[FilterIndex];
+		EachTermSanitized = &InSanitizedFilterTerms[FilterIndex];
+		int32 TermLen = EachTerm->Len();
 
-		// For every filter item the user has typed in (the text in the search bar, seperated by spaces)
-		for (int32 FilterIndex = 0; FilterIndex < InFilterTerms.Num(); ++FilterIndex)
+		// Now check the weighted lists	(We could further improve the hit weight by checking consecutive word matches)
+		for (int32 iFindCount = 0; iFindCount < WeightedArrayList.Num(); ++iFindCount)
 		{
-			EachTerm = &InFilterTerms[FilterIndex];
-			EachTermSanitized = &InSanitizedFilterTerms[FilterIndex];
-			int32 TermLen = EachTerm->Len();
+			const TArray<FString>& KeywordArray = *WeightedArrayList[iFindCount].Array;
+			float WeightPerList = 0.0f;
+			float KeywordArrayWeight = WeightedArrayList[iFindCount].WeightModifier;
 
-			// Now check the weighted lists	(We could further improve the hit weight by checking consecutive word matches)
-			for (int32 iFindCount = 0; iFindCount < WeightedArrayList.Num(); ++iFindCount)
-			{
-				const TArray<FString>& KeywordArray = *WeightedArrayList[iFindCount].Array;
-				float WeightPerList = 0.0f;
-				float KeywordArrayWeight = WeightedArrayList[iFindCount].WeightModifier;
+			// Count of how many words in this keyword array contain a filter(letter) that the user has typed in
+			int32 WordMatchCount = 0;
 
-				// Count of how many words in this keyword array contain a filter(letter) that the user has typed in
-				int32 WordMatchCount = 0;
+			// The number of characters in the best matching word
+			int32 BestMatchCharLength = 0;
 
-				// The number of characters in the best matching word
-				int32 BestMatchCharLength = 0;
+			// Loop through every word that the user could be looking for
+			for (int32 iEachWord = 0; iEachWord < KeywordArray.Num(); ++iEachWord)
+			{
+				float WeightPerWord = 0.0f;
 
-				// Loop through every word that the user could be looking for
-				for (int32 iEachWord = 0; iEachWord < KeywordArray.Num(); ++iEachWord)
+				// If a word contains the letter that the user has typed in, than increment the whole match count					
+				if (KeywordArray[iEachWord].Contains(*EachTermSanitized, ESearchCase::CaseSensitive) || KeywordArray[iEachWord].Contains(*EachTerm, ESearchCase::CaseSensitive))
 				{
-					float WeightPerWord = 0.0f;
+					++WordMatchCount;
+					WeightPerWord += KeywordArrayWeight * BPContextMenuConsoleVariables::WordContainsLetterWeightMultiplier;
 
-					// If a word contains the letter that the user has typed in, than increment the whole match count					
-					if (KeywordArray[iEachWord].Contains(*EachTermSanitized, ESearchCase::CaseSensitive) || KeywordArray[iEachWord].Contains(*EachTerm, ESearchCase::CaseSensitive))
+					// If the word starts with the letter, give it a little extra boost of weight
+					if (KeywordArray[iEachWord].StartsWith(*EachTermSanitized, ESearchCase::CaseSensitive) || KeywordArray[iEachWord].StartsWith(*EachTerm, ESearchCase::CaseSensitive))
 					{
-						++WordMatchCount;
-						WeightPerWord += KeywordArrayWeight * BPContextMenuConsoleVariables::WordContainsLetterWeightMultiplier;
-
-						// If the word starts with the letter, give it a little extra boost of weight
-						if (KeywordArray[iEachWord].StartsWith(*EachTermSanitized, ESearchCase::CaseSensitive) || KeywordArray[iEachWord].StartsWith(*EachTerm, ESearchCase::CaseSensitive))
-						{
-							WeightPerWord += KeywordArrayWeight * BPContextMenuConsoleVariables::StartsWithBonusWeightMultiplier;
-						}
+						WeightPerWord += KeywordArrayWeight * BPContextMenuConsoleVariables::StartsWithBonusWeightMultiplier;
+					}
 
-						if (WeightPerWord > WeightPerList)
-						{
-							// Use the best word match weight, we don't want to double-count redundant keywords like add and addmap here
-							WeightPerList = WeightPerWord;
-							BestMatchCharLength = KeywordArray[iEachWord].Len();
-						}
+					if (WeightPerWord > WeightPerList)
+					{
+						// Use the best word match weight, we don't want to double-count redundant keywords like add and addmap here
+						WeightPerList = WeightPerWord;
+						BestMatchCharLength = KeywordArray[iEachWord].Len();
 					}
 				}
+			}
 
-				// If the user has dragged off of a pin then do not prefer shorter things, because that will result
-				// in the matching of "Add" for a container instead of "+" for numeric types
-				// We only care about length penalty if something actually matched
-				if (BestMatchCharLength > 0 && WeightPerList > 0)
-				{
-					// How many words that we are checking had partial matches compared to what the user typed in?
-					float PercMatch = static_cast<float>(WordMatchCount) / static_cast<float>(KeywordArray.Num());
+			// If the user has dragged off of a pin then do not prefer shorter things, because that will result
+			// in the matching of "Add" for a container instead of "+" for numeric types
+			// We only care about length penalty if something actually matched
+			if (BestMatchCharLength > 0 && WeightPerList > 0)
+			{
+				// How many words that we are checking had partial matches compared to what the user typed in?
+				float PercMatch = static_cast<float>(WordMatchCount) / static_cast<float>(KeywordArray.Num());
 
-					float PercentageBonus = (WeightPerList * PercMatch * BPContextMenuConsoleVariables::PercentageMatchWeightMultiplier);
-					WeightPerList += PercentageBonus;
+				float PercentageBonus = (WeightPerList * PercMatch * BPContextMenuConsoleVariables::PercentageMatchWeightMultiplier);
+				WeightPerList += PercentageBonus;
 
-					// The shorter the matching word, the larger bonus it gets
-					float ShortFactor = static_cast<float>(BPContextMenuConsoleVariables::MaxWordLength - FMath::Min(BestMatchCharLength, BPContextMenuConsoleVariables::MaxWordLength));
-					float ShortWeight = ShortFactor * BPContextMenuConsoleVariables::ShorterWeight * (bIsFromDrag ? 0.25f : 1.0f);
-					WeightPerList += ShortWeight;
+				// The shorter the matching word, the larger bonus it gets
+				float ShortFactor = static_cast<float>(BPContextMenuConsoleVariables::MaxWordLength - FMath::Min(BestMatchCharLength, BPContextMenuConsoleVariables::MaxWordLength));
+				float ShortWeight = ShortFactor * BPContextMenuConsoleVariables::ShorterWeight * (bIsFromDrag ? 0.25f : 1.0f);
+				WeightPerList += ShortWeight;
 
-					OutDebugInfo.PercentMatch += PercMatch;
-					OutDebugInfo.ShorterMatchWeight += ShortWeight;
-					OutDebugInfo.PercentMatchWeight += PercentageBonus;
-				}
+				OutDebugInfo.PercentMatch += PercMatch;
+				OutDebugInfo.ShorterMatchWeight += ShortWeight;
+				OutDebugInfo.PercentMatchWeight += PercentageBonus;
+			}
 
-				TotalWeight += WeightPerList;
-				if (WeightedArrayList[iFindCount].DebugWeight)
-				{
-					// Each weight is used twice so add them
-					*WeightedArrayList[iFindCount].DebugWeight += WeightPerList;
-				}
+			TotalWeight += WeightPerList;
+			if (WeightedArrayList[iFindCount].DebugWeight)
+			{
+				// Each weight is used twice so add them
+				*WeightedArrayList[iFindCount].DebugWeight += WeightPerList;
 			}
 		}
-		OutDebugInfo.TotalWeight = TotalWeight;
-
-		PrintSearchTextDebugInfo(InFilterTerms, CurrentAction, &OutDebugInfo);
 	}
+	OutDebugInfo.TotalWeight = TotalWeight;
+
+	PrintSearchTextDebugInfo(InFilterTerms, InCurrentAction, &OutDebugInfo);
 
 	return TotalWeight;
 }
 
+PRAGMA_DISABLE_DEPRECATION_WARNINGS
+float UEdGraphSchema_K2::GetActionFilteredWeight(const FGraphActionListBuilderBase::ActionGroup& InCurrentAction, const TArray<FString>& InFilterTerms, const TArray<FString>& InSanitizedFilterTerms, const TArray<UEdGraphPin*>& DraggedFromPins) const
+{
+	int32 Action = 0;
+	if (InCurrentAction.Actions[Action].IsValid() == true)
+	{
+		return GetActionFilteredWeight(*InCurrentAction.Actions[Action], InFilterTerms, InSanitizedFilterTerms, DraggedFromPins);
+	}
+
+	return 0.f;
+}
+PRAGMA_ENABLE_DEPRECATION_WARNINGS
+
 #endif // WITH_EDITORONLY_DATA
 
 /////////////////////////////////////////////////////
diff --git a/Engine/Source/Editor/GraphEditor/Private/GraphActionNode.cpp b/Engine/Source/Editor/GraphEditor/Private/GraphActionNode.cpp
--- a/Engine/Source/Editor/GraphEditor/Private/GraphActionNode.cpp
+++ b/Engine/Source/Editor/GraphEditor/Private/GraphActionNode.cpp
@@ -184,10 +184,10 @@ FGraphActionNode::FGraphActionNode(int32 InGrouping, int32 InSectionID)
 }
 
 //------------------------------------------------------------------------------
-FGraphActionNode::FGraphActionNode(TArray< TSharedPtr<FEdGraphSchemaAction> > const& ActionList, int32 InGrouping, int32 InSectionID)
+FGraphActionNode::FGraphActionNode(const TSharedPtr<FEdGraphSchemaAction>& InAction, int32 InGrouping, int32 InSectionID)
 	: SectionID(InSectionID)
 	, Grouping(InGrouping)
-	, Actions(ActionList)
+	, Action(InAction)
 	, bPendingRenameRequest(false)
 	, InsertOrder(0)
 	, TotalLeafs(0)
@@ -195,11 +195,27 @@ FGraphActionNode::FGraphActionNode(TArray< TSharedPtr<FEdGraphSchemaAction> > co
 }
 
 //------------------------------------------------------------------------------
+TSharedPtr<FGraphActionNode> FGraphActionNode::AddChild(const TSharedPtr<FEdGraphSchemaAction>& InAction)
+{
+	const TArray<FString>& CategoryStack = InAction->GetCategoryChain();
+
+	TSharedPtr<FGraphActionNode> ActionNode = FGraphActionNode::NewActionNode(InAction);
+	if (!ActionNode->IsCategoryNode() && !ActionNode->IsSectionHeadingNode())
+	{
+		++TotalLeafs;
+	}
+
+	AddChildRecursively(CategoryStack, 0, ActionNode);
+
+	return ActionNode;
+}
+
+PRAGMA_DISABLE_DEPRECATION_WARNINGS
 TSharedPtr<FGraphActionNode> FGraphActionNode::AddChild(FGraphActionListBuilderBase::ActionGroup const& ActionSet)
 {
 	const TArray<FString>& CategoryStack = ActionSet.GetCategoryChain();
 
-	TSharedPtr<FGraphActionNode> ActionNode = FGraphActionNode::NewActionNode(ActionSet.Actions);
+	TSharedPtr<FGraphActionNode> ActionNode = FGraphActionNode::NewActionNode(ActionSet.Actions[0]);
 	if (!ActionNode->IsCategoryNode() && !ActionNode->IsSectionHeadingNode())
 	{
 		++TotalLeafs;
@@ -209,11 +225,46 @@ TSharedPtr<FGraphActionNode> FGraphActionNode::AddChild(FGraphActionListBuilderB
 	
 	return ActionNode;
 }
+PRAGMA_ENABLE_DEPRECATION_WARNINGS
 
 //------------------------------------------------------------------------------
+TSharedPtr<FGraphActionNode> FGraphActionNode::AddChildAlphabetical(const TSharedPtr<FEdGraphSchemaAction>& InAction)
+{
+	TSharedPtr<FGraphActionNode> ActionNode = FGraphActionNode::NewActionNode(InAction);
+	check(ActionNode->SectionID == INVALID_SECTION_ID); // this method does not support sections, those should be built statically
+
+	if (!ActionNode->IsCategoryNode() && !ActionNode->IsSectionHeadingNode())
+	{
+		++TotalLeafs;
+	}
+
+	// if a divider hasn't been created for the grouping, create one:
+	AddChildGrouping(ActionNode, this->AsShared());
+
+	// find or add categories iteratively, inserting as needed:
+	FGraphActionNode* OwningCategory = this;
+	const TArray<FString>& CategoryStack = InAction->GetCategoryChain();
+	for (const FString& CategorySection : CategoryStack)
+	{
+		TSharedPtr<FGraphActionNode> CategoryNode = OwningCategory->FindMatchingParent(CategorySection, ActionNode);
+		if (!CategoryNode.IsValid())
+		{
+			CategoryNode = NewCategoryNode(CategorySection, ActionNode->Grouping, ActionNode->SectionID);
+			OwningCategory->InsertChildAlphabetical(CategoryNode);
+		}
+
+		OwningCategory = CategoryNode.Get();
+	}
+
+	// finally insert the leaf:
+	OwningCategory->InsertChildAlphabetical(ActionNode);
+	return ActionNode;
+}
+
+PRAGMA_DISABLE_DEPRECATION_WARNINGS
 TSharedPtr<FGraphActionNode> FGraphActionNode::AddChildAlphabetical(FGraphActionListBuilderBase::ActionGroup const& ActionSet)
 {
-	TSharedPtr<FGraphActionNode> ActionNode = FGraphActionNode::NewActionNode(ActionSet.Actions);
+	TSharedPtr<FGraphActionNode> ActionNode = FGraphActionNode::NewActionNode(ActionSet.Actions[0]);
 	check(ActionNode->SectionID == INVALID_SECTION_ID); // this method does not support sections, those should be built statically
 
 	if (!ActionNode->IsCategoryNode() && !ActionNode->IsSectionHeadingNode())
@@ -243,6 +294,7 @@ TSharedPtr<FGraphActionNode> FGraphActionNode::AddChildAlphabetical(FGraphAction
 	OwningCategory->InsertChildAlphabetical(ActionNode);
 	return ActionNode;
 }
+PRAGMA_ENABLE_DEPRECATION_WARNINGS
 
 //------------------------------------------------------------------------------
 TSharedPtr<FGraphActionNode> FGraphActionNode::AddSection(int32 InGrouping, int32 InSectionID)
@@ -380,7 +432,7 @@ bool FGraphActionNode::IsCategoryNode() const
 //------------------------------------------------------------------------------
 bool FGraphActionNode::IsActionNode() const
 {
-	return Actions.Num() != 0;
+	return Action.IsValid();
 }
 
 //------------------------------------------------------------------------------
@@ -433,14 +485,7 @@ bool FGraphActionNode::HasValidAction() const
 //------------------------------------------------------------------------------
 TSharedPtr<FEdGraphSchemaAction> FGraphActionNode::GetPrimaryAction() const
 {
-	for (const TSharedPtr<FEdGraphSchemaAction>& NodeAction : Actions)
-	{
-		if (NodeAction.IsValid())
-		{
-			return NodeAction;
-		}
-	}
-	return TSharedPtr<FEdGraphSchemaAction>();
+	return Action;
 }
 
 //------------------------------------------------------------------------------
@@ -493,22 +538,19 @@ TSharedPtr<FGraphActionNode> FGraphActionNode::NewCategoryNode(FString const& Ca
 }
 
 //------------------------------------------------------------------------------
-TSharedPtr<FGraphActionNode> FGraphActionNode::NewActionNode(TArray< TSharedPtr<FEdGraphSchemaAction> > const& ActionList)
+TSharedPtr<FGraphActionNode> FGraphActionNode::NewActionNode(const TSharedPtr<FEdGraphSchemaAction>& Action)
 {
 	int32 Grouping  = FGraphActionNodeImpl::DEFAULT_GROUPING;
 	int32 SectionID = INVALID_SECTION_ID;
 
-	for (TSharedPtr<FEdGraphSchemaAction> const& Action : ActionList)
+	Grouping = FMath::Max(Grouping, Action->GetGrouping());
+	if (SectionID == INVALID_SECTION_ID)
 	{
-		Grouping = FMath::Max(Grouping, Action->GetGrouping());
-		if (SectionID == INVALID_SECTION_ID)
-		{
-			// take the first non-zero section ID
-			SectionID = Action->GetSectionID();
-		}
+		// take the first non-zero section ID
+		SectionID = Action->GetSectionID();
 	}
 
-	FGraphActionNode* ActionNode = new FGraphActionNode(ActionList, Grouping, SectionID);
+	FGraphActionNode* ActionNode = new FGraphActionNode(Action, Grouping, SectionID);
 	TSharedPtr<FEdGraphSchemaAction> PrimeAction = ActionNode->GetPrimaryAction();
 	checkSlow(PrimeAction.IsValid());
 	ActionNode->DisplayText = PrimeAction->GetMenuDescription();
diff --git a/Engine/Source/Editor/GraphEditor/Private/SGraphActionMenu.cpp b/Engine/Source/Editor/GraphEditor/Private/SGraphActionMenu.cpp
--- a/Engine/Source/Editor/GraphEditor/Private/SGraphActionMenu.cpp
+++ b/Engine/Source/Editor/GraphEditor/Private/SGraphActionMenu.cpp
@@ -510,7 +510,7 @@ void SGraphActionMenu::GetSelectedActions(TArray< TSharedPtr<FEdGraphSchemaActio
 	{
 		for ( int32 NodeIndex = 0; NodeIndex < SelectedNodes.Num(); NodeIndex++ )
 		{
-			OutSelectedActions.Append( SelectedNodes[NodeIndex]->Actions );
+			OutSelectedActions.Add( SelectedNodes[NodeIndex]->Action );
 		}
 	}
 }
@@ -570,10 +570,7 @@ void SGraphActionMenu::GetCategorySubActions(TWeakPtr<FGraphActionNode> InAction
 
 			if (CurrentChild.IsValid() && CurrentChild->IsActionNode())
 			{
-				for ( int32 ActionIndex = 0; ActionIndex != CurrentChild->Actions.Num(); ActionIndex++ )
-				{
-					OutActions.Add(CurrentChild->Actions[ActionIndex]);
-				}
+				OutActions.Add(CurrentChild->Action);
 			}
 		}
 	}
@@ -621,10 +618,8 @@ bool SGraphActionMenu::SelectItemByName(const FName& ItemName, ESelectInfo::Type
 			{
 				TSharedPtr<FGraphActionNode> CurrentChildNode = CurrentGraphNode->Children[ChildIdx];
 
-				for ( int32 ActionIndex = 0; ActionIndex < CurrentChildNode->Actions.Num(); ActionIndex++ )
+				if(FEdGraphSchemaAction* ChildGraphAction = CurrentChildNode->Action.Get())
 				{
-					FEdGraphSchemaAction* ChildGraphAction = CurrentChildNode->Actions[ActionIndex].Get();
-
 					// If the user is attempting to select a category, make sure it's a category
 					if( CurrentChildNode->IsCategoryNode() == bIsCategory )
 					{
@@ -1091,7 +1086,7 @@ void SGraphActionMenu::OnItemDoubleClicked( TSharedPtr< FGraphActionNode > InCli
 	{
 		if ( InClickedItem->IsActionNode() )
 		{
-			OnActionDoubleClicked.ExecuteIfBound(InClickedItem->Actions);
+			OnActionDoubleClicked.ExecuteIfBound({InClickedItem->Action});
 		}
 		else if (InClickedItem->Children.Num())
 		{
@@ -1245,12 +1240,8 @@ void SGraphActionMenu::AddReferencedObjects( FReferenceCollector& Collector )
 {
 	for (int32 CurTypeIndex = 0; CurTypeIndex < AllActions->GetNumActions(); ++CurTypeIndex)
 	{
-		FGraphActionListBuilderBase::ActionGroup& Action = AllActions->GetAction(CurTypeIndex);
-
-		for ( int32 ActionIndex = 0; ActionIndex < Action.Actions.Num(); ActionIndex++ )
-		{
-			Action.Actions[ActionIndex]->AddReferencedObjects(Collector);
-		}
+		TSharedPtr<FEdGraphSchemaAction> Action = AllActions->GetSchemaAction(CurTypeIndex);
+		Action->AddReferencedObjects(Collector);
 	}
 }
 
@@ -1266,7 +1257,7 @@ bool SGraphActionMenu::HandleSelection( TSharedPtr< FGraphActionNode > &InSelect
 	{
 		if ( InSelectedItem.IsValid() && InSelectedItem->IsActionNode() )
 		{
-			OnActionSelected.Execute(InSelectedItem->Actions, InSelectionType);
+			OnActionSelected.Execute({InSelectedItem->Action}, InSelectionType);
 			bResult = true;
 		}
 		else
@@ -1329,13 +1320,13 @@ void SGraphActionMenu::ScoreAndAddActions(int32 StartingIndex)
 	const int32 NumActions = AllActions->GetNumActions();
 	for (int32 CurTypeIndex = bIsPartialBuild ? StartingIndex : 0; CurTypeIndex < NumActions; ++CurTypeIndex)
 	{
-		FGraphActionListBuilderBase::ActionGroup& CurrentAction = AllActions->GetAction(CurTypeIndex);
+		TSharedPtr<FEdGraphSchemaAction> CurrentAction = AllActions->GetSchemaAction(CurTypeIndex);
 
 		// If we're filtering, search check to see if we need to show this action
 		bool bShowAction = true;
 		float EachWeight = TNumericLimits<float>::Lowest();
 
-		const FString& SearchText = CurrentAction.GetSearchTextForFirstAction();
+		const FString& SearchText = CurrentAction->GetFullSearchText();
 		for (int32 FilterIndex = 0; (FilterIndex < FilterTerms.Num()) && bShowAction; ++FilterIndex)
 		{
 			const bool bMatchesTerm = (SearchText.Contains(FilterTerms[FilterIndex], ESearchCase::CaseSensitive) || (SearchText.Contains(SanitizedFilterTerms[FilterIndex], ESearchCase::CaseSensitive) == true));
@@ -1351,7 +1342,7 @@ void SGraphActionMenu::ScoreAndAddActions(int32 StartingIndex)
 		if (bRequiresFiltering)
 		{
 			// Get the 'weight' of this in relation to the filter
-			EachWeight = ActionSchema->GetActionFilteredWeight(CurrentAction, FilterTerms, SanitizedFilterTerms, DraggedFromPins);
+			EachWeight = ActionSchema->GetActionFilteredWeight(*CurrentAction, FilterTerms, SanitizedFilterTerms, DraggedFromPins);
 			// If this action has a greater relevance than others, cache its index.
 			if (EachWeight > BestMatchCount)
 			{
diff --git a/Engine/Source/Editor/GraphEditor/Public/GraphActionNode.h b/Engine/Source/Editor/GraphEditor/Public/GraphActionNode.h
--- a/Engine/Source/Editor/GraphEditor/Public/GraphActionNode.h
+++ b/Engine/Source/Editor/GraphEditor/Public/GraphActionNode.h
@@ -26,7 +26,9 @@ public:
 	int32 const SectionID;
 	/** Identifies the menu group that this node belongs to (defaults to zero) */
 	int32 const Grouping;
-	/** A set of actions to execute when this node is picked from a menu */
+	/** An action to execute when this node is picked from a menu */
+	TSharedPtr<FEdGraphSchemaAction> const Action;
+	UE_DEPRECATED(5.5, "!! WARNING: This array is no longer populated!! FGraphActionNode::Actions array only functioned with a single Action (GetPrimaryAction), access via Action")
 	TArray< TSharedPtr<FEdGraphSchemaAction> > const Actions;
 
 	/** */
@@ -46,24 +48,32 @@ public:
 
 	/**
 	 * Inserts a new action node (and any accompanying category nodes) based off
-	 * the provided ActionSet. 
+	 * the provided Action. 
 	 *
 	 * NOTE: This does NOT insert the node in a sorted manner. Call SortChildren() 
 	 *       separately or use AddChildAlphabetical
 	 * 
-	 * @param  ActionSet	A list of actions that you want the node to execute when picked.
+	 * @param  Action	An action that you want the node to execute when picked.
 	 * @return The new action node.
 	 */
+	TSharedPtr<FGraphActionNode> AddChild(const TSharedPtr<FEdGraphSchemaAction>& Action);
+	PRAGMA_DISABLE_DEPRECATION_WARNINGS
+	UE_DEPRECATED(5.5, "FGraphActionListBuilderBase::ActionGroup has been deprecated, use TSharedPtr<FEdGraphSchemaAction> directly")
 	TSharedPtr<FGraphActionNode> AddChild(FGraphActionListBuilderBase::ActionGroup const& ActionSet);
+	PRAGMA_ENABLE_DEPRECATION_WARNINGS
 
 	/**
 	 * Inserts a new action node (and any required category nodes) based off
-	 * the provided ActionSet. Inserts in alphabetical order.
+	 * the provided Action. Inserts in alphabetical order.
 	 *
-	 * @param  ActionSet	A list of actions that you want the node to execute when picked.
+	 * @param  Action	An action that you want the node to execute when picked.
 	 * @return The new action node.
 	 */
+	TSharedPtr<FGraphActionNode> AddChildAlphabetical(const TSharedPtr<FEdGraphSchemaAction>& Action);
+	PRAGMA_DISABLE_DEPRECATION_WARNINGS
+	UE_DEPRECATED(5.5, "FGraphActionListBuilderBase::ActionGroup has been deprecated, use TSharedPtr<FEdGraphSchemaAction> directly")
 	TSharedPtr<FGraphActionNode> AddChildAlphabetical(FGraphActionListBuilderBase::ActionGroup const& ActionSet);
+	PRAGMA_ENABLE_DEPRECATION_WARNINGS
 
 	TSharedPtr<FGraphActionNode> AddSection(int32 Grouping, int32 InSectionID);
 
@@ -212,11 +222,11 @@ private:
 	/**
 	 * Constructor for action nodes. Private so that users go through AddChild().
 	 *
-	 * @param  ActionList
+	 * @param  InAction
 	 * @param  Grouping
 	 * @param  SectionID
 	 */
-	FGraphActionNode(TArray< TSharedPtr<FEdGraphSchemaAction> > const& ActionList, int32 Grouping, int32 SectionID);
+	FGraphActionNode(const TSharedPtr<FEdGraphSchemaAction>& InAction, int32 InGrouping, int32 InSectionID);
 
 	/**
 	 *
@@ -241,10 +251,10 @@ private:
 	/**
 	 *
 	 *
-	 * @param  ActionList
+	 * @param  ActionNode
 	 * @return
 	 */
-	static TSharedPtr<FGraphActionNode> NewActionNode(TArray< TSharedPtr<FEdGraphSchemaAction> > const& ActionList);
+	static TSharedPtr<FGraphActionNode> NewActionNode(TSharedPtr<FEdGraphSchemaAction> const& ActionNode);
 
 	/**
 	 *
diff --git a/Engine/Source/Editor/Kismet/Classes/BlueprintPaletteFavorites.h b/Engine/Source/Editor/Kismet/Classes/BlueprintPaletteFavorites.h
--- a/Engine/Source/Editor/Kismet/Classes/BlueprintPaletteFavorites.h
+++ b/Engine/Source/Editor/Kismet/Classes/BlueprintPaletteFavorites.h
@@ -29,6 +29,8 @@ struct KISMET_API FFavoritedBlueprintPaletteItem
 	GENERATED_USTRUCT_BODY()
 
 public:
+	FFavoritedBlueprintPaletteItem();
+
 	/**
 	 * Sometime favorites can be coming from user edited .ini files, so this 
 	 * converts that readable text into a favorite (since the strings are user 
@@ -44,7 +46,10 @@ public:
 	 * 
 	 * @param  PaletteAction	The action you wish to favorite.
 	 */
-	FFavoritedBlueprintPaletteItem(TSharedPtr<FEdGraphSchemaAction> PaletteAction = nullptr);
+	FFavoritedBlueprintPaletteItem(const FEdGraphSchemaAction& InPaletteAction);
+
+	UE_DEPRECATED(5.5, "Provide a dereferenced FEdGraphSchemaAction, after validating it in the caller")
+	FFavoritedBlueprintPaletteItem(TSharedPtr<FEdGraphSchemaAction> PaletteAction);
 
 	/**
 	 * 
@@ -123,7 +128,8 @@ public:
 	 * @return True if this action is already favorited, false if it not.
 	 */
 	bool IsFavorited(TSharedPtr<FEdGraphSchemaAction> PaletteAction) const;
-
+	bool IsFavorited(const FEdGraphSchemaAction& InCurrentAction) const;
+	
 	/**
 	 * 
 	 * 
diff --git a/Engine/Source/Editor/Kismet/Private/BlueprintActionMenuUtils.cpp b/Engine/Source/Editor/Kismet/Private/BlueprintActionMenuUtils.cpp
--- a/Engine/Source/Editor/Kismet/Private/BlueprintActionMenuUtils.cpp
+++ b/Engine/Source/Editor/Kismet/Private/BlueprintActionMenuUtils.cpp
@@ -794,41 +794,48 @@ void FBlueprintActionMenuUtils::MakeFavoritesMenu(FBlueprintActionContext const&
 }
 
 //------------------------------------------------------------------------------
-const UK2Node* FBlueprintActionMenuUtils::ExtractNodeTemplateFromAction(TSharedPtr<FEdGraphSchemaAction> PaletteAction)
+const UK2Node* FBlueprintActionMenuUtils::ExtractNodeTemplateFromAction(const TSharedPtr<FEdGraphSchemaAction>& PaletteAction)
 {
-	UK2Node const* TemplateNode = NULL;
 	if (PaletteAction.IsValid())
 	{
-		FName const ActionId = PaletteAction->GetTypeId();
-		if (ActionId == FBlueprintActionMenuItem::StaticGetTypeId())
-		{
-			FBlueprintActionMenuItem* NewNodeActionMenuItem = (FBlueprintActionMenuItem*)PaletteAction.Get();
-			TemplateNode = Cast<UK2Node>(NewNodeActionMenuItem->GetRawAction()->GetTemplateNode());
-		}
-		else if (ActionId == FBlueprintDragDropMenuItem::StaticGetTypeId())
-		{
-// 			FBlueprintDragDropMenuItem* DragDropActionMenuItem = (FBlueprintDragDropMenuItem*)PaletteAction.Get();
-// 			TemplateNode = Cast<UK2Node>(DragDropActionMenuItem->GetSampleAction()->GetTemplateNode());
-		}
-		// if this action inherits from FEdGraphSchemaAction_K2NewNode
-		else if (ActionId == FEdGraphSchemaAction_K2NewNode::StaticGetTypeId() ||
-			ActionId == FEdGraphSchemaAction_K2AssignDelegate::StaticGetTypeId() ||
-			ActionId == FEdGraphSchemaAction_K2AddComponent::StaticGetTypeId() ||
-			ActionId == FEdGraphSchemaAction_K2AddCustomEvent::StaticGetTypeId() ||
-			ActionId == FEdGraphSchemaAction_K2AddCallOnActor::StaticGetTypeId() ||
-			ActionId == FEdGraphSchemaAction_K2TargetNode::StaticGetTypeId() ||
-			ActionId == FEdGraphSchemaAction_K2Event::StaticGetTypeId() || 
-			ActionId == FEdGraphSchemaAction_K2AddEvent::StaticGetTypeId() ||
-			ActionId == FEdGraphSchemaAction_K2InputAction::StaticGetTypeId())
-		{
-			FEdGraphSchemaAction_K2NewNode* NewNodeAction = (FEdGraphSchemaAction_K2NewNode*)PaletteAction.Get();
-			TemplateNode = NewNodeAction->NodeTemplate;
-		}
-		else if (ActionId == FEdGraphSchemaAction_K2ViewNode::StaticGetTypeId())
-		{
-			FEdGraphSchemaAction_K2ViewNode* FocusNodeAction = (FEdGraphSchemaAction_K2ViewNode*)PaletteAction.Get();
-			TemplateNode = FocusNodeAction->NodePtr;
-		}
+		return ExtractNodeTemplateFromAction(*PaletteAction);
+	}
+
+	return nullptr;
+}
+
+const UK2Node* FBlueprintActionMenuUtils::ExtractNodeTemplateFromAction(const FEdGraphSchemaAction& PaletteAction)
+{
+	UK2Node const* TemplateNode = nullptr;
+	FName const ActionId = PaletteAction.GetTypeId();
+	if (ActionId == FBlueprintActionMenuItem::StaticGetTypeId())
+	{
+		const FBlueprintActionMenuItem& NewNodeActionMenuItem = (const FBlueprintActionMenuItem&)PaletteAction;
+		TemplateNode = Cast<UK2Node>(NewNodeActionMenuItem.GetRawAction()->GetTemplateNode());
+	}
+	else if (ActionId == FBlueprintDragDropMenuItem::StaticGetTypeId())
+	{
+		// 			const FBlueprintDragDropMenuItem& DragDropActionMenuItem = (const FBlueprintDragDropMenuItem&)PaletteAction.Get();
+		// 			TemplateNode = Cast<UK2Node>(DragDropActionMenuItem.GetSampleAction()->GetTemplateNode());
+	}
+	// if this action inherits from FEdGraphSchemaAction_K2NewNode
+	else if (ActionId == FEdGraphSchemaAction_K2NewNode::StaticGetTypeId() ||
+		ActionId == FEdGraphSchemaAction_K2AssignDelegate::StaticGetTypeId() ||
+		ActionId == FEdGraphSchemaAction_K2AddComponent::StaticGetTypeId() ||
+		ActionId == FEdGraphSchemaAction_K2AddCustomEvent::StaticGetTypeId() ||
+		ActionId == FEdGraphSchemaAction_K2AddCallOnActor::StaticGetTypeId() ||
+		ActionId == FEdGraphSchemaAction_K2TargetNode::StaticGetTypeId() ||
+		ActionId == FEdGraphSchemaAction_K2Event::StaticGetTypeId() ||
+		ActionId == FEdGraphSchemaAction_K2AddEvent::StaticGetTypeId() ||
+		ActionId == FEdGraphSchemaAction_K2InputAction::StaticGetTypeId())
+	{
+		const FEdGraphSchemaAction_K2NewNode& NewNodeAction = (const FEdGraphSchemaAction_K2NewNode&)PaletteAction;
+		TemplateNode = NewNodeAction.NodeTemplate;
+	}
+	else if (ActionId == FEdGraphSchemaAction_K2ViewNode::StaticGetTypeId())
+	{
+		const FEdGraphSchemaAction_K2ViewNode& FocusNodeAction = (const FEdGraphSchemaAction_K2ViewNode&)PaletteAction;
+		TemplateNode = FocusNodeAction.NodePtr;
 	}
 
 	return TemplateNode;
diff --git a/Engine/Source/Editor/Kismet/Private/BlueprintPaletteFavorites.cpp b/Engine/Source/Editor/Kismet/Private/BlueprintPaletteFavorites.cpp
--- a/Engine/Source/Editor/Kismet/Private/BlueprintPaletteFavorites.cpp
+++ b/Engine/Source/Editor/Kismet/Private/BlueprintPaletteFavorites.cpp
@@ -65,24 +65,24 @@ namespace BlueprintPaletteFavoritesImpl
 	 * @param  PaletteAction	The action you want a signature for.
 	 * @return A signature object, distinguishing the palette action from others (could also be invalid).
 	 */
-	static FBlueprintNodeSignature ConstructLegacySignature(TSharedPtr<FEdGraphSchemaAction> PaletteAction);
+	static FBlueprintNodeSignature ConstructLegacySignature(const FEdGraphSchemaAction& PaletteAction);
 }
 
 //------------------------------------------------------------------------------
-static FBlueprintNodeSignature BlueprintPaletteFavoritesImpl::ConstructLegacySignature(TSharedPtr<FEdGraphSchemaAction> InPaletteAction)
+static FBlueprintNodeSignature BlueprintPaletteFavoritesImpl::ConstructLegacySignature(const FEdGraphSchemaAction& InPaletteAction)
 {
 	TSubclassOf<UEdGraphNode> SignatureNodeClass;
 	FFieldVariant SignatureSubObject;
 	FName SignatureSubObjName;
 
-	FName const ActionId = InPaletteAction->GetTypeId();
+	FName const ActionId = InPaletteAction.GetTypeId();
 	if (ActionId == FEdGraphSchemaAction_K2AddComponent::StaticGetTypeId())
 	{
-		FEdGraphSchemaAction_K2AddComponent* AddComponentAction = (FEdGraphSchemaAction_K2AddComponent*)InPaletteAction.Get();
-		checkSlow(AddComponentAction->NodeTemplate != nullptr);
+		const FEdGraphSchemaAction_K2AddComponent& AddComponentAction = (const FEdGraphSchemaAction_K2AddComponent&)InPaletteAction;
+		checkSlow(AddComponentAction.NodeTemplate != nullptr);
 
-		SignatureNodeClass = AddComponentAction->NodeTemplate->GetClass();
-		SignatureSubObject = AddComponentAction->ComponentClass.Get();
+		SignatureNodeClass = AddComponentAction.NodeTemplate->GetClass();
+		SignatureSubObject = AddComponentAction.ComponentClass.Get();
 	}
 	else if (ActionId == FEdGraphSchemaAction_K2AddComment::StaticGetTypeId())
 	{
@@ -90,10 +90,10 @@ static FBlueprintNodeSignature BlueprintPaletteFavoritesImpl::ConstructLegacySig
 	}
 	else if (ActionId == FEdGraphSchemaAction_K2Delegate::StaticGetTypeId())
 	{
-		FEdGraphSchemaAction_K2Delegate* DelegateAction = (FEdGraphSchemaAction_K2Delegate*)InPaletteAction.Get();
+		const FEdGraphSchemaAction_K2Delegate& DelegateAction = (const FEdGraphSchemaAction_K2Delegate&)InPaletteAction;
 
 		SignatureNodeClass = UK2Node_BaseMCDelegate::StaticClass();
-		SignatureSubObject = DelegateAction->GetDelegateProperty();
+		SignatureSubObject = DelegateAction.GetDelegateProperty();
 	}
 	// if we can pull out a node associated with this action
 	else if (UK2Node const* NodeTemplate = FBlueprintActionMenuUtils::ExtractNodeTemplateFromAction(InPaletteAction))
@@ -171,6 +171,10 @@ static FBlueprintNodeSignature BlueprintPaletteFavoritesImpl::ConstructLegacySig
  * FFavoritedBlueprintPaletteItem
  ******************************************************************************/
 
+FFavoritedBlueprintPaletteItem::FFavoritedBlueprintPaletteItem()
+{
+}
+
 //------------------------------------------------------------------------------
 FFavoritedBlueprintPaletteItem::FFavoritedBlueprintPaletteItem(FString const& SerializedAction)
 	: ActionSignature(SerializedAction)
@@ -204,11 +208,37 @@ FFavoritedBlueprintPaletteItem::FFavoritedBlueprintPaletteItem(TSharedPtr<FEdGra
 		}
 		else
 		{
-			ActionSignature = BlueprintPaletteFavoritesImpl::ConstructLegacySignature(InPaletteAction);
+			ActionSignature = BlueprintPaletteFavoritesImpl::ConstructLegacySignature(*InPaletteAction);
 		}
 	}
 }
 
+FFavoritedBlueprintPaletteItem::FFavoritedBlueprintPaletteItem(const FEdGraphSchemaAction& InPaletteAction)
+{
+	if (InPaletteAction.GetTypeId() == FBlueprintActionMenuItem::StaticGetTypeId())
+	{
+		const FBlueprintActionMenuItem& ActionMenuItem = (const FBlueprintActionMenuItem&)InPaletteAction;
+		ActionSignature = ActionMenuItem.GetRawAction()->GetSpawnerSignature();
+	}
+	else if (InPaletteAction.GetTypeId() == FBlueprintDragDropMenuItem::StaticGetTypeId())
+	{
+		const FBlueprintDragDropMenuItem& CollectionMenuItem = (const FBlueprintDragDropMenuItem&)InPaletteAction;
+		ActionSignature = CollectionMenuItem.GetSampleAction()->GetSpawnerSignature();
+
+		// drag-n-drop menu items represent a collection of actions on the 
+		// same field (they spawn a sub-menu for the user to pick from), so
+		// they don't have a single node class
+		ActionSignature.SetNodeClass(nullptr);
+
+		static const FName CollectionSignatureKey(TEXT("ActionCollection"));
+		ActionSignature.AddNamedValue(CollectionSignatureKey, TEXT("true"));
+	}
+	else
+	{
+		ActionSignature = BlueprintPaletteFavoritesImpl::ConstructLegacySignature(InPaletteAction);
+	}
+}
+
 //------------------------------------------------------------------------------
 FFavoritedBlueprintPaletteItem::FFavoritedBlueprintPaletteItem(UBlueprintNodeSpawner const* BlueprintAction)
 	: ActionSignature(BlueprintAction->GetSpawnerSignature())
@@ -230,7 +260,7 @@ bool FFavoritedBlueprintPaletteItem::operator==(FFavoritedBlueprintPaletteItem c
 //------------------------------------------------------------------------------
 bool FFavoritedBlueprintPaletteItem::operator==(TSharedPtr<FEdGraphSchemaAction> PaletteAction) const
 {
-	return (*this == FFavoritedBlueprintPaletteItem(PaletteAction));
+	return PaletteAction.IsValid() && (*this == FFavoritedBlueprintPaletteItem(*PaletteAction));
 }
 
 //------------------------------------------------------------------------------
@@ -285,23 +315,25 @@ void UBlueprintPaletteFavorites::PostEditChangeProperty(struct FPropertyChangedE
 //------------------------------------------------------------------------------
 bool UBlueprintPaletteFavorites::CanBeFavorited(TSharedPtr<FEdGraphSchemaAction> PaletteAction) const
 {
-	return FFavoritedBlueprintPaletteItem(PaletteAction).IsValid();
+	return PaletteAction.IsValid() ? FFavoritedBlueprintPaletteItem(*PaletteAction).IsValid() : false;
 }
 
 //------------------------------------------------------------------------------
 bool UBlueprintPaletteFavorites::IsFavorited(TSharedPtr<FEdGraphSchemaAction> PaletteAction) const
+{
+	return PaletteAction.IsValid() ? IsFavorited(*PaletteAction) : false;
+}
+
+//------------------------------------------------------------------------------
+bool UBlueprintPaletteFavorites::IsFavorited(const FEdGraphSchemaAction& PaletteAction) const
 {
 	bool bIsFavorited = false;
-	if (!PaletteAction.IsValid())
-	{
-		bIsFavorited = false;
-	}
-	else if (PaletteAction->GetTypeId() == FBlueprintDragDropMenuItem::StaticGetTypeId())
+	if (PaletteAction.GetTypeId() == FBlueprintDragDropMenuItem::StaticGetTypeId())
 	{
-		FBlueprintDragDropMenuItem* CollectionMenuItem = (FBlueprintDragDropMenuItem*)PaletteAction.Get();
+		const FBlueprintDragDropMenuItem& CollectionMenuItem = (FBlueprintDragDropMenuItem&)PaletteAction;
 
 		bIsFavorited = true;
-		for (UBlueprintNodeSpawner const* Action : CollectionMenuItem->GetActionSet())
+		for (UBlueprintNodeSpawner const* Action : CollectionMenuItem.GetActionSet())
 		{
 			if (!IsFavorited(Action))
 			{
@@ -366,7 +398,7 @@ void UBlueprintPaletteFavorites::AddFavorite(TSharedPtr<FEdGraphSchemaAction> Pa
 		}
 		else
 		{
-			CurrentFavorites.Add(FFavoritedBlueprintPaletteItem(PaletteAction));
+			CurrentFavorites.Add(FFavoritedBlueprintPaletteItem(*PaletteAction));
 		}
 		SetProfile(BlueprintPaletteFavoritesImpl::CustomProfileId);
 	}
@@ -403,9 +435,9 @@ void UBlueprintPaletteFavorites::RemoveFavorite(TSharedPtr<FEdGraphSchemaAction>
 			SetProfile(BlueprintPaletteFavoritesImpl::CustomProfileId);
 		}
 	}
-	else if (IsFavorited(PaletteAction))
+	else if (IsFavorited(*PaletteAction))
 	{
-		CurrentFavorites.Remove(PaletteAction);
+		CurrentFavorites.Remove(*PaletteAction);
 		SetProfile(BlueprintPaletteFavoritesImpl::CustomProfileId);
 	}
 }
diff --git a/Engine/Source/Editor/Kismet/Private/SBlueprintPalette.cpp b/Engine/Source/Editor/Kismet/Private/SBlueprintPalette.cpp
--- a/Engine/Source/Editor/Kismet/Private/SBlueprintPalette.cpp
+++ b/Engine/Source/Editor/Kismet/Private/SBlueprintPalette.cpp
@@ -498,8 +498,12 @@ static void GetPaletteItemIcon(TSharedPtr<FEdGraphSchemaAction> ActionIn, UBluep
 static TSharedRef<IToolTip> ConstructToolTipWithActionPath(TSharedPtr<FEdGraphSchemaAction> ActionIn, TSharedPtr<IToolTip> OldToolTip)
 {
 	TSharedRef<IToolTip> NewToolTip = OldToolTip.ToSharedRef();
+	if (!ActionIn)
+	{
+		return NewToolTip;
+	}
 
-	FFavoritedBlueprintPaletteItem ActionItem(ActionIn);
+	FFavoritedBlueprintPaletteItem ActionItem(*ActionIn);
 	if (ActionItem.IsValid())
 	{
 		static FTextBlockStyle PathStyle = FTextBlockStyle()
@@ -1944,7 +1948,7 @@ FText SBlueprintPaletteItem::GetToolTipText() const
 				ToolTipText = ComponentClass->GetToolTipText();
 			}
 		}
-		else if (UK2Node const* const NodeTemplate = FBlueprintActionMenuUtils::ExtractNodeTemplateFromAction(PaletteAction))
+		else if (UK2Node const* const NodeTemplate = FBlueprintActionMenuUtils::ExtractNodeTemplateFromAction(*PaletteAction))
 		{
 			// If the node wants to create tooltip text, use that instead, because its probably more detailed
 			FText NodeToolTipText = NodeTemplate->GetTooltipText();
diff --git a/Engine/Source/Editor/Kismet/Private/SMyBlueprint.cpp b/Engine/Source/Editor/Kismet/Private/SMyBlueprint.cpp
--- a/Engine/Source/Editor/Kismet/Private/SMyBlueprint.cpp
+++ b/Engine/Source/Editor/Kismet/Private/SMyBlueprint.cpp
@@ -1124,11 +1124,10 @@ bool SMyBlueprint::CanRequestRenameOnActionNode(TWeakPtr<FGraphActionNode> InSel
 	}
 	else if (InSelectedNode.Pin()->IsActionNode())
 	{
-		check( InSelectedNode.Pin()->Actions.Num() > 0 && InSelectedNode.Pin()->Actions[0].IsValid() );
-		bIsReadOnly = FBlueprintEditorUtils::IsPaletteActionReadOnly(InSelectedNode.Pin()->Actions[0], BlueprintEditorPtr.Pin());
+		bIsReadOnly = FBlueprintEditorUtils::IsPaletteActionReadOnly(InSelectedNode.Pin()->Action, BlueprintEditorPtr.Pin());
 		if(!bIsReadOnly)
 		{
-			bIsReadOnly = !InSelectedNode.Pin()->Actions[0]->CanBeRenamed();
+			bIsReadOnly = !InSelectedNode.Pin()->Action->CanBeRenamed();
 		}
 	}
 
diff --git a/Engine/Source/Editor/Kismet/Public/BlueprintActionMenuUtils.h b/Engine/Source/Editor/Kismet/Public/BlueprintActionMenuUtils.h
--- a/Engine/Source/Editor/Kismet/Public/BlueprintActionMenuUtils.h
+++ b/Engine/Source/Editor/Kismet/Public/BlueprintActionMenuUtils.h
@@ -78,6 +78,7 @@ struct FBlueprintActionMenuUtils
 	 * @param  PaletteAction	The action you want a node-template for.
 	 * @return A pointer to the extracted node (NULL if the action doesn't have one, or we don't support the specific action type yet)
 	 */
-	KISMET_API static const UK2Node* ExtractNodeTemplateFromAction(TSharedPtr<FEdGraphSchemaAction> PaletteAction);
+	KISMET_API static const UK2Node* ExtractNodeTemplateFromAction(const FEdGraphSchemaAction& PaletteAction);
+	KISMET_API static const UK2Node* ExtractNodeTemplateFromAction(const TSharedPtr<FEdGraphSchemaAction>& PaletteAction);
 };
 
diff --git a/Engine/Source/Editor/UnrealEd/Classes/MaterialGraph/MaterialGraphSchema.h b/Engine/Source/Editor/UnrealEd/Classes/MaterialGraph/MaterialGraphSchema.h
--- a/Engine/Source/Editor/UnrealEd/Classes/MaterialGraph/MaterialGraphSchema.h
+++ b/Engine/Source/Editor/UnrealEd/Classes/MaterialGraph/MaterialGraphSchema.h
@@ -281,7 +281,11 @@ class UMaterialGraphSchema : public UEdGraphSchema
 	virtual bool SafeDeleteNodeFromGraph(UEdGraph* Graph, UEdGraphNode* NodeToDelete) const override;
 	virtual void GetAssetsGraphHoverMessage(const TArray<FAssetData>& Assets, const UEdGraph* HoverGraph, FString& OutTooltipText, bool& OutOkIcon) const;
 #if WITH_EDITORONLY_DATA
+	virtual float GetActionFilteredWeight(const FEdGraphSchemaAction& InCurrentAction, const TArray<FString>& InFilterTerms, const TArray<FString>& InSanitizedFilterTerms, const TArray<UEdGraphPin*>& DraggedFromPins) const override;
+	PRAGMA_DISABLE_DEPRECATION_WARNINGS
+	UE_DEPRECATED(5.5, "FGraphActionListBuilderBase::ActionGroup weight calculation was always performed on the first action, provide that if scoring a FGraphActionListBuilderBase::ActionGroup, or just provide the FEdGraphSchemaAction")
 	virtual float GetActionFilteredWeight(const FGraphActionListBuilderBase::ActionGroup& InCurrentAction, const TArray<FString>& InFilterTerms, const TArray<FString>& InSanitizedFilterTerms, const TArray<UEdGraphPin*>& DraggedFromPins) const override;
+	PRAGMA_ENABLE_DEPRECATION_WARNINGS
 	virtual FGraphSchemaSearchWeightModifiers GetSearchWeightModifiers() const override;
 #endif // WITH_EDITORONLY_DATA	
 	//~ End UEdGraphSchema Interface
diff --git a/Engine/Source/Editor/UnrealEd/Private/Commandlets/GenerateBlueprintAPICommandlet.cpp b/Engine/Source/Editor/UnrealEd/Private/Commandlets/GenerateBlueprintAPICommandlet.cpp
--- a/Engine/Source/Editor/UnrealEd/Private/Commandlets/GenerateBlueprintAPICommandlet.cpp
+++ b/Engine/Source/Editor/UnrealEd/Private/Commandlets/GenerateBlueprintAPICommandlet.cpp
@@ -182,7 +182,7 @@ GenerateBlueprintAPI commandlet params: \n\
 	 * distinguish similarly named actions). Can then be used to sort and 
 	 * uniquely identify actions.
 	 */
-	static FString GetActionKey(FGraphActionListBuilderBase::ActionGroup const& Action);
+	static FString GetActionKey(TSharedPtr<FEdGraphSchemaAction> const& Action);
 
 	/**
 	* Goes through all of the blueprint skeleton's object properties and pulls
@@ -234,13 +234,13 @@ GenerateBlueprintAPI commandlet params: \n\
 	/**
 	* Get the category information from a given action.
 	*/
-	static FString GetActionCategory(FGraphActionListBuilderBase::ActionGroup const& Action);
+	static FString GetActionCategory(TSharedPtr<FEdGraphSchemaAction> const& Action);
 
 	/**
 	 * Generic function that dumps information on a single action (like it's 
 	 * name, category, an associated node if it has one, etc.).
 	 */
-	static void DumpActionMenuItem(uint32 Indent, FGraphActionListBuilderBase::ActionGroup const& Action, FGraphActionListBuilderBase& ActionList, FArchive* FileOutWriter);
+	static void DumpActionMenuItem(uint32 Indent, TSharedPtr<FEdGraphSchemaAction> const& Action, FGraphActionListBuilderBase& ActionList, FArchive* FileOutWriter);
 }
 
 //------------------------------------------------------------------------------
@@ -578,9 +578,9 @@ static FString GenerateBlueprintAPIUtils::BuildIndentString(uint32 IndentCount,
 }
 
 //------------------------------------------------------------------------------
-static FString GenerateBlueprintAPIUtils::GetActionKey(FGraphActionListBuilderBase::ActionGroup const& Action)
+static FString GenerateBlueprintAPIUtils::GetActionKey(TSharedPtr<FEdGraphSchemaAction> const& Action)
 {
-	TArray<FString> MenuHierarchy = Action.GetCategoryChain();
+	TArray<FString> MenuHierarchy = Action->GetCategoryChain();
 
 	FString ActionKey;
 	for (FString const& SubCategory : MenuHierarchy)
@@ -592,8 +592,7 @@ static FString GenerateBlueprintAPIUtils::GetActionKey(FGraphActionListBuilderBa
 		ActionKey.RemoveAt(ActionKey.Len() - 1); // remove the trailing '|'
 	}
 
-	TSharedPtr<FEdGraphSchemaAction> MainAction = Action.Actions[0];
-	ActionKey += MainAction->GetMenuDescription().ToString();
+	ActionKey += Action->GetMenuDescription().ToString();
 
 	return ActionKey;
 }
@@ -767,14 +766,14 @@ static void GenerateBlueprintAPIUtils::DumpPalette(uint32 Indent, UBlueprint* Bl
 //------------------------------------------------------------------------------
 static void GenerateBlueprintAPIUtils::DumpActionList(uint32 Indent, FGraphActionListBuilderBase& ActionList, FArchive* FileOutWriter)
 {
-	TArray< FGraphActionListBuilderBase::ActionGroup const* > SortedActions;
+	TArray< TSharedPtr<FEdGraphSchemaAction> const* > SortedActions;
 	const TArray<FString>& CategoryFilter = GenerateBlueprintAPIUtils::CommandOptions.CategoryFilter;
 	bool HasCategoryFilter = CategoryFilter.Num() > 0;
 
 	for (int32 ActionIndex = 0; ActionIndex < ActionList.GetNumActions(); ++ActionIndex)
 	{
-		FGraphActionListBuilderBase::ActionGroup& Action = ActionList.GetAction(ActionIndex);
-		if (Action.Actions.Num() <= 0)
+		TSharedPtr<FEdGraphSchemaAction> Action = ActionList.GetSchemaAction(ActionIndex);
+		if (!Action.IsValid())
 		{
 			continue;
 		}
@@ -807,18 +806,15 @@ static void GenerateBlueprintAPIUtils::DumpActionList(uint32 Indent, FGraphActio
 
 	struct ActionSortFunctor
 	{
-		bool operator()(FGraphActionListBuilderBase::ActionGroup const& LHS, FGraphActionListBuilderBase::ActionGroup const& RHS) const
+		bool operator()(TSharedPtr<FEdGraphSchemaAction> const& LHSAction, TSharedPtr<FEdGraphSchemaAction> const& RHSAction) const
 		{
-			TSharedPtr<FEdGraphSchemaAction> LHSAction = LHS.Actions[0];
-			TSharedPtr<FEdGraphSchemaAction> RHSAction = RHS.Actions[0];
-			
 			if (LHSAction->GetGrouping() != RHSAction->GetGrouping())
 			{
 				return LHSAction->GetGrouping() > RHSAction->GetGrouping();
 			}
 			
-			FString LhKey = GetActionKey(LHS);
-			FString RhKey = GetActionKey(RHS);
+			FString LhKey = GetActionKey(LHSAction);
+			FString RhKey = GetActionKey(RHSAction);
 			return (LhKey.Compare(RhKey) < 0);
 		}
 	};
@@ -826,7 +822,7 @@ static void GenerateBlueprintAPIUtils::DumpActionList(uint32 Indent, FGraphActio
 	SortedActions.Sort(ActionSortFunctor());
 
 	FString LineEnding("\n");
-	for (FGraphActionListBuilderBase::ActionGroup const* Action : SortedActions)
+	for (TSharedPtr<FEdGraphSchemaAction> const* Action : SortedActions)
 	{
 		FileOutWriter->Serialize(TCHAR_TO_ANSI(*LineEnding), LineEnding.Len());
 		DumpActionMenuItem(Indent + 2, *Action, ActionList, FileOutWriter);
@@ -839,12 +835,10 @@ static void GenerateBlueprintAPIUtils::DumpActionList(uint32 Indent, FGraphActio
 	FileOutWriter->Serialize(TCHAR_TO_ANSI(*EndActionListEntry), EndActionListEntry.Len());
 }
 //------------------------------------------------------------------------------
-static FString GenerateBlueprintAPIUtils::GetActionCategory(FGraphActionListBuilderBase::ActionGroup const& Action)
+static FString GenerateBlueprintAPIUtils::GetActionCategory(TSharedPtr<FEdGraphSchemaAction> const& Action)
 {
-	check(Action.Actions.Num() > 0);
-
 	// Get action category info
-	const TArray<FString>& MenuHierarchy = Action.GetCategoryChain();
+	const TArray<FString>& MenuHierarchy = Action->GetCategoryChain();
 
 	FString ActionCategory;
 
@@ -859,7 +853,7 @@ static FString GenerateBlueprintAPIUtils::GetActionCategory(FGraphActionListBuil
 	return ActionCategory;
 }
 //------------------------------------------------------------------------------
-static void GenerateBlueprintAPIUtils::DumpActionMenuItem(uint32 Indent, FGraphActionListBuilderBase::ActionGroup const& Action, FGraphActionListBuilderBase& ActionList, FArchive* FileOutWriter)
+static void GenerateBlueprintAPIUtils::DumpActionMenuItem(uint32 Indent, TSharedPtr<FEdGraphSchemaAction> const& Action, FGraphActionListBuilderBase& ActionList, FArchive* FileOutWriter)
 {
 	const FString& ActionCategory = GenerateBlueprintAPIUtils::GetActionCategory(Action);
 	TArray<FString> Categories;
@@ -878,8 +872,7 @@ static void GenerateBlueprintAPIUtils::DumpActionMenuItem(uint32 Indent, FGraphA
 		}
 	}
 
-	TSharedPtr<FEdGraphSchemaAction> PrimeAction = Action.Actions[0];
-	const FString ActionName = PrimeAction->GetMenuDescription().ToString();
+	const FString ActionName = Action->GetMenuDescription().ToString();
 
 	const FString ActionEntryIndent = BuildIndentString(Indent);
 	FString ActionEntry = ActionEntryIndent + "\"" + MakeJsonString(ActionCategory + ActionName) + "\"";
@@ -890,12 +883,12 @@ static void GenerateBlueprintAPIUtils::DumpActionMenuItem(uint32 Indent, FGraphA
 	ActionEntry += " : {";
 
 	const FString TooltipFieldLabel("\"Tooltip\"      : \"");
-	const FString TooltipStr = PrimeAction->GetTooltipDescription().ToString().Replace(TEXT("\n"), *(IndentedNewline + BuildIndentString(TooltipFieldLabel.Len(), /*bUseSpaces =*/true)));
+	const FString TooltipStr = Action->GetTooltipDescription().ToString().Replace(TEXT("\n"), *(IndentedNewline + BuildIndentString(TooltipFieldLabel.Len(), /*bUseSpaces =*/true)));
 
 	ActionEntry += IndentedNewline + TooltipFieldLabel + MakeJsonString(TooltipStr) + "\"";
 		
 	// Get action node type info
-	UK2Node const* NodeTemplate = FBlueprintActionMenuUtils::ExtractNodeTemplateFromAction(PrimeAction);
+	UK2Node const* NodeTemplate = FBlueprintActionMenuUtils::ExtractNodeTemplateFromAction(Action);
 	if (NodeTemplate != nullptr)
 	{
 		UK2Node* Node = DuplicateObject<UK2Node>(NodeTemplate,ActionList.OwnerOfTemporaries);
diff --git a/Engine/Source/Editor/UnrealEd/Private/MaterialGraphSchema.cpp b/Engine/Source/Editor/UnrealEd/Private/MaterialGraphSchema.cpp
--- a/Engine/Source/Editor/UnrealEd/Private/MaterialGraphSchema.cpp
+++ b/Engine/Source/Editor/UnrealEd/Private/MaterialGraphSchema.cpp
@@ -1321,7 +1321,7 @@ FGraphSchemaSearchWeightModifiers UMaterialGraphSchema::GetSearchWeightModifiers
 	return Modifiers;
 }
 
-float UMaterialGraphSchema::GetActionFilteredWeight(const FGraphActionListBuilderBase::ActionGroup& InCurrentAction, const TArray<FString>& InFilterTerms, const TArray<FString>& InSanitizedFilterTerms, const TArray<UEdGraphPin*>& DraggedFromPins) const
+float UMaterialGraphSchema::GetActionFilteredWeight(const FEdGraphSchemaAction& InCurrentAction, const TArray<FString>& InFilterTerms, const TArray<FString>& InSanitizedFilterTerms, const TArray<UEdGraphPin*>& DraggedFromPins) const
 {
 	// The overall 'weight'
 	float TotalWeight = 0.0f;
@@ -1330,79 +1330,75 @@ float UMaterialGraphSchema::GetActionFilteredWeight(const FGraphActionListBuilde
 	TArray<FGraphSchemaSearchTextWeightInfo> WeightedArrayList;
 	FGraphSchemaSearchTextDebugInfo DebugInfo;
 
-	int32 Action = 0;
-	if (InCurrentAction.Actions[Action].IsValid() == true)
-	{
-		FGraphSchemaSearchWeightModifiers WeightModifiers = GetSearchWeightModifiers();
-		int32 NonLocalizedFirstIndex = CollectSearchTextWeightInfo(InCurrentAction, WeightModifiers, WeightedArrayList, &DebugInfo);
+	FGraphSchemaSearchWeightModifiers WeightModifiers = GetSearchWeightModifiers();
+	int32 NonLocalizedFirstIndex = CollectSearchTextWeightInfo(InCurrentAction, WeightModifiers, WeightedArrayList, &DebugInfo);
 
-		// Now iterate through all the filter terms and calculate a 'weight' using the values and multipliers
-		for (int32 FilterIndex = 0; FilterIndex < InFilterTerms.Num(); ++FilterIndex)
+	// Now iterate through all the filter terms and calculate a 'weight' using the values and multipliers
+	for (int32 FilterIndex = 0; FilterIndex < InFilterTerms.Num(); ++FilterIndex)
+	{
+		const FString& EachTerm = InFilterTerms[FilterIndex];
+		const FString& EachTermSanitized = InSanitizedFilterTerms[FilterIndex];
+		// Now check the weighted lists
+		for (int32 iFindCount = 0; iFindCount < WeightedArrayList.Num(); iFindCount++)
 		{
-			const FString& EachTerm = InFilterTerms[FilterIndex];
-			const FString& EachTermSanitized = InSanitizedFilterTerms[FilterIndex];
-			// Now check the weighted lists
-			for (int32 iFindCount = 0; iFindCount < WeightedArrayList.Num(); iFindCount++)
+			float WeightPerList = 0.0f;
+			const TArray<FString>& WordArray = *WeightedArrayList[iFindCount].Array;
+			float ArrayWeight = WeightedArrayList[iFindCount].WeightModifier;
+			int32 WholeMatchCount = 0;
+			float WholeMatchMultiplier = (iFindCount < NonLocalizedFirstIndex) ? MaterialEditorContextMenuConsoleVariables::WholeMatchLocalizedWeightMultiplier : MaterialEditorContextMenuConsoleVariables::WholeMatchWeightMultiplier;
+
+			// Count of how many words in this array contain a search term that the user has typed in
+			int32 WordMatchCount = 0;
+			// The number of characters in the best matching word
+			int32 BestMatchCharLength = 0;
+
+			for (int32 iEachWord = 0; iEachWord < WordArray.Num(); iEachWord++)
 			{
-				float WeightPerList = 0.0f;
-				const TArray<FString>& WordArray = *WeightedArrayList[iFindCount].Array;
-				float ArrayWeight = WeightedArrayList[iFindCount].WeightModifier;
-				int32 WholeMatchCount = 0;
-				float WholeMatchMultiplier = (iFindCount < NonLocalizedFirstIndex) ? MaterialEditorContextMenuConsoleVariables::WholeMatchLocalizedWeightMultiplier : MaterialEditorContextMenuConsoleVariables::WholeMatchWeightMultiplier;
-
-				// Count of how many words in this array contain a search term that the user has typed in
-				int32 WordMatchCount = 0;
-				// The number of characters in the best matching word
-				int32 BestMatchCharLength = 0;
-
-				for (int32 iEachWord = 0; iEachWord < WordArray.Num(); iEachWord++)
-				{
-					float WeightPerWord = 0.0f;
-
-					// If a word contains the search phrase that the user has typed in, then give it weight					
-					if (WordArray[iEachWord].Contains(EachTermSanitized, ESearchCase::CaseSensitive) || WordArray[iEachWord].Contains(EachTerm, ESearchCase::CaseSensitive))
-					{
-						++WordMatchCount;
-						WeightPerWord += ArrayWeight * WholeMatchMultiplier;
+				float WeightPerWord = 0.0f;
 
-						// If the word starts with the search term, give it extra boost of weight
-						if (WordArray[iEachWord].StartsWith(EachTermSanitized, ESearchCase::CaseSensitive) || WordArray[iEachWord].StartsWith(EachTerm, ESearchCase::CaseSensitive))
-						{
-							WeightPerWord += ArrayWeight * MaterialEditorContextMenuConsoleVariables::StartsWithBonusWeightMultiplier;
-						}
-					}
+				// If a word contains the search phrase that the user has typed in, then give it weight					
+				if (WordArray[iEachWord].Contains(EachTermSanitized, ESearchCase::CaseSensitive) || WordArray[iEachWord].Contains(EachTerm, ESearchCase::CaseSensitive))
+				{
+					++WordMatchCount;
+					WeightPerWord += ArrayWeight * WholeMatchMultiplier;
 
-					if (WeightPerWord > WeightPerList)
+					// If the word starts with the search term, give it extra boost of weight
+					if (WordArray[iEachWord].StartsWith(EachTermSanitized, ESearchCase::CaseSensitive) || WordArray[iEachWord].StartsWith(EachTerm, ESearchCase::CaseSensitive))
 					{
-						// Use the best word match weight, we don't want to count similar words more than one
-						WeightPerList = WeightPerWord;
-						BestMatchCharLength = WordArray[iEachWord].Len();
+						WeightPerWord += ArrayWeight * MaterialEditorContextMenuConsoleVariables::StartsWithBonusWeightMultiplier;
 					}
 				}
 
-				if (BestMatchCharLength > 0 && WeightPerList > 0)
+				if (WeightPerWord > WeightPerList)
 				{
-					// Higher number of matching words contributes to higher weight
-					float PercentMatch = (float)WordMatchCount / (float)WordArray.Num();
-					float PercentMatchWeight = (WeightPerList * PercentMatch * MaterialEditorContextMenuConsoleVariables::PercentageMatchWeightMultiplier);
-					WeightPerList += PercentMatchWeight;
-					DebugInfo.PercentMatchWeight += PercentMatchWeight;
-					DebugInfo.PercentMatch += PercentMatch;
-
-					// The shorter the best matched word, the larger bonus it gets
-					float ShorterMatchFactor = (float)EachTerm.Len() / (float)BestMatchCharLength;
-					float ShorterMatchWeight = ShorterMatchFactor * MaterialEditorContextMenuConsoleVariables::ShorterMatchWeight;
-					WeightPerList += ShorterMatchWeight;
-					DebugInfo.ShorterMatchWeight += ShorterMatchWeight;
+					// Use the best word match weight, we don't want to count similar words more than one
+					WeightPerList = WeightPerWord;
+					BestMatchCharLength = WordArray[iEachWord].Len();
 				}
+			}
 
-				if (WeightedArrayList[iFindCount].DebugWeight)
-				{
-					*WeightedArrayList[iFindCount].DebugWeight += WeightPerList;
-				}
+			if (BestMatchCharLength > 0 && WeightPerList > 0)
+			{
+				// Higher number of matching words contributes to higher weight
+				float PercentMatch = (float)WordMatchCount / (float)WordArray.Num();
+				float PercentMatchWeight = (WeightPerList * PercentMatch * MaterialEditorContextMenuConsoleVariables::PercentageMatchWeightMultiplier);
+				WeightPerList += PercentMatchWeight;
+				DebugInfo.PercentMatchWeight += PercentMatchWeight;
+				DebugInfo.PercentMatch += PercentMatch;
+
+				// The shorter the best matched word, the larger bonus it gets
+				float ShorterMatchFactor = (float)EachTerm.Len() / (float)BestMatchCharLength;
+				float ShorterMatchWeight = ShorterMatchFactor * MaterialEditorContextMenuConsoleVariables::ShorterMatchWeight;
+				WeightPerList += ShorterMatchWeight;
+				DebugInfo.ShorterMatchWeight += ShorterMatchWeight;
+			}
 
-				TotalWeight += WeightPerList;
+			if (WeightedArrayList[iFindCount].DebugWeight)
+			{
+				*WeightedArrayList[iFindCount].DebugWeight += WeightPerList;
 			}
+
+			TotalWeight += WeightPerList;
 		}
 
 		DebugInfo.TotalWeight = TotalWeight;
@@ -1411,6 +1407,19 @@ float UMaterialGraphSchema::GetActionFilteredWeight(const FGraphActionListBuilde
 
 	return TotalWeight;
 }
+
+PRAGMA_DISABLE_DEPRECATION_WARNINGS
+float UMaterialGraphSchema::GetActionFilteredWeight(const FGraphActionListBuilderBase::ActionGroup& InCurrentAction, const TArray<FString>& InFilterTerms, const TArray<FString>& InSanitizedFilterTerms, const TArray<UEdGraphPin*>& DraggedFromPins) const
+{
+	int32 Action = 0;
+	if (InCurrentAction.Actions[Action].IsValid() == true)
+	{
+		return GetActionFilteredWeight(*InCurrentAction.Actions[Action], InFilterTerms, InSanitizedFilterTerms, DraggedFromPins);
+	}
+	return 0.f;
+}
+PRAGMA_ENABLE_DEPRECATION_WARNINGS
+
 #endif // WITH_EDITORONLY_DATA
 
 #undef LOCTEXT_NAMESPACE
\ No newline at end of file
diff --git a/Engine/Source/Runtime/Engine/Classes/EdGraph/EdGraphSchema.h b/Engine/Source/Runtime/Engine/Classes/EdGraph/EdGraphSchema.h
--- a/Engine/Source/Runtime/Engine/Classes/EdGraph/EdGraphSchema.h
+++ b/Engine/Source/Runtime/Engine/Classes/EdGraph/EdGraphSchema.h
@@ -118,7 +118,12 @@ private:
 	/** This is just an arbitrary dump of extra text that search will match on, in addition to the description and tooltip, e.g., Add might have the keyword Math. */
 	UPROPERTY()
 	FText Keywords;
-
+	
+#if WITH_EDITORONLY_DATA
+	/** Cached chain of pipe separated categories */
+	UPROPERTY()
+	TArray<FString> CategoryChain;
+#endif
 public:
 	/** This is a priority number for overriding alphabetical order in the action list (higher value  == higher in the list). */
 	UPROPERTY()
@@ -197,6 +202,11 @@ public:
 	// (e.g., it will not actually move a user added variable or function to a new category)
 	ENGINE_API void CosmeticUpdateCategory(FText NewCategory);
 
+#if WITH_EDITOR
+	/** Updates the root category of the action, making any existing categories a subcategory e.g. NewRootCategory|OldRootCategory */
+	ENGINE_API void CosmeticUpdateRootCategory(FText NewRootCategory);
+#endif
+
 	ENGINE_API void UpdateSearchData(FText NewMenuDescription, FText NewToolTipDescription, FText NewCategory, FText NewKeywords);
 
 	int32 GetSectionID() const
@@ -317,6 +327,10 @@ public:
 	// Can be used to override the tooltip shown in the palette
 	virtual FText GetPaletteToolTip() const { return FText(); }
 
+#if WITH_EDITOR
+	// returns the full list of categories, after any pipe (|) separators have been parsed
+	ENGINE_API const TArray<FString>& GetCategoryChain() const;
+#endif
 private:
 	ENGINE_API void UpdateSearchText();
 };
@@ -455,14 +469,14 @@ struct FGraphActionListBuilderBase
 {
 public:
 	/** A single entry in the list - can contain multiple actions */
-	class ActionGroup
+	class UE_DEPRECATED(5.5, "ActionGroup has been deprecated - operate only on TSharedPtr<FEdGraphSchemaAction> or const FEdGraphSchemaAction& as appropriate") 
+	ActionGroup
 	{
 	public:
 		/** Constructor accepting a single action */
 		ENGINE_API ActionGroup( TSharedPtr<FEdGraphSchemaAction> InAction, FString RootCategory = FString());
 
 		/** Constructor accepting multiple actions */
-		UE_DEPRECATED(5.5, "Collections of Action lists were ambiguous and unused - make a compound FEdGraphSchemaAction and construct the ActionGroup with it")
 		ENGINE_API ActionGroup( const TArray< TSharedPtr<FEdGraphSchemaAction> >& InActions, FString RootCategory = FString());
 
 		/** Move constructor and move assignment operator */
@@ -485,7 +499,6 @@ public:
 		 * @param FromPins Optional pins that the action was dragged from.
 		 * @param Location The position on the graph to place new nodes.
 		 */
-		UE_DEPRECATED(5.5, "Invoke PerformAction on the ActionGroup::Actions directly")
 		ENGINE_API void PerformAction( class UEdGraph* ParentGraph, TArray<UEdGraphPin*>& FromPins, const FVector2D Location );
 		
 		/**
@@ -538,7 +551,7 @@ public:
 private:
 
 	/** All of the action entries */
-	TArray< ActionGroup > Entries;
+	TArray< TSharedPtr<FEdGraphSchemaAction> > Entries;
 
 public:
 
@@ -546,7 +559,9 @@ public:
 	virtual ~FGraphActionListBuilderBase() { }
 
 	/** Adds an action entry containing a single action */
-	ENGINE_API virtual void AddAction( const TSharedPtr<FEdGraphSchemaAction>& NewAction, FString const& Category = FString() );
+	ENGINE_API virtual void AddAction( const TSharedPtr<FEdGraphSchemaAction>& NewAction);
+	UE_DEPRECATED(5.5, "Overriding the schema action category is no longer supported, set the category via CosmeticUpdateRootCategory or CosmeticUpdateCategory")
+	ENGINE_API virtual void AddAction(const TSharedPtr<FEdGraphSchemaAction>& NewAction, FString const& Category);
 
 	/** Adds an action entry containing multiple actions */
 	UE_DEPRECATED(5.5, "Collections of Action lists were ambiguous and unused - make a compound FEdGraphSchemaAction and Add it")
@@ -559,7 +574,12 @@ public:
 	ENGINE_API int32 GetNumActions() const;
 
 	/** Returns the specified entry */
-	ENGINE_API ActionGroup& GetAction( int32 Index );
+	ENGINE_API TSharedPtr<FEdGraphSchemaAction>& GetSchemaAction(int32 Index);
+
+	PRAGMA_DISABLE_DEPRECATION_WARNINGS
+	UE_DEPRECATED(5.5, "FGraphActionListBuilderBase::ActionGroup has been deprecated, access TSharedPtr<FEdGraphSchemaAction> directly using GetSchemaAction")
+	ENGINE_API ActionGroup GetAction(int32 Index);
+	PRAGMA_ENABLE_DEPRECATION_WARNINGS
 
 	/** Clears the action entries */
 	ENGINE_API virtual void Empty();
@@ -587,7 +607,9 @@ public:
 	ENGINE_API FCategorizedGraphActionListBuilder(FString Category = FString());
 
 	// FGraphActionListBuilderBase Interface
-	ENGINE_API virtual void AddAction(const TSharedPtr<FEdGraphSchemaAction>& NewAction, FString const& Category = FString() ) override;
+	ENGINE_API virtual void AddAction(const TSharedPtr<FEdGraphSchemaAction>& NewAction) override;
+	UE_DEPRECATED(5.5, "Overriding the schema action category is no longer supported, specify category via FCategorizedGraphActionListBuilder constructor, or on the FEdGraphSchemaAction")
+	ENGINE_API virtual void AddAction(const TSharedPtr<FEdGraphSchemaAction>& NewAction, FString const& Category) override;
 	UE_DEPRECATED(5.5, "Collections of Action lists were ambiguous and unused - make a compound FEdGraphSchemaAction and Add it")
 	ENGINE_API virtual void AddActionList(const TArray<TSharedPtr<FEdGraphSchemaAction> >& NewActions, FString const& Category = FString()) override;
 	// End of FGraphActionListBuilderBase Interface
@@ -683,7 +705,11 @@ struct FGraphSchemaSearchTextDebugInfo
 	float ShorterMatchWeight = 0.0f;	// Weight for the shorter matched words
 
 	/** Print out the debug info about this weight info to the console */
+	ENGINE_API virtual void Print(const TArray<FString>& SearchForKeywords, const FEdGraphSchemaAction& Action) const;
+	PRAGMA_DISABLE_DEPRECATION_WARNINGS
+	UE_DEPRECATED(5.5, "FGraphActionListBuilderBase::ActionGroup weight calculation was always performed on the first action, provide that if scoring a FGraphActionListBuilderBase::ActionGroup, or just provide the FEdGraphSchemaAction")
 	ENGINE_API virtual void Print(const TArray<FString>& SearchForKeywords, const FGraphActionListBuilderBase::ActionGroup& Action) const;
+	PRAGMA_ENABLE_DEPRECATION_WARNINGS
 };
 #endif // WITH_EDITORONLY_DATA
 
@@ -922,7 +948,11 @@ class UEdGraphSchema : public UObject
 	 * @param InSanitizedFilterTerms		Sanitized search filters in all caps with no symbols or spaces
 	 * @param DraggedFromPins				Any pins that this action was dragged off of
 	 */
+	ENGINE_API virtual float GetActionFilteredWeight(const FEdGraphSchemaAction& InCurrentAction, const TArray<FString>& InFilterTerms, const TArray<FString>& InSanitizedFilterTerms, const TArray<UEdGraphPin*>& DraggedFromPins) const;
+	PRAGMA_DISABLE_DEPRECATION_WARNINGS
+	UE_DEPRECATED(5.5, "FGraphActionListBuilderBase::ActionGroup weight calculation was always performed on the first action, provide that if scoring a FGraphActionListBuilderBase::ActionGroup, or just provide the FEdGraphSchemaAction")
 	ENGINE_API virtual float GetActionFilteredWeight(const FGraphActionListBuilderBase::ActionGroup& InCurrentAction, const TArray<FString>& InFilterTerms, const TArray<FString>& InSanitizedFilterTerms, const TArray<UEdGraphPin*>& DraggedFromPins) const;
+	PRAGMA_ENABLE_DEPRECATION_WARNINGS
 
 	/** Get the weight modifiers from the console variable settings */
 	ENGINE_API virtual FGraphSchemaSearchWeightModifiers GetSearchWeightModifiers() const;
@@ -1342,9 +1372,18 @@ class UEdGraphSchema : public UObject
 #if WITH_EDITORONLY_DATA
 protected:
 	/** Build an array containing all search types, return the index of the first non-localized entry. */
+	ENGINE_API int32 CollectSearchTextWeightInfo(const FEdGraphSchemaAction& InCurrentAction, const FGraphSchemaSearchWeightModifiers& InWeightModifiers,
+		TArray<FGraphSchemaSearchTextWeightInfo>& OutWeightedArrayList, FGraphSchemaSearchTextDebugInfo* InDebugInfo) const;
+	PRAGMA_DISABLE_DEPRECATION_WARNINGS
+	UE_DEPRECATED(5.5, "FGraphActionListBuilderBase::ActionGroup weight calculation was always performed on the first action, provide that if scoring a FGraphActionListBuilderBase::ActionGroup, or just provide the FEdGraphSchemaAction")
 	ENGINE_API int32 CollectSearchTextWeightInfo(const FGraphActionListBuilderBase::ActionGroup& InCurrentAction, const FGraphSchemaSearchWeightModifiers& InWeightModifiers,
 		TArray<FGraphSchemaSearchTextWeightInfo>& OutWeightedArrayList, FGraphSchemaSearchTextDebugInfo* InDebugInfo) const;
+	PRAGMA_ENABLE_DEPRECATION_WARNINGS
 
+	ENGINE_API void PrintSearchTextDebugInfo(const TArray<FString>& InFilterTerms, const FEdGraphSchemaAction& InCurrentAction, const FGraphSchemaSearchTextDebugInfo* InDebugInfo) const;
+	PRAGMA_DISABLE_DEPRECATION_WARNINGS
+	UE_DEPRECATED(5.5, "Only the first entry in the action group has ever been used, provide it to PrintSearchTextDebugInfo")
 	ENGINE_API void PrintSearchTextDebugInfo(const TArray<FString>& InFilterTerms, const FGraphActionListBuilderBase::ActionGroup& InCurrentAction, const FGraphSchemaSearchTextDebugInfo* InDebugInfo) const;
+	PRAGMA_ENABLE_DEPRECATION_WARNINGS
 #endif // WITH_EDITORONLY_DATA
 };
diff --git a/Engine/Source/Runtime/Engine/Private/EdGraph/EdGraphSchema.cpp b/Engine/Source/Runtime/Engine/Private/EdGraph/EdGraphSchema.cpp
--- a/Engine/Source/Runtime/Engine/Private/EdGraph/EdGraphSchema.cpp
+++ b/Engine/Source/Runtime/Engine/Private/EdGraph/EdGraphSchema.cpp
@@ -27,6 +27,28 @@ void FEdGraphSchemaAction::CosmeticUpdateCategory(FText NewCategory)
 	UpdateSearchText();
 }
 
+#if WITH_EDITOR
+void FEdGraphSchemaAction::CosmeticUpdateRootCategory(FText NewRootCategory)
+{
+	if (!NewRootCategory.IsEmpty())
+	{
+		if (GetCategory().IsEmpty())
+		{
+			CosmeticUpdateCategory(NewRootCategory);
+		}
+		else
+		{
+			CosmeticUpdateCategory(FText::FromString(FString::Format(TEXT("{0}|{1}"), { NewRootCategory.ToString(), GetCategory().ToString()})));
+		}
+	}
+}
+
+const TArray<FString>& FEdGraphSchemaAction::GetCategoryChain() const
+{
+	return CategoryChain;
+}
+#endif // WITH_EDITOR
+
 void FEdGraphSchemaAction::UpdateSearchText()
 {
 	SearchText.Reset();
@@ -76,6 +98,19 @@ void FEdGraphSchemaAction::UpdateSearchText()
 		Entry.ToLowerInline();
 		SearchText += Entry;
 	}
+
+	#if WITH_EDITOR
+	ensure(CategoryChain.IsEmpty());
+	const FStringView CategoryDelim = FStringView(TEXT("|"));
+	const UE::String::EParseTokensOptions ParseOptions =
+		UE::String::EParseTokensOptions::IgnoreCase |
+		UE::String::EParseTokensOptions::SkipEmpty |
+		UE::String::EParseTokensOptions::Trim;
+	FString SubCategory = FEditorCategoryUtils::GetCategoryDisplayString(GetCategory().ToString());
+	UE::String::ParseTokens(FStringView(SubCategory), CategoryDelim,
+		[&CategoryChain = CategoryChain](FStringView Token) { CategoryChain.Emplace(Token); },
+		ParseOptions);
+	#endif
 }
 
 void FEdGraphSchemaAction::UpdateSearchData(FText NewMenuDescription, FText NewToolTipDescription, FText NewCategory, FText NewKeywords)
@@ -104,16 +139,25 @@ void FEdGraphSchemaAction::UpdateSearchData(FText NewMenuDescription, FText NewT
 /////////////////////////////////////////////////////
 // FGraphActionListBuilderBase
 
-void FGraphActionListBuilderBase::AddAction( const TSharedPtr<FEdGraphSchemaAction>& NewAction, FString const& Category)
+void FGraphActionListBuilderBase::AddAction( const TSharedPtr<FEdGraphSchemaAction>& NewAction)
+{
+	Entries.Add( NewAction );
+}
+
+void FGraphActionListBuilderBase::AddAction(const TSharedPtr<FEdGraphSchemaAction>& NewAction, FString const& Category)
 {
-	Entries.Add( ActionGroup( NewAction, Category ) );
+#if WITH_EDITOR
+	NewAction->CosmeticUpdateRootCategory(FText::FromString(Category));
+#endif
+	Entries.Add(NewAction);
 }
 
 void FGraphActionListBuilderBase::AddActionList( const TArray<TSharedPtr<FEdGraphSchemaAction> >& NewActions, FString const& Category)
 {
-	PRAGMA_DISABLE_DEPRECATION_WARNINGS
-	Entries.Add( ActionGroup( NewActions, Category ) );
-	PRAGMA_ENABLE_DEPRECATION_WARNINGS
+#if WITH_EDITOR
+	NewActions[0]->CosmeticUpdateRootCategory(FText::FromString(Category));
+#endif
+	Entries.Add( NewActions[0] );
 }
 
 void FGraphActionListBuilderBase::Append( FGraphActionListBuilderBase& Other )
@@ -126,11 +170,18 @@ int32 FGraphActionListBuilderBase::GetNumActions() const
 	return Entries.Num();
 }
 
-FGraphActionListBuilderBase::ActionGroup& FGraphActionListBuilderBase::GetAction( const int32 Index )
+TSharedPtr<FEdGraphSchemaAction>& FGraphActionListBuilderBase::GetSchemaAction( const int32 Index)
 {
 	return Entries[Index];
 }
 
+PRAGMA_DISABLE_DEPRECATION_WARNINGS
+FGraphActionListBuilderBase::ActionGroup FGraphActionListBuilderBase::GetAction(const int32 Index)
+{
+	return FGraphActionListBuilderBase::ActionGroup(Entries[Index]);
+}
+PRAGMA_ENABLE_DEPRECATION_WARNINGS
+
 void FGraphActionListBuilderBase::Empty()
 {
 	Entries.Empty();
@@ -138,7 +189,7 @@ void FGraphActionListBuilderBase::Empty()
 
 /////////////////////////////////////////////////////
 // FGraphActionListBuilderBase::GraphAction
-
+PRAGMA_DISABLE_DEPRECATION_WARNINGS
 FGraphActionListBuilderBase::ActionGroup::ActionGroup(TSharedPtr<FEdGraphSchemaAction> InAction, FString CategoryPrefix)
 	: RootCategory(MoveTemp(CategoryPrefix))
 {
@@ -243,6 +294,8 @@ void FGraphActionListBuilderBase::ActionGroup::InitCategoryChain()
 #endif
 }
 
+PRAGMA_ENABLE_DEPRECATION_WARNINGS
+
 /////////////////////////////////////////////////////
 // FCategorizedGraphActionListBuilder
 
@@ -263,9 +316,21 @@ FCategorizedGraphActionListBuilder::FCategorizedGraphActionListBuilder(FString C
 {
 }
 
-void FCategorizedGraphActionListBuilder::AddAction(TSharedPtr<FEdGraphSchemaAction> const& NewAction, FString const& CategoryIn)
+void FCategorizedGraphActionListBuilder::AddAction(const TSharedPtr<FEdGraphSchemaAction>& NewAction)
 {
-	FGraphActionListBuilderBase::AddAction(NewAction, ConcatCategories(Category, CategoryIn));
+	// apply Category to NewAction:
+#if WITH_EDITOR
+	NewAction->CosmeticUpdateRootCategory(FText::FromString(Category));
+#endif
+	FGraphActionListBuilderBase::AddAction(NewAction);
+}
+
+void FCategorizedGraphActionListBuilder::AddAction(const TSharedPtr<FEdGraphSchemaAction>& NewAction, FString const& CategoryIn)
+{
+#if WITH_EDITOR
+	NewAction->CosmeticUpdateRootCategory(FText::FromString(ConcatCategories(Category, CategoryIn)));
+#endif
+	FGraphActionListBuilderBase::AddAction(NewAction);
 }
 
 void FCategorizedGraphActionListBuilder::AddActionList(TArray<TSharedPtr<FEdGraphSchemaAction> > const& NewActions, FString const& CategoryIn)
@@ -769,140 +834,157 @@ FText UEdGraphSchema::GetPinDisplayName(const UEdGraphPin* Pin) const
 	return ResultPinName;
 }
 
-float UEdGraphSchema::GetActionFilteredWeight(const FGraphActionListBuilderBase::ActionGroup& InCurrentAction, const TArray<FString>& InFilterTerms, const TArray<FString>& InSanitizedFilterTerms, const TArray<UEdGraphPin*>& DraggedFromPins) const
+float UEdGraphSchema::GetActionFilteredWeight(const FEdGraphSchemaAction& InCurrentAction, const TArray<FString>& InFilterTerms, const TArray<FString>& InSanitizedFilterTerms, const TArray<UEdGraphPin*>& DraggedFromPins) const
 {
 	// The overall 'weight'
 	int32 TotalWeight = 0;
 
-	int32 Action = 0;
-	if (InCurrentAction.Actions[Action].IsValid() == true)
+	// Combine the actions string, separate with \n so terms don't run into each other, and remove the spaces (incase the user is searching for a variable)
+	// In the case of groups containing multiple actions, they will have been created and added at the same place in the code, using the same description
+	// and keywords, so we only need to use the first one for filtering.
+	const FString& SearchText = InCurrentAction.GetFullSearchText();
+
+	// Setup an array of arrays so we can do a weighted search			
+	TArray<FGraphSchemaSearchTextWeightInfo> WeightedArrayList;
+	FGraphSchemaSearchWeightModifiers WeightModifiers = GetSearchWeightModifiers();
+	FGraphSchemaSearchTextDebugInfo DebugInfo;
+	int32 NonLocalizedFirstIndex = CollectSearchTextWeightInfo(InCurrentAction, WeightModifiers, WeightedArrayList, &DebugInfo);
+
+	// Now iterate through all the filter terms and calculate a 'weight' using the values and multipliers
+	const FString* EachTerm = nullptr;
+	const FString* EachTermSanitized = nullptr;
+	for (int32 FilterIndex = 0; FilterIndex < InFilterTerms.Num(); ++FilterIndex)
 	{
-		// Combine the actions string, separate with \n so terms don't run into each other, and remove the spaces (incase the user is searching for a variable)
-		// In the case of groups containing multiple actions, they will have been created and added at the same place in the code, using the same description
-		// and keywords, so we only need to use the first one for filtering.
-		const FString& SearchText = InCurrentAction.GetSearchTextForFirstAction();
-
-		// Setup an array of arrays so we can do a weighted search			
-		TArray<FGraphSchemaSearchTextWeightInfo> WeightedArrayList;
-		FGraphSchemaSearchWeightModifiers WeightModifiers = GetSearchWeightModifiers();
-		FGraphSchemaSearchTextDebugInfo DebugInfo;
-		int32 NonLocalizedFirstIndex = CollectSearchTextWeightInfo(InCurrentAction, WeightModifiers, WeightedArrayList, &DebugInfo);
-
-		// Now iterate through all the filter terms and calculate a 'weight' using the values and multipliers
-		const FString* EachTerm = nullptr;
-		const FString* EachTermSanitized = nullptr;
-		for (int32 FilterIndex = 0; FilterIndex < InFilterTerms.Num(); ++FilterIndex)
+		EachTerm = &InFilterTerms[FilterIndex];
+		EachTermSanitized = &InSanitizedFilterTerms[FilterIndex];
+		if (SearchText.Contains(*EachTerm, ESearchCase::CaseSensitive))
 		{
-			EachTerm = &InFilterTerms[FilterIndex];
-			EachTermSanitized = &InSanitizedFilterTerms[FilterIndex];
-			if (SearchText.Contains(*EachTerm, ESearchCase::CaseSensitive))
-			{
-				TotalWeight += 2;
-			}
-			else if (SearchText.Contains(*EachTermSanitized, ESearchCase::CaseSensitive))
-			{
-				TotalWeight++;
-			}
-			// Now check the weighted lists	(We could further improve the hit weight by checking consecutive word matches)
-			for (int32 iFindCount = 0; iFindCount < WeightedArrayList.Num(); iFindCount++)
-			{
-				int32 WeightPerList = 0;
-				const TArray<FString>& KeywordArray = *WeightedArrayList[iFindCount].Array;
-				int32 EachWeight = WeightedArrayList[iFindCount].WeightModifier;
-				int32 WholeMatchCount = 0;
-				int32 WholeMatchMultiplier = (iFindCount < NonLocalizedFirstIndex) ? WeightModifiers.WholeMatchLocalizedWeightMultiplier : WeightModifiers.WholeMatchWeightMultiplier;
+			TotalWeight += 2;
+		}
+		else if (SearchText.Contains(*EachTermSanitized, ESearchCase::CaseSensitive))
+		{
+			TotalWeight++;
+		}
+		// Now check the weighted lists	(We could further improve the hit weight by checking consecutive word matches)
+		for (int32 iFindCount = 0; iFindCount < WeightedArrayList.Num(); iFindCount++)
+		{
+			int32 WeightPerList = 0;
+			const TArray<FString>& KeywordArray = *WeightedArrayList[iFindCount].Array;
+			int32 EachWeight = WeightedArrayList[iFindCount].WeightModifier;
+			int32 WholeMatchCount = 0;
+			int32 WholeMatchMultiplier = (iFindCount < NonLocalizedFirstIndex) ? WeightModifiers.WholeMatchLocalizedWeightMultiplier : WeightModifiers.WholeMatchWeightMultiplier;
 
-				for (int32 iEachWord = 0; iEachWord < KeywordArray.Num(); iEachWord++)
+			for (int32 iEachWord = 0; iEachWord < KeywordArray.Num(); iEachWord++)
+			{
+				// If we get an exact match weight the find count to get exact matches higher priority
+				if (KeywordArray[iEachWord].StartsWith(*EachTerm, ESearchCase::CaseSensitive))
 				{
-					// If we get an exact match weight the find count to get exact matches higher priority
-					if (KeywordArray[iEachWord].StartsWith(*EachTerm, ESearchCase::CaseSensitive))
+					if (iEachWord == 0)
 					{
-						if (iEachWord == 0)
-						{
-							WeightPerList += EachWeight * WholeMatchMultiplier;
-						}
-						else
-						{
-							WeightPerList += EachWeight;
-						}
-						WholeMatchCount++;
+						WeightPerList += EachWeight * WholeMatchMultiplier;
 					}
-					else if (KeywordArray[iEachWord].Contains(*EachTerm, ESearchCase::CaseSensitive))
+					else
 					{
 						WeightPerList += EachWeight;
 					}
-					if (KeywordArray[iEachWord].StartsWith(*EachTermSanitized, ESearchCase::CaseSensitive))
+					WholeMatchCount++;
+				}
+				else if (KeywordArray[iEachWord].Contains(*EachTerm, ESearchCase::CaseSensitive))
+				{
+					WeightPerList += EachWeight;
+				}
+				if (KeywordArray[iEachWord].StartsWith(*EachTermSanitized, ESearchCase::CaseSensitive))
+				{
+					if (iEachWord == 0)
 					{
-						if (iEachWord == 0)
-						{
-							WeightPerList += EachWeight * WholeMatchMultiplier;
-						}
-						else
-						{
-							WeightPerList += EachWeight;
-						}
-						WholeMatchCount++;
+						WeightPerList += EachWeight * WholeMatchMultiplier;
 					}
-					else if (KeywordArray[iEachWord].Contains(*EachTermSanitized, ESearchCase::CaseSensitive))
+					else
 					{
-						WeightPerList += EachWeight / 2;
+						WeightPerList += EachWeight;
 					}
+					WholeMatchCount++;
 				}
-
-				// Increase the weight if theres a larger % of matches in the keyword list
-				if (WholeMatchCount != 0)
+				else if (KeywordArray[iEachWord].Contains(*EachTermSanitized, ESearchCase::CaseSensitive))
 				{
-					int32 PercentAdjust = (100 / KeywordArray.Num()) * WholeMatchCount;
-					int32 PercentAdjustedWeight = WeightPerList * PercentAdjust;
-					DebugInfo.PercentMatch += (float)KeywordArray.Num() / (float)WholeMatchCount;
-					DebugInfo.PercentMatchWeight += (PercentAdjustedWeight - WeightPerList);
-					WeightPerList = PercentAdjustedWeight;
+					WeightPerList += EachWeight / 2;
 				}
+			}
 
-				if (WeightedArrayList[iFindCount].DebugWeight)
-				{
-					*WeightedArrayList[iFindCount].DebugWeight += WeightPerList;
-				}
-				TotalWeight += WeightPerList;
+			// Increase the weight if theres a larger % of matches in the keyword list
+			if (WholeMatchCount != 0)
+			{
+				int32 PercentAdjust = (100 / KeywordArray.Num()) * WholeMatchCount;
+				int32 PercentAdjustedWeight = WeightPerList * PercentAdjust;
+				DebugInfo.PercentMatch += (float)KeywordArray.Num() / (float)WholeMatchCount;
+				DebugInfo.PercentMatchWeight += (PercentAdjustedWeight - WeightPerList);
+				WeightPerList = PercentAdjustedWeight;
 			}
-		}
 
-		PrintSearchTextDebugInfo(InFilterTerms, InCurrentAction, &DebugInfo);
+			if (WeightedArrayList[iFindCount].DebugWeight)
+			{
+				*WeightedArrayList[iFindCount].DebugWeight += WeightPerList;
+			}
+			TotalWeight += WeightPerList;
+		}
 	}
+
+	PrintSearchTextDebugInfo(InFilterTerms, InCurrentAction, &DebugInfo);
 	return TotalWeight;
 
 }
 
-int32 UEdGraphSchema::CollectSearchTextWeightInfo(const FGraphActionListBuilderBase::ActionGroup& InCurrentAction, const FGraphSchemaSearchWeightModifiers& InWeightModifiers, 
-													TArray<FGraphSchemaSearchTextWeightInfo>& OutWeightedArrayList, FGraphSchemaSearchTextDebugInfo* InDebugInfo) const
+PRAGMA_DISABLE_DEPRECATION_WARNINGS
+float UEdGraphSchema::GetActionFilteredWeight(const FGraphActionListBuilderBase::ActionGroup& InCurrentAction, const TArray<FString>& InFilterTerms, const TArray<FString>& InSanitizedFilterTerms, const TArray<UEdGraphPin*>& DraggedFromPins) const
+{
+	int32 TotalWeight = 0;
+	int32 Action = 0;
+	if (InCurrentAction.Actions[Action].IsValid() == true)
+	{
+		return GetActionFilteredWeight(*InCurrentAction.Actions[Action], InFilterTerms, InSanitizedFilterTerms, DraggedFromPins);
+	}
+	return TotalWeight;
+}
+PRAGMA_ENABLE_DEPRECATION_WARNINGS
+
+int32 UEdGraphSchema::CollectSearchTextWeightInfo(const FEdGraphSchemaAction& InCurrentAction, const FGraphSchemaSearchWeightModifiers& InWeightModifiers,
+	TArray<FGraphSchemaSearchTextWeightInfo>& OutWeightedArrayList, FGraphSchemaSearchTextDebugInfo* InDebugInfo) const
 {
 	// First the localized keywords
-	OutWeightedArrayList.Add(FGraphSchemaSearchTextWeightInfo(&InCurrentAction.GetLocalizedSearchKeywordsArrayForFirstAction(), InWeightModifiers.KeywordWeight, InDebugInfo ? &InDebugInfo->KeywordWeight : nullptr));
+	OutWeightedArrayList.Add(FGraphSchemaSearchTextWeightInfo(&InCurrentAction.GetLocalizedSearchKeywordsArray(), InWeightModifiers.KeywordWeight, InDebugInfo ? &InDebugInfo->KeywordWeight : nullptr));
 
 	// The localized description
-	OutWeightedArrayList.Add(FGraphSchemaSearchTextWeightInfo(&InCurrentAction.GetLocalizedMenuDescriptionArrayForFirstAction(), InWeightModifiers.DescriptionWeight, InDebugInfo ? &InDebugInfo->DescriptionWeight : nullptr));
+	OutWeightedArrayList.Add(FGraphSchemaSearchTextWeightInfo(&InCurrentAction.GetLocalizedMenuDescriptionArray(), InWeightModifiers.DescriptionWeight, InDebugInfo ? &InDebugInfo->DescriptionWeight : nullptr));
 
 	// The node search localized title weight
-	OutWeightedArrayList.Add(FGraphSchemaSearchTextWeightInfo(&InCurrentAction.GetLocalizedSearchTitleArrayForFirstAction(), InWeightModifiers.NodeTitleWeight, InDebugInfo ? &InDebugInfo->NodeTitleWeight : nullptr));
+	OutWeightedArrayList.Add(FGraphSchemaSearchTextWeightInfo(&InCurrentAction.GetLocalizedSearchTitleArray(), InWeightModifiers.NodeTitleWeight, InDebugInfo ? &InDebugInfo->NodeTitleWeight : nullptr));
 
 	// The localized category
-	OutWeightedArrayList.Add(FGraphSchemaSearchTextWeightInfo(&InCurrentAction.GetLocalizedSearchCategoryArrayForFirstAction(), InWeightModifiers.CategoryWeight, InDebugInfo ? &InDebugInfo->CategoryWeight : nullptr));
+	OutWeightedArrayList.Add(FGraphSchemaSearchTextWeightInfo(&InCurrentAction.GetLocalizedSearchCategoryArray(), InWeightModifiers.CategoryWeight, InDebugInfo ? &InDebugInfo->CategoryWeight : nullptr));
 
 	// First the keywords
-	int32 NonLocalizedFirstIndex = OutWeightedArrayList.Add(FGraphSchemaSearchTextWeightInfo(&InCurrentAction.GetSearchKeywordsArrayForFirstAction(), InWeightModifiers.KeywordWeight, InDebugInfo ? &InDebugInfo->KeywordWeight : nullptr));
+	int32 NonLocalizedFirstIndex = OutWeightedArrayList.Add(FGraphSchemaSearchTextWeightInfo(&InCurrentAction.GetSearchKeywordsArray(), InWeightModifiers.KeywordWeight, InDebugInfo ? &InDebugInfo->KeywordWeight : nullptr));
 
 	// The description
-	OutWeightedArrayList.Add(FGraphSchemaSearchTextWeightInfo(&InCurrentAction.GetMenuDescriptionArrayForFirstAction(), InWeightModifiers.DescriptionWeight, InDebugInfo ? &InDebugInfo->DescriptionWeight : nullptr));
+	OutWeightedArrayList.Add(FGraphSchemaSearchTextWeightInfo(&InCurrentAction.GetMenuDescriptionArray(), InWeightModifiers.DescriptionWeight, InDebugInfo ? &InDebugInfo->DescriptionWeight : nullptr));
 
 	// The node search title weight
-	OutWeightedArrayList.Add(FGraphSchemaSearchTextWeightInfo(&InCurrentAction.GetSearchTitleArrayForFirstAction(), InWeightModifiers.NodeTitleWeight, InDebugInfo ? &InDebugInfo->NodeTitleWeight : nullptr));
+	OutWeightedArrayList.Add(FGraphSchemaSearchTextWeightInfo(&InCurrentAction.GetSearchTitleArray(), InWeightModifiers.NodeTitleWeight, InDebugInfo ? &InDebugInfo->NodeTitleWeight : nullptr));
 
 	// The category
-	OutWeightedArrayList.Add(FGraphSchemaSearchTextWeightInfo(&InCurrentAction.GetSearchCategoryArrayForFirstAction(), InWeightModifiers.CategoryWeight, InDebugInfo ? &InDebugInfo->CategoryWeight : nullptr));
+	OutWeightedArrayList.Add(FGraphSchemaSearchTextWeightInfo(&InCurrentAction.GetSearchCategoryArray(), InWeightModifiers.CategoryWeight, InDebugInfo ? &InDebugInfo->CategoryWeight : nullptr));
 
 	return NonLocalizedFirstIndex;
 }
 
+PRAGMA_DISABLE_DEPRECATION_WARNINGS
+int32 UEdGraphSchema::CollectSearchTextWeightInfo(const FGraphActionListBuilderBase::ActionGroup& InCurrentAction, const FGraphSchemaSearchWeightModifiers& InWeightModifiers, 
+													TArray<FGraphSchemaSearchTextWeightInfo>& OutWeightedArrayList, FGraphSchemaSearchTextDebugInfo* InDebugInfo) const
+{
+	return CollectSearchTextWeightInfo(*InCurrentAction.Actions[0], InWeightModifiers, OutWeightedArrayList, InDebugInfo);
+}
+PRAGMA_ENABLE_DEPRECATION_WARNINGS
+
 //////////////////////////////////////////////////////////////////////////
 /** CVars for tweaking how context menu search picks the best match */
 namespace ContextMenuConsoleVariables
@@ -977,7 +1059,17 @@ FGraphSchemaSearchWeightModifiers UEdGraphSchema::GetSearchWeightModifiers() con
 	return Modifiers;
 }
 
+PRAGMA_DISABLE_DEPRECATION_WARNINGS
 void UEdGraphSchema::PrintSearchTextDebugInfo(const TArray<FString>& InFilterTerms, const FGraphActionListBuilderBase::ActionGroup& InCurrentAction, const FGraphSchemaSearchTextDebugInfo* InDebugInfo) const
+{
+	if (InCurrentAction.Actions.Num() > 0)
+	{
+		PrintSearchTextDebugInfo(InFilterTerms, *InCurrentAction.Actions[0], InDebugInfo);
+	}
+}
+PRAGMA_ENABLE_DEPRECATION_WARNINGS
+
+void UEdGraphSchema::PrintSearchTextDebugInfo(const TArray<FString>& InFilterTerms, const FEdGraphSchemaAction& InCurrentAction, const FGraphSchemaSearchTextDebugInfo* InDebugInfo) const
 {
 	if (ContextMenuConsoleVariables::bPrintDebugInfo && InDebugInfo)
 	{
@@ -985,7 +1077,7 @@ void UEdGraphSchema::PrintSearchTextDebugInfo(const TArray<FString>& InFilterTer
 	}
 }
 
-void FGraphSchemaSearchTextDebugInfo::Print(const TArray<FString>& SearchForKeywords, const FGraphActionListBuilderBase::ActionGroup& Action) const
+void FGraphSchemaSearchTextDebugInfo::Print(const TArray<FString>& SearchForKeywords, const FEdGraphSchemaAction& Action) const
 {
 	auto CombineStrings = [](const TArray<FString>& Strings) -> FString
 	{
@@ -999,20 +1091,27 @@ void FGraphSchemaSearchTextDebugInfo::Print(const TArray<FString>& SearchForKeyw
 	};
 
 	FString SearchText = CombineStrings(SearchForKeywords);
-	FString LocalizedTitle = CombineStrings(Action.GetLocalizedSearchTitleArrayForFirstAction());
-	FString LocalizedKeywords = CombineStrings(Action.GetLocalizedSearchKeywordsArrayForFirstAction());
-	FString LocalizedDescription = CombineStrings(Action.GetLocalizedMenuDescriptionArrayForFirstAction());
-	FString LocalizedCategory = CombineStrings(Action.GetLocalizedSearchCategoryArrayForFirstAction());
-	FString Title = CombineStrings(Action.GetSearchTitleArrayForFirstAction());
-	FString Keywords = CombineStrings(Action.GetSearchKeywordsArrayForFirstAction());
-	FString Description = CombineStrings(Action.GetMenuDescriptionArrayForFirstAction());
-	FString Category = CombineStrings(Action.GetSearchCategoryArrayForFirstAction());
+	FString LocalizedTitle = CombineStrings(Action.GetLocalizedSearchTitleArray());
+	FString LocalizedKeywords = CombineStrings(Action.GetLocalizedSearchKeywordsArray());
+	FString LocalizedDescription = CombineStrings(Action.GetLocalizedMenuDescriptionArray());
+	FString LocalizedCategory = CombineStrings(Action.GetLocalizedSearchCategoryArray());
+	FString Title = CombineStrings(Action.GetSearchTitleArray());
+	FString Keywords = CombineStrings(Action.GetSearchKeywordsArray());
+	FString Description = CombineStrings(Action.GetMenuDescriptionArray());
+	FString Category = CombineStrings(Action.GetSearchCategoryArray());
 
 	UE_LOG(LogTemp, Log, TEXT("Searching for \"%s\" in [LocTitle:\"%s\" LocKeywords:\"%s\" LocDescription:\"%s\" LocCategory:\"%s\" Title:\"%s\" Keywords:\"%s\" Description:\"%s\" Category:\"%s\"] \
 TotalWeight: %.2f | NodeTitleWeight: %.2f | KeywordWeight: %.2f | CategoryWeight: %.2f | PercentMatchWeight: %.2f | ShorterMatchWeight: %.2f"),
 		*SearchText, *LocalizedTitle, *LocalizedKeywords, *LocalizedDescription, *LocalizedCategory, *Title, *Keywords, *Description, *Category, 
 		TotalWeight, NodeTitleWeight, KeywordWeight, CategoryWeight, PercentMatchWeight, ShorterMatchWeight);
 }
+
+PRAGMA_DISABLE_DEPRECATION_WARNINGS
+void FGraphSchemaSearchTextDebugInfo::Print(const TArray<FString>& SearchForKeywords, const FGraphActionListBuilderBase::ActionGroup& Action) const
+{
+	return Print(SearchForKeywords, *Action.Actions[0]);
+}
+PRAGMA_ENABLE_DEPRECATION_WARNINGS
 #endif // WITH_EDITORONLY_DATA
 
 void UEdGraphSchema::ConstructBasicPinTooltip(UEdGraphPin const& Pin, FText const& PinDescription, FString& TooltipOut) const
