VT fix for change due to platform preview

diff --git a/Engine/Source/Runtime/Engine/Private/VT/VirtualTextureDataBuilder.cpp b/Engine/Source/Runtime/Engine/Private/VT/VirtualTextureDataBuilder.cpp
--- a/Engine/Source/Runtime/Engine/Private/VT/VirtualTextureDataBuilder.cpp
+++ b/Engine/Source/Runtime/Engine/Private/VT/VirtualTextureDataBuilder.cpp
@@ -727,6 +727,7 @@ void FVirtualTextureDataBuilder::BuildBlockTiles(uint32 LayerIndex, uint32 Block
 		FTextureBuildSettings TBSettings;
 		TBSettings.MaxTextureResolution = FTextureBuildSettings::MaxTextureResolutionDefault;
 		TBSettings.TextureFormatName = LayerData.TextureFormatName;
+		TBSettings.BaseTextureFormatName = TBSettings.TextureFormatName; // VTs never have platform prefix
 		TBSettings.bSRGB = BuildSettingsForLayer.bSRGB;
 		TBSettings.bUseLegacyGamma = BuildSettingsForLayer.bUseLegacyGamma;
 		TBSettings.MipGenSettings = TMGS_NoMipmaps;
@@ -912,6 +913,7 @@ void FVirtualTextureDataBuilder::BuildLayerBlocks(FSlowTask& BuildTask, uint32 L
 		FTextureBuildSettings TBSettings = SettingsPerLayer[0];
 		//TBSettings.MaxTextureResolution = FTextureBuildSettings::MaxTextureResolutionDefault;
 		TBSettings.TextureFormatName = LayerData.FormatName;
+		TBSettings.BaseTextureFormatName = LayerData.FormatName; // VTs never have platform prefix
 
 		if (LayerIndex != 0)
 		{
@@ -1109,6 +1111,7 @@ void FVirtualTextureDataBuilder::BuildLayerBlocks(FSlowTask& BuildTask, uint32 L
 		FTextureBuildSettings TBSettings = SettingsPerLayer[0];
 		TBSettings.MaxTextureResolution = FTextureBuildSettings::MaxTextureResolutionDefault;; // don't limit the size of the mip-tail, this limit only applies to each source block
 		TBSettings.TextureFormatName = LayerData.FormatName;
+		TBSettings.BaseTextureFormatName = LayerData.FormatName; // VTs never have platform prefix
 		TBSettings.bSRGB = BuildSettingsForLayer.bSRGB;
 		TBSettings.bUseLegacyGamma = BuildSettingsForLayer.bUseLegacyGamma;
 
