- Replaced PLATFORM_SUPPORTS_VECTOR_ROUNDING_MODE with PLATFORM_SUPPORTS_VECTOR_CONTROL_REGISTERS, and replaced VectorSetControlRegister, VECTOR_ROUND_TOWARD_ZERO, and VECTOR_DENORMALS_FLUSH_TO_ZERO with noops when PLATFORM_SUPPORTS_VECTOR_CONTROL_REGISTERS is true

diff --git a/Engine/Source/Runtime/Core/Public/HAL/Platform.h b/Engine/Source/Runtime/Core/Public/HAL/Platform.h
--- a/Engine/Source/Runtime/Core/Public/HAL/Platform.h
+++ b/Engine/Source/Runtime/Core/Public/HAL/Platform.h
@@ -167,8 +167,10 @@
 	#define PLATFORM_ENABLE_VECTORINTRINSICS	0
 #endif
 
-#ifndef PLATFORM_SUPPORTS_VECTOR_ROUNDING_MODE
-	#define PLATFORM_SUPPORTS_VECTOR_ROUNDING_MODE 1
+// Defines if the platform can handle setting vectorized exception
+// and rounding modes
+#ifndef PLATFORM_SUPPORTS_VECTOR_CONTROL_REGISTERS
+	#define PLATFORM_SUPPORTS_VECTOR_CONTROL_REGISTERS 1
 #endif
 
 #ifndef PLATFORM_USE_SSE2_FOR_THREAD_YIELD
diff --git a/Engine/Source/Runtime/Core/Public/Math/UnrealMathSSE.h b/Engine/Source/Runtime/Core/Public/Math/UnrealMathSSE.h
--- a/Engine/Source/Runtime/Core/Public/Math/UnrealMathSSE.h
+++ b/Engine/Source/Runtime/Core/Public/Math/UnrealMathSSE.h
@@ -2687,6 +2687,7 @@ FORCEINLINE int VectorAnyGreaterThan(const VectorRegister4Double& Vec1, const Ve
  */
 #define VectorGetControlRegister()		_mm_getcsr()
 
+#if PLATFORM_SUPPORTS_VECTOR_CONTROL_REGISTERS
 /**
  * Sets the control register.
  *
@@ -2704,6 +2705,12 @@ FORCEINLINE int VectorAnyGreaterThan(const VectorRegister4Double& Vec1, const Ve
   */
 #define VECTOR_DENORMALS_FLUSH_TO_ZERO	_MM_FLUSH_ZERO_ON
 
+#else
+#define VectorSetControlRegister(...)
+#define VECTOR_ROUND_TOWARD_ZERO
+#define VECTOR_DENORMALS_FLUSH_TO_ZERO
+#endif
+
 /**
 * Multiplies two quaternions; the order matters.
 *
diff --git a/Engine/Source/Runtime/Engine/Private/SkeletalRenderCPUSkin.cpp b/Engine/Source/Runtime/Engine/Private/SkeletalRenderCPUSkin.cpp
--- a/Engine/Source/Runtime/Engine/Private/SkeletalRenderCPUSkin.cpp
+++ b/Engine/Source/Runtime/Engine/Private/SkeletalRenderCPUSkin.cpp
@@ -1029,9 +1029,7 @@ static void SkinVertices(
 	const FVector& WorldScale)
 {
 	uint32 StatusRegister = VectorGetControlRegister();
-#if PLATFORM_SUPPORTS_VECTOR_ROUNDING_MODE
 	VectorSetControlRegister( StatusRegister | VECTOR_ROUND_TOWARD_ZERO );
-#endif
 
 	// Create array to track state during morph blending
 	TArray<FMorphTargetInfo> MorphEvalInfos;
@@ -1061,9 +1059,7 @@ static void SkinVertices(
 		SkinVertexSection<VertexType, NumberOfUVs>(DestVertex, MorphEvalInfos, MorphTargetWeights, Section, LOD, WeightBuffer, VertexBufferBaseIndex, NumValidMorphs, CurBaseVertIdx, LODIndex, ReferenceToLocal, ClothSimData, ClothBlendWeight, WorldToLocal, WorldScaleAbs);
 	}
 
-#if PLATFORM_SUPPORTS_VECTOR_ROUNDING_MODE
 	VectorSetControlRegister( StatusRegister );
-#endif
 }
 
 /**
