Skeleton Modifier: skeleton/skeletal mesh transaction improvement

    skip skeletal meshes that are also being transacted to avoid re-registering the components in USkeleton::PostEditUndo() as will be done later in USkeletalMesh::PostEditUndo()
    note that this also avoid re-registering the same component several times when transacting

#rnx

diff --git a/Engine/Source/Runtime/Engine/Classes/Engine/SkeletalMesh.h b/Engine/Source/Runtime/Engine/Classes/Engine/SkeletalMesh.h
--- a/Engine/Source/Runtime/Engine/Classes/Engine/SkeletalMesh.h
+++ b/Engine/Source/Runtime/Engine/Classes/Engine/SkeletalMesh.h
@@ -2456,6 +2456,12 @@ public:
 
 	ENGINE_API virtual void PostEditChangeProperty(FPropertyChangedEvent& PropertyChangedEvent) override;
 
+private:	
+	bool bTransacting = false;
+public:
+	bool IsTransacting() const;
+	
+	ENGINE_API virtual void PreEditUndo() override;
 	ENGINE_API virtual void PostEditUndo() override;
 	ENGINE_API virtual void GetAssetRegistryTagMetadata(TMap<FName, FAssetRegistryTagMetadata>& OutMetadata) const override;
 
diff --git a/Engine/Source/Runtime/Engine/Private/Animation/Skeleton.cpp b/Engine/Source/Runtime/Engine/Private/Animation/Skeleton.cpp
--- a/Engine/Source/Runtime/Engine/Private/Animation/Skeleton.cpp
+++ b/Engine/Source/Runtime/Engine/Private/Animation/Skeleton.cpp
@@ -2032,11 +2032,15 @@ void USkeleton::RenameVirtualBone(const FName OriginalBoneName, const FName NewB
 
 void USkeleton::HandleVirtualBoneChanges()
 {
-	const bool bRebuildNameMap = false;
+	constexpr bool bRebuildNameMap = false;
 	ReferenceSkeleton.RebuildRefSkeleton(this, bRebuildNameMap);
 
 	UE::Anim::FSkeletonRemappingRegistry::Get().RefreshMappings(this);
 
+	// store skeletal meshes that are also transacting to avoid re-registering the component here
+	// as it will be done later in USkeletalMesh::PostEditUndo()
+	TArray<USkeletalMesh*> SkeletalMeshTransacting;
+	
 	for (TObjectIterator<USkeletalMesh> ItMesh; ItMesh; ++ItMesh)
 	{
 		USkeletalMesh* SkelMesh = *ItMesh;
@@ -2045,19 +2049,43 @@ void USkeleton::HandleVirtualBoneChanges()
 			// also have to update retarget base pose
 			SkelMesh->GetRefSkeleton().RebuildRefSkeleton(this, bRebuildNameMap);
 			RebuildLinkup(SkelMesh);
+
+#if WITH_EDITOR
+			if (SkelMesh->IsTransacting())
+			{
+				SkeletalMeshTransacting.Add(SkelMesh);
+			}
+#endif
 		}
 	}
 
 	// refresh curve meta data that contains joint info
 	RefreshSkeletonMetaData();
 
+	auto NeedsReRegistration = [this, &SkeletalMeshTransacting](const USkinnedMeshComponent* InMeshComponent)
+	{
+		if (!InMeshComponent || InMeshComponent->IsTemplate())
+		{
+			return false;
+		}
+
+		USkinnedAsset* SkinnedAsset = InMeshComponent->GetSkinnedAsset();
+		if (!SkinnedAsset || SkinnedAsset->GetSkeleton() != this)
+		{
+			return false;
+		}
+
+#if WITH_EDITOR
+		return !SkeletalMeshTransacting.Contains(SkinnedAsset);
+#else
+		return true;
+#endif
+	};
+	
 	for (TObjectIterator<USkinnedMeshComponent> It; It; ++It)
 	{
 		USkinnedMeshComponent* MeshComponent = *It;
-		if (MeshComponent &&
-			MeshComponent->GetSkinnedAsset() &&
-			MeshComponent->GetSkinnedAsset()->GetSkeleton() == this &&
-			!MeshComponent->IsTemplate())
+		if (NeedsReRegistration(MeshComponent))
 		{
 			FComponentReregisterContext Context(MeshComponent);
 		}
diff --git a/Engine/Source/Runtime/Engine/Private/SkeletalMesh.cpp b/Engine/Source/Runtime/Engine/Private/SkeletalMesh.cpp
--- a/Engine/Source/Runtime/Engine/Private/SkeletalMesh.cpp
+++ b/Engine/Source/Runtime/Engine/Private/SkeletalMesh.cpp
@@ -1392,6 +1392,17 @@ void USkeletalMesh::PostEditChangeProperty(FPropertyChangedEvent& PropertyChange
 	check(PostEditChangeStackCounter == 1);
 }
 
+bool USkeletalMesh::IsTransacting() const
+{
+	return bTransacting;
+}
+
+void USkeletalMesh::PreEditUndo()
+{
+	bTransacting = true;
+	Super::PreEditUndo();
+}
+
 void USkeletalMesh::PostEditUndo()
 {
 	check(IsInGameThread());
@@ -1413,6 +1424,8 @@ void USkeletalMesh::PostEditUndo()
 		// A morph target remove has been undone, reinitialise
 		InitMorphTargets();
 	}
+	
+	bTransacting = false;
 }
 
 void USkeletalMesh::UpdateGenerateUpToData()
