Make TargetFrameRate WITH_EDITORDATA_ONLY and remove from details panel (to reduce confusion for users)

diff --git a/Engine/Source/Runtime/Engine/Classes/Animation/AnimSequence.h b/Engine/Source/Runtime/Engine/Classes/Animation/AnimSequence.h
--- a/Engine/Source/Runtime/Engine/Classes/Animation/AnimSequence.h
+++ b/Engine/Source/Runtime/Engine/Classes/Animation/AnimSequence.h
@@ -878,13 +878,13 @@ public:
 #endif // WITH_EDITOR
 
 protected:
-	UPROPERTY(VisibleAnywhere, AssetRegistrySearchable, Category = "Animation")
-	FFrameRate TargetFrameRate;
-
 	UPROPERTY(VisibleAnywhere, AssetRegistrySearchable, Category = "Animation")
 	FPerPlatformFrameRate PlatformTargetFrameRate;
 
 #if WITH_EDITORONLY_DATA
+	UPROPERTY()
+	FFrameRate TargetFrameRate;	
+
 	UPROPERTY(VisibleAnywhere, AssetRegistrySearchable, Category = "Animation", Transient, DuplicateTransient)
 	int32 NumberOfSampledKeys;
 
diff --git a/Engine/Source/Runtime/Engine/Private/Animation/AnimSequence.cpp b/Engine/Source/Runtime/Engine/Private/Animation/AnimSequence.cpp
--- a/Engine/Source/Runtime/Engine/Private/Animation/AnimSequence.cpp
+++ b/Engine/Source/Runtime/Engine/Private/Animation/AnimSequence.cpp
@@ -666,11 +666,14 @@ void UAnimSequence::Serialize(FArchive& Ar)
 			UE::Anim::FAnimSequenceCompilingManager::Get().FinishCompilation({this});
 		}
 	}
-#endif // WITH_EDITOR	
+#endif // WITH_EDITOR
+
+#if WITH_EDITORONLY_DATA
 	if ( Ar.IsLoading() && Ar.CustomVer(FUE5MainStreamObjectVersion::GUID) < FUE5MainStreamObjectVersion::PerPlatformAnimSequenceTargetFrameRate)
 	{
 		PlatformTargetFrameRate = TargetFrameRate;
-	}
+	}	
+#endif // WITH_EDITORONLY_DATA
 
 	if (Ar.CustomVer(FFrameworkObjectVersion::GUID) < FFrameworkObjectVersion::MoveCompressedAnimDataToTheDDC)
 	{
