Fix UConsole using a different font for text and cursor position estimation

diff --git a/Engine/Source/Runtime/Engine/Private/UserInterface/Console.cpp b/Engine/Source/Runtime/Engine/Private/UserInterface/Console.cpp
--- a/Engine/Source/Runtime/Engine/Private/UserInterface/Console.cpp
+++ b/Engine/Source/Runtime/Engine/Private/UserInterface/Console.cpp
@@ -1352,7 +1352,7 @@ void UConsole::PostRender_Console_Open(UCanvas* Canvas)
 
 	if (Scrollback.Num())
 	{
-		FCanvasTextItem ConsoleText(FVector2D(LeftPos, TopPos + Height - 5 - yl), FText::FromString(TEXT("")), GEngine->GetLargeFont(), ConsoleSettings->InputColor);
+		FCanvasTextItem ConsoleText(FVector2D(LeftPos, TopPos + Height - 5 - yl), FText::FromString(TEXT("")), Font, ConsoleSettings->InputColor);
 		// change the text color to white
 		ConsoleText.SetColor(FLinearColor::White);
 
@@ -1519,7 +1519,7 @@ void UConsole::PostRender_InputLine(UCanvas* Canvas, FIntPoint UserInputLinePos)
 
 	// Currently typed string
 	FText Str = FText::FromString(TypedInputText);
-	FCanvasTextItem ConsoleText(FVector2D(UserInputLinePos.X, UserInputLinePos.Y - 3 - yl), Str, GEngine->GetLargeFont(), ConsoleSettings->InputColor);
+	FCanvasTextItem ConsoleText(FVector2D(UserInputLinePos.X, UserInputLinePos.Y - 3 - yl), Str, Font, ConsoleSettings->InputColor);
 	ConsoleText.EnableShadow(FLinearColor::Black);
 	Canvas->DrawItem(ConsoleText);
 
