[NetPhysics] Resimulation improvements to NetworkPhysicsComponent

- CompareData for inputs was called multiple times for the same input since we get redundant inputs from the server, this is now fixed so we only check each incoming input once.
- CompareData was also called on simulated proxies when receiving redundant inputs (inputs we already received earlier)
- Made State and Input replication push-based and to only call OnRep_ when a change is made. Requires the project to run with Net.IsPushModelEnabled 1 (will function as previously without the CVar set).
- Made it possible to run NetworkPhysicsComponent with only a custom input struct registered, since custom states are not always needed and consume CPU and network bandwidth.
- IsLocallyControlled can be set on the server through UNetworkPhysicsComponent::bIsRelayingLocalInputs for AI to apply inputs and for the server to send those down to simulated proxies.
- Server send input data from the history buffer to clients instead of sending from incoming input data from clients, this ensures that clients receive the inputs the server actually uses. If the input buffer runs empty and the server predicts inputs, those are the inputs that gets sent to clients now.
- Make replicated state and input send the latest entry/entries in the history instead of the last used entries.
- Switched State::ApplyData to get called before Input::ApplyData so that the state is correct when inputs are set and processed.
- Rename internal variables from "frame" to "index" for clarity in RewindHistory

diff --git a/Engine/Source/Runtime/Engine/Private/PhysicsEngine/NetworkPhysicsComponent.cpp b/Engine/Source/Runtime/Engine/Private/PhysicsEngine/NetworkPhysicsComponent.cpp
--- a/Engine/Source/Runtime/Engine/Private/PhysicsEngine/NetworkPhysicsComponent.cpp
+++ b/Engine/Source/Runtime/Engine/Private/PhysicsEngine/NetworkPhysicsComponent.cpp
@@ -11,6 +11,7 @@
 #include "Net/UnrealNetwork.h"
 #include "PhysicsProxy/SingleParticlePhysicsProxy.h"
 #include "Chaos/PhysicsObjectInternalInterface.h"
+#include "Net/Core/PushModel/PushModel.h"
 
 #if UE_WITH_IRIS
 #include "Iris/ReplicationState/PropertyNetSerializerInfoRegistry.h"
@@ -22,11 +23,11 @@ namespace PhysicsReplicationCVars
 {
 	namespace ResimulationCVars
 	{
-		static bool bAllowRewindToClosestState = true;
+		bool bAllowRewindToClosestState = true;
 		static FAutoConsoleVariableRef CVarResimAllowRewindToClosestState(TEXT("np2.Resim.AllowRewindToClosestState"), bAllowRewindToClosestState, TEXT("When rewinding to a specific frame, if the client doens't have state data for that frame, use closest data available. Only affects the first rewind frame, when FPBDRigidsEvolution is set to Reset."));
-		static bool bCompareStateToTriggerRewind = false;
+		bool bCompareStateToTriggerRewind = false;
 		static FAutoConsoleVariableRef CVarResimCompareStateToTriggerRewind(TEXT("np2.Resim.CompareStateToTriggerRewind"), bCompareStateToTriggerRewind, TEXT("When true, cache local players custom state struct in rewind history and compare the predicted state with incoming server state to trigger resimulations if they differ, comparison done through FNetworkPhysicsData::CompareData"));
-		static bool bCompareInputToTriggerRewind = false;
+		bool bCompareInputToTriggerRewind = false;
 		static FAutoConsoleVariableRef CVarResimCompareInputToTriggerRewind(TEXT("np2.Resim.CompareInputToTriggerRewind"), bCompareInputToTriggerRewind, TEXT("When true, compare local players predicted inputs with incoming server inputs to trigger resimulations if they differ, comparison done through FNetworkPhysicsData::CompareData."));
 	}
 }
@@ -391,8 +392,14 @@ UNetworkPhysicsComponent::UNetworkPhysicsComponent() : Super()
 
 void UNetworkPhysicsComponent::InitPhysics()
 {
-	bCompareStateToTriggerRewind = PhysicsReplicationCVars::ResimulationCVars::bCompareStateToTriggerRewind;
-	bCompareInputToTriggerRewind = PhysicsReplicationCVars::ResimulationCVars::bCompareInputToTriggerRewind;
+	if (const IConsoleVariable* CVarCompareInputToTriggerRewind = IConsoleManager::Get().FindConsoleVariable(TEXT("np2.Resim.CompareInputToTriggerRewind")))
+	{
+		bCompareInputToTriggerRewind = CVarCompareInputToTriggerRewind->GetBool();
+	}
+	if (const IConsoleVariable* CVarCompareStateToTriggerRewind = IConsoleManager::Get().FindConsoleVariable(TEXT("np2.Resim.CompareStateToTriggerRewind")))
+	{
+		bCompareStateToTriggerRewind = CVarCompareStateToTriggerRewind->GetBool();
+	}
 
 	AActor* Owner = GetOwner();
 	if (Owner)
@@ -401,8 +408,8 @@ void UNetworkPhysicsComponent::InitPhysics()
 		{
 			InputRedundancy = PhysicsSettings->ResimulationSettings.bOverrideRedundantInputs ? PhysicsSettings->ResimulationSettings.RedundantInputs : InputRedundancy;
 			StateRedundancy = PhysicsSettings->ResimulationSettings.bOverrideRedundantStates ? PhysicsSettings->ResimulationSettings.RedundantStates : StateRedundancy;
-			bCompareStateToTriggerRewind = PhysicsSettings->ResimulationSettings.GetCompareStateToTriggerRewind(PhysicsReplicationCVars::ResimulationCVars::bCompareStateToTriggerRewind);
-			bCompareInputToTriggerRewind = PhysicsSettings->ResimulationSettings.GetCompareInputToTriggerRewind(PhysicsReplicationCVars::ResimulationCVars::bCompareInputToTriggerRewind);
+			bCompareInputToTriggerRewind = PhysicsSettings->ResimulationSettings.GetCompareInputToTriggerRewind(bCompareInputToTriggerRewind);
+			bCompareStateToTriggerRewind = PhysicsSettings->ResimulationSettings.GetCompareStateToTriggerRewind(bCompareStateToTriggerRewind);
 		}
 
 		FRepMovement& RepMovement = Owner->GetReplicatedMovement_Mutable();
@@ -420,9 +427,6 @@ void UNetworkPhysicsComponent::InitPhysics()
 	bWantsInitializeComponent = true;
 	SetIsReplicatedByDefault(true);
 	SetAsyncPhysicsTickEnabled(true);
-
-	StateOffsets.SetNumZeroed(StateRedundancy + 1);
-	InputOffsets.SetNumZeroed(InputRedundancy + 1);
 }
 
 void UNetworkPhysicsComponent::BeginPlay()
@@ -483,8 +487,13 @@ void UNetworkPhysicsComponent::GetLifetimeReplicatedProps(TArray< FLifetimePrope
 {
 	Super::GetLifetimeReplicatedProps(OutLifetimeProps);
 
-	DOREPLIFETIME_CONDITION_NOTIFY(UNetworkPhysicsComponent, ReplicatedInputs, COND_None, REPNOTIFY_Always);
-	DOREPLIFETIME_CONDITION_NOTIFY(UNetworkPhysicsComponent, ReplicatedStates, COND_None, REPNOTIFY_Always);
+	FDoRepLifetimeParams ReplicatedParams;
+	ReplicatedParams.Condition = COND_None;
+	ReplicatedParams.RepNotifyCondition = REPNOTIFY_OnChanged;
+	ReplicatedParams.bIsPushBased = true;
+
+	DOREPLIFETIME_WITH_PARAMS_FAST(UNetworkPhysicsComponent, ReplicatedInputs, ReplicatedParams);
+	DOREPLIFETIME_WITH_PARAMS_FAST(UNetworkPhysicsComponent, ReplicatedStates, ReplicatedParams);
 }
 
 void UNetworkPhysicsComponent::AsyncPhysicsTickComponent(float DeltaTime, float SimTime)
@@ -521,10 +530,6 @@ void UNetworkPhysicsComponent::AsyncPhysicsTickComponent(float DeltaTime, float
 
 				// Replicate states across the network
 				SendStateData();
-
-				// Advance the NetworkIndex
-				InputIndex = (InputIndex + 1) % (InputRedundancy + 1);
-				StateIndex = (StateIndex + 1) % (StateRedundancy + 1);
 			}
 		}
 	}
@@ -532,7 +537,8 @@ void UNetworkPhysicsComponent::AsyncPhysicsTickComponent(float DeltaTime, float
 
 void UNetworkPhysicsComponent::SendInputData()
 {
-	if (IsLocallyControlled() && InputHistory)
+	// Inputs are sent from the server or locally controlled actors/pawns
+	if (InputHistory && (IsLocallyControlled() || HasServerWorld()))
 	{
 		const APlayerController* PlayerController = GetPlayerController();
 		if (!PlayerController)
@@ -540,19 +546,22 @@ void UNetworkPhysicsComponent::SendInputData()
 			PlayerController = GetWorld()->GetFirstPlayerController();
 		}
 
-		// We just check that the local client to server offset is valid before doing something
-		if (PlayerController && (HasServerWorld() || PlayerController->GetNetworkPhysicsTickOffsetAssigned()))
+		// Check that client and server have synced the physics tick offset
+		if (PlayerController && PlayerController->GetNetworkPhysicsTickOffsetAssigned())
 		{
 			const int32 LocalOffset = HasServerWorld() ? 0 : PlayerController->GetNetworkPhysicsTickOffset();
-			const int32 NextIndex = (InputIndex + 1) % (InputRedundancy + 1);
 
-			// if on server (Listen server) we should send the inputs onto all the clients through repnotify
-			ReplicatedInputs.History = InputHistory->CopyFramesWithOffset(InputOffsets[NextIndex], InputOffsets[InputIndex], LocalOffset);
-			
+			// Send latest N frames from history
+			const int32 ToFrame = InputHistory->GetLatestFrame() + 1; // ToFrame is excluded so add 1.
+			const int32 FromFrame = ToFrame - InputRedundancy;
+
+			ReplicatedInputs.History = InputHistory->CopyFramesWithOffset(FromFrame, ToFrame, LocalOffset);
+			MARK_PROPERTY_DIRTY_FROM_NAME(UNetworkPhysicsComponent, ReplicatedInputs, this);
+
 			if (!HasServerWorld())
 			{
 #if DEBUG_NETWORK_PHYSICS
-				FAsyncPhysicsTimestamp Timestamp = PlayerController->GetPhysicsTimestamp();
+				FAsyncPhysicsTimestamp Timestamp = const_cast<APlayerController*>(PlayerController)->GetPhysicsTimestamp();
 
 				TArray<int32> LocalFrames, ServerFrames, InputFrames;
 				PRAGMA_DISABLE_DEPRECATION_WARNINGS // TODO: Change to DebugData() in UE 5.6 and remove deprecation pragma
@@ -579,10 +588,13 @@ void UNetworkPhysicsComponent::SendStateData()
 {
 	if (HasServerWorld() && StateHistory)
 	{
-		const int32 NextIndex = (StateIndex + 1) % (StateRedundancy + 1);
+		// Send latest N frames from history
+		const int32 ToFrame = StateHistory->GetLatestFrame() + 1; // ToFrame is excluded so add 1.
+		const int32 FromFrame = ToFrame - StateRedundancy;
 
 		// if on server we should send the states onto all the clients through repnotify
-		ReplicatedStates.History = StateHistory->CopyFramesWithOffset(StateOffsets[NextIndex], StateOffsets[StateIndex], 0);
+		ReplicatedStates.History = StateHistory->CopyFramesWithOffset(FromFrame, ToFrame, 0);
+		MARK_PROPERTY_DIRTY_FROM_NAME(UNetworkPhysicsComponent, ReplicatedStates, this);
 	}
 }
 
@@ -720,7 +732,7 @@ void UNetworkPhysicsComponent::OnRep_SetReplicatedInputs()
 				TSharedPtr<Chaos::FBaseRewindHistory> ReceivedInputs = MakeShareable(ReplicatedInputs.History->Clone().Release());
 				PhysScene->EnqueueAsyncPhysicsCommand(0, this, [this, PhysScene, ReceivedInputs, LocalOffset]()
 				{
-						if (bCompareInputToTriggerRewind)
+						if (bCompareInputToTriggerRewind && IsLocallyControlled())
 						{
 							int32 ResimFrame = InputHistory->ReceiveNewData(*ReceivedInputs, LocalOffset, /*CompareDataForRewind*/ bCompareInputToTriggerRewind);
 							if (ResimFrame != INDEX_NONE)
@@ -782,21 +794,18 @@ void UNetworkPhysicsComponent::ServerReceiveInputData_Implementation(const FNetw
 	{ 
 		// We could probably skip that test since the server RPC is on server
 		ensure(HasServerWorld());
-
-		// Record the received inputs from the client into the history for future use
-		ReplicatedInputs.History = ClientInputs.History->Clone();
-
-		// Validate data in the received inputs
-		ReplicatedInputs.History->ValidateDataInHistory(ActorComponent);
-
+		
 		if (UWorld* World = GetWorld())
 		{
 			if (FPhysScene* PhysScene = World->GetPhysicsScene())
 			{
 				// Make another copy of the client inputs for the physics thread to consume
-				TSharedPtr<Chaos::FBaseRewindHistory> ReceivedInputs = MakeShareable(ReplicatedInputs.History->Clone().Release());
+				TSharedPtr<Chaos::FBaseRewindHistory> ReceivedInputs = MakeShareable(ClientInputs.History->Clone().Release());
 				PhysScene->EnqueueAsyncPhysicsCommand(0, this, [this, ReceivedInputs, PhysScene]()
 				{
+					// Validate data in the received inputs
+					ReceivedInputs->ValidateDataInHistory(ActorComponent);
+
 					PRAGMA_DISABLE_DEPRECATION_WARNINGS // TODO: Change to ReceiveNewData() in UE 5.6 and remove deprecation pragma
 					InputHistory->ReceiveNewDatas(*ReceivedInputs, 0);
 					PRAGMA_ENABLE_DEPRECATION_WARNINGS
@@ -837,7 +846,7 @@ void UNetworkPhysicsComponent::OnPreProcessInputsInternal(const int32 PhysicsSte
 	}
 #endif
 
-	if (InputHistory && StateHistory && ActorComponent)
+	if (ActorComponent)
 	{ 
 		bool bIsSolverReset = false;
 		bool bIsSolverResim = false;
@@ -851,11 +860,28 @@ void UNetworkPhysicsComponent::OnPreProcessInputsInternal(const int32 PhysicsSte
 			}
 		}
 
+		// Apply replicated state on clients if we are resimulating
+		if (StateHistory && !HasServerWorld() && bIsSolverResim)
+		{
+			FNetworkPhysicsData* PhysicsData = StateData.Get();
+			PhysicsData->LocalFrame = PhysicsStep;
+#if DEBUG_NETWORK_PHYSICS
+			UE_LOG(LogChaos, Log, TEXT("		Extracting history states at frame %d | Component = %s"), PhysicsStep, *GetFullName());
+#endif
+			const bool bExactFrame = PhysicsReplicationCVars::ResimulationCVars::bAllowRewindToClosestState ? !bIsSolverReset : true;
+			PRAGMA_DISABLE_DEPRECATION_WARNINGS // TODO: Change to ExtractData() in UE 5.6 and remove deprecation pragma
+			if (StateHistory->ExtractDatas(PhysicsStep, bIsSolverReset, PhysicsData, bExactFrame) && PhysicsData->bReceivedData)
+			PRAGMA_ENABLE_DEPRECATION_WARNINGS
+			{
+				PhysicsData->ApplyData(ActorComponent);
+			}
+		}
+
 		// Apply replicated inputs on server and simulated proxies (and on local player if we are resimulating)
-		if (!IsLocallyControlled() || bIsSolverResim)
+		if (InputHistory && (!IsLocallyControlled() || bIsSolverResim))
 		{
 			FNetworkPhysicsData* PhysicsData = InputData.Get();
-			const int32 ExpectedInputFrame = PhysicsData->LocalFrame + 1;
+			const int32 NextExpectedLocalFrame = PhysicsData->LocalFrame + 1;
 			PhysicsData->LocalFrame = PhysicsStep;
 
 	#if DEBUG_NETWORK_PHYSICS
@@ -876,31 +902,14 @@ void UNetworkPhysicsComponent::OnPreProcessInputsInternal(const int32 PhysicsSte
 					}
 				}
 				// Merge all inputs since last used input, if not resimulating
-				else if (PhysicsData->LocalFrame > ExpectedInputFrame)
+				else if (PhysicsData->LocalFrame > NextExpectedLocalFrame)
 				{
-					InputHistory->MergeData(ExpectedInputFrame, PhysicsData);
+					InputHistory->MergeData(NextExpectedLocalFrame, PhysicsData);
 				}
 
 				PhysicsData->ApplyData(ActorComponent);
 			}
 		}
-
-		// Apply replicated state on clients if we are resimulating
-		if (!HasServerWorld() && bIsSolverResim)
-		{
-			FNetworkPhysicsData* PhysicsData = StateData.Get();
-			PhysicsData->LocalFrame = PhysicsStep;
-	#if DEBUG_NETWORK_PHYSICS
-			UE_LOG(LogChaos, Log, TEXT("		Extracting history states at frame %d | Component = %s"), PhysicsStep, *GetFullName());
-	#endif
-			const bool bExactFrame = PhysicsReplicationCVars::ResimulationCVars::bAllowRewindToClosestState ? !bIsSolverReset : true;
-			PRAGMA_DISABLE_DEPRECATION_WARNINGS // TODO: Change to ExtractData() in UE 5.6 and remove deprecation pragma
-			if (StateHistory->ExtractDatas(PhysicsStep, bIsSolverReset, PhysicsData, bExactFrame) && PhysicsData->bReceivedData)
-			PRAGMA_ENABLE_DEPRECATION_WARNINGS
-			{
-				PhysicsData->ApplyData(ActorComponent);
-			}
-		}
 	}
 }
 
@@ -919,7 +928,7 @@ void UNetworkPhysicsComponent::OnPostProcessInputsInternal(const int32 PhysicsSt
 	}
 #endif
 
-	if (InputHistory && StateHistory && ActorComponent)
+	if (InputHistory && ActorComponent)
 	{
 		if (FPhysScene* PhysScene = GetWorld()->GetPhysicsScene())
 		{
@@ -948,7 +957,6 @@ void UNetworkPhysicsComponent::OnPostProcessInputsInternal(const int32 PhysicsSt
 
 			PhysicsData->BuildData(ActorComponent);
 
-			InputOffsets[InputIndex] = FMath::Max(InputOffsets[InputIndex], PhysicsStep + 1);
 			PRAGMA_DISABLE_DEPRECATION_WARNINGS // TODO: Change to RecordData() in UE 5.6 and remove deprecation pragma
 			InputHistory->RecordDatas(PhysicsStep, PhysicsData);
 			PRAGMA_ENABLE_DEPRECATION_WARNINGS
@@ -960,7 +968,7 @@ void UNetworkPhysicsComponent::OnPostProcessInputsInternal(const int32 PhysicsSt
 
 		const bool bIsServer = HasServerWorld();
 		const bool bShouldCacheStateHistory = bIsServer || (bCompareStateToTriggerRewind && bShouldCacheInputHistory);
-		if (bShouldCacheStateHistory)
+		if (StateHistory && bShouldCacheStateHistory)
 		{
 			// Compute of the local frame coming from the client that was used to generate this state
 			int32 InputFrame = INDEX_NONE;
@@ -982,7 +990,6 @@ void UNetworkPhysicsComponent::OnPostProcessInputsInternal(const int32 PhysicsSt
 
 			PhysicsData->BuildData(ActorComponent);
 
-			StateOffsets[StateIndex] = FMath::Max(StateOffsets[StateIndex], PhysicsStep + 1);
 			PRAGMA_DISABLE_DEPRECATION_WARNINGS // TODO: Change to RecordData() in UE 5.6 and remove deprecation pragma
 			StateHistory->RecordDatas(PhysicsStep, PhysicsData);
 			PRAGMA_ENABLE_DEPRECATION_WARNINGS
@@ -1044,7 +1051,7 @@ bool UNetworkPhysicsComponent::HasLocalController() const
 
 bool UNetworkPhysicsComponent::IsLocallyControlled() const
 {
-	if (bIsRelayingLocalInputs && !GetWorld()->IsNetMode(NM_DedicatedServer))
+	if (bIsRelayingLocalInputs)
 	{
 		return true;
 	}
@@ -1107,7 +1114,10 @@ void UNetworkPhysicsComponent::AddDataHistory()
 			if (Chaos::FRewindData* RewindData = Solver->GetRewindData())
 			{
 				RewindData->AddInputHistory(InputHistory);
-				RewindData->AddStateHistory(StateHistory);
+				if (StateHistory)
+				{
+					RewindData->AddStateHistory(StateHistory);
+				}
 			}
 		}
 	}
diff --git a/Engine/Source/Runtime/Engine/Public/Physics/NetworkPhysicsComponent.h b/Engine/Source/Runtime/Engine/Public/Physics/NetworkPhysicsComponent.h
--- a/Engine/Source/Runtime/Engine/Public/Physics/NetworkPhysicsComponent.h
+++ b/Engine/Source/Runtime/Engine/Public/Physics/NetworkPhysicsComponent.h
@@ -525,7 +525,7 @@ public:
 	// Unregister the component from the network manager
 	ENGINE_API virtual void UninitializeComponent() override;
 
-	// Register and create the states/inputs history
+	/** Register and create both state and input to be both networked and cached in history */
 	template<typename PhysicsTraits>
 	void CreateDataHistory(UActorComponent* HistoryComponent);
 	
@@ -533,7 +533,12 @@ public:
 	UE_DEPRECATED(5.4, "Deprecated, use CreateDataHistory() instead")
 	void CreateDatasHistory(UActorComponent* HistoryComponent);
 	
-	
+	/**  Register and create input history
+	* Please use CreateDataHistory() if both input and custom state are supposed to be networked and cached in history.
+	* NOTE: Registering Input without State requires networking push-model to be enabled to take advantage of all the CPU and Network Bandwidth savings, CVar: Net.IsPushModelEnabled 1 */
+	template<class InputsType>
+	void CreateInputHistory(UActorComponent* HistoryComponent);
+
 
 	// Remove state/input history from rewind data
 	ENGINE_API void RemoveDataHistory();
@@ -623,21 +628,9 @@ private:
 	// Send last N number of inputs each replication call to patch up holes due to packet loss
 	int8 InputRedundancy = 3;
 
-	// Current index used in the inputs offsets
-	int8 InputIndex = 0;
-
-	// Input offsets defined on PT based on the newly recorded input data
-	TArray<int32> InputOffsets;
-
 	// Send last N number of states each replication call to patch up holes due to packet loss
 	int8 StateRedundancy = 1;
 
-	// Current index used in the states offsets
-	int8 StateIndex = 0;
-
-	// State offsets defined on PT based on the newly recorded state data
-	TArray<int32> StateOffsets;
-
 	// Actor component that will be used to fill the histories
 	TObjectPtr<UActorComponent> ActorComponent;
 
@@ -685,3 +678,24 @@ FORCEINLINE void UNetworkPhysicsComponent::CreateDataHistory(UActorComponent* Hi
 	
 	AddDataHistory();
 }
+
+template<class InputsType>
+FORCEINLINE void UNetworkPhysicsComponent::CreateInputHistory(UActorComponent* HistoryComponent)
+{
+	const int32 NumFrames = SetupRewindData();
+
+	APlayerController* Controller = GetPlayerController();
+	const bool bIsLocalHistory = (Controller && Controller->IsLocalController()); // FIXME: The controller is null at this point, but bIsLocalHistory isn't currently used so doesn't create an issue.
+
+	InputHistory = MakeShared<TNetRewindHistory<InputsType>>(NumFrames, bIsLocalHistory);
+
+	InputData = MakeUnique<InputsType>();
+
+	ReplicatedInputs.History = MakeUnique<TNetRewindHistory<InputsType>>(NumFrames, bIsLocalHistory);
+	ReplicatedInputs.Owner = this;
+
+	ActorComponent = HistoryComponent;
+
+	AddDataHistory();
+}
+
diff --git a/Engine/Source/Runtime/Experimental/Chaos/Public/RewindData.h b/Engine/Source/Runtime/Experimental/Chaos/Public/RewindData.h
--- a/Engine/Source/Runtime/Experimental/Chaos/Public/RewindData.h
+++ b/Engine/Source/Runtime/Experimental/Chaos/Public/RewindData.h
@@ -56,7 +56,10 @@ struct FBaseRewindHistory
 	UE_DEPRECATED(5.4, "Deprecated, use RecordData() instead")
 	FORCEINLINE virtual bool RecordDatas(const int32 RecordFrame, const void* HistoryDatas) { return RecordData(RecordFrame, HistoryDatas); }
 
-	/** Create a polymorphic copy of only a range of frames, applying the frame offset to the copies */
+	/** Create a polymorphic copy of only a range of frames, applying the frame offset to the copies
+	* @param StartFrame = Included
+	* @param EndFrame = Excluded
+	*/
 	virtual TUniquePtr<FBaseRewindHistory> CopyFramesWithOffset(const uint32 StartFrame, const uint32 EndFrame, const int32 FrameOffset) = 0;
 
 	/** Copy new data (received from the network) into this history, returns frame to resimulate from if @param CompareDataForRewind is set to true and compared data differ enough */
@@ -83,6 +86,9 @@ struct FBaseRewindHistory
 
 	/** Legacy interface to apply inputs */
 	FORCEINLINE virtual bool ApplyInputs(const int32 ApplyFrame, const bool bResetSolver) { return false; }
+
+	/** Return the most up to date frame entry in history */
+	virtual int32 GetLatestFrame() { return INDEX_NONE; }
 };
 
 /** Templated data history holding a data buffer */
@@ -90,7 +96,7 @@ template<typename DataType>
 struct TDataRewindHistory : public FBaseRewindHistory
 {
 	FORCEINLINE TDataRewindHistory(const int32 FrameCount, const bool bIsHistoryLocal) :
-		bIsLocalHistory(bIsHistoryLocal), DataHistory(), CurrentFrame(0), CurrentIndex(0), NumFrames(FrameCount)
+		bIsLocalHistory(bIsHistoryLocal), DataHistory(), LatestFrame(0), CurrentFrame(0), CurrentIndex(0), NumFrames(FrameCount)
 	{
 		DataHistory.SetNum(NumFrames);
 	}
@@ -123,18 +129,18 @@ public :
 	/** Check if the history buffer contains an entry for the given frame*/
 	FORCEINLINE virtual bool HasValidData(const int32 ValidFrame) const override
 	{
-		const int32 LocalFrame = ValidFrame % NumFrames;
-		return ValidFrame == DataHistory[LocalFrame].LocalFrame;
+		const int32 ValidIndex = ValidFrame % NumFrames;
+		return ValidFrame == DataHistory[ValidIndex].LocalFrame;
 	}
 
 	/** Extract states at a given time */
 	FORCEINLINE virtual bool ExtractData(const int32 ExtractFrame, const bool bResetSolver, void* HistoryData, const bool bExactFrame = false) override
 	{
-		const int32 LocalFrame = ExtractFrame % NumFrames;
-		if (ExtractFrame == DataHistory[LocalFrame].LocalFrame)
+		const int32 ExtractIndex = ExtractFrame % NumFrames;
+		if (ExtractFrame == DataHistory[ExtractIndex].LocalFrame)
 		{
 			CurrentFrame = ExtractFrame;
-			CurrentIndex = LocalFrame;
+			CurrentIndex = ExtractIndex;
 
 #if DEBUG_NETWORK_PHYSICS
 			UE_LOG(LogTemp, Log, TEXT("		Found matching data into history at frame %d"), ExtractFrame);
@@ -149,40 +155,40 @@ public :
 				UE_LOG(LogChaos, Warning, TEXT("		Unable to extract data at frame %d while rewinding the simulation"), ExtractFrame);
 			}
 			PRAGMA_DISABLE_DEPRECATION_WARNINGS // TODO: Change to ClosestData() in UE 5.6 and remove deprecation pragma
-			const int32 MinFrame = ClosestDatas(ExtractFrame, true);
-			const int32 MaxFrame = ClosestDatas(ExtractFrame, false);
+			const int32 MinFrameIndex = ClosestDatas(ExtractFrame, true);
+			const int32 MaxFrameIndex = ClosestDatas(ExtractFrame, false);
 			PRAGMA_ENABLE_DEPRECATION_WARNINGS
 
-			if (MinFrame != INDEX_NONE && MaxFrame != INDEX_NONE)
+			if (MinFrameIndex != INDEX_NONE && MaxFrameIndex != INDEX_NONE)
 			{
 				PRAGMA_DISABLE_DEPRECATION_WARNINGS
 				// TODO: Change to InterpolateData() in UE 5.6 and remove deprecation pragma
 				static_cast<DataType*>(HistoryData)->InterpolateDatas(
-					DataHistory[MinFrame], DataHistory[MaxFrame]); 
+					DataHistory[MinFrameIndex], DataHistory[MaxFrameIndex]);
 				PRAGMA_ENABLE_DEPRECATION_WARNINGS
 
-				const int32 DeltaFrame = FMath::Abs(ExtractFrame - DataHistory[MinFrame].LocalFrame);
+				const int32 DeltaFrame = FMath::Abs(ExtractFrame - DataHistory[MinFrameIndex].LocalFrame);
 
 				static_cast<DataType*>(HistoryData)->LocalFrame = ExtractFrame;
-				static_cast<DataType*>(HistoryData)->ServerFrame = DataHistory[MinFrame].ServerFrame + DeltaFrame;
-				static_cast<DataType*>(HistoryData)->InputFrame = DataHistory[MinFrame].InputFrame + DeltaFrame;
-				return true;
+				static_cast<DataType*>(HistoryData)->ServerFrame = DataHistory[MinFrameIndex].ServerFrame + DeltaFrame;
+				static_cast<DataType*>(HistoryData)->InputFrame = DataHistory[MinFrameIndex].InputFrame + DeltaFrame;
 #if DEBUG_NETWORK_PHYSICS
-				UE_LOG(LogTemp, Log, TEXT("		Smoothing data between frame %d and %d - > [%d %d]"), DataHistory[MinFrame].LocalFrame, DataHistory[MaxFrame].LocalFrame, static_cast<DataType*>(HistoryData)->InputFrame, static_cast<DataType*>(HistoryData)->ServerFrame);
+				UE_LOG(LogTemp, Log, TEXT("		Smoothing data between frame %d and %d - > [%d %d]"), DataHistory[MinFrameIndex].LocalFrame, DataHistory[MaxFrameIndex].LocalFrame, static_cast<DataType*>(HistoryData)->InputFrame, static_cast<DataType*>(HistoryData)->ServerFrame);
 #endif
+				return true;
 			}
-			else if (MinFrame != INDEX_NONE)
+			else if (MinFrameIndex != INDEX_NONE)
 			{
-				*static_cast<DataType*>(HistoryData) = DataHistory[MinFrame];
-				return true;
+				*static_cast<DataType*>(HistoryData) = DataHistory[MinFrameIndex];
 #if DEBUG_NETWORK_PHYSICS
-				UE_LOG(LogTemp, Log, TEXT("		Setting data to frame %d"), DataHistory[MinFrame].LocalFrame);
+				UE_LOG(LogTemp, Log, TEXT("		Setting data to frame %d"), DataHistory[MinFrameIndex].LocalFrame);
 #endif
+				return true;
 			}
 			else
 			{
 #if DEBUG_NETWORK_PHYSICS
-				UE_LOG(LogTemp, Log, TEXT("		Failed to find data bounds : Min = %d | Max = %d"), MinFrame, MaxFrame);
+				UE_LOG(LogTemp, Log, TEXT("		Failed to find data bounds : Min = %d | Max = %d"), MinFrameIndex, MaxFrameIndex);
 #endif
 				return false;
 			}
@@ -195,17 +201,17 @@ public :
 		const int32 ToFrame = static_cast<DataType*>(ToData)->LocalFrame;
 		for (; FromFrame < ToFrame; FromFrame++)
 		{
-			const int32 LocalFrame = FromFrame % NumFrames;
-			if (FromFrame == DataHistory[LocalFrame].LocalFrame)
+			const int32 FromIndex = FromFrame % NumFrames;
+			if (FromFrame == DataHistory[FromIndex].LocalFrame)
 			{
-				static_cast<DataType*>(ToData)->MergeData(DataHistory[LocalFrame]);
+				static_cast<DataType*>(ToData)->MergeData(DataHistory[FromIndex]);
 			}
 		}
 	}
 
 	FORCEINLINE virtual bool TriggerRewindFromNewData(void* NewData) override
 	{
-		if (EvalData(static_cast<DataType*>(NewData)->LocalFrame))
+		if (EvalData(static_cast<DataType*>(NewData)->LocalFrame) && !DataHistory[CurrentIndex].bReceivedData)
 		{
 			return !static_cast<DataType*>(NewData)->CompareData(DataHistory[CurrentIndex]);
 		}
@@ -216,10 +222,10 @@ public :
 	/** Load the data from the buffer at a specific frame */
 	FORCEINLINE bool LoadData(const int32 LoadFrame)
 	{
-		const int32 LocalFrame = LoadFrame % NumFrames;
-		DataHistory[LocalFrame].LocalFrame = LoadFrame;
+		const int32 LoadIndex = LoadFrame % NumFrames;
+		DataHistory[LoadIndex].LocalFrame = LoadFrame;
 		CurrentFrame = LoadFrame;
-		CurrentIndex = LocalFrame;
+		CurrentIndex = LoadIndex;
 		return true;
 	}
 	UE_DEPRECATED(5.4, "Deprecated, use LoadData() instead")
@@ -228,11 +234,11 @@ public :
 	/** Eval the data from the buffer at a specific frame */
 	FORCEINLINE bool EvalData(const int32 EvalFrame)
 	{
-		const int32 LocalFrame = EvalFrame % NumFrames;
-		if (EvalFrame == DataHistory[LocalFrame].LocalFrame)
+		const int32 EvalIndex = EvalFrame % NumFrames;
+		if (EvalFrame == DataHistory[EvalIndex].LocalFrame)
 		{
 			CurrentFrame = EvalFrame;
-			CurrentIndex = LocalFrame;
+			CurrentIndex = EvalIndex;
 			return true;
 		}
 		return false;
@@ -247,6 +253,8 @@ public :
 		LoadDatas(RecordFrame);
 		PRAGMA_ENABLE_DEPRECATION_WARNINGS
 		DataHistory[CurrentIndex] = *static_cast<const DataType*>(HistoryData);
+		
+		LatestFrame = FMath::Max(LatestFrame, DataHistory[CurrentIndex].LocalFrame);
 		return true;
 	}
 
@@ -263,10 +271,10 @@ public :
 	FORCEINLINE uint32 NumValidData(const uint32 StartFrame, const uint32 EndFrame) const
 	{
 		uint32 NumData = 0;
-		for (uint32 FrameIndex = StartFrame; FrameIndex < EndFrame; ++FrameIndex)
+		for (uint32 ValidFrame = StartFrame; ValidFrame < EndFrame; ++ValidFrame)
 		{
-			const int32 LocalFrame = FrameIndex % NumFrames;
-			if (FrameIndex == DataHistory[LocalFrame].LocalFrame)
+			const int32 ValidIndex = ValidFrame % NumFrames;
+			if (ValidFrame == DataHistory[ValidIndex].LocalFrame)
 			{
 				++NumData;
 			}
@@ -280,6 +288,12 @@ public :
 	UE_DEPRECATED(5.4, "Deprecated, use GetDataHistory() instead")
 	TArray<DataType>& GetDatasArray() { return GetDataHistory(); }
 
+	/** Return the most up to date frame entry in history */
+	virtual int32 GetLatestFrame() override
+	{
+		return LatestFrame;
+	}
+
 protected : 
 
 	/** Check if the history is on the local/remote client*/
@@ -288,6 +302,9 @@ protected :
 	/**  Data buffer holding the history */
 	TArray<DataType> DataHistory;
 
+	/** The most up to date frame entry in history */
+	int32 LatestFrame;
+
 	/** Current frame that is being loaded/recorded */
 	int32 CurrentFrame;
 
