robustinate the Texture size limit code to make sure we don't call into platform texture dimension queries that can't handle 64 bit sizes

diff --git a/Engine/Source/Runtime/Engine/Private/Texture.cpp b/Engine/Source/Runtime/Engine/Private/Texture.cpp
--- a/Engine/Source/Runtime/Engine/Private/Texture.cpp
+++ b/Engine/Source/Runtime/Engine/Private/Texture.cpp
@@ -577,6 +577,12 @@ void UTexture::ValidateSettingsAfterImportOrEdit(bool * pRequiresNotifyMaterials
 		
 #if WITH_EDITORONLY_DATA
 
+	// GetMaximumDimension is virtual, for the current texture type, on the current (host) RHI
+	//	not really right to ever be using it, quereis should be about the target platform
+	const int32 RHIMaximumDimension = GetMaximumDimension();
+	check( RHIMaximumDimension > 0 );
+	check( RHIMaximumDimension <= GetMaximumDimensionOfNonVT() ); // GetMaximumDimensionOfNonVT is just a constant (16384)
+
 	if (Source.IsValid()) // we can have an empty source if the last source in a texture2d array is removed via the editor.
 	{
 		if ( MipGenSettings == TMGS_LeaveExistingMips && PowerOfTwoMode != ETexturePowerOfTwoSetting::None )
@@ -595,8 +601,8 @@ void UTexture::ValidateSettingsAfterImportOrEdit(bool * pRequiresNotifyMaterials
 			PowerOfTwoMode = ETexturePowerOfTwoSetting::None;
 		}
 
-		ResizeDuringBuildX = FMath::Max(0, FMath::Min((int32)GetMaximumDimension(), ResizeDuringBuildX));
-		ResizeDuringBuildY = FMath::Max(0, FMath::Min((int32)GetMaximumDimension(), ResizeDuringBuildY));
+		ResizeDuringBuildX = FMath::Max(0, FMath::Min(RHIMaximumDimension, ResizeDuringBuildX));
+		ResizeDuringBuildY = FMath::Max(0, FMath::Min(RHIMaximumDimension, ResizeDuringBuildY));
 
 		// IsPowerOfTwo only checks XY
 		bool bIsPowerOfTwo = Source.AreAllBlocksPowerOfTwo();
@@ -630,42 +636,50 @@ void UTexture::ValidateSettingsAfterImportOrEdit(bool * pRequiresNotifyMaterials
 			NeverStream = true;	
 		}
 	
-		int32 MaxDimension = FMath::Max( Source.GetSizeX() , Source.GetSizeY() );
-		bool bLargeTextureMustBeVT = MaxDimension > GetMaximumDimensionOfNonVT();
+		const int32 LargerSourceDimension = FMath::Max( Source.GetSizeX() , Source.GetSizeY() );
+		bool bLargeTextureMustBeVT = LargerSourceDimension > GetMaximumDimensionOfNonVT();
+
+		// note : checking the VirtualTextureStreaming without checking the TargetPlatform is potentially buggy
+		//	if the VT-enabled-ness of the platforms is not all the same as the Editor host platform
 
-		if ( bLargeTextureMustBeVT && ! VirtualTextureStreaming && MaxTextureSize == 0 )
+		if ( VirtualTextureStreaming && ! CVarVirtualTextures.GetValueOnAnyThread() )
 		{
-			// ?? why does this not just use CVarVirtualTextures ??
-			static const auto CVarVirtualTexturesEnabled = IConsoleManager::Get().FindTConsoleVariableDataInt(TEXT("r.VirtualTextures"));
-			check( CVarVirtualTexturesEnabled != nullptr );
+			// VT was turned on for this texture, but off for the project
+			//	fix it now, turn off on the texture:
+			UE_LOG(LogTexture, Display, TEXT("Texture %s was VT, but VT is off on the project, turning off on texture."), *GetName());
+			VirtualTextureStreaming = false;
+			bRequiresNotifyMaterials = true;
+		}
 
-			if ( CVarVirtualTexturesEnabled->GetValueOnAnyThread() )
+		if ( bLargeTextureMustBeVT && ! VirtualTextureStreaming && ( MaxTextureSize == 0 || MaxTextureSize > RHIMaximumDimension ) )
+		{
+			if ( CVarVirtualTextures.GetValueOnAnyThread() )
 			{
 				if ( GetTextureClass() == ETextureClass::TwoD )
 				{
-					UE_LOG(LogTexture, Display, TEXT("Large Texture %s Dimension=%d changed to VT; to disable VT set MaxTextureSize first"), *GetName(),MaxDimension);
+					UE_LOG(LogTexture, Display, TEXT("Large Texture %s Dimension=%d changed to VT; to disable VT set MaxTextureSize first"), *GetName(), LargerSourceDimension);
 					VirtualTextureStreaming = true;
 					bRequiresNotifyMaterials = true;
 				}
 				else
 				{
-					UE_LOG(LogTexture, Warning, TEXT("Large Texture %s Dimension=%d needs to be VT but is not 2d, changing MaxTextureSize"), *GetName(),MaxDimension);
+					UE_LOG(LogTexture, Warning, TEXT("Large Texture %s Dimension=%d needs to be VT but is not 2d, changing MaxTextureSize"), *GetName(), LargerSourceDimension);
 				
-					// GetMaximumDimension is the max size for this texture type on the current RHI
-					MaxTextureSize = GetMaximumDimension();
+					MaxTextureSize = RHIMaximumDimension;
 				}
 			}
 			else
 			{
-				UE_LOG(LogTexture, Warning, TEXT("Large Texture %s Dimension=%d must be VT but VirtualTextures are disabled, changing MaxTextureSize"), *GetName(),MaxDimension);
+				UE_LOG(LogTexture, Warning, TEXT("Large Texture %s Dimension=%d must be VT but VirtualTextures are disabled, changing MaxTextureSize"), *GetName(), LargerSourceDimension);
 
-				// GetMaximumDimension is the max size for this texture type on the current RHI
-				MaxTextureSize = GetMaximumDimension();
+				MaxTextureSize = RHIMaximumDimension;
 			}
 		}
 	
 		if (VirtualTextureStreaming)
 		{
+			// note: does not check CVAR VT enabled! may not actually be VT
+
 			if (!bIsPowerOfTwo)
 			{
 				if ( bLargeTextureMustBeVT || Source.GetNumBlocks() > 1 )
@@ -684,13 +698,14 @@ void UTexture::ValidateSettingsAfterImportOrEdit(bool * pRequiresNotifyMaterials
 				}
 			}
 
-			// VTs require mips as VT memory management assumes 1:1 texel/pixel mapping, which requires mips to enforce.
 			if (LODGroup == TEXTUREGROUP_ColorLookupTable)
 			{
 				UE_LOG(LogTexture, Warning, TEXT("VirtualTextureStreaming is not compatible with ColorLookupTable LODGroup as virtual textures require mips (%s)"), *GetName());
 				VirtualTextureStreaming = false;
 				bRequiresNotifyMaterials = true;
 			}
+
+			// VTs require mips as VT memory management assumes 1:1 texel/pixel mapping, which requires mips to enforce.
 			if (MipGenSettings == TMGS_NoMipmaps)
 			{
 				UE_LOG(LogTexture, Display, TEXT("Virtual textures require mips and MipGenSettings is NoMipmaps: Forcing to SimpleAverage (%s)"), *GetName());
@@ -734,8 +749,7 @@ void UTexture::ValidateSettingsAfterImportOrEdit(bool * pRequiresNotifyMaterials
 	}
 	else
 	{
-		// note : GetMaximumDimension is the max dim for this texture type in the current RHI
-		MaxTextureSize = FMath::Min<int32>(FMath::RoundUpToPowerOfTwo(MaxTextureSize), GetMaximumDimension());
+		MaxTextureSize = FMath::Min<int32>(FMath::RoundUpToPowerOfTwo(MaxTextureSize), RHIMaximumDimension);
 	}
 #endif
 	
@@ -3542,7 +3556,6 @@ void FTextureSource::SetId(const FGuid& InId, bool bInGuidIsHash)
 
 // GetMaximumDimensionOfNonVT is static
 // not for current texture type, not for current RHI
-
 int32 UTexture::GetMaximumDimensionOfNonVT()
 {
 	// 16384 limit ; larger must be VT
@@ -3555,6 +3568,8 @@ int32 UTexture::GetMaximumDimensionOfNonVT()
 // and on the current RHI
 uint32 UTexture::GetMaximumDimension() const
 {
+	// the various virtual implementations of this wind up returning GRHIGlobals.MaxCubeTextureDimensions etc.
+
 	// just assume anyone who doesn't implement this virtual is 2d
 	return GetMax2DTextureDimension();
 }
@@ -3814,9 +3829,8 @@ void UTexture::GetBuiltTextureSize(const ITargetPlatformSettings* TargetPlatform
 		}
 	}
 	
-	// ?? why does this not just use CVarVirtualTextures ??
-	static const auto CVarVirtualTexturesEnabled = IConsoleManager::Get().FindTConsoleVariableDataInt(TEXT("r.VirtualTextures")); check(CVarVirtualTexturesEnabled);
-	const bool bVirtualTextureStreaming = CVarVirtualTexturesEnabled->GetValueOnAnyThread() && TargetPlatformSettings->SupportsFeature(ETargetPlatformFeatures::VirtualTextureStreaming) && VirtualTextureStreaming;
+	// bug? this queries CVarVirtualTextures, but should be asking for that in the TargetPlatform config
+	const bool bVirtualTextureStreaming = CVarVirtualTextures.GetValueOnAnyThread() && TargetPlatformSettings->SupportsFeature(ETargetPlatformFeatures::VirtualTextureStreaming) && VirtualTextureStreaming;
 
 	const UTextureLODSettings& LODSettings = TargetPlatformSettings->GetTextureLODSettings();
  	const uint32 LODBiasNoCinematics = FMath::Max<int32>(LODSettings.CalculateLODBias(SizeX, SizeY, MaxTextureSize, LODGroup, LODBias, 0, MipGenSettings, bVirtualTextureStreaming), 0);
diff --git a/Engine/Source/Runtime/Engine/Private/TextureDerivedData.cpp b/Engine/Source/Runtime/Engine/Private/TextureDerivedData.cpp
--- a/Engine/Source/Runtime/Engine/Private/TextureDerivedData.cpp
+++ b/Engine/Source/Runtime/Engine/Private/TextureDerivedData.cpp
@@ -657,46 +657,59 @@ static void GetBuiltTextureSizeBytesEstimate(
 	BuildSettings.GetEncodedTextureDescription(&TextureDescription, TextureFormat, TopMipSizeX, TopMipSizeY, TopMipSizeZ, NumMips, bHasAlpha);
 	check(TextureDescription.PixelFormat == PixelFormat);
 
-	uint64 TopMipSizeBytes = 0;
-	uint64 TotalImageSizeBytes = 0;
+	uint64 LinearTopMipSizeBytes = 0;
+	uint64 LinearTotalImageSizeBytes = 0;
 
-	int32 LODBias = 0;
-	FEncodedTextureExtendedData ExtendedData = TextureFormat->GetExtendedDataForTexture(TextureDescription, LODBias);
-	if ( ExtendedData.MipSizesInBytes.Num() > 0 )
+	// calculate bytes for linear unpadded/untiled layout :
+	for (int32 MipIndex = 0; MipIndex < TextureDescription.NumMips; MipIndex++)
 	{
-		// ExtendedData is only valid for platform/tiled images
-
-		TopMipSizeBytes = ExtendedData.MipSizesInBytes[0];
-		
-		TotalImageSizeBytes = 0;
-		for(const uint64 & MipSize : ExtendedData.MipSizesInBytes )
+		if (MipIndex == 0)
 		{
-			TotalImageSizeBytes += MipSize;
+			LinearTopMipSizeBytes = TextureDescription.GetMipSizeInBytes(0);
+			LinearTotalImageSizeBytes = LinearTopMipSizeBytes;
+		}
+		else
+		{
+			LinearTotalImageSizeBytes += TextureDescription.GetMipSizeInBytes(MipIndex);
 		}
-
 	}
-	else
+	
+	check( LinearTotalImageSizeBytes > 0 );
+
+	OutTopMipSizeBytes = LinearTopMipSizeBytes;
+	OutTotalImageSizeBytes = LinearTotalImageSizeBytes;
+
+	if ( LinearTotalImageSizeBytes < (2LL<<30) )
 	{
-		// calculate bytes for linear unpadded/untiled layout :
-		for (int32 MipIndex = 0; MipIndex < TextureDescription.NumMips; MipIndex++)
+		// only call GetExtendedDataForTexture if total size is under 2 GB
+		//	because it calls into platform texture lib functions that are not 64-bit math safe
+
+		int32 LODBias = 0;
+		FEncodedTextureExtendedData ExtendedData = TextureFormat->GetExtendedDataForTexture(TextureDescription, LODBias);
+		if ( ExtendedData.MipSizesInBytes.Num() > 0 )
 		{
-			if (MipIndex == 0)
-			{
-				TopMipSizeBytes = TextureDescription.GetMipSizeInBytes(0);
-				TotalImageSizeBytes = TopMipSizeBytes;
-			}
-			else
+			// ExtendedData is only valid for platform/tiled images
+		
+			uint64 TiledTopMipSizeBytes = 0;
+			uint64 TiledTotalImageSizeBytes = 0;
+
+			TiledTopMipSizeBytes = ExtendedData.MipSizesInBytes[0];
+		
+			TiledTotalImageSizeBytes = 0;
+			for(const uint64 & MipSize : ExtendedData.MipSizesInBytes )
 			{
-				TotalImageSizeBytes += TextureDescription.GetMipSizeInBytes(MipIndex);
+				TiledTotalImageSizeBytes += MipSize;
 			}
+			
+			check( TiledTotalImageSizeBytes > 0 );
+			check( TiledTopMipSizeBytes >= LinearTopMipSizeBytes );
+			check( TiledTotalImageSizeBytes >= LinearTotalImageSizeBytes );
+
+			OutTopMipSizeBytes = TiledTopMipSizeBytes;
+			OutTotalImageSizeBytes = TiledTotalImageSizeBytes;
 		}
 	}
 
-	check(TotalImageSizeBytes > 0);
-
-
-	OutTopMipSizeBytes = TopMipSizeBytes;
-	OutTotalImageSizeBytes = TotalImageSizeBytes;
 }
 
 // may reduce OutSettings.MaxTextureResolution
@@ -717,7 +730,45 @@ static void ModifyMaxTextureResolutionBuildSettingsForPlatformLimit(
 	// GetBuiltTextureSize is the size after LODBias
 	int32 BuiltSizeX=0,BuiltSizeY=0,BuiltSizeZ=0;
 	Texture.GetBuiltTextureSize(TargetPlatform,BuiltSizeX,BuiltSizeY,BuiltSizeZ);
-		
+	
+	const int32 MaxDimension = UTexture::GetMaximumDimensionOfNonVT();
+	
+	// OriginalMaxTextureResolution is uint32_max if Texture did not have a max size set
+	uint32 OriginalMaxTextureResolution = OutSettings.MaxTextureResolution;
+	
+	OutSettings.MaxTextureResolution = FMath::Min<uint32>(MaxDimension,OutSettings.MaxTextureResolution);
+
+	if ( BuiltSizeX > MaxDimension || BuiltSizeY > MaxDimension || BuiltSizeZ > MaxDimension )
+	{
+		// this should have already happened in Texture.cpp ValidateSettingsAfterImportOrEdit
+		//	no harm in doing it again to make sure
+
+		if ( BuiltSizeZ > MaxDimension && ! OutSettings.bVolume )
+		{
+			UE_LOG(LogTexture, Error, TEXT("Texture %s non-volume has huge Z depth!"), 
+				*Texture.GetPathName());
+
+			OutSettings.MaxTextureResolution = 4;
+			return;
+		}
+		else
+		{
+			UE_LOG(LogTexture, Warning, TEXT("Texture %s exceeds maximum dimensions : %d x %d x %d > %d , shrinking..."), *Texture.GetPathName(),
+				BuiltSizeX, BuiltSizeY, BuiltSizeZ, MaxDimension
+				);
+		}
+
+		while ( BuiltSizeX > MaxDimension || BuiltSizeY > MaxDimension || BuiltSizeZ > MaxDimension )
+		{
+			BuiltSizeX = FMath::Max(1,BuiltSizeX>>1);
+			BuiltSizeY = FMath::Max(1,BuiltSizeY>>1);
+			if ( OutSettings.bVolume )
+			{
+				BuiltSizeZ = FMath::Max(1,BuiltSizeZ>>1);
+			}
+		}
+	}
+
 	uint64 MaxSurfaceBytes,MaxPackageBytes;
 	TargetPlatform->GetTextureSizeLimits(MaxSurfaceBytes,MaxPackageBytes);
 	
@@ -743,9 +794,6 @@ static void ModifyMaxTextureResolutionBuildSettingsForPlatformLimit(
 			MaxSurfaceBytes,MaxPackageBytes
 			);
 		
-		// OriginalMaxTextureResolution is uint32_max if Texture did not have a max size set
-		uint32 OriginalMaxTextureResolution = OutSettings.MaxTextureResolution;
-
 		do
 		{
 			// change MaxTextureResolution so that it causes us to do one mip step down
@@ -768,18 +816,19 @@ static void ModifyMaxTextureResolutionBuildSettingsForPlatformLimit(
 			GetBuiltTextureSizeBytesEstimate(OutSettings,TextureFormat, BuiltSizeX,BuiltSizeY,BuiltSizeZ,OutSettings.bVolume,PixelFormat,SurfaceBytes,TotalBytes);
 		}
 		while ( SurfaceBytes > MaxSurfaceBytes || TotalBytes > MaxPackageBytes );
-		
-		{
-			// compensate for LODBias that will be applied
-			// after scaling to MaxTextureResolution, LODBiasNoCinematics will be applied
+	}
+
+	if ( OutSettings.MaxTextureResolution != OriginalMaxTextureResolution )
+	{
+		// compensate for LODBias that will be applied
+		// after scaling to MaxTextureResolution, LODBiasNoCinematics will be applied
 			
-			const UTextureLODSettings& LODSettings = TargetPlatform->GetTextureLODSettings();
- 			const uint32 LODBiasNoCinematics = FMath::Max<int32>(LODSettings.CalculateLODBias(BuiltSizeX, BuiltSizeY, Texture.MaxTextureSize, Texture.LODGroup, Texture.LODBias, 0, Texture.MipGenSettings, OutSettings.bVirtualStreamable), 0);
+		const UTextureLODSettings& LODSettings = TargetPlatform->GetTextureLODSettings();
+ 		const uint32 LODBiasNoCinematics = FMath::Max<int32>(LODSettings.CalculateLODBias(BuiltSizeX, BuiltSizeY, Texture.MaxTextureSize, Texture.LODGroup, Texture.LODBias, 0, Texture.MipGenSettings, OutSettings.bVirtualStreamable), 0);
 
-			int64 MaxTextureResolutionUp = ((int64)OutSettings.MaxTextureResolution)<<LODBiasNoCinematics;
+		int64 MaxTextureResolutionUp = ((int64)OutSettings.MaxTextureResolution)<<LODBiasNoCinematics;
 
-			OutSettings.MaxTextureResolution = (uint32) FMath::Min<int64>((int64)OriginalMaxTextureResolution,MaxTextureResolutionUp);
-		}
+		OutSettings.MaxTextureResolution = (uint32) FMath::Min<int64>((int64)OriginalMaxTextureResolution,MaxTextureResolutionUp);
 
 		// ensure MaxTextureResolution never goes up :
 		OutSettings.MaxTextureResolution = FMath::Min(OriginalMaxTextureResolution,OutSettings.MaxTextureResolution);
diff --git a/Engine/Source/Runtime/Engine/Private/TextureDerivedDataTask.cpp b/Engine/Source/Runtime/Engine/Private/TextureDerivedDataTask.cpp
--- a/Engine/Source/Runtime/Engine/Private/TextureDerivedDataTask.cpp
+++ b/Engine/Source/Runtime/Engine/Private/TextureDerivedDataTask.cpp
@@ -193,6 +193,12 @@ void FTextureSourceData::InitAsPlaceholder()
 void FTextureSourceData::Init(UTexture& InTexture, TextureMipGenSettings InMipGenSettings, bool bInCubeMap, bool bInTextureArray, bool bInVolumeTexture, bool bAllowAsyncLoading)
 {
 	check( bValid == false ); // we set to true at the end, acts as our return value
+	
+	if ( ! InTexture.Source.IsValid() )
+	{
+		UE_LOG(LogTexture, Warning, TEXT("FTextureSourceData::Init on Invalid texture: %s"), *InTexture.GetPathName());
+		return;
+	}
 
 	const int32 NumBlocks = InTexture.Source.GetNumBlocks();
 	const int32 NumLayers = InTexture.Source.GetNumLayers();
@@ -724,7 +730,10 @@ static void DDC1_StoreClassicTextureInDerivedData(
 
 	// VT can be bigger than (1<<(MAX_TEXTURE_MIP_COUNT-1)) , but doesn't actually make all those mips
 	// bForVirtualTextureStreamingBuild is false in this branch
-	check(MipCount <= MAX_TEXTURE_MIP_COUNT);
+	// MipCount here can actually be more than MAX_TEXTURE_MIP_COUNT and that's okay if LODBias will drop those mips
+	//	because of the delightful way LODBias works, we have to actually build the too-big mips and they will be dropped later
+	//	(it's very unusual to reach this case, typically MaxTextureSize should have been set to 16384 preventing that)
+	//check(MipCount <= MAX_TEXTURE_MIP_COUNT);
 
 	for (int32 MipIndex = 0; MipIndex < MipCount; ++MipIndex)
 	{
@@ -1526,6 +1535,8 @@ FTextureCacheDerivedDataWorker::FTextureCacheDerivedDataWorker(
 {
 	check(DerivedData);
 	
+	// if ! InTexture->Source.IsValid() -> fail now ?
+
 	RequiredMemoryEstimate = GetBuildRequiredMemoryEstimate(InTexture,InSettingsPerLayerFetchOrBuild);
 
 	if (InSettingsPerLayerFetchFirst)
