World Partition - WorldPartitionBlueprintLibrary: Fixed small typo

diff --git a/Engine/Source/Runtime/Engine/Public/WorldPartition/WorldPartitionBlueprintLibrary.h b/Engine/Source/Runtime/Engine/Public/WorldPartition/WorldPartitionBlueprintLibrary.h
--- a/Engine/Source/Runtime/Engine/Public/WorldPartition/WorldPartitionBlueprintLibrary.h
+++ b/Engine/Source/Runtime/Engine/Public/WorldPartition/WorldPartitionBlueprintLibrary.h
@@ -40,7 +40,7 @@ struct FActorDesc
 	UPROPERTY(BlueprintReadOnly, VisibleAnywhere, Category=Actor)
 	FSoftObjectPath Class;
 
-	/** Internal name of the acgor. */
+	/** Internal name of the actor. */
 	UPROPERTY(BlueprintReadOnly, VisibleAnywhere, Category=Actor)
 	FName Name;
 
