Implemented stochastic sampling of SparseVolumeTexture mip levels to save the cost of always having to sample two mip levels for manual linear mip level interpolation. Currently disabled by default to not regress Heterogeneous Volumes, which at the moment does not support stochastic sampling. Preflighted in 662666be2267e5ebbc0385f1 but SubmitTool refuses to accept the PF due to unrelated build errors in EngineTest.

#rnx

diff --git a/Engine/Shaders/Private/SparseVolumeTexture/SparseVolumeTextureCommon.ush b/Engine/Shaders/Private/SparseVolumeTexture/SparseVolumeTextureCommon.ush
--- a/Engine/Shaders/Private/SparseVolumeTexture/SparseVolumeTextureCommon.ush
+++ b/Engine/Shaders/Private/SparseVolumeTexture/SparseVolumeTextureCommon.ush
@@ -2,6 +2,8 @@
 
 #pragma once
 
+#include "../Random.ush"
+
 // SVT_TODO: Unify with macros in SparseVolumeTexture.h
 #define SPARSE_VOLUME_TILE_RES 16
 #define SPARSE_VOLUME_TILE_BORDER 1
@@ -11,6 +13,10 @@
 #define SVTADDRESSMODE_WRAP 1u
 #define SVTADDRESSMODE_MIRROR 2u
 
+#ifndef SPARSE_VOLUME_TEXTURE_SOFTWARE_LINEAR_MIP_FILTERING
+#define SPARSE_VOLUME_TEXTURE_SOFTWARE_LINEAR_MIP_FILTERING 1 // SVT_TODO: Default to stochastic mip sampling once Heterogeneous Volumes properly initializes FMaterialPixelParameters.
+#endif
+
 float SparseVolumeTextureApplyAddressModeMirror(float v)
 {
 	float t = frac(v * 0.5f) * 2.0f;
@@ -64,7 +70,7 @@ FSparseVolumeTextureUniforms SparseVolumeTextureUnpackUniforms(const uint4 Packe
 	return Result;
 }
 
-float SparseVolumeTextureCalculateMipLevel(FSparseVolumeTextureUniforms Uniforms, float3 UVWDDX, float3 UVWDDY, float MipBias)
+float SparseVolumeTextureCalculateMipLevel(FSparseVolumeTextureUniforms Uniforms, float3 UVWDDX, float3 UVWDDY, float MipBias, float2 SvPositionXY)
 {
 	const float3 Resolution = Uniforms.VolumePageResolution * float(SPARSE_VOLUME_TILE_RES);
 	const float3 TexCoordDDX = UVWDDX * Resolution;
@@ -72,9 +78,16 @@ float SparseVolumeTextureCalculateMipLevel(FSparseVolumeTextureUniforms Uniforms
 	const float DDXLengthSquared = dot(TexCoordDDX, TexCoordDDX);
 	const float DDYLengthSquared = dot(TexCoordDDY, TexCoordDDY);
 	const float MaxLengthSquared = max(DDXLengthSquared, DDYLengthSquared);
-	const float MipLevel = 0.5f * log2(max(MaxLengthSquared, 1e-8f));
-	const float MipLevelBiased = max(0.0f, MipLevel + MipBias);
-	return MipLevelBiased;
+	const float ComputedMipLevel = 0.5f * log2(max(MaxLengthSquared, 1e-8f));
+
+#if SPARSE_VOLUME_TEXTURE_SOFTWARE_LINEAR_MIP_FILTERING
+	const float Noise = 0.0f;
+#else
+	const float Noise = InterleavedGradientNoise(SvPositionXY, View.StateFrameIndexMod8);
+#endif
+
+	const float FinalMipLevelClamped = max(0.0f, ComputedMipLevel + MipBias + Noise);
+	return FinalMipLevelClamped;
 }
 
 float3 SparseVolumeTextureGetVoxelCoord(Texture3D<uint> PageTable, FSparseVolumeTextureUniforms Uniforms, float3 PageTableCoord, int MipLevel)
@@ -132,3 +145,30 @@ float4 SparseVolumeTextureSamplePhysicalTileData(Texture3D PhysicalTileDataA, Te
 	default: return 0.0f;
 	}
 }
+
+float3 SparseVolumeTextureSamplePageTableSecondMipWrapper(Texture3D<uint> PageTable, FSparseVolumeTextureUniforms Uniforms, float3 UVW, uint AddressU, uint AddressV, uint AddressW, int MipLevel = 0)
+{
+#if SPARSE_VOLUME_TEXTURE_SOFTWARE_LINEAR_MIP_FILTERING
+	return SparseVolumeTextureSamplePageTable(PageTable, Uniforms, UVW, AddressU, AddressV, AddressW, MipLevel);
+#else
+	return 0.0f;
+#endif
+}
+
+float4 SparseVolumeTextureSamplePhysicalTileDataSecondMipWrapper(Texture3D PhysicalTileDataA, Texture3D PhysicalTileDataB, SamplerState TileDataSampler, float3 VoxelUVW, int PhysicalTileDataIndex)
+{
+#if SPARSE_VOLUME_TEXTURE_SOFTWARE_LINEAR_MIP_FILTERING
+	return SparseVolumeTextureSamplePhysicalTileData(PhysicalTileDataA, PhysicalTileDataB, TileDataSampler, VoxelUVW, PhysicalTileDataIndex);
+#else
+	return 0.0f;
+#endif
+}
+
+float4 SparseVolumeTextureCombineMipSamples(float4 LowerMipSample, float4 UpperMipSample, float MipLevelFractionalPart)
+{
+#if SPARSE_VOLUME_TEXTURE_SOFTWARE_LINEAR_MIP_FILTERING
+	return lerp(LowerMipSample, UpperMipSample, MipLevelFractionalPart);
+#else
+	return LowerMipSample;
+#endif
+}
diff --git a/Engine/Source/Runtime/Core/Private/UObject/FortniteMainBranchObjectVersion.cpp b/Engine/Source/Runtime/Core/Private/UObject/FortniteMainBranchObjectVersion.cpp
--- a/Engine/Source/Runtime/Core/Private/UObject/FortniteMainBranchObjectVersion.cpp
+++ b/Engine/Source/Runtime/Core/Private/UObject/FortniteMainBranchObjectVersion.cpp
@@ -14,6 +14,6 @@ TMap<FGuid, FGuid> FFortniteMainBranchObjectVersion::GetSystemGuids()
 	SystemGuids.Add(DevGuids.Niagara_LatestScriptCompileVersion, FGuid("E267F600682A4E94885E8AEC5CE1DC91"));
 	SystemGuids.Add(DevGuids.SkeletalMeshDerivedDataVersion, FGuid("B9AEAA3EE7AC4FCFB3860E8F3E06DD75"));
 	SystemGuids.Add(DevGuids.STATICMESH_DERIVEDDATA_VER, FGuid("543843B103794E0E9BA4BE80FB602F79"));
-	SystemGuids.Add(DevGuids.MaterialTranslationDDCVersion, FGuid("5F7AE5F7FE724C4CA448A63427BC4BB7"));
+	SystemGuids.Add(DevGuids.MaterialTranslationDDCVersion, FGuid("92556B12C25341D595C44D0FC9119E2C"));
 	return SystemGuids;
 }
diff --git a/Engine/Source/Runtime/Engine/Private/Materials/HLSLMaterialTranslator.cpp b/Engine/Source/Runtime/Engine/Private/Materials/HLSLMaterialTranslator.cpp
--- a/Engine/Source/Runtime/Engine/Private/Materials/HLSLMaterialTranslator.cpp
+++ b/Engine/Source/Runtime/Engine/Private/Materials/HLSLMaterialTranslator.cpp
@@ -15420,7 +15420,7 @@ static const TCHAR* GetSparseVolumeTextureAddressMode(TextureAddress Address)
 	}
 }
 
-int32 FHLSLMaterialTranslator::SparseVolumeTextureSamplePageTable(int32 SparseVolumeTextureIndex, int32 UVWIndex, int32 MipLevelIndex, ESamplerSourceMode SamplerSource)
+int32 FHLSLMaterialTranslator::SparseVolumeTextureSamplePageTable(int32 SparseVolumeTextureIndex, int32 UVWIndex, int32 MipLevelIndex, ESamplerSourceMode SamplerSource, bool bIsManualLinearMipMapSecondSample)
 {
 	if (SparseVolumeTextureIndex == INDEX_NONE || UVWIndex == INDEX_NONE || MipLevelIndex == INDEX_NONE)
 	{
@@ -15487,7 +15487,9 @@ int32 FHLSLMaterialTranslator::SparseVolumeTextureSamplePageTable(int32 SparseVo
 	}
 
 	AddEstimatedTextureSample();
-	FString SampleCode = FString::Printf(TEXT("SparseVolumeTextureSamplePageTable(Material.SparseVolumeTexturePageTable_%d, %s, %s, %s, %s, %s, %s)"),
+	const TCHAR* FunctionPostfix = bIsManualLinearMipMapSecondSample ? TEXT("SecondMipWrapper") : TEXT("");
+	FString SampleCode = FString::Printf(TEXT("SparseVolumeTextureSamplePageTable%s(Material.SparseVolumeTexturePageTable_%d, %s, %s, %s, %s, %s, %s)"),
+		FunctionPostfix,
 		SVTReferenceIndex, 
 		*GetParameterCode(SparseVolumeTextureIndex), 
 		*GetParameterCode(UVWAsFloat3Index), 
@@ -15498,7 +15500,7 @@ int32 FHLSLMaterialTranslator::SparseVolumeTextureSamplePageTable(int32 SparseVo
 	return AddCodeChunk(MCT_Float3, *SampleCode);
 }
 
-int32 FHLSLMaterialTranslator::SparseVolumeTextureSamplePhysicalTileData(int32 SparseVolumeTextureIndex, int32 VoxelCoordIndex, int32 PhysicalTileDataIdxIndex)
+int32 FHLSLMaterialTranslator::SparseVolumeTextureSamplePhysicalTileData(int32 SparseVolumeTextureIndex, int32 VoxelCoordIndex, int32 PhysicalTileDataIdxIndex, bool bIsManualLinearMipMapSecondSample)
 {
 	if (SparseVolumeTextureIndex == INDEX_NONE || VoxelCoordIndex == INDEX_NONE || PhysicalTileDataIdxIndex == INDEX_NONE)
 	{
@@ -15532,9 +15534,14 @@ int32 FHLSLMaterialTranslator::SparseVolumeTextureSamplePhysicalTileData(int32 S
 	check(UniformTextureExpressions[(uint32)EMaterialTextureParameterType::SparseVolume].IsValidIndex(SVTReferenceIndex));
 
 	AddEstimatedTextureSample();
-	FString SampleCode = FString::Printf(TEXT("SparseVolumeTextureSamplePhysicalTileData(Material.SparseVolumeTexturePhysicalA_%d, Material.SparseVolumeTexturePhysicalB_%d, Material.SparseVolumeTexturePhysical_%dSampler, %s, %s)"),
-		SVTReferenceIndex, SVTReferenceIndex, SVTReferenceIndex, *GetParameterCode(VoxelCoordAsFloat3Index), *GetParameterCode(IndexAsFloatIndex));
-	return AddCodeChunk(MCT_Float4, *SampleCode);
+	const TCHAR* FunctionPostfix = bIsManualLinearMipMapSecondSample ? TEXT("SecondMipWrapper") : TEXT("");
+	return AddCodeChunk(MCT_Float4, TEXT("SparseVolumeTextureSamplePhysicalTileData%s(Material.SparseVolumeTexturePhysicalA_%d, Material.SparseVolumeTexturePhysicalB_%d, Material.SparseVolumeTexturePhysical_%dSampler, %s, %s)"),
+		FunctionPostfix,
+		SVTReferenceIndex,
+		SVTReferenceIndex,
+		SVTReferenceIndex,
+		*GetParameterCode(VoxelCoordAsFloat3Index),
+		*GetParameterCode(IndexAsFloatIndex));
 }
 
 int32 FHLSLMaterialTranslator::SparseVolumeTextureSample(int32 SparseVolumeTextureIndex, int32 UVWIndex, int32 MipValue0Index, int32 MipValue1Index, int32 PhysicalTileDataIdxIndex, ETextureMipValueMode MipValueMode, ESamplerSourceMode SamplerSource)
@@ -15595,26 +15602,27 @@ int32 FHLSLMaterialTranslator::SparseVolumeTextureSample(int32 SparseVolumeTextu
 
 	if (bDeriveMipLevel)
 	{
-		FString MipLevelCode = FString::Printf(TEXT("SparseVolumeTextureCalculateMipLevel(%s, %s, %s, %s)"), *GetParameterCode(SparseVolumeTextureIndex), *GetParameterCode(DDXIndex), *GetParameterCode(DDYIndex), *GetParameterCode(MipLevelBiasIndex));
-		MipLevelIndex = AddCodeChunk(MCT_Float1, *MipLevelCode);
+		MipLevelIndex = AddCodeChunk(MCT_Float1, TEXT("SparseVolumeTextureCalculateMipLevel(%s, %s, %s, %s, %s)"), 
+			*GetParameterCode(SparseVolumeTextureIndex), 
+			*GetParameterCode(DDXIndex), 
+			*GetParameterCode(DDYIndex), 
+			*GetParameterCode(MipLevelBiasIndex),
+			TEXT("Parameters.SvPosition.xy"));
 	}
 
 	// Sample the first mip
 	int32 MipLevel0Index = Floor(MipLevelIndex);
-	int32 VoxelCoordMip0Index = SparseVolumeTextureSamplePageTable(SparseVolumeTextureIndex, UVWIndex, MipLevel0Index, SamplerSource);
-	int32 Mip0SampleIndex = SparseVolumeTextureSamplePhysicalTileData(SparseVolumeTextureIndex, VoxelCoordMip0Index, PhysicalTileDataIdxIndex);
+	int32 VoxelCoordMip0Index = SparseVolumeTextureSamplePageTable(SparseVolumeTextureIndex, UVWIndex, MipLevel0Index, SamplerSource, false /*bIsManualLinearMipMapSecondSample*/);
+	int32 Mip0SampleIndex = SparseVolumeTextureSamplePhysicalTileData(SparseVolumeTextureIndex, VoxelCoordMip0Index, PhysicalTileDataIdxIndex, false /*bIsManualLinearMipMapSecondSample*/);
 
 	// Sample the second mip
-	// SVT_TODO: Try to optimize out this second sample if LerpAlpha == 0. Might need to do that in HLSL.
 	int32 MipLevel1Index = Ceil(MipLevelIndex);
-	int32 VoxelCoordMip1Index = SparseVolumeTextureSamplePageTable(SparseVolumeTextureIndex, UVWIndex, MipLevel1Index, SamplerSource);
-	int32 Mip1SampleIndex = SparseVolumeTextureSamplePhysicalTileData(SparseVolumeTextureIndex, VoxelCoordMip1Index, PhysicalTileDataIdxIndex);
+	int32 VoxelCoordMip1Index = SparseVolumeTextureSamplePageTable(SparseVolumeTextureIndex, UVWIndex, MipLevel1Index, SamplerSource, true /*bIsManualLinearMipMapSecondSample*/);
+	int32 Mip1SampleIndex = SparseVolumeTextureSamplePhysicalTileData(SparseVolumeTextureIndex, VoxelCoordMip1Index, PhysicalTileDataIdxIndex, true /*bIsManualLinearMipMapSecondSample*/);
 
 	// Lerp
 	int32 LerpAlphaIndex = Frac(MipLevelIndex);
-	int32 LerpedResultIndex = Lerp(Mip0SampleIndex, Mip1SampleIndex, LerpAlphaIndex);
-
-	return LerpedResultIndex;
+	return AddCodeChunk(MCT_Float4, TEXT("SparseVolumeTextureCombineMipSamples(%s, %s, %s)"), *GetParameterCode(Mip0SampleIndex), *GetParameterCode(Mip1SampleIndex), *GetParameterCode(LerpAlphaIndex));
 }
 
 #endif // WITH_EDITORONLY_DATA
diff --git a/Engine/Source/Runtime/Engine/Private/Materials/HLSLMaterialTranslator.h b/Engine/Source/Runtime/Engine/Private/Materials/HLSLMaterialTranslator.h
--- a/Engine/Source/Runtime/Engine/Private/Materials/HLSLMaterialTranslator.h
+++ b/Engine/Source/Runtime/Engine/Private/Materials/HLSLMaterialTranslator.h
@@ -1081,8 +1081,8 @@ protected:
 	virtual int32 SparseVolumeTextureParameter(FName ParameterName, USparseVolumeTexture* InDefaultTexture, int32& TextureReferenceIndex, EMaterialSamplerType SamplerType) override;
 	virtual int32 SparseVolumeTextureUniform(int32 TextureIndex, int32 VectorIndex, UE::Shader::EValueType Type) override;
 	virtual int32 SparseVolumeTextureUniformParameter(FName ParameterName, int32 TextureIndex, int32 VectorIndex, UE::Shader::EValueType Type) override;
-	virtual int32 SparseVolumeTextureSamplePageTable(int32 SparseVolumeTextureIndex, int32 UVWIndex, int32 MipLevelIndex, ESamplerSourceMode SamplerSource) override;
-	virtual int32 SparseVolumeTextureSamplePhysicalTileData(int32 SparseVolumeTextureIndex, int32 VoxelCoordIndex, int32 PhysicalTileDataIdxIndex) override;
+	virtual int32 SparseVolumeTextureSamplePageTable(int32 SparseVolumeTextureIndex, int32 UVWIndex, int32 MipLevelIndex, ESamplerSourceMode SamplerSource, bool bIsManualLinearMipMapSecondSample) override;
+	virtual int32 SparseVolumeTextureSamplePhysicalTileData(int32 SparseVolumeTextureIndex, int32 VoxelCoordIndex, int32 PhysicalTileDataIdxIndex, bool bIsManualLinearMipMapSecondSample) override;
 	virtual int32 SparseVolumeTextureSample(int32 SparseVolumeTextureIndex, int32 UVWIndex, int32 MipValue0Index, int32 MipValue1Index, int32 PhysicalTileDataIdxIndex, ETextureMipValueMode MipValueMode, ESamplerSourceMode SamplerSource) override;
 
 	virtual UObject* GetReferencedTexture(int32 Index) override;
diff --git a/Engine/Source/Runtime/Engine/Public/MaterialCompiler.h b/Engine/Source/Runtime/Engine/Public/MaterialCompiler.h
--- a/Engine/Source/Runtime/Engine/Public/MaterialCompiler.h
+++ b/Engine/Source/Runtime/Engine/Public/MaterialCompiler.h
@@ -352,17 +352,19 @@ public:
 	 * @param UVWIndex The UVW coordinate at which to sample the texture.
 	 * @param MipLevelIndex The mip level at which to sample the texture.
 	 * @param SamplerSource The type of sampler to use to sample the texture.
+	 * @param bIsManualLinearMipMapSecondSample is true when this is the second/higher sample for manual linear mip map interpolation.
 	 * @return The code chunk index of the result of the texture sample.
 	 */
-	virtual int32 SparseVolumeTextureSamplePageTable(int32 SparseVolumeTextureIndex, int32 UVWIndex, int32 MipLevelIndex, ESamplerSourceMode SamplerSource) = 0;
+	virtual int32 SparseVolumeTextureSamplePageTable(int32 SparseVolumeTextureIndex, int32 UVWIndex, int32 MipLevelIndex, ESamplerSourceMode SamplerSource, bool bIsManualLinearMipMapSecondSample) = 0;
 	/**
 	 * Sample a sparse volume texture physical tile data texture.
 	 * @param SparseVolumeTextureIndex The code chunk index of the texture.
 	 * @param VoxelCoordIndex The coordinate at which to sample the texture.
 	 * @param PhysicalTileDataIdxIndex The code chunk of the index (0 or 1) of the physical tile data texture to sample.
+	 * @param bIsManualLinearMipMapSecondSample is true when this is the second/higher sample for manual linear mip map interpolation.
 	 * @return The code chunk index of the result of the texture sample.
 	 */
-	virtual int32 SparseVolumeTextureSamplePhysicalTileData(int32 SparseVolumeTextureIndex, int32 VoxelCoordIndex, int32 PhysicalTileDataIdxIndex) = 0;
+	virtual int32 SparseVolumeTextureSamplePhysicalTileData(int32 SparseVolumeTextureIndex, int32 VoxelCoordIndex, int32 PhysicalTileDataIdxIndex, bool bIsManualLinearMipMapSecondSample) = 0;
 
 	virtual int32 SparseVolumeTextureSample(int32 SparseVolumeTextureIndex, int32 UVWIndex, int32 MipValue0Index, int32 MipValue1Index, int32 PhysicalTileDataIdxIndex, ETextureMipValueMode MipValueMode, ESamplerSourceMode SamplerSource) = 0;
 
@@ -960,13 +962,13 @@ public:
 	virtual int32 ExternalTextureCoordinateOffset(int32 TextureReferenceIndex, TOptional<FName> ParameterName) override { return Compiler->ExternalTextureCoordinateOffset(TextureReferenceIndex, ParameterName); }
 	virtual int32 ExternalTextureCoordinateOffset(const FGuid& ExternalTextureGuid) override { return Compiler->ExternalTextureCoordinateOffset(ExternalTextureGuid); }
 
-	virtual int32 SparseVolumeTexture(USparseVolumeTexture* Texture, int32& TextureReferenceIndex, EMaterialSamplerType SamplerType)																		{ return Compiler->SparseVolumeTexture(Texture, TextureReferenceIndex, SamplerType); }
-	virtual int32 SparseVolumeTextureParameter(FName ParameterName, USparseVolumeTexture* InDefaultTexture, int32& TextureReferenceIndex, EMaterialSamplerType SamplerType)									{ return Compiler->SparseVolumeTextureParameter(ParameterName, InDefaultTexture, TextureReferenceIndex, SamplerType); }
-	virtual int32 SparseVolumeTextureUniform(int32 TextureIndex, int32 VectorIndex, UE::Shader::EValueType Type) override																					{ return Compiler->SparseVolumeTextureUniform(TextureIndex, VectorIndex, Type); }
-	virtual int32 SparseVolumeTextureUniformParameter(FName ParameterName, int32 TextureIndex, int32 VectorIndex, UE::Shader::EValueType Type) override														{ return Compiler->SparseVolumeTextureUniformParameter(ParameterName, TextureIndex, VectorIndex, Type); }
-	virtual int32 SparseVolumeTextureSamplePageTable(int32 SparseVolumeTextureIndex, int32 UVWIndex, int32 MipLevelIndex, ESamplerSourceMode SamplerSource) override										{ return Compiler->SparseVolumeTextureSamplePageTable(SparseVolumeTextureIndex, UVWIndex, MipLevelIndex, SamplerSource); }
-	virtual int32 SparseVolumeTextureSamplePhysicalTileData(int32 SparseVolumeTextureIndex, int32 VoxelCoordIndex, int32 PhysicalTileDataIdxIndex) override													{ return Compiler->SparseVolumeTextureSamplePhysicalTileData(SparseVolumeTextureIndex, VoxelCoordIndex, PhysicalTileDataIdxIndex); }
-	virtual int32 SparseVolumeTextureSample(int32 SparseVolumeTextureIndex, int32 UVWIndex, int32 MipValue0Index, int32 MipValue1Index, int32 PhysicalTileDataIdxIndex, ETextureMipValueMode MipValueMode, ESamplerSourceMode SamplerSource) { return Compiler->SparseVolumeTextureSample(SparseVolumeTextureIndex, UVWIndex, MipValue0Index, MipValue1Index, PhysicalTileDataIdxIndex, MipValueMode, SamplerSource); }
+	virtual int32 SparseVolumeTexture(USparseVolumeTexture* Texture, int32& TextureReferenceIndex, EMaterialSamplerType SamplerType) override																									{ return Compiler->SparseVolumeTexture(Texture, TextureReferenceIndex, SamplerType); }
+	virtual int32 SparseVolumeTextureParameter(FName ParameterName, USparseVolumeTexture* InDefaultTexture, int32& TextureReferenceIndex, EMaterialSamplerType SamplerType)	override															{ return Compiler->SparseVolumeTextureParameter(ParameterName, InDefaultTexture, TextureReferenceIndex, SamplerType); }
+	virtual int32 SparseVolumeTextureUniform(int32 TextureIndex, int32 VectorIndex, UE::Shader::EValueType Type) override																														{ return Compiler->SparseVolumeTextureUniform(TextureIndex, VectorIndex, Type); }
+	virtual int32 SparseVolumeTextureUniformParameter(FName ParameterName, int32 TextureIndex, int32 VectorIndex, UE::Shader::EValueType Type) override																							{ return Compiler->SparseVolumeTextureUniformParameter(ParameterName, TextureIndex, VectorIndex, Type); }
+	virtual int32 SparseVolumeTextureSamplePageTable(int32 SparseVolumeTextureIndex, int32 UVWIndex, int32 MipLevelIndex, ESamplerSourceMode SamplerSource, bool bIsManualLinearMipMapSecondSample) override									{ return Compiler->SparseVolumeTextureSamplePageTable(SparseVolumeTextureIndex, UVWIndex, MipLevelIndex, SamplerSource, bIsManualLinearMipMapSecondSample); }
+	virtual int32 SparseVolumeTextureSamplePhysicalTileData(int32 SparseVolumeTextureIndex, int32 VoxelCoordIndex, int32 PhysicalTileDataIdxIndex, bool bIsManualLinearMipMapSecondSample) override												{ return Compiler->SparseVolumeTextureSamplePhysicalTileData(SparseVolumeTextureIndex, VoxelCoordIndex, PhysicalTileDataIdxIndex, bIsManualLinearMipMapSecondSample); }
+	virtual int32 SparseVolumeTextureSample(int32 SparseVolumeTextureIndex, int32 UVWIndex, int32 MipValue0Index, int32 MipValue1Index, int32 PhysicalTileDataIdxIndex, ETextureMipValueMode MipValueMode, ESamplerSourceMode SamplerSource)	{ return Compiler->SparseVolumeTextureSample(SparseVolumeTextureIndex, UVWIndex, MipValue0Index, MipValue1Index, PhysicalTileDataIdxIndex, MipValueMode, SamplerSource); }
 
 	virtual UObject* GetReferencedTexture(int32 Index) override
 	{
