PR #9995: Fix Editor crash if FEngineAnalytics::Analytics fails to initialize

diff --git a/Engine/Source/Runtime/Engine/Private/EngineAnalytics.cpp b/Engine/Source/Runtime/Engine/Private/EngineAnalytics.cpp
--- a/Engine/Source/Runtime/Engine/Private/EngineAnalytics.cpp
+++ b/Engine/Source/Runtime/Engine/Private/EngineAnalytics.cpp
@@ -144,7 +144,8 @@ void FEngineAnalytics::Initialize()
 		}
 
 #if WITH_EDITOR
-		if (!AnalyticsSessionSummaryManager)
+		// Avoid crashing if Analytics did not initialize properly for whatever reason.
+		if (Analytics.IsValid() && !AnalyticsSessionSummaryManager)
 		{
 			// Create the session summary manager for the Editor instance.
 			AnalyticsSessionSummaryManager = MakeUnique<FAnalyticsSessionSummaryManager>(
