Only update max texture resolution if we violate a constraint.

diff --git a/Engine/Source/Runtime/Engine/Private/TextureDerivedData.cpp b/Engine/Source/Runtime/Engine/Private/TextureDerivedData.cpp
--- a/Engine/Source/Runtime/Engine/Private/TextureDerivedData.cpp
+++ b/Engine/Source/Runtime/Engine/Private/TextureDerivedData.cpp
@@ -736,10 +736,13 @@ static void ModifyMaxTextureResolutionBuildSettingsForPlatformLimit(
 	// OriginalMaxTextureResolution is uint32_max if Texture did not have a max size set
 	uint32 OriginalMaxTextureResolution = OutSettings.MaxTextureResolution;
 	
-	OutSettings.MaxTextureResolution = FMath::Min<uint32>(MaxDimension,OutSettings.MaxTextureResolution);
 
 	if ( BuiltSizeX > MaxDimension || BuiltSizeY > MaxDimension || BuiltSizeZ > MaxDimension )
 	{
+		// Only update the max texture resolution if we are affected by this so that previously conforming
+		// textures don't get rebuilt.
+		OutSettings.MaxTextureResolution = FMath::Min<uint32>(MaxDimension, OutSettings.MaxTextureResolution);
+
 		// this should have already happened in Texture.cpp ValidateSettingsAfterImportOrEdit
 		//	no harm in doing it again to make sure
 
