Fixed case where platform ShadingModel mask was not applied in all cases

diff --git a/Engine/Source/Runtime/Core/Private/UObject/FortniteMainBranchObjectVersion.cpp b/Engine/Source/Runtime/Core/Private/UObject/FortniteMainBranchObjectVersion.cpp
--- a/Engine/Source/Runtime/Core/Private/UObject/FortniteMainBranchObjectVersion.cpp
+++ b/Engine/Source/Runtime/Core/Private/UObject/FortniteMainBranchObjectVersion.cpp
@@ -1,4 +1,4 @@
-// Copyright Epic Games, Inc. All Rights Reserved.
+﻿// Copyright Epic Games, Inc. All Rights Reserved.
 #include "UObject/FortniteMainBranchObjectVersion.h"
 
 TMap<FGuid, FGuid> FFortniteMainBranchObjectVersion::GetSystemGuids()
@@ -14,6 +14,6 @@ TMap<FGuid, FGuid> FFortniteMainBranchObjectVersion::GetSystemGuids()
 	SystemGuids.Add(DevGuids.Niagara_LatestScriptCompileVersion, FGuid("E267F600682A4E94885E8AEC5CE1DC91"));
 	SystemGuids.Add(DevGuids.SkeletalMeshDerivedDataVersion, FGuid("B9AEAA3EE7AC4FCFB3860E8F3E06DD75"));
 	SystemGuids.Add(DevGuids.STATICMESH_DERIVEDDATA_VER, FGuid("543843B103794E0E9BA4BE80FB602F79"));
-	SystemGuids.Add(DevGuids.MaterialTranslationDDCVersion, FGuid("92556B12C25341D595C44D0FC9119E2C"));
+	SystemGuids.Add(DevGuids.MaterialTranslationDDCVersion, FGuid("D83409BD645A4B86AF0ECC671B0962C9"));
 	return SystemGuids;
 }
diff --git a/Engine/Source/Runtime/Engine/Private/Materials/HLSLMaterialTranslator.cpp b/Engine/Source/Runtime/Engine/Private/Materials/HLSLMaterialTranslator.cpp
--- a/Engine/Source/Runtime/Engine/Private/Materials/HLSLMaterialTranslator.cpp
+++ b/Engine/Source/Runtime/Engine/Private/Materials/HLSLMaterialTranslator.cpp
@@ -4279,8 +4279,12 @@ FMaterialShadingModelField FHLSLMaterialTranslator::GetCompiledShadingModels() c
 	{
 		return ShadingModelsFromCompilation;
 	}
-
-	return Material->GetShadingModels();
+	else
+	{
+		FMaterialShadingModelField MaterialShadingModels = Material->GetShadingModels();
+		UMaterialInterface::FilterOutPlatformShadingModels(Platform, MaterialShadingModels);
+		return MaterialShadingModels;
+	}
 }
 
 int32 FHLSLMaterialTranslator::Error(const TCHAR* Text)
@@ -14751,16 +14755,10 @@ void FHLSLMaterialTranslator::PrepareEnvironmentDefines()
 	}
 
 	// Set all the shading models for this material here 
-	FMaterialShadingModelField ShadingModels = Material->GetShadingModels();
-
 	// If the material gets its shading model from the material expressions, then we use the result from the compilation (assuming it's valid).
 	// This result will potentially be tighter than what GetShadingModels() returns, because it only picks up the shading models from the expressions that get compiled for a specific feature level and quality level
 	// For example, the material might have shading models behind static switches. GetShadingModels() will return both the true and the false paths from that switch, whereas the shading model field from the compilation will only contain the actual shading model selected 
-	if (Material->IsShadingModelFromMaterialExpression() && ShadingModelsFromCompilation.IsValid())
-	{
-		// Shading models fetched from the compilation of the expression graph
-		ShadingModels = ShadingModelsFromCompilation;
-	}
+	FMaterialShadingModelField ShadingModels = GetCompiledShadingModels();
 
 	ensure(ShadingModels.IsValid());
 
