Iris - Fixed replication of DestructionInfos for Initially Dormant actors that never has replicated.

diff --git a/Engine/Source/Runtime/Engine/Private/NetDriver.cpp b/Engine/Source/Runtime/Engine/Private/NetDriver.cpp
--- a/Engine/Source/Runtime/Engine/Private/NetDriver.cpp
+++ b/Engine/Source/Runtime/Engine/Private/NetDriver.cpp
@@ -3802,13 +3802,32 @@ FActorDestructionInfo* UNetDriver::CreateDestructionInfo(AActor* ThisActor, FAct
 	return &NewInfo;
 }
 
-void UNetDriver::NotifyActorDestroyed( AActor* ThisActor, bool IsSeamlessTravel )
+void UNetDriver::NotifyActorDestroyed(AActor* ThisActor, bool IsSeamlessTravel)
 {
 #if UE_WITH_IRIS
-	// For Iris this is handled through a call to AActor::EndReplication
-	// NOTE: The reason for doing this after the removal from the call to RepChangedPropertyTrackerMap.Remove is that Iris currently relies on this to invoke PreReplication which will create a RepChangePropertyTracker
 	if (ReplicationSystem)
 	{
+		// For actors that are DORM_Initial when destroyed we need to create and replicate a destruction info.
+		if (ReplicationSystem->IsServer())
+		{
+			const bool bActorHasRole = ThisActor->GetRemoteRole() != ROLE_None;
+			const bool bShouldCreateDestructionInfoForInitiallyDormantActor = IsDormInitialStartupActor(ThisActor) && bActorHasRole && !IsSeamlessTravel && !GIsReconstructingBlueprintInstances && !UE::Net::ShouldIgnoreStaticActorDestruction();
+
+			if (bShouldCreateDestructionInfoForInitiallyDormantActor)
+			{
+				if (UActorReplicationBridge* Bridge = ReplicationSystem->GetReplicationBridgeAs<UActorReplicationBridge>())
+				{
+					const UE::Net::FNetObjectReference Reference = Bridge->GetOrCreateObjectReference(ThisActor);
+					if (Reference.GetRefHandle().IsValid() && Reference.GetRefHandle().IsStatic())
+					{
+						const UReplicationBridge::FEndReplicationParameters Params = { .Location = ThisActor->GetActorLocation(), .Level = ThisActor->GetLevel(), .bUseDistanceBasedPrioritization = true };
+						Bridge->AddStaticDestructionInfo(ThisActor->GetName(), ThisActor->GetOuter(), Params);
+					}
+				}
+			}
+		}
+
+		// For the other cases, Iris handles replication of destroy through EndReplication
 		return;
 	}
 #endif // UE_WITH_IRIS
