Scene Capture:  variation of DeprojectSceneCaptureToWorld that operates on scene capture component (DeprojectSceneCaptureComponentToWorld), based on GitHub pull request.

#11101

diff --git a/Engine/Source/Runtime/Engine/Classes/Kismet/GameplayStatics.h b/Engine/Source/Runtime/Engine/Classes/Kismet/GameplayStatics.h
--- a/Engine/Source/Runtime/Engine/Classes/Kismet/GameplayStatics.h
+++ b/Engine/Source/Runtime/Engine/Classes/Kismet/GameplayStatics.h
@@ -28,6 +28,7 @@ class UDialogueWave;
 class UParticleSystem;
 class UParticleSystemComponent;
 class USaveGame;
+class USceneCaptureComponent2D;
 class USceneComponent;
 class USoundAttenuation;
 class USoundBase;
@@ -1445,6 +1446,16 @@ public:
 	UFUNCTION(BlueprintPure, Category = "Camera", meta = (Keywords = "unproject"))
 	static ENGINE_API bool DeprojectSceneCaptureToWorld(ASceneCapture2D const* SceneCapture2D, const FVector2D& TargetUV, FVector& WorldPosition, FVector& WorldDirection);
 
+	/**
+	 * Transforms the given 2D UV coordinate into a 3D world-space point and direction.
+	 * @param SceneCaptureComponent2D	Deproject using this scene capture component's view.
+	 * @param ScreenPosition			UV in scene capture render target to deproject.
+	 * @param WorldPosition				(out) Corresponding 3D position on camera near plane, in world space.
+	 * @param WorldDirection			(out) World space direction vector away from the camera at the given 2d point.
+	 */
+	UFUNCTION(BlueprintPure, Category = "Camera", meta = (Keywords = "unproject"))
+	static ENGINE_API bool DeprojectSceneCaptureComponentToWorld(USceneCaptureComponent2D* SceneCaptureComponent2D, const FVector2D& TargetUV, FVector& WorldPosition, FVector& WorldDirection);
+
 	/** 
 	 * Transforms the given 3D world-space point into a its 2D screen space coordinate. 
 	 * @param Player			Project using this player's view.
diff --git a/Engine/Source/Runtime/Engine/Private/GameplayStatics.cpp b/Engine/Source/Runtime/Engine/Private/GameplayStatics.cpp
--- a/Engine/Source/Runtime/Engine/Private/GameplayStatics.cpp
+++ b/Engine/Source/Runtime/Engine/Private/GameplayStatics.cpp
@@ -3243,7 +3243,20 @@ bool UGameplayStatics::DeprojectScreenToWorld(APlayerController const* Player, c
 
 bool UGameplayStatics::DeprojectSceneCaptureToWorld(ASceneCapture2D const* SceneCapture2D, const FVector2D& TargetUV, FVector& WorldPosition, FVector& WorldDirection)
 {
-	if (USceneCaptureComponent2D* SceneCaptureComponent2D = SceneCapture2D->GetCaptureComponent2D())
+	if (SceneCapture2D)
+	{
+		return DeprojectSceneCaptureComponentToWorld(SceneCapture2D->GetCaptureComponent2D(), TargetUV, WorldPosition, WorldDirection);
+	}
+
+	// something went wrong, zero things and return false
+	WorldPosition = FVector::ZeroVector;
+	WorldDirection = FVector::ZeroVector;
+	return false;
+}
+
+bool UGameplayStatics::DeprojectSceneCaptureComponentToWorld(USceneCaptureComponent2D* SceneCaptureComponent2D, const FVector2D& TargetUV, FVector& WorldPosition, FVector& WorldDirection)
+{
+	if (SceneCaptureComponent2D)
 	{
 		if (SceneCaptureComponent2D->TextureTarget)
 		{
