[ALT]
 - Deprecate GIsEditorLoadingPackage global and use accessors instead so we can ensure thread-safety
 - Make GIsEditorLoadingPackage compatible with ALT
 - Prevent GIsEditorLoadingPackage from being misused

diff --git a/Engine/Plugins/Experimental/DynamicMaterial/Source/DynamicMaterialEditor/Private/Model/DynamicMaterialModelEditorOnlyData.cpp b/Engine/Plugins/Experimental/DynamicMaterial/Source/DynamicMaterialEditor/Private/Model/DynamicMaterialModelEditorOnlyData.cpp
--- a/Engine/Plugins/Experimental/DynamicMaterial/Source/DynamicMaterialEditor/Private/Model/DynamicMaterialModelEditorOnlyData.cpp
+++ b/Engine/Plugins/Experimental/DynamicMaterial/Source/DynamicMaterialEditor/Private/Model/DynamicMaterialModelEditorOnlyData.cpp
@@ -544,7 +544,7 @@ void UDynamicMaterialModelEditorOnlyData::RequestMaterialBuild()
 {
 	if (!HasAnyFlags(RF_ClassDefaultObject | RF_ArchetypeObject))
 	{
-		FDynamicMaterialEditorModule::Get().AddBuildRequest(this, /* Dirty Packages */ !GIsEditorLoadingPackage);
+		FDynamicMaterialEditorModule::Get().AddBuildRequest(this, /* Dirty Packages */ !UE::GetIsEditorLoadingPackage());
 	}
 }
 
diff --git a/Engine/Plugins/Experimental/Water/Source/Editor/Private/WaterLandscapeBrush.cpp b/Engine/Plugins/Experimental/Water/Source/Editor/Private/WaterLandscapeBrush.cpp
--- a/Engine/Plugins/Experimental/Water/Source/Editor/Private/WaterLandscapeBrush.cpp
+++ b/Engine/Plugins/Experimental/Water/Source/Editor/Private/WaterLandscapeBrush.cpp
@@ -240,7 +240,7 @@ void AWaterLandscapeBrush::RegisterDelegates()
 				&& (Level != nullptr)
 				&& Algo::AnyOf(Level->Actors, [this](AActor* Actor) { return IsActorAffectingLandscape(Actor); })))
 			{
-				UpdateActors(!GIsEditorLoadingPackage);
+				UpdateActors(!UE::GetIsEditorLoadingPackage());
 			}
 		});
 
@@ -251,7 +251,7 @@ void AWaterLandscapeBrush::RegisterDelegates()
 				&& (Level != nullptr)
 				&& Algo::AnyOf(Level->Actors, [this](AActor* Actor) { return IsActorAffectingLandscape(Actor); })))
 			{
-				UpdateActors(!GIsEditorLoadingPackage);
+				UpdateActors(!UE::GetIsEditorLoadingPackage());
 			}
 		});
 
diff --git a/Engine/Source/Developer/AssetTools/Private/AssetViewUtils.cpp b/Engine/Source/Developer/AssetTools/Private/AssetViewUtils.cpp
--- a/Engine/Source/Developer/AssetTools/Private/AssetViewUtils.cpp
+++ b/Engine/Source/Developer/AssetTools/Private/AssetViewUtils.cpp
@@ -22,6 +22,7 @@
 #include "Settings/EditorLoadingSavingSettings.h"
 #include "SourceControlOperations.h"
 #include "ISourceControlModule.h"
+#include "Templates/GuardValueAccessors.h"
 #include "SourceControlHelpers.h"
 #include "FileHelpers.h"
 #include "AssetRegistry/ARFilter.h"
@@ -245,7 +246,7 @@ AssetViewUtils::ELoadAssetsResult AssetViewUtils::LoadAssetsIfNeeded(TConstArray
 	bool bSomeObjectsFailedToLoad = false;
 	bool bCancelled = false;
 	{
-		TGuardValue<bool> IsEditorLoadingPackageGuard(GIsEditorLoadingPackage, true);
+		TGuardValueAccessors<bool> IsEditorLoadingPackageGuard(UE::GetIsEditorLoadingPackage, UE::SetIsEditorLoadingPackage, true);
 		const ELoadFlags LoadFlags = Settings.bFollowRedirectors ? LOAD_None : LOAD_NoRedirects;
 
 		for (const FSoftObjectPath& ObjectPath : UnloadedObjectPaths)
diff --git a/Engine/Source/Editor/Kismet/Private/FindInBlueprints.cpp b/Engine/Source/Editor/Kismet/Private/FindInBlueprints.cpp
--- a/Engine/Source/Editor/Kismet/Private/FindInBlueprints.cpp
+++ b/Engine/Source/Editor/Kismet/Private/FindInBlueprints.cpp
@@ -74,6 +74,7 @@
 #include "UObject/UnrealType.h"
 #include "UObject/WeakObjectPtr.h"
 #include "UObject/WeakObjectPtrTemplates.h"
+#include "Templates/GuardValueAccessors.h"
 #include "Widgets/Docking/SDockTab.h"
 #include "Widgets/Images/SImage.h"
 #include "Widgets/Images/SThrobber.h"
@@ -306,7 +307,7 @@ UBlueprint* FFindInBlueprintsResult::GetParentBlueprint() const
 	{
 		UObject* Object;
 		{
-			TGuardValue<bool> IsEditorLoadingPackageGuard(GIsEditorLoadingPackage, true);
+			TGuardValueAccessors<bool> IsEditorLoadingPackageGuard(UE::GetIsEditorLoadingPackage, UE::SetIsEditorLoadingPackage, true);
 			Object = LoadObject<UObject>(NULL, *DisplayText.ToString(), NULL, 0, NULL);
 		}
 
diff --git a/Engine/Source/Editor/UnrealEd/Private/EditorEngine.cpp b/Engine/Source/Editor/UnrealEd/Private/EditorEngine.cpp
--- a/Engine/Source/Editor/UnrealEd/Private/EditorEngine.cpp
+++ b/Engine/Source/Editor/UnrealEd/Private/EditorEngine.cpp
@@ -99,6 +99,7 @@
 #include "UObject/ReferenceChainSearch.h"
 #include "UObject/ArchiveCookContext.h"
 #include "HAL/PlatformApplicationMisc.h"
+#include "Templates/GuardValueAccessors.h"
 #include "IMediaModule.h"
 #include "Scalability.h"
 #include "PlatformInfo.h"
@@ -7536,7 +7537,7 @@ void UEditorEngine::InitializeNewlyCreatedInactiveWorld(UWorld* World)
 	if (!World->bIsWorldInitialized && World->WorldType == EWorldType::Inactive && !World->IsInstanced())
 	{
 		// Guard against dirtying packages while initializing the map
-		TGuardValue<bool> IsEditorLoadingPackageGuard(GIsEditorLoadingPackage, true);
+		TGuardValueAccessors<bool> IsEditorLoadingPackageGuard(UE::GetIsEditorLoadingPackage, UE::SetIsEditorLoadingPackage, true);
 		// This is probably no longer needed with the EditorLoadingPackage guard but doesn't hurt to keep for safety.
 		const bool bOldDirtyState = World->GetOutermost()->IsDirty();
 
diff --git a/Engine/Source/Editor/UnrealEd/Private/EditorServer.cpp b/Engine/Source/Editor/UnrealEd/Private/EditorServer.cpp
--- a/Engine/Source/Editor/UnrealEd/Private/EditorServer.cpp
+++ b/Engine/Source/Editor/UnrealEd/Private/EditorServer.cpp
@@ -25,6 +25,7 @@
 #include "UObject/UObjectAnnotation.h"
 #include "Serialization/ArchiveCountMem.h"
 #include "Misc/PackageName.h"
+#include "Templates/GuardValueAccessors.h"
 #include "UObject/PackageFileSummary.h"
 #include "UObject/ReferenceChainSearch.h"
 #include "Widgets/DeclarativeSyntaxSupport.h"
@@ -1093,7 +1094,7 @@ bool UEditorEngine::CanTransact()
 	// we can transact if we have a transaction buffer and aren't currently loading packages or  routing postload.
 	// No transaction should be created during loading
 	return Trans != nullptr &&
-		!GIsEditorLoadingPackage &&
+		!UE::GetIsEditorLoadingPackage() &&
 		!FUObjectThreadContext::Get().IsRoutingPostLoad;
 }
 
@@ -2347,7 +2348,7 @@ bool UEditorEngine::Map_Load(const TCHAR* Str, FOutputDevice& Ar)
 
 #define LOCTEXT_NAMESPACE "EditorEngine"
 	// We are beginning a map load
-	TGuardValue<bool> IsEditorLoadingPackageGuard(GIsEditorLoadingPackage, true);
+	TGuardValueAccessors<bool> IsEditorLoadingPackageGuard(UE::GetIsEditorLoadingPackage, UE::SetIsEditorLoadingPackage, true);
 
 	FWorldContext &Context = GetEditorWorldContext();
 	check(Context.World() == GWorld);
diff --git a/Engine/Source/Editor/UnrealEd/Private/PackageTools.cpp b/Engine/Source/Editor/UnrealEd/Private/PackageTools.cpp
--- a/Engine/Source/Editor/UnrealEd/Private/PackageTools.cpp
+++ b/Engine/Source/Editor/UnrealEd/Private/PackageTools.cpp
@@ -49,6 +49,7 @@
 #include "Engine/BlueprintGeneratedClass.h"
 #include "Engine/GameEngine.h"
 #include "Engine/LevelStreaming.h"
+#include "Templates/GuardValueAccessors.h"
 #include "Engine/Selection.h"
 #include "AssetRegistry/IAssetRegistry.h"
 #include "Logging/MessageLog.h"
@@ -842,7 +843,7 @@ UPackageTools::UPackageTools(const FObjectInitializer& ObjectInitializer)
 	{
 		bool bResult = false;
 
-		TGuardValue<bool> IsEditorLoadingPackageGuard(GIsEditorLoadingPackage, true);
+		TGuardValueAccessors<bool> IsEditorLoadingPackageGuard(UE::GetIsEditorLoadingPackage, UE::SetIsEditorLoadingPackage, true);
 
 		FTextBuilder ErrorMessageBuilder;
 
diff --git a/Engine/Source/Editor/UnrealEd/Private/UnrealEdEngine.cpp b/Engine/Source/Editor/UnrealEd/Private/UnrealEdEngine.cpp
--- a/Engine/Source/Editor/UnrealEd/Private/UnrealEdEngine.cpp
+++ b/Engine/Source/Editor/UnrealEd/Private/UnrealEdEngine.cpp
@@ -602,7 +602,7 @@ void UUnrealEdEngine::OnPackageDirtyStateUpdated( UPackage* Pkg)
 		const bool bIsAutoSaving = PackageAutoSaver.Get() && PackageAutoSaver->IsAutoSaving();
 
 		if( !bIsAutoSaving && 
-			!GIsEditorLoadingPackage && // Don't ask if the package was modified as a result of a load
+			!UE::GetIsEditorLoadingPackage() && // Don't ask if the package was modified as a result of a load
 			!GIsCookerLoadingPackage && // don't ask if the package was modified as a result of a cooker load
 			!(Package->GetPackageFlags() & PKG_CompiledIn)) // don't ask if the package is a script package (changes are saved elsewhere via config files)
 		{
diff --git a/Engine/Source/Editor/UnrealEd/Private/UnrealEdMisc.cpp b/Engine/Source/Editor/UnrealEd/Private/UnrealEdMisc.cpp
--- a/Engine/Source/Editor/UnrealEd/Private/UnrealEdMisc.cpp
+++ b/Engine/Source/Editor/UnrealEd/Private/UnrealEdMisc.cpp
@@ -1261,7 +1261,7 @@ void FUnrealEdMisc::CB_RedrawAllViewports()
 
 void FUnrealEdMisc::CB_LevelActorsAdded(AActor* InActor)
 {
-	if (!GIsEditorLoadingPackage &&
+	if (!UE::GetIsEditorLoadingPackage() &&
 		!GIsCookerLoadingPackage &&
 		FEngineAnalytics::IsAvailable() &&
 		InActor &&
diff --git a/Engine/Source/Runtime/Core/Private/Misc/CoreGlobals.cpp b/Engine/Source/Runtime/Core/Private/Misc/CoreGlobals.cpp
--- a/Engine/Source/Runtime/Core/Private/Misc/CoreGlobals.cpp
+++ b/Engine/Source/Runtime/Core/Private/Misc/CoreGlobals.cpp
@@ -384,8 +384,10 @@ const TCHAR* LexToString(ELoaderType Type)
 	}
 }
 
+PRAGMA_DISABLE_DEPRECATION_WARNINGS
 /** Whether the editor is currently loading a package or not												*/
-bool					GIsEditorLoadingPackage			= false;
+FIsEditorLoadingPackage	GIsEditorLoadingPackage;
+PRAGMA_ENABLE_DEPRECATION_WARNINGS
 /** Whether the cooker is currently loading a package or not												*/
 bool					GIsCookerLoadingPackage			= false;
 /** Whether GWorld points to the play in editor world														*/
@@ -890,6 +892,64 @@ FPlayInEditorID::operator int32() const
 }
 PRAGMA_ENABLE_DEPRECATION_WARNINGS
 
+bool PRIVATE_GIsEditorLoadingPackage = false;
+
+namespace UE
+{
+	bool GetIsEditorLoadingPackage()
+	{
+		if (IsInGameThread())
+		{
+			if (GIsEditor && IsInAsyncLoadingThread())
+			{
+				ensureMsgf(PRIVATE_GIsEditorLoadingPackage, TEXT("GIsEditorLoadingPackage should always be true during editor loading."));
+			}
+			return PRIVATE_GIsEditorLoadingPackage;
+		}
+		else if (IsInAsyncLoadingThread())
+		{
+			// Directly depends on GIsEditor on real async-loading thread
+			return GIsEditor;
+		}
+		else
+		{
+			ensureMsgf(false, TEXT("Querying GIsEditorLoadingPackage from any thread is not safe, you should capture the value instead if needed."));
+			return false;
+		}
+	}
+	void SetIsEditorLoadingPackage(bool InValue)
+	{
+		if (IsInGameThread())
+		{
+			if (GIsEditor && !InValue && IsInAsyncLoadingThread())
+			{
+				ensureMsgf(false, TEXT("Overriding GIsEditorLoadingPackage to false during editor loading is prohibited."));
+			}
+			else
+			{
+				PRIVATE_GIsEditorLoadingPackage = InValue;
+			}
+		}
+		else
+		{
+			ensureMsgf(false, TEXT("Overriding GIsEditorLoadingPackage only makes sense for the game thread."));
+		}
+	}
+}
+
+PRAGMA_DISABLE_DEPRECATION_WARNINGS
+FIsEditorLoadingPackage& FIsEditorLoadingPackage::operator= (bool InValue)
+{
+	UE::SetIsEditorLoadingPackage(InValue);
+	return *this;
+}
+
+FIsEditorLoadingPackage::operator bool() const
+{
+	return UE::GetIsEditorLoadingPackage();
+}
+PRAGMA_ENABLE_DEPRECATION_WARNINGS
+
 #undef LOCTEXT_NAMESPACE
 
 bool IsRunningCookOnTheFly()
diff --git a/Engine/Source/Runtime/Core/Public/CoreGlobals.h b/Engine/Source/Runtime/Core/Public/CoreGlobals.h
--- a/Engine/Source/Runtime/Core/Public/CoreGlobals.h
+++ b/Engine/Source/Runtime/Core/Public/CoreGlobals.h
@@ -479,8 +479,24 @@ extern CORE_API void (*SuspendTextureStreamingRenderTasks)();
 /** Resume texture updates caused by completed async IOs. */
 extern CORE_API void (*ResumeTextureStreamingRenderTasks)();
 
+class UE_DEPRECATED(5.5, "This will be removed.") FIsEditorLoadingPackage
+{
+public:
+	CORE_API FIsEditorLoadingPackage& operator= (bool InValue);
+	CORE_API operator bool() const;
+};
+
+PRAGMA_DISABLE_DEPRECATION_WARNINGS
 /** Whether the editor is currently loading a package or not */
-extern CORE_API bool GIsEditorLoadingPackage;
+UE_DEPRECATED(5.5, "Use UE::GetIsEditorLoadingPackage() and UE::SetIsEditorLoadingPackage() instead.")
+extern CORE_API FIsEditorLoadingPackage GIsEditorLoadingPackage;
+PRAGMA_ENABLE_DEPRECATION_WARNINGS
+
+namespace UE
+{
+	CORE_API bool GetIsEditorLoadingPackage();
+	CORE_API void SetIsEditorLoadingPackage(bool InValue);
+}
 
 /** Whether the cooker is currently loading a package or not */
 extern CORE_API bool GIsCookerLoadingPackage;
diff --git a/Engine/Source/Runtime/CoreUObject/Private/Serialization/AsyncLoading.cpp b/Engine/Source/Runtime/CoreUObject/Private/Serialization/AsyncLoading.cpp
--- a/Engine/Source/Runtime/CoreUObject/Private/Serialization/AsyncLoading.cpp
+++ b/Engine/Source/Runtime/CoreUObject/Private/Serialization/AsyncLoading.cpp
@@ -42,6 +42,7 @@
 #include "ProfilingDebugging/LoadTimeTracker.h"
 #include "HAL/ThreadHeartBeat.h"
 #include "HAL/ExceptionHandling.h"
+#include "Templates/GuardValueAccessors.h"
 #include "Serialization/AsyncLoadingPrivate.h"
 #include "UObject/UObjectHash.h"
 #include "Templates/UniquePtr.h"
@@ -4977,7 +4978,7 @@ EAsyncPackageState::Type FAsyncLoadingThread::TickAsyncLoading(bool bUseTimeLimi
 #if WITH_EDITOR
 	// In the editor loading cannot be part of a transaction as it cannot be undone, and may result in recording half-loaded objects. So we suppress any active transaction while in this stack, and set the editor loading flag
 	TGuardValue<ITransaction*> SuppressTransaction(GUndo, nullptr);
-	TGuardValue<bool> IsEditorLoadingPackage(GIsEditorLoadingPackage, GIsEditor || GIsEditorLoadingPackage);
+	TGuardValueAccessors<bool> IsEditorLoadingPackageGuard(UE::GetIsEditorLoadingPackage, UE::SetIsEditorLoadingPackage, GIsEditor || UE::GetIsEditorLoadingPackage());
 #endif
 
 	const bool bLoadingSuspended = IsAsyncLoadingSuspendedInternal();
diff --git a/Engine/Source/Runtime/CoreUObject/Private/Serialization/AsyncLoading2.cpp b/Engine/Source/Runtime/CoreUObject/Private/Serialization/AsyncLoading2.cpp
--- a/Engine/Source/Runtime/CoreUObject/Private/Serialization/AsyncLoading2.cpp
+++ b/Engine/Source/Runtime/CoreUObject/Private/Serialization/AsyncLoading2.cpp
@@ -68,6 +68,7 @@
 #include "ProfilingDebugging/CountersTrace.h"
 #include "ProfilingDebugging/AssetMetadataTrace.h"
 #include "Async/Async.h"
+#include "Templates/GuardValueAccessors.h"
 #include "Async/ParallelFor.h"
 #include "Async/ManualResetEvent.h"
 #include "HAL/LowLevelMemStats.h"
@@ -8079,7 +8080,7 @@ EAsyncPackageState::Type FAsyncLoadingThread2::TickAsyncLoadingFromGameThread(FA
 #if WITH_EDITOR
 	// In the editor loading cannot be part of a transaction as it cannot be undone, and may result in recording half-loaded objects. So we suppress any active transaction while in this stack, and set the editor loading flag
 	TGuardValue<ITransaction*> SuppressTransaction(GUndo, nullptr);
-	TGuardValue<bool> IsEditorLoadingPackage(GIsEditorLoadingPackage, GIsEditor || GIsEditorLoadingPackage);
+	TGuardValueAccessors<bool> IsEditorLoadingPackageGuard(UE::GetIsEditorLoadingPackage, UE::SetIsEditorLoadingPackage, GIsEditor || UE::GetIsEditorLoadingPackage());
 #endif
 
 	const bool bLoadingSuspended = IsAsyncLoadingSuspended();
diff --git a/Engine/Source/Runtime/CoreUObject/Private/UObject/UObjectBaseUtility.cpp b/Engine/Source/Runtime/CoreUObject/Private/UObject/UObjectBaseUtility.cpp
--- a/Engine/Source/Runtime/CoreUObject/Private/UObject/UObjectBaseUtility.cpp
+++ b/Engine/Source/Runtime/CoreUObject/Private/UObject/UObjectBaseUtility.cpp
@@ -237,7 +237,7 @@ bool UObjectBaseUtility::MarkPackageDirty() const
 			// we explicitly disable the ability to dirty a package or map during load/undo/redo.  Commandlets can still
 			// set the dirty state on load.
 			if( IsRunningCommandlet() || 
-				(!IsInAsyncLoadingThread() && GIsEditor && !GIsEditorLoadingPackage && !GIsCookerLoadingPackage && !GIsPlayInEditorWorld && !IsReloadActive()
+				(!IsInAsyncLoadingThread() && GIsEditor && !UE::GetIsEditorLoadingPackage() && !GIsCookerLoadingPackage && !GIsPlayInEditorWorld && !IsReloadActive()
 #if WITH_EDITORONLY_DATA
 				&& !GIsTransacting
 				&& !Package->bIsCookedForEditor // Cooked packages can't be modified nor marked as dirty
diff --git a/Engine/Source/Runtime/CoreUObject/Private/UObject/UObjectGlobals.cpp b/Engine/Source/Runtime/CoreUObject/Private/UObject/UObjectGlobals.cpp
--- a/Engine/Source/Runtime/CoreUObject/Private/UObject/UObjectGlobals.cpp
+++ b/Engine/Source/Runtime/CoreUObject/Private/UObject/UObjectGlobals.cpp
@@ -75,6 +75,7 @@
 #include "Misc/PackageAccessTracking.h"
 #include "UObject/PropertyWithSetterAndGetter.h"
 #include "UObject/AnyPackagePrivate.h"
+#include "Templates/GuardValueAccessors.h"
 #include "UObject/UObjectGlobalsInternal.h"
 #include "Serialization/AsyncPackageLoader.h"
 #include "Containers/VersePath.h"
@@ -1773,7 +1774,7 @@ UPackage* LoadPackageInternal(UPackage* InOuter, const FPackagePath& PackagePath
 #if WITH_EDITOR
 	// In the editor loading cannot be part of a transaction as it cannot be undone, and may result in recording half-loaded objects. So we suppress any active transaction while in this stack, and set the editor loading flag
 	TGuardValue<ITransaction*> SuppressTransaction(GUndo, nullptr);
-	TGuardValue<bool> IsEditorLoadingPackage(GIsEditorLoadingPackage, GIsEditor || GIsEditorLoadingPackage);
+	TGuardValueAccessors<bool> IsEditorLoadingPackage(UE::GetIsEditorLoadingPackage, UE::SetIsEditorLoadingPackage, GIsEditor || UE::GetIsEditorLoadingPackage());
 #endif
 
 	TOptional<FScopedSlowTask> SlowTask;
@@ -1977,7 +1978,7 @@ UPackage* LoadPackageInternal(UPackage* InOuter, const FPackagePath& PackagePath
 		EndLoadAndCopyLocalizationGatherFlag();
 
 #if WITH_EDITOR
-		GIsEditorLoadingPackage = IsEditorLoadingPackage.GetOriginalValue();
+		UE::SetIsEditorLoadingPackage(IsEditorLoadingPackage.GetOriginalValue());
 #endif
 
 		// if we are calculating the script SHA for a package, do the comparison now
diff --git a/Engine/Source/Runtime/Engine/Private/ActorDeferredScriptManager.cpp b/Engine/Source/Runtime/Engine/Private/ActorDeferredScriptManager.cpp
--- a/Engine/Source/Runtime/Engine/Private/ActorDeferredScriptManager.cpp
+++ b/Engine/Source/Runtime/Engine/Private/ActorDeferredScriptManager.cpp
@@ -9,6 +9,7 @@
 #include "AsyncCompilationHelpers.h"
 #include "AssetCompilingManager.h"
 #include "Engine/Level.h"
+#include "Templates/GuardValueAccessors.h"
 #include "Engine/StaticMesh.h"
 #include "Logging/LogMacros.h"
 #include "ProfilingDebugging/CountersTrace.h"
@@ -188,7 +189,7 @@ void FActorDeferredScriptManager::ProcessAsyncTasks(bool bLimitExecutionTime)
 
 		// Since this deferred run of construction script was supposed to be done during level load
 		// temporarily set the global flag to prevent dirtying the level package.
-		TGuardValue<bool> IsEditorLoadingPackageGuard(GIsEditorLoadingPackage, true);
+		TGuardValueAccessors<bool> IsEditorLoadingPackageGuard(UE::GetIsEditorLoadingPackage, UE::SetIsEditorLoadingPackage, true);
 
 		int32 CurrentIndex = NextIndexToProcess;
 		UPackage* EnteredPackage = nullptr;
diff --git a/Engine/Source/Runtime/Engine/Private/LevelUtils.cpp b/Engine/Source/Runtime/Engine/Private/LevelUtils.cpp
--- a/Engine/Source/Runtime/Engine/Private/LevelUtils.cpp
+++ b/Engine/Source/Runtime/Engine/Private/LevelUtils.cpp
@@ -175,7 +175,7 @@ bool FLevelUtils::IsLevelLocked(ULevel* Level)
 {
 	//We should not check file status on disk if we are not running the editor
 	// Don't permit spawning in read only levels if they are locked
-	if ( GIsEditor && !GIsEditorLoadingPackage )
+	if ( GIsEditor && !UE::GetIsEditorLoadingPackage() )
 	{
 		if ( GEngine && GEngine->bLockReadOnlyLevels )
 		{
diff --git a/Engine/Source/Runtime/Engine/Private/StaticMesh.cpp b/Engine/Source/Runtime/Engine/Private/StaticMesh.cpp
--- a/Engine/Source/Runtime/Engine/Private/StaticMesh.cpp
+++ b/Engine/Source/Runtime/Engine/Private/StaticMesh.cpp
@@ -8762,7 +8762,7 @@ FStaticMeshCompilationContext::FStaticMeshCompilationContext()
 {
 	// Remember if the editor was loading a package when initiating the build so that we can temporarily restore that state when 
 	//  executing FinishBuildInternal on the game thread at the end of the build :
-	bIsEditorLoadingPackage = GIsEditorLoadingPackage;
+	bIsEditorLoadingPackage = UE::GetIsEditorLoadingPackage();
 }
 
 #undef LOCTEXT_NAMESPACE
diff --git a/Engine/Source/Runtime/Engine/Private/StaticMeshCompiler.cpp b/Engine/Source/Runtime/Engine/Private/StaticMeshCompiler.cpp
--- a/Engine/Source/Runtime/Engine/Private/StaticMeshCompiler.cpp
+++ b/Engine/Source/Runtime/Engine/Private/StaticMeshCompiler.cpp
@@ -15,6 +15,7 @@
 #include "Settings/EditorExperimentalSettings.h"
 #include "GameFramework/Pawn.h"
 #include "Misc/QueuedThreadPoolWrapper.h"
+#include "Templates/GuardValueAccessors.h"
 #include "ProfilingDebugging/CountersTrace.h"
 #include "TextureCompiler.h"
 #include "ShaderCompiler.h"
@@ -252,7 +253,7 @@ void FStaticMeshCompilingManager::PostCompilation(UStaticMesh* StaticMesh)
 
 		// If async (post load or build), restore the state of GIsEditorLoadingPackage for the duration of this function (including outside of this scope) as it was when the build of the static mesh was initiated, 
 		//  so that async builds have the same result as synchronous ones (e.g. don't dirty packages when the components referencing this static mesh call Modify because GIsEditorLoadingPackage is true) :
-		TUniquePtr<TGuardValue<bool>> IsEditorLoadingPackageGuard;
+		TUniquePtr<TGuardValueAccessors<bool>> IsEditorLoadingPackageGuard;
 
 		// The scope is important here to destroy the FStaticMeshAsyncBuildScope before broadcasting events
 		{
@@ -273,7 +274,7 @@ void FStaticMeshCompilingManager::PostCompilation(UStaticMesh* StaticMesh)
 
 			if (LocalAsyncTask->GetTask().PostLoadContext.IsValid())
 			{
-				IsEditorLoadingPackageGuard.Reset(new TGuardValue<bool>(GIsEditorLoadingPackage, LocalAsyncTask->GetTask().PostLoadContext->bIsEditorLoadingPackage));
+				IsEditorLoadingPackageGuard.Reset(new TGuardValueAccessors<bool>(UE::GetIsEditorLoadingPackage, UE::SetIsEditorLoadingPackage, LocalAsyncTask->GetTask().PostLoadContext->bIsEditorLoadingPackage));
 
 				StaticMesh->FinishPostLoadInternal(*LocalAsyncTask->GetTask().PostLoadContext);
 
@@ -282,7 +283,7 @@ void FStaticMeshCompilingManager::PostCompilation(UStaticMesh* StaticMesh)
 
 			if (LocalAsyncTask->GetTask().BuildContext.IsValid())
 			{
-				IsEditorLoadingPackageGuard.Reset(new TGuardValue<bool>(GIsEditorLoadingPackage, LocalAsyncTask->GetTask().BuildContext->bIsEditorLoadingPackage));
+				IsEditorLoadingPackageGuard.Reset(new TGuardValueAccessors<bool>(UE::GetIsEditorLoadingPackage, UE::SetIsEditorLoadingPackage, LocalAsyncTask->GetTask().BuildContext->bIsEditorLoadingPackage));
 
 				TArray<IStaticMeshComponent*> ComponentsToUpdate;
 				for (IStaticMeshComponent* Component : ObjectCacheScope.GetContext().GetStaticMeshComponents(StaticMesh))
diff --git a/Engine/Source/Runtime/Engine/Private/World.cpp b/Engine/Source/Runtime/Engine/Private/World.cpp
--- a/Engine/Source/Runtime/Engine/Private/World.cpp
+++ b/Engine/Source/Runtime/Engine/Private/World.cpp
@@ -118,6 +118,7 @@
 #include "Net/PerfCountersHelpers.h"
 #include "InGamePerformanceTracker.h"
 #include "Engine/AssetManager.h"
+#include "Templates/GuardValueAccessors.h"
 #include "Engine/HLODProxy.h"
 #include "MoviePlayerProxy.h"
 #include "ObjectTrace.h"
@@ -3313,7 +3314,7 @@ void UWorld::AddToWorld( ULevel* Level, const FTransform& LevelTransform, bool b
 
 #if WITH_EDITOR
 			// Pretend here that we are loading package to avoid package dirtying during components registration
-		TGuardValue<bool> IsEditorLoadingPackage(GIsEditorLoadingPackage, (GIsEditor ? true : GIsEditorLoadingPackage));
+		TGuardValueAccessors<bool> IsEditorLoadingPackage(UE::GetIsEditorLoadingPackage, UE::SetIsEditorLoadingPackage, (GIsEditor ? true : UE::GetIsEditorLoadingPackage()));
 #endif
 
 		// We don't need to rerun construction scripts if we have cooked data or we are playing in editor unless the PIE world was loaded
diff --git a/Engine/Source/Runtime/Engine/Private/WorldPartition/WorldPartition.cpp b/Engine/Source/Runtime/Engine/Private/WorldPartition/WorldPartition.cpp
--- a/Engine/Source/Runtime/Engine/Private/WorldPartition/WorldPartition.cpp
+++ b/Engine/Source/Runtime/Engine/Private/WorldPartition/WorldPartition.cpp
@@ -871,7 +871,7 @@ bool UWorldPartition::IsMainWorldPartition() const
 #if WITH_EDITOR
 void UWorldPartition::OnLevelActorDeleted(AActor* Actor)
 {
-	if (GIsEditorLoadingPackage)
+	if (UE::GetIsEditorLoadingPackage())
 	{
 		if (UActorDescContainerInstance* DescContainerInstance = GetActorDescContainerInstance())
 		{
diff --git a/Engine/Source/Runtime/Engine/Private/WorldPartition/WorldPartitionHandle.cpp b/Engine/Source/Runtime/Engine/Private/WorldPartition/WorldPartitionHandle.cpp
--- a/Engine/Source/Runtime/Engine/Private/WorldPartition/WorldPartitionHandle.cpp
+++ b/Engine/Source/Runtime/Engine/Private/WorldPartition/WorldPartitionHandle.cpp
@@ -10,6 +10,7 @@
 #include "WorldPartition/ActorDescContainerInstanceCollection.h"
 #include "WorldPartition/ActorDescContainerInstance.h"
 #include "WorldPartition/WorldPartitionActorDescInstance.h"
+#include "Templates/GuardValueAccessors.h"
 
 /**
 * FWorldPartitionImplBase
@@ -92,7 +93,7 @@ FWorldPartitionLoadingContext::IContext::~IContext()
 void FWorldPartitionLoadingContext::FImmediate::RegisterActor(FWorldPartitionActorDescInstance* InActorDescInstance)
 {
 	// Set GIsEditorLoadingPackage to avoid dirtying the Actor package if Modify() is called during the load sequence
-	TGuardValue<bool> IsEditorLoadingPackageGuard(GIsEditorLoadingPackage, true);
+	TGuardValueAccessors<bool> IsEditorLoadingPackageGuard(UE::GetIsEditorLoadingPackage, UE::SetIsEditorLoadingPackage, true);
 
 	if (InActorDescInstance->StartAsyncLoad())
 	{
@@ -115,7 +116,7 @@ void FWorldPartitionLoadingContext::FImmediate::RegisterActor(FWorldPartitionAct
 void FWorldPartitionLoadingContext::FImmediate::UnregisterActor(FWorldPartitionActorDescInstance* InActorDescInstance)
 {
 	// Set GIsEditorLoadingPackage to avoid dirtying the Actor package if Modify() is called during the unload sequence
-	TGuardValue<bool> IsEditorLoadingPackageGuard(GIsEditorLoadingPackage, true);
+	TGuardValueAccessors<bool> IsEditorLoadingPackageGuard(UE::GetIsEditorLoadingPackage, UE::SetIsEditorLoadingPackage, true);
 
 	// When cleaning up worlds, actors are already marked as garbage at this point, so no need to remove them from the world
 	if (AActor* Actor = InActorDescInstance->GetActor(); IsValid(Actor))
@@ -174,7 +175,7 @@ FWorldPartitionLoadingContext::FDeferred::~FDeferred()
 			if (ContainerOp.Registrations.Num())
 			{
 				// Set GIsEditorLoadingPackage to avoid dirtying the Actor package if Modify() is called during the load sequence
-				TGuardValue<bool> IsEditorLoadingPackageGuard(GIsEditorLoadingPackage, true);
+				TGuardValueAccessors<bool> IsEditorLoadingPackageGuard(UE::GetIsEditorLoadingPackage, UE::SetIsEditorLoadingPackage, true);
 
 				TArray<AActor*> ActorList;
 				if (ULevel* Level = CreateActorList(ActorList, ContainerOp.Registrations))
@@ -192,7 +193,7 @@ FWorldPartitionLoadingContext::FDeferred::~FDeferred()
 			if (ContainerOp.Unregistrations.Num())
 			{
 				// Set GIsEditorLoadingPackage to avoid dirtying the Actor package if Modify() is called during the unload sequence
-				TGuardValue<bool> IsEditorLoadingPackageGuard(GIsEditorLoadingPackage, true);
+				TGuardValueAccessors<bool> IsEditorLoadingPackageGuard(UE::GetIsEditorLoadingPackage, UE::SetIsEditorLoadingPackage, true);
 
 				TArray<AActor*> ActorList;
 				if (ULevel* Level = CreateActorList(ActorList, ContainerOp.Unregistrations))
@@ -216,7 +217,7 @@ FWorldPartitionLoadingContext::FDeferred::~FDeferred()
 
 void FWorldPartitionLoadingContext::FDeferred::RegisterActor(FWorldPartitionActorDescInstance* InActorDescInstance)
 {
-	TGuardValue<bool> IsEditorLoadingPackageGuard(GIsEditorLoadingPackage, true);
+	TGuardValueAccessors<bool> IsEditorLoadingPackageGuard(UE::GetIsEditorLoadingPackage, UE::SetIsEditorLoadingPackage, true);
 
 	check(InActorDescInstance);
 	if (InActorDescInstance->StartAsyncLoad())
