NetDriver
* Added some functions to allow swapping a NetDriver from a ReplicationDriver (RepGraph) to a ReplicationSystem (Iris) without destroying and restarting the netdriver

#[review]-33100777

diff --git a/Engine/Source/Runtime/Engine/Classes/Engine/Engine.h b/Engine/Source/Runtime/Engine/Classes/Engine/Engine.h
--- a/Engine/Source/Runtime/Engine/Classes/Engine/Engine.h
+++ b/Engine/Source/Runtime/Engine/Classes/Engine/Engine.h
@@ -3101,6 +3101,18 @@ public:
 	/** A list of Iris NetDriverConfigs */
 	UPROPERTY(Config, transient)
 	TArray<FIrisNetDriverConfig> IrisNetDriverConfigs;
+
+	/** 
+	 * Returns the Iris config for the corresponding NetDriver 
+	 * Priority order for the IrisNetDriverConfigs are:
+	 *		1. NetDriverName exact match
+	 *		2. NetDriverName wildcard match
+	 *		3. NetDriverDefinition match
+	 */
+	ENGINE_API const FIrisNetDriverConfig* GetIrisNetDriverConfig(FName InNetDriverDefinition, FName InNetDriverName) const;
+
+	/** Returns true if the netdriver will run with Iris enable. */
+	ENGINE_API bool WillNetDriverUseIris(const FWorldContext& Context, FName InNetDriverDefinition, FName InNetDriverName) const;
 	
 	/** A configurable list of actors that are automatically spawned upon server startup (just prior to InitGame) */
 	UPROPERTY(config)
diff --git a/Engine/Source/Runtime/Engine/Classes/Engine/NetDriver.h b/Engine/Source/Runtime/Engine/Classes/Engine/NetDriver.h
--- a/Engine/Source/Runtime/Engine/Classes/Engine/NetDriver.h
+++ b/Engine/Source/Runtime/Engine/Classes/Engine/NetDriver.h
@@ -1050,8 +1050,10 @@ public:
 	FName GetNetDriverDefinition() const { return NetDriverDefinition; }
 
 	/** Callback after the engine created the NetDriver and set our name for the first time */
-	void PostCreation(bool bInitializeWithIris);
+	ENGINE_API void PostCreation(bool bInitializeWithIris);
 
+	/** Reset some NetDriver settings after destroying and recreating a replication system or replication driver */
+	ENGINE_API void ReinitBase();
 
 	void InitPacketSimulationSettings();
 
@@ -1850,18 +1852,18 @@ public:
 	ENGINE_API EEngineNetworkRuntimeFeatures GetNetworkRuntimeFeatures() const;
 	
 #if UE_WITH_IRIS
-	/** Remove references to the Iris bridge and system without deleting it */
-	ENGINE_API void ClearIrisSystem();
-
-	/** Set a previously initialized IrisSystem into this NetDriver */
-	ENGINE_API void RestoreIrisSystem(UReplicationSystem* InReplicationSystem);
-
 	/**
-	 * Destroy and recreate the iris replication system for an active netdrive.
+	 * Destroy and recreate the iris replication system for an active netdriver.
 	 * This will re-add all existing replicated actors back in the system.
 	 * Useful if you need to reapply hotfix configs downloaded post-initialization.
 	 */
 	ENGINE_API void RestartIrisSystem();
+
+	/** Destroy the Iris replication system before it gets recreated. */
+	ENGINE_API void DestroyIrisSystem();
+
+	/** Create the Iris replication system after it got destroyed. */
+	ENGINE_API void RecreateIrisSystem();
 #endif // UE_WITH_IRIS
 
 	template<class T>
diff --git a/Engine/Source/Runtime/Engine/Private/NetDriver.cpp b/Engine/Source/Runtime/Engine/Private/NetDriver.cpp
--- a/Engine/Source/Runtime/Engine/Private/NetDriver.cpp
+++ b/Engine/Source/Runtime/Engine/Private/NetDriver.cpp
@@ -6061,6 +6061,11 @@ void UNetDriver::SetNetDriverDefinition(FName NewNetDriverDefinition)
 	InitPacketSimulationSettings();
 }
 
+void UNetDriver::ReinitBase()
+{
+	UE_LOG(LogNet, Log, TEXT("Re-InitBase %s (NetDriverDefinition %s) using replication model %s"), *NetDriverName.ToString(), *NetDriverDefinition.ToString(), *GetReplicationModelName());
+}
+
 void UNetDriver::PostCreation(bool bInitializeWithIris)
 {
 #if UE_WITH_IRIS
@@ -6072,7 +6077,6 @@ void UNetDriver::PostCreation(bool bInitializeWithIris)
 		//Add to CSV whether we're using Iris on the GameNetDriver or not
 		CSV_METADATA(TEXT("Iris"), IsUsingIrisReplication() ? TEXT("1") : TEXT("0"));
 	}
-
 }
 
 #if NET_DEBUG_RELEVANT_ACTORS
@@ -6895,61 +6899,54 @@ void UNetDriver::SetReplicationSystem(UReplicationSystem* InReplicationSystem)
 	}
 }
 
-void UNetDriver::ClearIrisSystem()
+void UNetDriver::RestartIrisSystem()
 {
-	if (ReplicationSystem)
+	if (ReplicationSystem == nullptr)
 	{
-		UReplicationBridge* Bridge = ReplicationSystem->GetReplicationBridge();
-		if (ensureAlways(Bridge))
-		{
-			Bridge->SetNetDriver(nullptr);
-		}
+		ensureMsgf(false, TEXT("RestartIrisSystem called while no system existed."));
+		return;
 	}
 
-	ReplicationSystem = nullptr;
-}
-
-void UNetDriver::RestoreIrisSystem(UReplicationSystem* InReplicationSystem)
-{
-	check(InReplicationSystem != nullptr);
-	check(InReplicationSystem->GetReplicationBridge() != nullptr);
-	checkf(ReplicationSystem == nullptr, TEXT("Cannot restore IrisSystem in %s since one system is already initialized."), *GetName());
+	if (ClientConnections.Num() > 0)
+	{
+		ensureMsgf(false, TEXT("RestartIrisSystem called while there were active connections."));
+		return;
+	}
 
-	ReplicationSystem = InReplicationSystem;
-	ReplicationSystem->GetReplicationBridge()->SetNetDriver(this);
+	DestroyIrisSystem();
 
-	// When we run using Iris, we use ReplicationSystemId as our unique identifier
-	NetTraceId = ReplicationSystem->GetId();
-	
-	// World Actors have already been registered in this IrisSystem, prevent adding them twice when the World gets set.
-	bSkipBeginReplicationForWorld = true;
+	CreateReplicationSystem(!IsServer());
 }
 
-void UNetDriver::RestartIrisSystem()
+void UNetDriver::DestroyIrisSystem()
 {
 	if (ReplicationSystem == nullptr)
 	{
-		ensureMsgf(false, TEXT("RestartIrisSystem called while no system existed."));
+		ensureMsgf(false, TEXT("DestroyIrisSystem called while no system existed."));
 		return;
 	}
 
 	if (ClientConnections.Num() > 0)
 	{
-		ensureMsgf(false, TEXT("RestartIrisSystem called while there were active connections."));
+		ensureMsgf(false, TEXT("DestroyIrisSystem called while there were active connections."));
 		return;
 	}
 
-	QUICK_SCOPE_CYCLE_COUNTER(STAT_NetDriver_RestartIrisSystem);
+	QUICK_SCOPE_CYCLE_COUNTER(STAT_NetDriver_DestroyIrisSystem);
 
 	UE::Net::FReplicationSystemFactory::DestroyReplicationSystem(ReplicationSystem);
 	ReplicationSystem = nullptr;
+}
 
+void UNetDriver::RecreateIrisSystem()
+{
 	CreateReplicationSystem(!IsServer());
 }
 
 void UNetDriver::CreateReplicationSystem(bool bInitAsClient)
 {
 	LLM_SCOPE_BYTAG(Iris);
+	QUICK_SCOPE_CYCLE_COUNTER(STAT_NetDriver_CreateReplicationSystem);
 
 	const bool bBridgeClassExists = InitReplicationBridgeClass();
 
diff --git a/Engine/Source/Runtime/Engine/Private/UnrealEngine.cpp b/Engine/Source/Runtime/Engine/Private/UnrealEngine.cpp
--- a/Engine/Source/Runtime/Engine/Private/UnrealEngine.cpp
+++ b/Engine/Source/Runtime/Engine/Private/UnrealEngine.cpp
@@ -13800,8 +13800,10 @@ UNetDriver* UEngine::FindNamedNetDriver(const UPendingNetGame* InPendingNetGame,
 	return FindNamedNetDriver_Local(GetWorldContextFromPendingNetGameChecked(InPendingNetGame).ActiveNetDrivers, NetDriverName);
 }
 
-namespace UE::Private
+const FIrisNetDriverConfig* UEngine::GetIrisNetDriverConfig(FName InNetDriverDefinition, FName InNetDriverName) const
 {
+	const FIrisNetDriverConfig* IrisConfig = nullptr;
+
 	/**
 	 * Look for a config setting telling if this new NetDriver can use Iris or not
 	 * To enable iris for a netdriver you need to add the proper configuration in Engine.ini like so:
@@ -13810,119 +13812,122 @@ namespace UE::Private
 	 *		+IrisNetDriverConfigs=(NetDriverDefinition="MyNetDriverDef",bCanUseIris=false)
 	 *		+IrisNetDriverConfigs=(NetDriverWildcardName="SecondaryNetDriver*",bCanUseIris=true);
 	 *
-	 * Priority order for the IrisNetDriverConfigs are:
-	 *		1. NetDriverName exact match
-	 *		2. NetDriverName wildcard match
-	 *		3. NetDriverDefinition match
-	 *
 	 */
-	bool IsNetDriverUsingIris(UEngine* Engine, const FWorldContext& Context, FName InNetDriverDefinition, FName InNetDriverName)
-	{
-#if UE_WITH_IRIS
-		FIrisNetDriverConfig* IrisConfig = nullptr;
 
-		// Search for the exact name match
-		if (IrisConfig == nullptr)
+#if UE_WITH_IRIS
+	// Search for the exact name match
+	if (IrisConfig == nullptr)
+	{
+		IrisConfig = IrisNetDriverConfigs.FindByPredicate([InNetDriverName](const FIrisNetDriverConfig& Config)
 		{
-			IrisConfig = Engine->IrisNetDriverConfigs.FindByPredicate([InNetDriverName](const FIrisNetDriverConfig& Config)
-				{
-					return Config.NetDriverName.IsNone() ? false : Config.NetDriverName == InNetDriverName;
-				});
-		}
+			return Config.NetDriverName.IsNone() ? false : Config.NetDriverName == InNetDriverName;
+		});
+	}
 
-		// Search for a wildcard match
-		if (IrisConfig == nullptr)
+	// Search for a wildcard match
+	if (IrisConfig == nullptr)
+	{
+		const FString TempNetDriverName = InNetDriverName.ToString();
+		IrisConfig = IrisNetDriverConfigs.FindByPredicate([TempNetDriverName](const FIrisNetDriverConfig& Config)
 		{
-			const FString TempNetDriverName = InNetDriverName.ToString();
-			IrisConfig = Engine->IrisNetDriverConfigs.FindByPredicate([TempNetDriverName](const FIrisNetDriverConfig& Config)
-				{
-					return Config.NetDriverWildcardName.IsEmpty() ? false : TempNetDriverName.MatchesWildcard(Config.NetDriverWildcardName);
-				});
-		}
+			return Config.NetDriverWildcardName.IsEmpty() ? false : TempNetDriverName.MatchesWildcard(Config.NetDriverWildcardName);
+		});
+	}
 
-		// Search for the definition match
-		if (IrisConfig == nullptr)
+	// Search for the definition match
+	if (IrisConfig == nullptr)
+	{
+		IrisConfig = IrisNetDriverConfigs.FindByPredicate([InNetDriverDefinition](const FIrisNetDriverConfig& Config)
 		{
-			IrisConfig = Engine->IrisNetDriverConfigs.FindByPredicate([InNetDriverDefinition](const FIrisNetDriverConfig& Config)
-				{
-					return Config.NetDriverDefinition.IsNone() ? false : Config.NetDriverDefinition == InNetDriverDefinition;
-				});
-		}
+			return Config.NetDriverDefinition.IsNone() ? false : Config.NetDriverDefinition == InNetDriverDefinition;
+		});
+	}
+#endif //#if UE_WITH_IRIS
+
+	return IrisConfig;
+}
+
+bool UEngine::WillNetDriverUseIris(const FWorldContext& Context, FName InNetDriverDefinition, FName InNetDriverName) const
+{
+#if UE_WITH_IRIS
+	const FIrisNetDriverConfig* IrisConfig = GetIrisNetDriverConfig(InNetDriverDefinition, InNetDriverName);
 
-		const bool bConfigCanUseIris = IrisConfig && IrisConfig->bCanUseIris;
-		const bool bIsEngineDefaultIris = UE::Net::ShouldUseIrisReplication();
+	const bool bConfigCanUseIris = IrisConfig && IrisConfig->bCanUseIris;
+	const bool bIsEngineDefaultIris = UE::Net::ShouldUseIrisReplication();
 
-		bool bUseIrisRepSystem = bConfigCanUseIris && bIsEngineDefaultIris;
+	bool bUseIrisRepSystem = bConfigCanUseIris && bIsEngineDefaultIris;
 
-		if (UGameInstance* ContextGameInstance = Context.OwningGameInstance)
-		{
-			EReplicationSystem GameInstanceDesiredRepSystem = ContextGameInstance->GetDesiredReplicationSystem(InNetDriverDefinition);
+	if (UGameInstance* ContextGameInstance = Context.OwningGameInstance)
+	{
+		EReplicationSystem GameInstanceDesiredRepSystem = ContextGameInstance->GetDesiredReplicationSystem(InNetDriverDefinition);
 
-			// If the game instance requested to use the generic repsystem
-			if (GameInstanceDesiredRepSystem == EReplicationSystem::Generic)
-			{
-				UE_CLOG(bUseIrisRepSystem, LogNet, Log, TEXT("GameInstance %s is forcing NetDriver %s (NetDefinition %s) to use the Generic replication system."), *GetNameSafe(ContextGameInstance), *InNetDriverName.ToString(), *InNetDriverDefinition.ToString());
-				bUseIrisRepSystem = false;
-			}
-			// If the game instance requested to use the Iris repsystem
-			else if (GameInstanceDesiredRepSystem == EReplicationSystem::Iris)
-			{
-				UE_CLOG(!bUseIrisRepSystem && bConfigCanUseIris, LogNet, Log, TEXT("GameInstance %s is forcing NetDriver %s (NetDefinition %s) to use the Iris replication system."), *GetNameSafe(ContextGameInstance), *InNetDriverName.ToString(), *InNetDriverDefinition.ToString());
+		// If the game instance requested to use the generic repsystem
+		if (GameInstanceDesiredRepSystem == EReplicationSystem::Generic)
+		{
+			UE_CLOG(bUseIrisRepSystem, LogNet, Log, TEXT("GameInstance %s is forcing NetDriver %s (NetDefinition %s) to use the Generic replication system."), *GetNameSafe(ContextGameInstance), *InNetDriverName.ToString(), *InNetDriverDefinition.ToString());
+			bUseIrisRepSystem = false;
+		}
+		// If the game instance requested to use the Iris repsystem
+		else if (GameInstanceDesiredRepSystem == EReplicationSystem::Iris)
+		{
+			UE_CLOG(!bUseIrisRepSystem && bConfigCanUseIris, LogNet, Log, TEXT("GameInstance %s is forcing NetDriver %s (NetDefinition %s) to use the Iris replication system."), *GetNameSafe(ContextGameInstance), *InNetDriverName.ToString(), *InNetDriverDefinition.ToString());
 				
-				// Enable Iris ONLY if the config supports it.
-				bUseIrisRepSystem = bConfigCanUseIris;
-			}
+			// Enable Iris ONLY if the config supports it.
+			bUseIrisRepSystem = bConfigCanUseIris;
+		}
 
 #if WITH_EDITOR
-			// In PIE let's make sure the clients follow what the server's game net driver is using
-			if (Context.WorldType == EWorldType::PIE && !Context.RunAsDedicated && InNetDriverDefinition == NAME_GameNetDriver)
+		// In PIE let's make sure the clients follow what the server's game net driver is using
+		if (Context.WorldType == EWorldType::PIE && !Context.RunAsDedicated && InNetDriverDefinition == NAME_GameNetDriver)
+		{
+			if (FWorldContext* ServerPIEContext = GEngine->GetWorldContextFromPIEInstance(0))
 			{
-				if (FWorldContext* ServerPIEContext = GEngine->GetWorldContextFromPIEInstance(0))
+				if (ServerPIEContext->RunAsDedicated)
 				{
-					if (ServerPIEContext->RunAsDedicated)
+					for (const FNamedNetDriver& PieNetDriver : ServerPIEContext->ActiveNetDrivers)
 					{
-						for (const FNamedNetDriver& PieNetDriver : ServerPIEContext->ActiveNetDrivers)
+						if (PieNetDriver.NetDriverDef->DefName == NAME_GameNetDriver)
 						{
-							if (PieNetDriver.NetDriverDef->DefName == NAME_GameNetDriver)
-							{
-								bUseIrisRepSystem = PieNetDriver.NetDriver->IsUsingIrisReplication();
-								break;
-							}
+							bUseIrisRepSystem = PieNetDriver.NetDriver->IsUsingIrisReplication();
+							break;
 						}
 					}
 				}
 			}
-#endif
 		}
+#endif
+	}
 
-		// Ignore all of the above if the cmdline is requesting a specific system
-		const EReplicationSystem CmdlineRequest = UE::Net::GetUseIrisReplicationCmdlineValue();
-		if (CmdlineRequest == EReplicationSystem::Iris)
-		{
-			UE_CLOG(!bUseIrisRepSystem && bConfigCanUseIris, LogNet, Log, TEXT("Cmdline -UseIrisReplication=1 is forcing NetDriver %s (NetDefinition %s) to use the Iris replication system."), *InNetDriverName.ToString(), *InNetDriverDefinition.ToString());
-			bUseIrisRepSystem = bConfigCanUseIris;
-		}
-		else if (CmdlineRequest == EReplicationSystem::Generic)
-		{
-			UE_CLOG(bUseIrisRepSystem, LogNet, Log, TEXT("Cmdline -UseIrisReplication=0 is forcing NetDriver %s (NetDefinition %s) to use the Iris replication system."), *InNetDriverName.ToString(), *InNetDriverDefinition.ToString());
-			bUseIrisRepSystem = false;
-		}
+	// Ignore all of the above if the cmdline is requesting a specific system
+	const EReplicationSystem CmdlineRequest = UE::Net::GetUseIrisReplicationCmdlineValue();
+	if (CmdlineRequest == EReplicationSystem::Iris)
+	{
+		UE_CLOG(!bUseIrisRepSystem && bConfigCanUseIris, LogNet, Log, TEXT("Cmdline -UseIrisReplication=1 is forcing NetDriver %s (NetDefinition %s) to use the Iris replication system."), *InNetDriverName.ToString(), *InNetDriverDefinition.ToString());
+		bUseIrisRepSystem = bConfigCanUseIris;
+	}
+	else if (CmdlineRequest == EReplicationSystem::Generic)
+	{
+		UE_CLOG(bUseIrisRepSystem, LogNet, Log, TEXT("Cmdline -UseIrisReplication=0 is forcing NetDriver %s (NetDefinition %s) to use the Iris replication system."), *InNetDriverName.ToString(), *InNetDriverDefinition.ToString());
+		bUseIrisRepSystem = false;
+	}
 
-		// Only use Iris if the module is loaded (this happens automatically if the Iris plugin is enabled)
-		if (bUseIrisRepSystem)
+	// Only use Iris if the module is loaded (this happens automatically if the Iris plugin is enabled)
+	if (bUseIrisRepSystem)
+	{
+		if (!ensureMsgf(FModuleManager::Get().IsModuleLoaded("IrisCore"), TEXT("%s is not using Iris because the IrisCore module isn't loaded. Check whether the Iris plugin is enabled."), *InNetDriverName.ToString()))
 		{
-			if (!ensureMsgf(FModuleManager::Get().IsModuleLoaded("IrisCore"), TEXT("%s is not using Iris because the IrisCore module isn't loaded. Check whether the Iris plugin is enabled."), *InNetDriverName.ToString()))
-			{
-				return false;
-			}
+			return false;
 		}
+	}
 
-		return bUseIrisRepSystem;
+	return bUseIrisRepSystem;
 #else
-		return false;
+	return false;
 #endif //UE_WITH_IRIS
-	}
+}
 
+namespace UE::Private
+{
 	UNetDriver* CreateNetDriver_Local(UEngine* Engine, FWorldContext& Context, FName NetDriverDefinition, FName InNetDriverName)
 	{
 		UNetDriver* ReturnVal = nullptr;
@@ -14057,7 +14062,7 @@ namespace UE::Private
 				check(ReturnVal != nullptr);
 
 				const FName DriverName = InNetDriverName.IsNone() ? ReturnVal->GetFName() : InNetDriverName;
-				const bool bInitializeWithIris = IsNetDriverUsingIris(Engine, Context, NetDriverDefinition, DriverName);
+				const bool bInitializeWithIris = Engine->WillNetDriverUseIris(Context, NetDriverDefinition, DriverName);
 
 				ReturnVal->SetNetDriverName(DriverName);
 				ReturnVal->SetNetDriverDefinition(NetDriverDefinition);
