Return the latest time for components with scene proxies that are marked as always visible rather than updating the component time for each one. Saves multiple ms of CPU time in CitySample.
 - Also deprecated unused primitive RegistrationSerialNumber member.

diff --git a/Engine/Source/Runtime/Engine/Classes/Components/PrimitiveComponent.h b/Engine/Source/Runtime/Engine/Classes/Components/PrimitiveComponent.h
--- a/Engine/Source/Runtime/Engine/Classes/Components/PrimitiveComponent.h
+++ b/Engine/Source/Runtime/Engine/Classes/Components/PrimitiveComponent.h
@@ -900,8 +900,9 @@ public:
 	ENGINE_API bool WasRecentlyRendered(float Tolerance = 0.2f) const;
 
 	ENGINE_API void SetLastRenderTime(float InLastRenderTime);
-	float GetLastRenderTime() const { return SceneData.LastRenderTime; }
-	float GetLastRenderTimeOnScreen() const { return SceneData.LastRenderTimeOnScreen; }
+
+	ENGINE_API float GetLastRenderTime() const;
+	ENGINE_API float GetLastRenderTimeOnScreen() const;
 
 #if MESH_DRAW_COMMAND_STATS
 	ENGINE_API void SetMeshDrawCommandStatsCategory(FName StatsCategory);
@@ -1956,7 +1957,12 @@ public:
 	FPrimitiveSceneProxy* SceneProxy;
 
 	FPrimitiveSceneProxy* GetSceneProxy() const { check(SceneProxy == SceneData.SceneProxy); return SceneData.SceneProxy; }
-	void ReleaseSceneProxy() { check(SceneProxy == SceneData.SceneProxy); SceneProxy = nullptr;  SceneData.SceneProxy = nullptr; }
+
+	ENGINE_API void AssignSceneProxy(FPrimitiveSceneProxy* SceneProxy);
+	ENGINE_API void ReleaseSceneProxy();
+
+	/** Returns whether the primitive component's scene proxy is always visible. */
+	bool IsAlwaysVisible() const { return SceneData.bAlwaysVisible; }
 
 	/** A fence to track when the primitive is detached from the scene in the rendering thread. */
 	FRenderCommandFence DetachFence;
diff --git a/Engine/Source/Runtime/Engine/Classes/GameFramework/Actor.h b/Engine/Source/Runtime/Engine/Classes/GameFramework/Actor.h
--- a/Engine/Source/Runtime/Engine/Classes/GameFramework/Actor.h
+++ b/Engine/Source/Runtime/Engine/Classes/GameFramework/Actor.h
@@ -131,6 +131,20 @@ enum class UE_DEPRECATED(5.0, "EActorGridPlacement is deprecated.") EActorGridPl
 };
 #endif
 
+/** Helper struct that allows UPrimitiveComponent and FPrimitiveSceneInfo write to the Actor's LastRenderTime member */
+struct FActorLastRenderTime
+{
+	float LastRenderTime;
+	std::atomic_int32_t NumAlwaysVisibleComponents = 0;
+
+private:
+	static void Set(AActor* InActor, float LastRenderTime);
+	static FActorLastRenderTime* GetPtr(AActor* InActor);
+
+	friend class UPrimitiveComponent;
+	friend struct FPrimitiveSceneInfoAdapter;
+};
+
 ENGINE_API DECLARE_LOG_CATEGORY_EXTERN(LogActor, Log, Warning);
 
 // Delegate signatures
@@ -857,7 +871,7 @@ private:
 	 * from the render thread, which is up to a frame behind the game thread, so you should allow this time to
 	 * be at least a frame behind the game thread's world time before you consider the actor non-visible.
 	 */
-	float LastRenderTime;
+	FActorLastRenderTime LastRenderTime;
 
 	friend struct FActorLastRenderTime;
 
@@ -3695,6 +3709,7 @@ public:
 	}
 
 private:
+
 	/**
 	 * Internal helper function to call a compile-time lambda on all components of a given type
 	 * Use template parameter bClassIsActorComponent to avoid doing unnecessary IsA checks when the ComponentClass is exactly UActorComponent
@@ -4488,23 +4503,16 @@ private:
 	friend UWorld;
 };
 
-/** Helper struct that allows UPrimitiveComponent and FPrimitiveSceneInfo write to the Actor's LastRenderTime member */
-struct FActorLastRenderTime
-{
-private:
-	static void Set(AActor* InActor, float LastRenderTime)
-	{
-		InActor->LastRenderTime = LastRenderTime;
-	}
 
-	static float* GetPtr(AActor* InActor)
-	{
-		return (InActor ? &InActor->LastRenderTime : nullptr);
-	}
+inline void FActorLastRenderTime::Set(AActor* InActor, float LastRenderTime)
+{
+	InActor->LastRenderTime.LastRenderTime = LastRenderTime;
+}
 
-	friend class UPrimitiveComponent;
-	friend struct FPrimitiveSceneInfoAdapter;
-};
+inline FActorLastRenderTime* FActorLastRenderTime::GetPtr(AActor* InActor)
+{
+	return (InActor ? &InActor->LastRenderTime : nullptr);
+}
 
 #if WITH_EDITOR
 struct FSetActorHiddenInSceneOutliner
diff --git a/Engine/Source/Runtime/Engine/Private/Actor.cpp b/Engine/Source/Runtime/Engine/Private/Actor.cpp
--- a/Engine/Source/Runtime/Engine/Private/Actor.cpp
+++ b/Engine/Source/Runtime/Engine/Private/Actor.cpp
@@ -2094,7 +2094,14 @@ bool AActor::WasRecentlyRendered(float Tolerance) const
 
 float AActor::GetLastRenderTime() const
 {
-	return LastRenderTime;
+	if (LastRenderTime.NumAlwaysVisibleComponents.load(std::memory_order_relaxed) > 0)
+	{
+		if (const UWorld* World = GetWorld())
+		{
+			return World->GetTimeSeconds();
+		}
+	}
+	return LastRenderTime.LastRenderTime;
 }
 
 void AActor::SetOwner(AActor* NewOwner)
@@ -6549,5 +6556,4 @@ void AActor::ForEachComponentOfActorClassDefault(const TSubclassOf<AActor>& Acto
 	}
 }
 
-
 #undef LOCTEXT_NAMESPACE
diff --git a/Engine/Source/Runtime/Engine/Private/Components/PrimitiveComponent.cpp b/Engine/Source/Runtime/Engine/Private/Components/PrimitiveComponent.cpp
--- a/Engine/Source/Runtime/Engine/Private/Components/PrimitiveComponent.cpp
+++ b/Engine/Source/Runtime/Engine/Private/Components/PrimitiveComponent.cpp
@@ -622,10 +622,6 @@ void UPrimitiveComponent::CreateRenderState_Concurrent(FRegisterComponentContext
 		CachedMaxDrawDistance = bNeverCull ? 0.f : LDMaxDrawDistance;
 	}
 
-	// Always setup our ptr to the OwnerLastRenderTimer for rendering time feedback from the renderer
-	// The owner can change after calls to OnRegister so we must resynchronize this value
-	SceneData.OwnerLastRenderTimePtr = FActorLastRenderTime::GetPtr(GetOwner());
-
 	Super::CreateRenderState_Concurrent(Context);
 
 	UpdateBounds();
@@ -639,6 +635,8 @@ void UPrimitiveComponent::CreateRenderState_Concurrent(FRegisterComponentContext
 #endif
 	)
 	{
+		SceneData.OwnerLastRenderTimePtr = FActorLastRenderTime::GetPtr(GetOwner());
+
 		if (Context != nullptr)
 		{
 			Context->AddPrimitive(this);
@@ -669,15 +667,6 @@ void UPrimitiveComponent::SendRenderTransform_Concurrent()
 
 void UPrimitiveComponent::OnRegister()
 {
-	// Both those are initalized before call Super::OnRegister since the primitive can be added to the scene
-	// before this method completes, for example through FNiagaraSystem::PollForCompilationComplete()
-	 
-	// Setup our ptr to the OwnerLastRenderTimer for rendering time feedback from the renderer
-	SceneData.OwnerLastRenderTimePtr = FActorLastRenderTime::GetPtr(GetOwner());
-	
-	// Deterministically track primitives via registration sequence numbers.
- 	SceneData.RegistrationSerialNumber = FPrimitiveSceneInfoData::GetNextRegistrationSerialNumber(); 
-
 	Super::OnRegister();
 	
 	if (bCanEverAffectNavigation)
@@ -708,8 +697,6 @@ void UPrimitiveComponent::OnRegister()
 
 void UPrimitiveComponent::OnUnregister()
 {
-	SceneData.OwnerLastRenderTimePtr = nullptr;
-
 	// If this is being garbage collected we don't really need to worry about clearing this
 	if (!HasAnyFlags(RF_BeginDestroyed) && !IsUnreachable())
 	{
@@ -825,6 +812,8 @@ void UPrimitiveComponent::DestroyRenderState_Concurrent()
 		World->Scene->RemovePrimitive(this);
 	}
 
+	SceneData.OwnerLastRenderTimePtr = nullptr;
+
 	Super::DestroyRenderState_Concurrent();
 }
 
@@ -4694,6 +4683,24 @@ void UPrimitiveComponent::SetLastRenderTime(float InLastRenderTime)
 	}
 }
 
+float UPrimitiveComponent::GetLastRenderTime() const
+{
+	if (IsAlwaysVisible())
+	{
+		return GetWorld()->GetTimeSeconds();
+	}
+	return SceneData.LastRenderTime;
+}
+
+float UPrimitiveComponent::GetLastRenderTimeOnScreen() const
+{
+	if (IsAlwaysVisible())
+	{
+		return GetWorld()->GetTimeSeconds();
+	}
+	return SceneData.LastRenderTimeOnScreen;
+}
+
 #if MESH_DRAW_COMMAND_STATS
 void UPrimitiveComponent::SetMeshDrawCommandStatsCategory(FName StatsCategory)
 {
@@ -4955,6 +4962,34 @@ void UPrimitiveComponent::GetPrimitiveStats(FPrimitiveStats& PrimitiveStats) con
 	// no default values returned
 }
 
+void UPrimitiveComponent::AssignSceneProxy(FPrimitiveSceneProxy* InSceneProxy)
+{
+	check(SceneProxy == nullptr && SceneData.SceneProxy == nullptr);
+	SceneProxy = InSceneProxy;
+	SceneData.SceneProxy = InSceneProxy;
+	if (SceneProxy)
+	{
+		SceneData.bAlwaysVisible = SceneProxy->IsAlwaysVisible();
+		if (SceneData.bAlwaysVisible && SceneData.OwnerLastRenderTimePtr)
+		{
+			SceneData.OwnerLastRenderTimePtr->NumAlwaysVisibleComponents.fetch_add(1, std::memory_order_relaxed);
+		}
+	}
+}
+
+void UPrimitiveComponent::ReleaseSceneProxy()
+{
+	check(SceneProxy == SceneData.SceneProxy);
+	if (SceneData.bAlwaysVisible && SceneData.OwnerLastRenderTimePtr)
+	{
+		const uint32 NumRefs = SceneData.OwnerLastRenderTimePtr->NumAlwaysVisibleComponents.fetch_sub(1, std::memory_order_relaxed);
+		check(NumRefs > 0);
+	}
+	SceneProxy = nullptr;
+	SceneData.SceneProxy = nullptr;
+	SceneData.bAlwaysVisible = false;
+}
+
 bool FActorPrimitiveComponentInterface::IsRenderStateCreated() const 
 {
 	return UPrimitiveComponent::GetPrimitiveComponent(this)->IsRenderStateCreated();
@@ -5075,11 +5110,9 @@ FString FActorPrimitiveComponentInterface::GetOwnerName() const
 FPrimitiveSceneProxy* FActorPrimitiveComponentInterface::CreateSceneProxy() 
 {
 	UPrimitiveComponent* Component = UPrimitiveComponent::GetPrimitiveComponent(this);
-	check(Component->SceneProxy == nullptr && Component->SceneData.SceneProxy == nullptr);
-	FPrimitiveSceneProxy* Proxy = Component->CreateSceneProxy();
-	Component->SceneData.SceneProxy = Proxy;
-	Component->SceneProxy = Proxy;
-	return Proxy;
+	FPrimitiveSceneProxy* SceneProxy = Component->CreateSceneProxy();
+	Component->AssignSceneProxy(SceneProxy);
+	return SceneProxy;
 }
 
 #if WITH_EDITOR
diff --git a/Engine/Source/Runtime/Engine/Private/PrimitiveSceneInfoData.cpp b/Engine/Source/Runtime/Engine/Private/PrimitiveSceneInfoData.cpp
--- a/Engine/Source/Runtime/Engine/Private/PrimitiveSceneInfoData.cpp
+++ b/Engine/Source/Runtime/Engine/Private/PrimitiveSceneInfoData.cpp
@@ -1,12 +1,23 @@
 // Copyright Epic Games, Inc. All Rights Reserved.
 
 #include "PrimitiveSceneInfoData.h"
-
+#include "GameFramework/Actor.h"
 
 /** Next id to be used by a component. */
 // 0 is reserved to mean invalid
 FThreadSafeCounter FPrimitiveSceneInfoData::NextPrimitiveId;
 
-/** Next registration serial number to be assigned to a component when it is registered. */
-// -1 is reserved to mean invalid
-FThreadSafeCounter FPrimitiveSceneInfoData::NextRegistrationSerialNumber;
\ No newline at end of file
+void FPrimitiveSceneInfoData::SetLastRenderTime(float InLastRenderTime, bool bUpdateLastRenderTimeOnScreen) const
+{
+	LastRenderTime = InLastRenderTime;
+
+	if (bUpdateLastRenderTimeOnScreen)
+	{
+		LastRenderTimeOnScreen = InLastRenderTime;
+	}
+
+	if (OwnerLastRenderTimePtr)
+	{
+		OwnerLastRenderTimePtr->LastRenderTime = InLastRenderTime;
+	}
+}
\ No newline at end of file
diff --git a/Engine/Source/Runtime/Engine/Public/PrimitiveSceneDesc.h b/Engine/Source/Runtime/Engine/Public/PrimitiveSceneDesc.h
--- a/Engine/Source/Runtime/Engine/Public/PrimitiveSceneDesc.h
+++ b/Engine/Source/Runtime/Engine/Public/PrimitiveSceneDesc.h
@@ -23,39 +23,38 @@ class UPrimtiiveComponent;
 */
 
 struct FPrimitiveSceneDesc 
-{			
-	FPrimitiveSceneProxyDesc* ProxyDesc = nullptr;			
+{
+	FPrimitiveSceneProxyDesc* ProxyDesc = nullptr;
 	IPrimitiveComponent* PrimitiveComponentInterface = nullptr;
 	FPrimitiveSceneInfoData* PrimitiveSceneData = nullptr;
 	
 	FPrimitiveComponentId LightingAttachmentComponentId;
 	FPrimitiveComponentId LodParentComponentId;
-	
+
 	bool bShouldAddtoScene = true; // for UpdatePrimitiveAttachment
 	bool bRecreateProxyOnUpdateTransform = false; 
-	bool bIsUnreachable = false;	
+	bool bIsUnreachable = false;
 	bool bBulkReregister = false;
-		
-	EComponentMobility::Type Mobility = EComponentMobility::Movable;	
+
+	EComponentMobility::Type Mobility = EComponentMobility::Movable;
 	FBoxSphereBounds Bounds;
 	FBoxSphereBounds LocalBounds;
 
 	FMatrix RenderMatrix;
 	FVector AttachmentRootPosition;
-	
+
 	UWorld*  World = nullptr;
 	UObject* PrimitiveUObject = nullptr;
 	
 	// @todo: Possibly add uninitialized FStrings in this object to allow overriding the name without having a corresponding UObject
 	FString GetFullName() { return PrimitiveUObject->GetFullName(); }
-	FString GetName() { return PrimitiveUObject->GetName(); }	
+	FString GetName() { return PrimitiveUObject->GetName(); }
 
 	bool IsUnreachable() { return bIsUnreachable; }
 	bool ShouldRecreateProxyOnUpdateTransform() { return bRecreateProxyOnUpdateTransform; }
 
 	FThreadSafeCounter& GetAttachmentCounter() { return PrimitiveSceneData->AttachmentCounter; }
 	FPrimitiveComponentId GetPrimitiveSceneId() const { return PrimitiveSceneData->PrimitiveSceneId; }
-	int32 GetRegistrationSerialNumber() const { return PrimitiveSceneData->RegistrationSerialNumber; }
 	FPrimitiveComponentId GetLODParentId() const { return LodParentComponentId; }
 	FPrimitiveComponentId GetLightingAttachmentId() const { return LightingAttachmentComponentId; }
 
@@ -83,6 +82,9 @@ struct FPrimitiveSceneDesc
 	const IPrimitiveComponent* GetPrimitiveComponentInterface() const { return PrimitiveComponentInterface; }
 
 	UPackage* GetOutermost() const { return PrimitiveUObject->GetOutermost(); }
+
+	UE_DEPRECATED(5.5, "GetRegistrationSerialNumber is no longer used.")
+	int32 GetRegistrationSerialNumber() const { return -1; }
 };
 
 struct FInstancedStaticMeshSceneDesc
diff --git a/Engine/Source/Runtime/Engine/Public/PrimitiveSceneInfoData.h b/Engine/Source/Runtime/Engine/Public/PrimitiveSceneInfoData.h
--- a/Engine/Source/Runtime/Engine/Public/PrimitiveSceneInfoData.h
+++ b/Engine/Source/Runtime/Engine/Public/PrimitiveSceneInfoData.h
@@ -6,18 +6,19 @@
 #include "PrimitiveComponentId.h"
 
 class FPrimitiveSceneProxy;
+struct FActorLastRenderTime;
 
 /*
  * All the necessary information for scene primitive to component feedback 
  */
 struct FPrimitiveSceneInfoData
-{		
+{
 	/** The primitive's scene info. */
-	FPrimitiveSceneProxy* SceneProxy = nullptr;	
+	FPrimitiveSceneProxy* SceneProxy = nullptr;
 
 	/** Last time the component was submitted for rendering (called FScene::AddPrimitive). */
 	float LastSubmitTime = -1000;
-	
+
 	/**
 	 * The value of WorldSettings->TimeSeconds for the frame when this component was last rendered.  This is written
 	 * from the render thread, which is up to a frame behind the game thread, so you should allow this time to
@@ -27,7 +28,6 @@ struct FPrimitiveSceneInfoData
 
 	/** Same as LastRenderTime but only updated if the component is on screen. Used by the texture streamer. */
 	mutable float LastRenderTimeOnScreen = -1000.0f;
-	
 
 	/**
 	* Incremented by the main thread before being attached to the scene, decremented
@@ -54,13 +54,13 @@ struct FPrimitiveSceneInfoData
 	/** Used by the renderer, to identify a primitive across re-registers. */
 	FPrimitiveComponentId PrimitiveSceneId;
 
-	/**
-	* Identifier used to track the time that this component was registered with the world / renderer.
-	* Updated to unique incremental value each time OnRegister() is called. The value of 0 is unused.
-	* */
-	int32 RegistrationSerialNumber = -1;
+	UE_DEPRECATED(5.5, "RegistrationSerialNumber is no longer used")
+	int32 RegistrationSerialNumber : 31;
+
+	/** Whether the primitive is always visible. If true the last render time will be unset. */
+	int32 bAlwaysVisible : 1;
 
-	/** 
+	/**
 	 * Pointer to the last render time variable on the primitive's owning actor or other UObject (if owned), which is written to by the RT and read by the GT.
 	 * The value of LastRenderTime will therefore not be deterministic due to race conditions, but the GT uses it in a way that allows this.
 	 * Storing a pointer to the UObject member variable only works in the AActor/UPrimitiveComponent case because:
@@ -69,43 +69,30 @@ struct FPrimitiveSceneInfoData
 	 *	If the UPrimitiveComponent and the Actor are GC'd together, neither will be deleted until FinishDestroy has been executed on both.
 	 *	UPrimitiveComponent's FinishDestroy will not execute until the primitive has been detached from the Scene through it's DetachFence.
 	 * In general feedback from the renderer to the game thread like this should be avoided.
-	 * 
+	 *
 	 * Any other user of this struct that intends to add it's own primitives in the Scene must provide the same guarantees. 
-	 * 
+	 *
 	 */
-	float* OwnerLastRenderTimePtr = nullptr;
+	FActorLastRenderTime* OwnerLastRenderTimePtr = nullptr;
 
-
-	void SetLastRenderTime(float InLastRenderTime, bool bUpdateLastRenderTimeOnScreen) const
-	{
-		LastRenderTime = InLastRenderTime;
-		if (bUpdateLastRenderTimeOnScreen)
-			LastRenderTimeOnScreen = InLastRenderTime;
-
-		if (OwnerLastRenderTimePtr)
-		{
-			*OwnerLastRenderTimePtr = InLastRenderTime;
-		}
-	}	
+	ENGINE_API void SetLastRenderTime(float InLastRenderTime, bool bUpdateLastRenderTimeOnScreen) const;
 
 protected:
 
 	/** Next id to be used by a component. */
 	static ENGINE_API FThreadSafeCounter NextPrimitiveId;
 
-	/** Next registration serial number to be assigned to a component when it is registered. */
-	static ENGINE_API FThreadSafeCounter NextRegistrationSerialNumber;
-
 public:
 
-	static int32 GetNextRegistrationSerialNumber() { return NextRegistrationSerialNumber.Increment(); }
+	UE_DEPRECATED(5.5, "GetNextRegistrationSerialNumber is no longer used")
+	static int32 GetNextRegistrationSerialNumber() { return -1; }
 
 	FPrimitiveSceneInfoData()
+		: RegistrationSerialNumber(-1)
+		, bAlwaysVisible(false)
 	{
-		RegistrationSerialNumber = GetNextRegistrationSerialNumber();
 		PrimitiveSceneId.PrimIDValue = NextPrimitiveId.Increment();
 	}
-
 };
 
 #if UE_ENABLE_INCLUDE_ORDER_DEPRECATED_IN_5_5
diff --git a/Engine/Source/Runtime/Renderer/Private/PrimitiveSceneInfo.cpp b/Engine/Source/Runtime/Renderer/Private/PrimitiveSceneInfo.cpp
--- a/Engine/Source/Runtime/Renderer/Private/PrimitiveSceneInfo.cpp
+++ b/Engine/Source/Runtime/Renderer/Private/PrimitiveSceneInfo.cpp
@@ -195,30 +195,28 @@ struct FPrimitiveSceneInfoAdapter
 		{
 			// Support for legacy path for proxy creation, if not handled it'll internally invoke the IPrimitiveComponentInterface path
 			if (UPrimitiveComponent* PrimitiveComponent =  PrimitiveComponentInterface->GetUObject<UPrimitiveComponent>())
-			{		
-				DefaultHitProxy = SceneProxy->CreateHitProxies(PrimitiveComponent, HitProxies);			
+			{
+				DefaultHitProxy = SceneProxy->CreateHitProxies(PrimitiveComponent, HitProxies);
 			}
 			else 
 			{
 				// For all other implementers
-				DefaultHitProxy = SceneProxy->CreateHitProxies(PrimitiveComponentInterface, HitProxies);			
+				DefaultHitProxy = SceneProxy->CreateHitProxies(PrimitiveComponentInterface, HitProxies);
 			}
 		}
 	}
 
-	FPrimitiveSceneInfoAdapter(UPrimitiveComponent* InComponent)		
-	{		
+	FPrimitiveSceneInfoAdapter(UPrimitiveComponent* InComponent)
+	{
 		SceneProxy = InComponent->SceneProxy;
 		SceneData = &InComponent->SceneData;
-		ComponentId = SceneData->PrimitiveSceneId;			
-		check(InComponent->GetSceneData().RegistrationSerialNumber != -1);
-		RegistrationSerialNumber = InComponent->GetSceneData().RegistrationSerialNumber;
+		ComponentId = SceneData->PrimitiveSceneId;
 		Component = InComponent;
 		PrimitiveComponentInterface = InComponent->GetPrimitiveComponentInterface();
-		PrimitiveDesc = nullptr;		
+		PrimitiveDesc = nullptr;
 		
 		// This validates the UPrimitiveComponent has properly initialized its OwnerLastRenderTimePtr
-		check(InComponent->SceneData.OwnerLastRenderTimePtr == FActorLastRenderTime::GetPtr(InComponent->GetOwner()));		
+		check(InComponent->SceneData.OwnerLastRenderTimePtr == FActorLastRenderTime::GetPtr(InComponent->GetOwner()));
 		Mobility = InComponent->Mobility;
 
 		const UPrimitiveComponent* SearchParentComponent = InComponent->GetLightingAttachmentRoot();
@@ -233,18 +231,17 @@ struct FPrimitiveSceneInfoAdapter
 		if (LODParent)
 		{
 			LODParentComponentId = LODParent->GetPrimitiveSceneId();
-		}		
+		}
 
 		if (GIsEditor)
 		{
 			CreateHitProxies();
 		}
-		
 	}
 	
 	FPrimitiveSceneInfoAdapter(FPrimitiveSceneDesc* InPrimitiveSceneDesc)
-	{		
-		check(InPrimitiveSceneDesc);		
+	{
+		check(InPrimitiveSceneDesc);
 
 		Component = nullptr;
 		PrimitiveComponentInterface = InPrimitiveSceneDesc->GetPrimitiveComponentInterface();
@@ -253,27 +250,24 @@ struct FPrimitiveSceneInfoAdapter
 		SceneProxy = InPrimitiveSceneDesc->GetSceneProxy();
 		check(SceneProxy);
 		ComponentId = InPrimitiveSceneDesc->GetPrimitiveSceneId();
-		RegistrationSerialNumber = InPrimitiveSceneDesc->GetRegistrationSerialNumber();
 		LODParentComponentId = InPrimitiveSceneDesc->GetLODParentId();
-		LightingAttachmentComponentId = InPrimitiveSceneDesc->GetLightingAttachmentId();			 		
-		Mobility = InPrimitiveSceneDesc->GetMobility();				
+		LightingAttachmentComponentId = InPrimitiveSceneDesc->GetLightingAttachmentId();
+		Mobility = InPrimitiveSceneDesc->GetMobility();
 		
 		if (GIsEditor && PrimitiveComponentInterface)
 		{
 			CreateHitProxies();
 		}
-
 	}
 	
 	FPrimitiveSceneProxy* SceneProxy;
 	FPrimitiveComponentId ComponentId;
-	int32 RegistrationSerialNumber;
 	FPrimitiveComponentId LODParentComponentId;
-	FPrimitiveComponentId LightingAttachmentComponentId;	
+	FPrimitiveComponentId LightingAttachmentComponentId;
 	EComponentMobility::Type Mobility;
 
 	// mutable so that hit proxies can be moved to final destination
-	mutable TArray<TRefCountPtr<HHitProxy> > HitProxies;
+	mutable TArray<TRefCountPtr<HHitProxy>> HitProxies;
 	HHitProxy* DefaultHitProxy = nullptr;
 
 	FPrimitiveSceneInfoData* SceneData;
@@ -285,7 +279,6 @@ struct FPrimitiveSceneInfoAdapter
 FPrimitiveSceneInfo::FPrimitiveSceneInfo(const FPrimitiveSceneInfoAdapter& InAdapter, FScene* InScene):
 	Proxy(InAdapter.SceneProxy),
 	PrimitiveComponentId(InAdapter.ComponentId),
-	RegistrationSerialNumber(InAdapter.RegistrationSerialNumber),
 	IndirectLightingCacheAllocation(NULL),
 	CachedPlanarReflectionProxy(NULL),
 	CachedReflectionCaptureProxy(NULL),
diff --git a/Engine/Source/Runtime/Renderer/Private/SceneVisibility.cpp b/Engine/Source/Runtime/Renderer/Private/SceneVisibility.cpp
--- a/Engine/Source/Runtime/Renderer/Private/SceneVisibility.cpp
+++ b/Engine/Source/Runtime/Renderer/Private/SceneVisibility.cpp
@@ -667,12 +667,6 @@ static void UpdateAlwaysVisible(const FScene& Scene, FViewInfo& View, FFrustumCu
 				RayTracingBits |= Mask;
 			}
 		#endif
-
-			FPrimitiveSceneInfo* PrimitiveSceneInfo = Scene.Primitives[Index];
-			PrimitiveSceneInfo->LastRenderTime = CurrentWorldTime;
-
-			const bool bUpdateLastRenderTimeOnScreen = true;
-			PrimitiveSceneInfo->UpdateComponentLastRenderTime(CurrentWorldTime, bUpdateLastRenderTimeOnScreen);
 		}
 
 		VisWords[StartWord + WordIndex] = VisBits;
diff --git a/Engine/Source/Runtime/Renderer/Public/PrimitiveSceneInfo.h b/Engine/Source/Runtime/Renderer/Public/PrimitiveSceneInfo.h
--- a/Engine/Source/Runtime/Renderer/Public/PrimitiveSceneInfo.h
+++ b/Engine/Source/Runtime/Renderer/Public/PrimitiveSceneInfo.h
@@ -294,7 +294,8 @@ public:
 	 * Number assigned to this component when it was registered with the world.
 	 * This will only ever be updated if the object is re-registered.
 	 */
-	int32 RegistrationSerialNumber;
+	UE_DEPRECATED(5.5, "RegistrationSerialNumber is no longer used")
+	int32 RegistrationSerialNumber = -1;
 
 	/** 
 	 * The root attachment component id for use with lighting, if valid.
