print IBuild messages for parent builds.

diff --git a/Engine/Source/Runtime/Engine/Private/TextureDerivedDataTask.cpp b/Engine/Source/Runtime/Engine/Private/TextureDerivedDataTask.cpp
--- a/Engine/Source/Runtime/Engine/Private/TextureDerivedDataTask.cpp
+++ b/Engine/Source/Runtime/Engine/Private/TextureDerivedDataTask.cpp
@@ -438,6 +438,68 @@
 	}
 }
 
+// Dumps the output messages that were created during the given build.
+static void PrintIBuildOutputMessages(const UE::DerivedData::FBuildOutput& InBuildOutput)
+{
+	using namespace UE::DerivedData;
+
+	const FSharedString& Name = InBuildOutput.GetName();
+	const FUtf8SharedString& Function = InBuildOutput.GetFunction();
+
+	for (const FBuildOutputMessage& Message : InBuildOutput.GetMessages())
+	{
+		switch (Message.Level)
+		{
+		case EBuildOutputMessageLevel::Error:
+			// We drop errors to warnings so that they don't stop e.g. a cook from occurring as the cook is likely still
+			// usable
+			UE_LOG(LogTexture, Warning, TEXT("[Error] %s (Build of '%s' by %s.)"),
+				*WriteToString<256>(Message.Message), *Name, *WriteToString<32>(Function));
+			break;
+		case EBuildOutputMessageLevel::Warning:
+			UE_LOG(LogTexture, Warning, TEXT("%s (Build of '%s' by %s.)"),
+				*WriteToString<256>(Message.Message), *Name, *WriteToString<32>(Function));
+			break;
+		case EBuildOutputMessageLevel::Display:
+			UE_LOG(LogTexture, Display, TEXT("%s (Build of '%s' by %s.)"),
+				*WriteToString<256>(Message.Message), *Name, *WriteToString<32>(Function));
+			break;
+		default:
+			checkNoEntry();
+			break;
+		}
+	}
+
+	for (const FBuildOutputLog& Log : InBuildOutput.GetLogs())
+	{
+		switch (Log.Level)
+		{
+		case EBuildOutputLogLevel::Error:
+			// We drop errors to warnings so that they don't stop e.g. a cook from occurring as the cook is likely still
+			// usable
+			UE_LOG(LogTexture, Warning, TEXT("[Error] %s: %s (Build of '%s' by %s.)"),
+				*WriteToString<64>(Log.Category), *WriteToString<256>(Log.Message),
+				*Name, *WriteToString<32>(Function));
+			break;
+		case EBuildOutputLogLevel::Warning:
+			UE_LOG(LogTexture, Warning, TEXT("%s: %s (Build of '%s' by %s.)"),
+				*WriteToString<64>(Log.Category), *WriteToString<256>(Log.Message),
+				*Name, *WriteToString<32>(Function));
+			break;
+		default:
+			checkNoEntry();
+			break;
+		}
+	}
+
+	if (InBuildOutput.HasError())
+	{
+		UE_LOG(LogTexture, Warning, TEXT("Failed to build derived data for build of '%s' by %s."),
+			*Name, *WriteToString<32>(Function));
+		return;
+	}
+}
+
 namespace UE::TextureDerivedData
 {
 
@@ -481,6 +543,8 @@
 		{
 			FBuildInputMetadataArray ChildInputMetadata;
 
+			PrintIBuildOutputMessages(InCompleteParams.Output);
+
 			UE::DerivedData::EStatus Status = InCompleteParams.Status;
 			if (Status == UE::DerivedData::EStatus::Ok)
 			{
@@ -490,6 +554,8 @@
 				// Find everything we want from this build and pipe them over.
 				InChildDefinition.IterateInputBuilds([&Status, &InParentBuild, &ChildInputMetadata](FUtf8StringView InOurKey, const UE::DerivedData::FBuildValueKey& InBuildValueKey)
 				{
+					//UE_LOG(LogTexture, Warning, TEXT("Input Build parent=%s key=%s"), *WriteToString<128>(InParentBuild.Definition.GetKey()), *WriteToString<128>(InOurKey));
+
 					// Filter to things _this_ build produces as the child could be pulling values from different parents.
 					if (InBuildValueKey.BuildKey == InParentBuild.Definition.GetKey())
 					{
@@ -2639,56 +2705,10 @@
 {
 	using namespace UE::DerivedData;
 
-	const FBuildOutput& Output = InBuildCompleteParams.Output;
-	const FSharedString& Name = Output.GetName();
-	const FUtf8SharedString& Function = Output.GetFunction();
+	PrintIBuildOutputMessages(InBuildCompleteParams.Output);
 
-	for (const FBuildOutputMessage& Message : Output.GetMessages())
+	if (InBuildCompleteParams.Output.HasError())
 	{
-		switch (Message.Level)
-		{
-		case EBuildOutputMessageLevel::Error:
-			UE_LOG(LogTexture, Warning, TEXT("[Error] %s (Build of '%s' by %s.)"),
-				*WriteToString<256>(Message.Message), *Name, *WriteToString<32>(Function));
-			break;
-		case EBuildOutputMessageLevel::Warning:
-			UE_LOG(LogTexture, Warning, TEXT("%s (Build of '%s' by %s.)"),
-				*WriteToString<256>(Message.Message), *Name, *WriteToString<32>(Function));
-			break;
-		case EBuildOutputMessageLevel::Display:
-			UE_LOG(LogTexture, Display, TEXT("%s (Build of '%s' by %s.)"),
-				*WriteToString<256>(Message.Message), *Name, *WriteToString<32>(Function));
-			break;
-		default:
-			checkNoEntry();
-			break;
-		}
-	}
-
-	for (const FBuildOutputLog& Log : Output.GetLogs())
-	{
-		switch (Log.Level)
-		{
-		case EBuildOutputLogLevel::Error:
-			UE_LOG(LogTexture, Warning, TEXT("[Error] %s: %s (Build of '%s' by %s.)"),
-				*WriteToString<64>(Log.Category), *WriteToString<256>(Log.Message),
-				*Name, *WriteToString<32>(Function));
-			break;
-		case EBuildOutputLogLevel::Warning:
-			UE_LOG(LogTexture, Warning, TEXT("%s: %s (Build of '%s' by %s.)"),
-				*WriteToString<64>(Log.Category), *WriteToString<256>(Log.Message),
-				*Name, *WriteToString<32>(Function));
-			break;
-		default:
-			checkNoEntry();
-			break;
-		}
-	}
-
-	if (Output.HasError())
-	{
-		UE_LOG(LogTexture, Warning, TEXT("Failed to build derived data for build of '%s' by %s."),
-			*Name, *WriteToString<32>(Function));
 		return;
 	}
 
