[Navmesh] Brush Navmesh Generation Behavior Fix and Change

- Added a console variable for ignoring navigation flags on Brush. Defaulted it to ignore those flags to be consistent with the old behavior we had on brush navmesh generation. NOTE: after toggling the console variable, we need to trigger a static Brush's update(by dragging or changing a property in its brush component's details panel) to apply.
- Fixed a case where the navmesh didn't get updated when chaning the brush from static type to non-static type, due to not updating the static nav geometry data for BSP.
- Fleshed out comments

diff --git a/Engine/Source/Editor/UnrealEd/Private/EditorExporters.cpp b/Engine/Source/Editor/UnrealEd/Private/EditorExporters.cpp
--- a/Engine/Source/Editor/UnrealEd/Private/EditorExporters.cpp
+++ b/Engine/Source/Editor/UnrealEd/Private/EditorExporters.cpp
@@ -2322,8 +2322,7 @@ void UEditorEngine::RebuildStaticNavigableGeometry(ULevel* Level)
 				FBspNode* Node = &Model->Nodes[i];
 				FBspSurf& Surf = Model->Surfs[Node->iSurf];
 				
-				const bool bIsNavRelevant = Surf.Actor && Surf.Actor->GetBrushComponent() && Surf.Actor->GetBrushComponent()->IsNavigationRelevant();
-				if (!bIsNavRelevant)
+				if (!(Surf.Actor && Surf.Actor->ShouldExportStaticNavigableGeometry()))
 				{
 					continue;
 				}
diff --git a/Engine/Source/Runtime/Engine/Classes/Components/BrushComponent.h b/Engine/Source/Runtime/Engine/Classes/Components/BrushComponent.h
--- a/Engine/Source/Runtime/Engine/Classes/Components/BrushComponent.h
+++ b/Engine/Source/Runtime/Engine/Classes/Components/BrushComponent.h
@@ -37,12 +37,15 @@ class UBrushComponent : public UPrimitiveComponent
 
 	//~ Begin UObject Interface
 	virtual void PostLoad() override;
+#if WITH_EDITOR
+	ENGINE_API virtual void PostEditChangeProperty(FPropertyChangedEvent& PropertyChangedEvent) override;
 	//~ End UObject Interface
 
-#if WITH_EDITOR
 	//~Begin UActorComponent Interface
 	ENGINE_API virtual void OnRegister() override;
 	//~End UActorComponent Interface
+
+	void ConditionalRebuildAlteredBSP() const;
 #endif //WITH_EDITOR
 
 	//~ Begin USceneComponent Interface
diff --git a/Engine/Source/Runtime/Engine/Classes/Engine/Brush.h b/Engine/Source/Runtime/Engine/Classes/Engine/Brush.h
--- a/Engine/Source/Runtime/Engine/Classes/Engine/Brush.h
+++ b/Engine/Source/Runtime/Engine/Classes/Engine/Brush.h
@@ -189,6 +189,9 @@ public:
 	ENGINE_API virtual void InitPosRotScale();
 	ENGINE_API virtual void CopyPosRotScaleFrom( ABrush* Other );
 
+	/** @return true if its Brush Component is static and should affect navmesh based on user settings*/
+	bool ENGINE_API ShouldExportStaticNavigableGeometry() const;
+	
 	static void SetSuppressBSPRegeneration(bool bSuppress) { bSuppressBSPRegeneration = bSuppress; }
 	static bool GetSuppressBSPRegeneration() { return bSuppressBSPRegeneration; }
 private:
diff --git a/Engine/Source/Runtime/Engine/Private/Brush.cpp b/Engine/Source/Runtime/Engine/Private/Brush.cpp
--- a/Engine/Source/Runtime/Engine/Private/Brush.cpp
+++ b/Engine/Source/Runtime/Engine/Private/Brush.cpp
@@ -49,6 +49,15 @@ namespace BrushUtils
 	}
 }
 
+namespace BrushNavmeshGenerationCVars
+{
+	static bool bForceNavmeshGenerationOnStaticBrush = true;
+	static FAutoConsoleVariableRef CVarForceNavmeshGenerationOnStaticBrush(
+		TEXT("brush.ForceNavmeshGenerationOnStaticBrush"), 
+		bForceNavmeshGenerationOnStaticBrush, 
+		TEXT("Force exporting static brush's geometry data to Navigation System regardless of CanEverAffectNavigation and CollisionPreset values. Enabled by default to be backward compatible with brushes on legacy maps."), 
+		ECVF_Default);
+}
 #endif
 
 DEFINE_LOG_CATEGORY_STATIC(LogBrush, Log, All);
@@ -184,6 +193,18 @@ void ABrush::SetNeedRebuild(ULevel* InLevel)
 	}
 }
 
+bool ABrush::ShouldExportStaticNavigableGeometry() const
+{
+	if (BrushComponent)
+	{
+		const bool bNavRelevant = BrushComponent->IsNavigationRelevant();
+
+		return IsStaticBrush() && (BrushNavmeshGenerationCVars::bForceNavmeshGenerationOnStaticBrush || bNavRelevant);
+	}
+
+	return false;
+}
+
 void ABrush::InitPosRotScale()
 {
 	check(BrushComponent);
diff --git a/Engine/Source/Runtime/Engine/Private/Components/BrushComponent.cpp b/Engine/Source/Runtime/Engine/Private/Components/BrushComponent.cpp
--- a/Engine/Source/Runtime/Engine/Private/Components/BrushComponent.cpp
+++ b/Engine/Source/Runtime/Engine/Private/Components/BrushComponent.cpp
@@ -519,25 +519,53 @@ void UBrushComponent::PostLoad()
 }
 
 #if WITH_EDITOR
+void UBrushComponent::PostEditChangeProperty(FPropertyChangedEvent& PropertyChangedEvent)
+{
+	if (const ABrush* BrushOwner = Cast<ABrush>(GetOwner()))
+	{
+		if (!BrushOwner->IsVolumeBrush())
+		{
+			if (!BrushOwner->IsStaticBrush())
+			{
+				if (PropertyChangedEvent.Property && PropertyChangedEvent.Property->GetFName() == GET_MEMBER_NAME_CHECKED(UBrushComponent, Mobility))
+				{
+					// This is to cover a case when the user changes the non-volume brush from static type to non-static type(i.e. changes its mobility)
+					// The previously registered static navigation data for BSP didn't get updated and this component's part was staying there
+					ConditionalRebuildAlteredBSP();
+				}
+			}
+		}
+	}
+
+	Super::PostEditChangeProperty(PropertyChangedEvent);
+}
+
 void UBrushComponent::OnRegister()
 {
 	Super::OnRegister();
 
-	// BSP can only be rebuilt during a transaction
-	if (!ABrush::GetSuppressBSPRegeneration() && GEditor && GUndo)
+	if (const ABrush* BrushOwner = Cast<ABrush>(GetOwner()))
 	{
-		if (ABrush* BrushOwner = Cast<ABrush>(GetOwner()))
+		// Volume Brush and Non-Static Brush's Geometric Data will be exported via registering the component
+		// which is handled in UPrimitiveComponent::OnRegister
+		if (!BrushOwner->IsVolumeBrush() && BrushOwner->IsStaticBrush())
 		{
-			const bool bIsBrushOwnerBSP = !BrushOwner->IsVolumeBrush() && BrushOwner->IsStaticBrush();
 			// Not checking NavRelavency is intentional to cover cases when it changes from being Nav Relevant to Nav Irrelevant,
-			// we need to clear the previously registered BSP Navigation Data.
-			if (bIsBrushOwnerBSP)
-			{
-				GEditor->RebuildAlteredBSP();
-			}
+			// we need to update the registered static Navigation Data for BSP and remove this component's part.
+			ConditionalRebuildAlteredBSP();
 		}
 	}
 }
+
+void UBrushComponent::ConditionalRebuildAlteredBSP() const
+{
+	// Not entirely clear why we should limit the rebuilding to within a transaction,
+	// but this is done in ABrush::PostEditChangeProperty, so we are following the same pattern
+	if (!ABrush::GetSuppressBSPRegeneration() && GEditor && GUndo)
+	{
+		GEditor->RebuildAlteredBSP();
+	}
+}
 #endif //WITH_EDITOR
 
 ESceneDepthPriorityGroup UBrushComponent::GetStaticDepthPriorityGroup() const
