Adding some logs to hunt down an issue.

diff --git a/Engine/Source/Runtime/Engine/Private/SoundWave.cpp b/Engine/Source/Runtime/Engine/Private/SoundWave.cpp
--- a/Engine/Source/Runtime/Engine/Private/SoundWave.cpp
+++ b/Engine/Source/Runtime/Engine/Private/SoundWave.cpp
@@ -1825,6 +1825,18 @@ FByteBulkData* USoundWave::GetCompressedData(FName Format, const FPlatformAudioC
 			UE_LOG(LogAudio, Error, TEXT("Attempt to access the DDC when there is none available on sound '%s', format = %s. Should have been cooked."), *GetFullName(), *PlatformSpecificFormat.ToString());
 		}
 	}
+	else
+	{
+		// We had valid data in the format container - it should be non zero size!
+		// If we have a DDC then we don't want to log because we emplaced an empty buffer on purpose to
+		// avoid hitting the ddc over and over.
+#if !WITH_EDITOR
+		if (Result->GetBulkDataSize() == 0)
+		{
+			UE_LOG(LogAudio, Warning, TEXT("Sound format container returned empty sound! sound '%s', format = %s."), *GetFullName(), *PlatformSpecificFormat.ToString());
+		}
+#endif
+	}
 	check(Result);
 	return Result->GetBulkDataSize() > 0 ? Result : NULL; // we don't return empty bulk data...but we save it to avoid thrashing the DDC
 }
@@ -2275,6 +2287,10 @@ bool USoundWave::InitAudioResource(FName Format)
 			check(SoundWaveDataPtr->ResourceSize > 0);
 #endif
 		}
+		else
+		{
+			UE_LOG(LogAudio, Warning, TEXT("Soundwave: %s doesn't have compressed data! Format = %s"), *GetName(), *WriteToString<64>(Format));
+		}
 	}
 
 	return SoundWaveDataPtr->ResourceSize > 0;
@@ -4583,10 +4599,25 @@ FSoundWaveProxy::FSoundWaveProxy(USoundWave* InWave)
 	// this should have been allocated by the USoundWave and should always be valid
 	check(SoundWaveDataPtr);
 
+	bool bIsStreaming = InWave->IsStreaming(nullptr);
+
 	// non-streaming sources need resource data initialized before the FSoundWaveProxy
 	// can be used. 
-	InWave->InitAudioResource(SoundWaveDataPtr->GetRuntimeFormat());
-	check((InWave->IsStreaming(nullptr)) || (SoundWaveDataPtr->GetResourceSize() > 0));}
+	if (!InWave->InitAudioResource(SoundWaveDataPtr->GetRuntimeFormat()) && !bIsStreaming)
+	{
+		UE_LOG(LogAudio, Warning, TEXT("FSoundWaveProxy failed to InitAudioResource: %s, format %s"), *InWave->GetName(), *WriteToString<64>(SoundWaveDataPtr->GetRuntimeFormat()));
+	}
+
+	// We must either be able to stream data, or have the data already.
+	
+	bool bHasResourceData = SoundWaveDataPtr->GetResourceSize() > 0;
+
+	check(bIsStreaming || bHasResourceData);
+	if (!bIsStreaming && !bHasResourceData)
+	{
+		UE_LOG(LogAudio, Warning, TEXT("FSoundWaveProxy doesn't have audio data! bIsStreaming = %d, bHasResourceData = %d, Sound = %s"), bIsStreaming, bHasResourceData, *InWave->GetName());
+	}
+}
 
 FSoundWaveProxy::~FSoundWaveProxy()
 {
diff --git a/Engine/Source/Runtime/Engine/Private/SoundWaveProxyReader.cpp b/Engine/Source/Runtime/Engine/Private/SoundWaveProxyReader.cpp
--- a/Engine/Source/Runtime/Engine/Private/SoundWaveProxyReader.cpp
+++ b/Engine/Source/Runtime/Engine/Private/SoundWaveProxyReader.cpp
@@ -72,6 +72,10 @@ FSoundWaveProxyReader::FSoundWaveProxyReader(FSoundWaveProxyRef InWaveProxy, con
 
 	// Prepare to read audio
 	bIsDecoderValid = InitializeDecoder(Settings.StartTimeInSeconds);
+	if (!bIsDecoderValid)
+	{
+		UE_LOG(LogAudio, Warning, TEXT("Failed to InitializeDecoder in FSoundWaveProxyReader(), Package: %s"), *WriteToString<64>(InWaveProxy->GetPackageName()));
+	}
 
 	// set the decoder to "Fail" if we're unable to create a decoder
 	DecodeResult = bIsDecoderValid ? DecodeResult : EDecodeResult::Fail;
