GPU Skin Cache:  Speculative fix for crashes in "PrepareUpdateSkinning", asserting PreviousPositionBuffer != PositionBuffer.  Simplified buffer allocation logic to make it categorically impossible for the same buffer to be allocated, by changing the logic to implicitly select opposite double buffered elements.  If a used element is passed to the allocator, its opposite is returned, rather than considering revision information on the element.

#rnx

diff --git a/Engine/Source/Runtime/Engine/Private/GPUSkinCache.cpp b/Engine/Source/Runtime/Engine/Private/GPUSkinCache.cpp
--- a/Engine/Source/Runtime/Engine/Private/GPUSkinCache.cpp
+++ b/Engine/Source/Runtime/Engine/Private/GPUSkinCache.cpp
@@ -1999,47 +1999,21 @@ void FGPUSkinCache::PrepareUpdateSkinning(FGPUSkinCacheEntry* Entry, int32 Secti
 
 	DispatchData.DispatchFlags = 0;
 
-	auto BufferUpdate = [&DispatchData, OverlappedUAVs](
-		FSkinCacheRWBuffer*& PositionBuffer,
-		const FVertexBufferAndSRV& BoneBuffer, 
-		uint32 Revision,
-		const FVertexBufferAndSRV& PrevBoneBuffer,
-		uint32 PrevRevision,
-		uint32 UpdateFlag
-		)
-	{
-		PositionBuffer = DispatchData.PositionTracker.Find(BoneBuffer, Revision);
-		if (!PositionBuffer)
-		{
-			PositionBuffer = DispatchData.PositionTracker.Advance(BoneBuffer, Revision, PrevBoneBuffer, PrevRevision);
-			check(PositionBuffer);
-
-			DispatchData.DispatchFlags |= UpdateFlag;
+	// Find existing buffers
+	DispatchData.PositionBuffer = DispatchData.PositionTracker.Find(BoneBuffer, CurrentRevision);
+	DispatchData.PreviousPositionBuffer = DispatchData.PositionTracker.Find(PrevBoneBuffer, PreviousRevision);
 
-			if (OverlappedUAVs)
-			{
-				(*OverlappedUAVs).Add(PositionBuffer);
-			}
-		}
-	};
-
-	BufferUpdate(
-		DispatchData.PreviousPositionBuffer,
-		PrevBoneBuffer,
-		PreviousRevision,
-		BoneBuffer,
-		CurrentRevision,
-		(uint32)EGPUSkinCacheDispatchFlags::DispatchPrevPosition
-		);
-
-	BufferUpdate(
-		DispatchData.PositionBuffer, 
-		BoneBuffer, 
-		CurrentRevision, 
-		PrevBoneBuffer, 
-		PreviousRevision, 
-		(uint32)EGPUSkinCacheDispatchFlags::DispatchPosition
-		);
+	// Allocate buffers if not found, excluding buffers already in use.  Or make the current buffer distinct if it happens to equal previous.
+	if (!DispatchData.PositionBuffer || DispatchData.PositionBuffer == DispatchData.PreviousPositionBuffer)
+	{
+		DispatchData.PositionBuffer = DispatchData.PositionTracker.AllocateUnused(BoneBuffer, CurrentRevision, DispatchData.PreviousPositionBuffer, OverlappedUAVs);
+		DispatchData.DispatchFlags |= (uint32)EGPUSkinCacheDispatchFlags::DispatchPosition;
+	}
+	if (!DispatchData.PreviousPositionBuffer)
+	{
+		DispatchData.PreviousPositionBuffer = DispatchData.PositionTracker.AllocateUnused(PrevBoneBuffer, PreviousRevision, DispatchData.PositionBuffer, OverlappedUAVs);
+		DispatchData.DispatchFlags |= (uint32)EGPUSkinCacheDispatchFlags::DispatchPrevPosition;
+	}
 
 	DispatchData.TangentBuffer = DispatchData.PositionTracker.GetTangentBuffer();
 	DispatchData.IntermediateTangentBuffer = DispatchData.PositionTracker.GetIntermediateTangentBuffer();
diff --git a/Engine/Source/Runtime/Engine/Public/GPUSkinCache.h b/Engine/Source/Runtime/Engine/Public/GPUSkinCache.h
--- a/Engine/Source/Runtime/Engine/Public/GPUSkinCache.h
+++ b/Engine/Source/Runtime/Engine/Public/GPUSkinCache.h
@@ -359,34 +359,17 @@ public:
 			return Allocation ? Allocation->GetIntermediateAccumulatedTangentBuffer() : nullptr;
 		}
 
-		FSkinCacheRWBuffer* Advance(const FVertexBufferAndSRV& BoneBuffer1, uint32 Revision1, const FVertexBufferAndSRV& BoneBuffer2, uint32 Revision2)
+		// Allocates an element that's not the "Used" element passed in (or if Used is NULL, allocates any element).  Optionally adds element to "AllocatedItems" if non-null.
+		FSkinCacheRWBuffer* AllocateUnused(const FVertexBufferAndSRV& BoneBuffer, uint32 Revision, const FSkinCacheRWBuffer* Used, TArray<FSkinCacheRWBuffer*>* AllocatedItems)
 		{
-			FSkinCacheRWBuffer* Result = nullptr;
-			const FVertexBufferAndSRV* InBoneBuffers[2] = { &BoneBuffer1 , &BoneBuffer2 };
-			uint32 InRevisions[2] = { Revision1 , Revision2 };
+			int32 UnusedIndex = Used == &Allocation->PositionBuffers[0] ? 1 : 0;
+			Revisions[UnusedIndex] = Revision;
+			BoneBuffers[UnusedIndex] = &BoneBuffer;
 
-			for (int32 Index = 0; Index < NUM_BUFFERS; ++Index)
+			FSkinCacheRWBuffer* Result = &Allocation->PositionBuffers[UnusedIndex];
+			if (AllocatedItems)
 			{
-				bool Needed = false;
-				for (int32 i = 0; i < 2; ++i)
-				{
-					if (Revisions[Index] == InRevisions[i] && BoneBuffers[Index] == InBoneBuffers[i])
-					{
-						if (i == 0)
-						{
-							Result = &Allocation->PositionBuffers[Index];
-						}
-						Needed = true;
-					}
-				}
-
-				if (!Needed)
-				{
-					Revisions[Index] = Revision1;
-					BoneBuffers[Index] = &BoneBuffer1;
-					Result = &Allocation->PositionBuffers[Index];
-					break;
-				}
+				AllocatedItems->Add(Result);
 			}
 			return Result;
 		}
