Ensure the displacement material attribute is always defaulted to the center and make it so that displacement set to a constant that is equal to center disables displacement for the material.

diff --git a/Engine/Source/Developer/MaterialBaking/Private/ExportMaterialProxy.h b/Engine/Source/Developer/MaterialBaking/Private/ExportMaterialProxy.h
--- a/Engine/Source/Developer/MaterialBaking/Private/ExportMaterialProxy.h
+++ b/Engine/Source/Developer/MaterialBaking/Private/ExportMaterialProxy.h
@@ -76,6 +76,16 @@ struct FExportMaterialCompiler : public FProxyMaterialCompiler
 		return Compiler->VertexNormal();
 	}
 
+	virtual int32 DisplacementCenter() override
+	{		
+		return Compiler->DisplacementCenter();
+	}
+
+	virtual int32 DisplacementMagnitude() override
+	{		
+		return Compiler->DisplacementMagnitude();
+	}
+
 	virtual int32 ParticleRandom() override
 	{
 		return Compiler->Constant(0.0f);
diff --git a/Engine/Source/Developer/MaterialUtilities/Private/MaterialUtilities.cpp b/Engine/Source/Developer/MaterialUtilities/Private/MaterialUtilities.cpp
--- a/Engine/Source/Developer/MaterialUtilities/Private/MaterialUtilities.cpp
+++ b/Engine/Source/Developer/MaterialUtilities/Private/MaterialUtilities.cpp
@@ -289,6 +289,16 @@ struct FExportMaterialCompiler : public FProxyMaterialCompiler
 		// Current returning vertex normal since pixel normal will contain incorrect data (normal calculated from uv data used as vertex positions to render out the material)
 		return Compiler->VertexNormal();
 	}
+	
+	virtual int32 DisplacementCenter() override
+	{
+		return Compiler->DisplacementCenter();
+	}
+	
+	virtual int32 DisplacementMagnitude() override
+	{
+		return Compiler->DisplacementMagnitude();
+	}
 
 	virtual int32 ParticleRandom() override
 	{
diff --git a/Engine/Source/Runtime/Engine/Classes/Materials/MaterialExpressionDisplacementScaling.h b/Engine/Source/Runtime/Engine/Classes/Materials/MaterialExpressionDisplacementScaling.h
new file mode 100644
--- /dev/null
+++ b/Engine/Source/Runtime/Engine/Classes/Materials/MaterialExpressionDisplacementScaling.h
@@ -0,0 +1,24 @@
+// Copyright Epic Games, Inc. All Rights Reserved.
+
+#pragma once
+
+#include "CoreMinimal.h"
+#include "Engine/EngineTypes.h"
+#include "Materials/MaterialExpression.h"
+#include "MaterialExpressionDisplacementScaling.generated.h"
+
+class FMaterialCompiler;
+
+UCLASS(collapsecategories, hidecategories = Object, MinimalAPI)
+class UMaterialExpressionDisplacementScaling : public UMaterialExpression
+{
+	GENERATED_UCLASS_BODY()
+
+	//~ Begin UMaterialExpression Interface
+#if WITH_EDITOR
+	bool GenerateHLSLExpression(FMaterialHLSLGenerator& Generator, UE::HLSLTree::FScope& Scope, int32 OutputIndex, UE::HLSLTree::FExpression const*& OutExpression) const override;
+	virtual int32 Compile(FMaterialCompiler* Compiler, int32 OutputIndex) override;
+	virtual void GetCaption(TArray<FString>& OutCaptions) const override;
+#endif
+	//~ End UMaterialExpression Interface
+};
\ No newline at end of file
diff --git a/Engine/Source/Runtime/Engine/Private/Materials/HLSLMaterialTranslator.cpp b/Engine/Source/Runtime/Engine/Private/Materials/HLSLMaterialTranslator.cpp
--- a/Engine/Source/Runtime/Engine/Private/Materials/HLSLMaterialTranslator.cpp
+++ b/Engine/Source/Runtime/Engine/Private/Materials/HLSLMaterialTranslator.cpp
@@ -1700,9 +1700,11 @@ void FHLSLMaterialTranslator::DoTranslate()
 	bool bUsesWorldPositionOffsetPrevious = IsMaterialPropertyUsed(MP_WorldPositionOffset, Chunk[CompiledMP_PrevWorldPositionOffset], FLinearColor(0, 0, 0, 0), 3);
 	bUsesWorldPositionOffset = bUsesWorldPositionOffsetCurrent || bUsesWorldPositionOffsetPrevious;
 
+	const FDisplacementScaling DisplacementScaling = Material->GetDisplacementScaling();
 	bUsesDisplacement = DoesPlatformSupportNanite(Platform) &&
 		Material->IsTessellationEnabled() &&
-		!FMath::IsNearlyZero(Material->GetDisplacementScaling().Magnitude, UE_KINDA_SMALL_NUMBER);
+		!FMath::IsNearlyZero(DisplacementScaling.Magnitude, UE_KINDA_SMALL_NUMBER) &&
+		IsMaterialPropertyUsed(MP_Displacement, Chunk[MP_Displacement], FLinearColor(DisplacementScaling.Center, 0, 0, 0), 1);
 
 	MaterialCompilationOutput.bModifiesMeshPosition = bUsesPixelDepthOffset || bUsesWorldPositionOffset || bUsesDisplacement;
 	MaterialCompilationOutput.bUsesWorldPositionOffset = bUsesWorldPositionOffset;
@@ -11042,6 +11044,16 @@ int32 FHLSLMaterialTranslator::PixelNormalWS()
 	return AddInlinedCodeChunk(MCT_Float3,TEXT("Parameters.WorldNormal"));	
 }
 
+int32 FHLSLMaterialTranslator::DisplacementCenter()
+{
+	return Constant(Material->GetDisplacementScaling().Center);
+}
+
+int32 FHLSLMaterialTranslator::DisplacementMagnitude()
+{
+	return Constant(Material->GetDisplacementScaling().Magnitude);
+}
+
 int32 FHLSLMaterialTranslator::Derivative(int32 A, EDervativeComponent Component)
 {
 	if (A == INDEX_NONE)
diff --git a/Engine/Source/Runtime/Engine/Private/Materials/HLSLMaterialTranslator.h b/Engine/Source/Runtime/Engine/Private/Materials/HLSLMaterialTranslator.h
--- a/Engine/Source/Runtime/Engine/Private/Materials/HLSLMaterialTranslator.h
+++ b/Engine/Source/Runtime/Engine/Private/Materials/HLSLMaterialTranslator.h
@@ -1161,6 +1161,10 @@ protected:
 	virtual int32 VertexNormal() override;
 	virtual int32 VertexTangent() override;
 	virtual int32 PixelNormalWS() override;
+
+	virtual int32 DisplacementCenter() override;
+	virtual int32 DisplacementMagnitude() override;
+
 	virtual int32 DDX(int32 A) override;
 	virtual int32 DDY(int32 A) override;
 
diff --git a/Engine/Source/Runtime/Engine/Private/Materials/MaterialAttributeDefinitionMap.cpp b/Engine/Source/Runtime/Engine/Private/Materials/MaterialAttributeDefinitionMap.cpp
--- a/Engine/Source/Runtime/Engine/Private/Materials/MaterialAttributeDefinitionMap.cpp
+++ b/Engine/Source/Runtime/Engine/Private/Materials/MaterialAttributeDefinitionMap.cpp
@@ -71,6 +71,11 @@ int32 FMaterialAttributeDefintion::CompileDefaultValue(FMaterialCompiler* Compil
 		return Compiler->SubstrateCreateAndRegisterNullMaterial();
 	}
 
+	if (Property == MP_Displacement)
+	{
+		return Compiler->DisplacementCenter();
+	}
+
 	if (TexCoordIndex == INDEX_NONE)
 	{
 		// Standard value type
diff --git a/Engine/Source/Runtime/Engine/Private/Materials/MaterialExpressionHLSL.cpp b/Engine/Source/Runtime/Engine/Private/Materials/MaterialExpressionHLSL.cpp
--- a/Engine/Source/Runtime/Engine/Private/Materials/MaterialExpressionHLSL.cpp
+++ b/Engine/Source/Runtime/Engine/Private/Materials/MaterialExpressionHLSL.cpp
@@ -65,6 +65,7 @@
 #include "Materials/MaterialExpressionDepthOfFieldFunction.h"
 #include "Materials/MaterialExpressionDeriveNormalZ.h"
 #include "Materials/MaterialExpressionDesaturation.h"
+#include "Materials/MaterialExpressionDisplacementScaling.h"
 #include "Materials/MaterialExpressionDistance.h"
 #include "Materials/MaterialExpressionDistanceCullFade.h"
 #include "Materials/MaterialExpressionDistanceFieldApproxAO.h"
@@ -5204,4 +5205,22 @@ bool UMaterialExpressionNeuralNetworkOutput::GenerateHLSLExpression(FMaterialHLS
 	return true;
 }
 
+bool UMaterialExpressionDisplacementScaling::GenerateHLSLExpression(FMaterialHLSLGenerator& Generator, UE::HLSLTree::FScope& Scope, int32 OutputIndex, UE::HLSLTree::FExpression const*& OutExpression) const
+{
+	const UMaterial* BaseMaterial = Generator.GetTargetMaterial();
+	const FDisplacementScaling Scaling = BaseMaterial->GetDisplacementScaling();
+	const UE::HLSLTree::FExpression* Expr = nullptr;
+
+	if (OutputIndex == 0)
+	{
+		Expr = Generator.NewConstant(Scaling.Center);
+	}
+	else if (OutputIndex == 1)
+	{
+		Expr = Generator.NewConstant(Scaling.Magnitude);
+	}
+
+	return Expr != nullptr;
+}
+
 #endif // WITH_EDITOR
diff --git a/Engine/Source/Runtime/Engine/Private/Materials/MaterialExpressions.cpp b/Engine/Source/Runtime/Engine/Private/Materials/MaterialExpressions.cpp
--- a/Engine/Source/Runtime/Engine/Private/Materials/MaterialExpressions.cpp
+++ b/Engine/Source/Runtime/Engine/Private/Materials/MaterialExpressions.cpp
@@ -95,6 +95,7 @@
 #include "Materials/MaterialExpressionDistance.h"
 #include "Materials/MaterialExpressionDistanceCullFade.h"
 #include "Materials/MaterialExpressionDistanceFieldsRenderingSwitch.h"
+#include "Materials/MaterialExpressionDisplacementScaling.h"
 #include "Materials/MaterialExpressionDivide.h"
 #include "Materials/MaterialExpressionDotProduct.h"
 #include "Materials/MaterialExpressionDynamicParameter.h"
@@ -29288,4 +29289,41 @@ uint32 UMaterialExpressionNeuralNetworkOutput::GetInputType(int32 InputIndex)
 }
 #endif // WITH_EDITOR
 
+///////////////////////////////////////////////////////////////////////////////
+// UMaterialExpressionActorPositionWS
+///////////////////////////////////////////////////////////////////////////////
+
+UMaterialExpressionDisplacementScaling::UMaterialExpressionDisplacementScaling(const FObjectInitializer& ObjectInitializer)
+	: Super(ObjectInitializer)
+{
+#if WITH_EDITOR
+	Outputs.Reset();
+	Outputs.Add(FExpressionOutput(TEXT("Center")));
+	Outputs.Add(FExpressionOutput(TEXT("Magnitude")));
+
+	bShowOutputNameOnPin = true;
+#endif
+}
+
+#if WITH_EDITOR
+int32 UMaterialExpressionDisplacementScaling::Compile(FMaterialCompiler* Compiler, int32 OutputIndex)
+{
+	switch (OutputIndex)
+	{
+	case 0:
+		return Compiler->DisplacementCenter();
+	case 1:
+		return Compiler->DisplacementMagnitude();
+	default:
+		return INDEX_NONE;
+	}
+}
+
+void UMaterialExpressionDisplacementScaling::GetCaption(TArray<FString>& OutCaptions) const
+{
+	OutCaptions.Add(TEXT("Displacement Scaling"));
+}
+#endif // WITH_EDITOR
+
+
 #undef LOCTEXT_NAMESPACE
diff --git a/Engine/Source/Runtime/Engine/Public/MaterialCompiler.h b/Engine/Source/Runtime/Engine/Public/MaterialCompiler.h
--- a/Engine/Source/Runtime/Engine/Public/MaterialCompiler.h
+++ b/Engine/Source/Runtime/Engine/Public/MaterialCompiler.h
@@ -515,6 +515,9 @@ public:
 	virtual int32 VertexNormal() = 0;
 	virtual int32 VertexTangent() = 0;
 	virtual int32 PixelNormalWS() = 0;
+	
+	virtual int32 DisplacementCenter() = 0;
+	virtual int32 DisplacementMagnitude() = 0;
 
 	virtual int32 CustomExpression(class UMaterialExpressionCustom* Custom, int32 OutputIndex, TArray<int32>& CompiledInputs) = 0;
 	virtual int32 CustomOutput(class UMaterialExpressionCustomOutput* Custom, int32 OutputIndex, int32 OutputCode) = 0;
@@ -1066,6 +1069,9 @@ public:
 	virtual int32 VertexTangent() override { return Compiler->VertexTangent(); }
 	virtual int32 PixelNormalWS() override { return Compiler->PixelNormalWS(); }
 
+	virtual int32 DisplacementCenter() override { return Compiler->DisplacementCenter(); }
+	virtual int32 DisplacementMagnitude() override { return Compiler->DisplacementMagnitude(); }
+
 	virtual int32 CustomExpression(class UMaterialExpressionCustom* Custom, int32 OutputIndex, TArray<int32>& CompiledInputs) override { return Compiler->CustomExpression(Custom, OutputIndex, CompiledInputs); }
 	virtual int32 CustomOutput(class UMaterialExpressionCustomOutput* Custom, int32 OutputIndex, int32 OutputCode) override { return Compiler->CustomOutput(Custom, OutputIndex, OutputCode); }
 	virtual int32 VirtualTextureOutput(uint8 AttributeMask) override { return Compiler->VirtualTextureOutput(AttributeMask); }
