Texture.cpp fix bPNGCompressed_DEPRECATED was not being cleared
that should have been okay because deprecated fields are not written out
but lets just make sure it is cleared to be safe
also fix TextureSource::Reset not resetting all fields

diff --git a/Engine/Source/Runtime/Engine/Private/Texture.cpp b/Engine/Source/Runtime/Engine/Private/Texture.cpp
--- a/Engine/Source/Runtime/Engine/Private/Texture.cpp
+++ b/Engine/Source/Runtime/Engine/Private/Texture.cpp
@@ -793,6 +793,7 @@ void UTexture::PostEditChangeProperty(FPropertyChangedEvent& PropertyChangedEven
 		// if RequiresNotifyMaterials was turned on by Validate
 		// for a PostEditChange() with no Property
 		// no need to Notify
+		// @@ ?? what ? that doesn't seem right, if Validate changed VirtualStreaming you do need to notify
 		RequiresNotifyMaterials = false;
 	}
 
@@ -1095,16 +1096,24 @@ void UTexture::Serialize(FArchive& Ar)
 			Source.CompressionFormat = TSCF_None;
 		}
 
-		if ( Source.bPNGCompressed_DEPRECATED && Source.CompressionFormat != TSCF_PNG )
+		if ( Source.bPNGCompressed_DEPRECATED )
 		{
-			// loaded with deprecated "bPNGCompressed" (but not the newer CompressionFormat)
-			// change to CompressionFormat PNG
-			// this is expected on assets older than the CompressionFormat field
-			check( Source.CompressionFormat == TSCF_None );
-			Source.CompressionFormat = TSCF_PNG;
+			if ( Source.CompressionFormat == TSCF_None )
+			{
+				// loaded with deprecated "bPNGCompressed" (but not the newer CompressionFormat)
+				// change to CompressionFormat PNG
+				// this is expected on assets older than the CompressionFormat field
+				Source.CompressionFormat = TSCF_PNG;
+			}
+			else if ( Source.CompressionFormat != TSCF_PNG )
+			{
+				UE_LOG(LogTexture, Warning, TEXT("Texture \"%s\" has CompressionFormat=%d=%s with bPNGCompressed, ignoring bPNGCompressed."), 
+					*GetPathName(), (int)Source.CompressionFormat, *Source.GetSourceCompressionAsString() );
+			}
 		}
 		
 		// bPNGCompressed_DEPRECATED is not kept in sync with CompressionFormat any more, do not check it after this point
+		Source.bPNGCompressed_DEPRECATED = false;
 
 		if ( Source.GetFormat() == TSF_RGBA8_DEPRECATED
 			|| Source.GetFormat() == TSF_RGBE8_DEPRECATED )
@@ -3296,7 +3305,13 @@ void FTextureSource::Reset()
 #endif
 
 	check( LockState == ELockState::None && NumLockedMips == 0 );
+	
+	// Owner not Reset
+
+	// TornOff members not reset ?
 
+	BaseBlockX = 0;
+	BaseBlockY = 0;
 	SizeX = 0;
 	SizeY = 0;
 	NumSlices = 0;
@@ -3306,6 +3321,8 @@ void FTextureSource::Reset()
 	LayerFormat.Empty();
 	Blocks.Empty();
 	BlockDataOffsets.Empty();
+	bPNGCompressed_DEPRECATED = false;
+	bLongLatCubemap = false;
 	CompressionFormat = TSCF_None;
 	LockedMipData.Reset();
 	NumLockedMips = 0u;
@@ -3315,7 +3332,7 @@ void FTextureSource::Reset()
 
 	BulkData.Reset();
 
-	ForceGenerateGuid();
+	ForceGenerateGuid(); // sets Id and bGuidIsHash
 }
 
 // total size in bytes including all blocks and layers
