Fix NetConnection crash when receiving packets from an IP address that was previously remapped after the user's client connection was cleaned-up.

diff --git a/Engine/Plugins/Online/OnlineSubsystemUtils/Source/OnlineSubsystemUtils/Private/IpNetDriver.cpp b/Engine/Plugins/Online/OnlineSubsystemUtils/Source/OnlineSubsystemUtils/Private/IpNetDriver.cpp
--- a/Engine/Plugins/Online/OnlineSubsystemUtils/Source/OnlineSubsystemUtils/Private/IpNetDriver.cpp
+++ b/Engine/Plugins/Online/OnlineSubsystemUtils/Source/OnlineSubsystemUtils/Private/IpNetDriver.cpp
@@ -1184,22 +1184,22 @@ void UIpNetDriver::TickDispatch(float DeltaTime)
 
 		if (Connection == nullptr)
 		{
-			auto* Result = MappedClientConnections.Find(FromAddr);
-
-			if (Result != nullptr)
+			if (TObjectPtr<UNetConnection>* ConnectionMapValue = MappedClientConnections.Find(FromAddr))
 			{
-				UNetConnection* ConnVal = *Result;
-
-				if (ConnVal != nullptr)
+				UNetConnection* FoundConnection = *ConnectionMapValue;
+				if (FoundConnection)
 				{
-					Connection = ConnVal;
+					if (ensureMsgf(FoundConnection->Driver, TEXT("Found invalid cleaned-up connection in Map: %s"), *Connection->Describe()))
+					{
+						Connection = FoundConnection;
+						check(CastChecked<UIpConnection>(Connection)->RemoteAddr->CompareEndpoints(*FromAddr));
+					}
 				}
 				else
 				{
 					ReceivedTraits.bFromRecentlyDisconnected = true;
 				}
 			}
-			check(Connection == nullptr || CastChecked<UIpConnection>(Connection)->RemoteAddr->CompareEndpoints(*FromAddr));
 		}
 
 
@@ -1367,25 +1367,23 @@ UNetConnection* UIpNetDriver::ProcessConnectionlessPacket(FReceivedPacketView& P
 
 					if (FoundConn != nullptr)
 					{
-						TObjectPtr<UNetConnection> RemovedConn = nullptr;
 						TSharedRef<FInternetAddr> RemoteAddrRef = FoundConn->RemoteAddr.ToSharedRef();
 
-						//verify(MappedClientConnections.RemoveAndCopyValue(RemoteAddrRef, RemovedConn) && RemovedConn == FoundConn);
-
 						// @todo: There needs to be a proper/standardized copy API for this. Also in IpConnection.cpp
 						bool bIsValid = false;
 
 						RemoveFromNewIPTracking(RemoteAddrRef.Get());
+						
+						// Remove the old address mapping
+						MappedClientConnections.Remove(RemoteAddrRef);
 
 						const FString OldAddress = RemoteAddrRef->ToString(true);
 
 						RemoteAddrRef->SetIp(*Address->ToString(false), bIsValid);
 						RemoteAddrRef->SetPort(Address->GetPort());
 
-
 						MappedClientConnections.Add(RemoteAddrRef, FoundConn);
 
-
 						// Make sure we didn't just invalidate a RecentlyDisconnectedClients entry, with the same address
 						int32 RecentDisconnectIdx = RecentlyDisconnectedClients.IndexOfByPredicate(
 							[&RemoteAddrRef](const FDisconnectedClient& CurElement)
diff --git a/Engine/Source/Runtime/Engine/Private/NetConnection.cpp b/Engine/Source/Runtime/Engine/Private/NetConnection.cpp
--- a/Engine/Source/Runtime/Engine/Private/NetConnection.cpp
+++ b/Engine/Source/Runtime/Engine/Private/NetConnection.cpp
@@ -109,6 +109,9 @@ static int32 GNetCloseTimingDebug = 0;
 static FAutoConsoleVariableRef CVarCloseTimingDebug(TEXT("net.CloseTimingDebug"), GNetCloseTimingDebug,
 	TEXT("Logs the last packet send/receive and TickFlush/TickDispatch times, on connection close - for debugging blocked send/recv paths."));
 
+static bool GCVarLogAllConnectionCleanup = false;
+static FAutoConsoleVariableRef CVarLogAllConnectionCleanup(TEXT("net.Connection.LogAllCleanup"), GCVarLogAllConnectionCleanup, TEXT("When true log every connection CleanUp even when it was a normal socket closure."));
+
 extern int32 GNetDormancyValidate;
 extern bool GbNetReuseReplicatorsForDormantObjects;
 
@@ -1257,10 +1260,9 @@ void UNetConnection::CleanUp()
 	}
 	Children.Empty();
 
-	if ( GetConnectionState() != USOCK_Closed )
-	{
-		UE_LOG( LogNet, Log, TEXT( "UNetConnection::Cleanup: Closing open connection. %s" ), *Describe() );
-	}
+	const bool bWasSocketClosed = GetConnectionState() == USOCK_Closed;
+
+	UE_CLOG((GCVarLogAllConnectionCleanup || !bWasSocketClosed), LogNet, Log, TEXT("UNetConnection::Cleanup: Closing open connection. %s"), *Describe());
 
 	Close(ENetCloseResult::Cleanup);
 
diff --git a/Engine/Source/Runtime/Engine/Private/NetDriver.cpp b/Engine/Source/Runtime/Engine/Private/NetDriver.cpp
--- a/Engine/Source/Runtime/Engine/Private/NetDriver.cpp
+++ b/Engine/Source/Runtime/Engine/Private/NetDriver.cpp
@@ -6335,8 +6335,7 @@ void UNetDriver::AddClientConnection(UNetConnection* NewConnection)
 
 	SCOPE_CYCLE_COUNTER(Stat_NetDriverAddClientConnection);
 
-	UE_CLOG(!DDoS.CheckLogRestrictions(), LogNet, Log, TEXT("AddClientConnection: Added client connection: %s"),
-		*NewConnection->Describe());
+	UE_CLOG(!DDoS.CheckLogRestrictions(), LogNet, Log, TEXT("AddClientConnection: Added client connection: %s"), *NewConnection->Describe());
 
 	ClientConnections.Add(NewConnection);
 
@@ -6556,29 +6555,35 @@ void UNetDriver::NotifyActorFullyDormantForConnection(AActor* Actor, UNetConnect
 
 void UNetDriver::RemoveClientConnection(UNetConnection* ClientConnectionToRemove)
 {
+	check(ClientConnectionToRemove);
 	verify(ClientConnections.Remove(ClientConnectionToRemove) == 1);
 
 	TSharedPtr<const FInternetAddr> AddrToRemove = ClientConnectionToRemove->GetRemoteAddr();
 
 	if (AddrToRemove.IsValid())
 	{
+		constexpr bool bPrintPort = true;
+
 		TSharedRef<const FInternetAddr> ConstAddrRef = AddrToRemove.ToSharedRef();
 
 		if (RecentlyDisconnectedTrackingTime > 0)
 		{
-			auto* FoundVal = MappedClientConnections.Find(ConstAddrRef);
+			const bool bIsAddressTracked = MappedClientConnections.Find(ConstAddrRef) != nullptr;
 
-			// Mark recently disconnected clients as nullptr (don't wait for GC), and keep the MappedClientConections entry for a while.
-			// Required for identifying/ignoring packets from recently disconnected clients, with the same performance as for NetConnection's (important for DDoS detection)
-			if (ensure(FoundVal != nullptr))
+			if (ensureMsgf(bIsAddressTracked, TEXT("RemoveClientConnection for %s was not in MappedClientConnections list"), *(ConstAddrRef->ToString(bPrintPort))))
 			{
+				UE_LOG(LogNet, Log, TEXT("UNetDriver::RemoveClientConnection - Removed address %s from MappedClientConnections for: %s"), *(ConstAddrRef->ToString(bPrintPort)), *ClientConnectionToRemove->Describe());
+
 				RecentlyDisconnectedClients.Add(FDisconnectedClient(ConstAddrRef, FPlatformTime::Seconds()));
 
-				*FoundVal = nullptr;
+				// Mark recently disconnected clients as nullptr (don't wait for GC), and keep the MappedClientConections entry for a while.
+				// Required for identifying/ignoring packets from recently disconnected clients, with the same performance as for NetConnection's (important for DDoS detection)
+				MappedClientConnections.Emplace(ConstAddrRef, nullptr);
 			}
 		}
 		else
 		{
+			UE_LOG(LogNet, Log, TEXT("UNetDriver::RemoveClientConnection - Removed address %s from MappedClientConnections for: %s"), *(ConstAddrRef->ToString(bPrintPort)), *ClientConnectionToRemove->Describe());
 			verify(MappedClientConnections.Remove(ConstAddrRef) == 1);
 		}
 	}
