Add option to strip debug names from streamable manager in shipping ( as a memory saving )

diff --git a/Engine/Source/Runtime/Engine/Private/StreamableManager.cpp b/Engine/Source/Runtime/Engine/Private/StreamableManager.cpp
--- a/Engine/Source/Runtime/Engine/Private/StreamableManager.cpp
+++ b/Engine/Source/Runtime/Engine/Private/StreamableManager.cpp
@@ -35,6 +35,13 @@ static FAutoConsoleVariableRef CVarStreamableFlushAllAsyncLoadRequestsOnWait(
 	ECVF_Default
 );
 
+static bool GStreamableStripDebugNameInShipping = false;
+static FAutoConsoleVariableRef CVarStreamableStripDebugNameInShipping(
+	TEXT("s.StreamableStripDebugNameInShipping"),
+	GStreamableStripDebugNameInShipping,
+	TEXT("Whether to strip the bundle names DebugName when changing bundle. Saves memory when disabled, provides more debug information when enabled."),
+	ECVF_Default
+);
 
 const FString FStreamableHandle::HandleDebugName_Preloading = FString(TEXT("Preloading"));
 const FString FStreamableHandle::HandleDebugName_AssetList = FString(TEXT("LoadAssetList"));
@@ -1392,9 +1399,18 @@ TSharedPtr<FStreamableHandle> FStreamableManager::RequestAsyncLoad(TArray<FSoftO
 	NewRequest->CompleteDelegate = MoveTemp(DelegateToCall);
 	NewRequest->OwningManager = this;
 	NewRequest->RequestedAssets = MoveTemp(TargetsToStream);
-#if (!PLATFORM_IOS && !PLATFORM_ANDROID)
-	NewRequest->DebugName = MoveTemp(DebugName);
+
+	bool bShouldStripDebugName = false;
+#if (PLATFORM_IOS || PLATFORM_ANDROID)
+	bShouldStripDebugName = true;
+#elif UE_BUILD_SHIPPING
+	bShouldStripDebugName = GStreamableStripDebugNameInShipping;
 #endif
+	if (!bShouldStripDebugName)
+	{
+		NewRequest->DebugName = MoveTemp(DebugName);
+	}
+
 	NewRequest->Priority = Priority;
 #if UE_WITH_PACKAGE_ACCESS_TRACKING
 	PackageAccessTracking_Private::FTrackedData* AccumulatedScopeData = PackageAccessTracking_Private::FPackageAccessRefScope::GetCurrentThreadAccumulatedData();
