Fixed exported templates in FTraceFilter.

diff --git a/Engine/Source/Runtime/Engine/Private/TraceFilter.cpp b/Engine/Source/Runtime/Engine/Private/TraceFilter.cpp
--- a/Engine/Source/Runtime/Engine/Private/TraceFilter.cpp
+++ b/Engine/Source/Runtime/Engine/Private/TraceFilter.cpp
@@ -272,70 +272,76 @@ FAutoConsoleCommand FlushFilterStateCommand(TEXT("TraceFilter.FlushState"), TEXT
 		})
 	);
 
-template<>
-ENGINE_API bool FTraceFilter::IsObjectTraceable</*bForceThreadSafe = */ true>(const UObject* InObject)
+template<bool bForceThreadSafe>
+bool FTraceFilter::IsObjectTraceable(const UObject* InObject)
 {
-	// Object not found in the AnnotationMap means that it is at the default value, which is bIsTraceable == true
-	return GObjectFilterAnnotations.GetAnnotationMap().Find(InObject) == nullptr;
-}
+	if constexpr (!bForceThreadSafe)
+	{
+		check(GObjectFilterAnnotations.IsLocked());
+	}
 
-template<>
-ENGINE_API bool FTraceFilter::IsObjectTraceable</*bForceThreadSafe = */ false>(const UObject* InObject)
-{
-	check(GObjectFilterAnnotations.IsLocked());
 	// Object not found in the AnnotationMap means that it is at the default value, which is bIsTraceable == true
 	return GObjectFilterAnnotations.GetAnnotationMap().Find(InObject) == nullptr;
 }
 
-template<>
-ENGINE_API void FTraceFilter::SetObjectIsTraceable</*bForceThreadSafe = */ true>(const UObject* InObject, bool bIsTraceable)
+template bool FTraceFilter::IsObjectTraceable</*bForceThreadSafe = */ false>(const UObject* InObject);
+template bool FTraceFilter::IsObjectTraceable</*bForceThreadSafe = */ true>(const UObject* InObject);
+
+template<bool bForceThreadSafe>
+void FTraceFilter::SetObjectIsTraceable(const UObject* InObject, bool bIsTraceable)
 {
 	ensure(InObject);
-		
-	FTraceFilterObjectAnnotation Annotation;
-	Annotation.bIsTraceable = bIsTraceable;
-	GObjectFilterAnnotations.AddAnnotation(InObject, Annotation);
 
-	if (bIsTraceable)
+	if constexpr (bForceThreadSafe)
 	{
-		TRACE_OBJECT(InObject);
+		FTraceFilterObjectAnnotation Annotation;
+		Annotation.bIsTraceable = bIsTraceable;
+		GObjectFilterAnnotations.AddAnnotation(InObject, Annotation);
+
+		if (bIsTraceable)
+		{
+			TRACE_OBJECT(InObject);
+		}
+	}
+	else
+	{
+		check(GObjectFilterAnnotations.IsLocked());
+		TMap<const UObjectBase*, FTraceFilterObjectAnnotation>& AnnotationMap = GObjectFilterAnnotations.GetAnnotationMap();
+		if (!bIsTraceable)
+		{
+			AnnotationMap.FindOrAdd(InObject).bIsTraceable = false;
+		}
+		else
+		{
+			AnnotationMap.Remove(InObject);
+			TRACE_OBJECT(InObject);
+		}
 	}
 }
 
-template<>
-ENGINE_API void FTraceFilter::SetObjectIsTraceable</*bForceThreadSafe = */ false>(const UObject* InObject, bool bIsTraceable)
+template void FTraceFilter::SetObjectIsTraceable</*bForceThreadSafe = */ true>(const UObject* InObject, bool bIsTraceable);
+template void FTraceFilter::SetObjectIsTraceable</*bForceThreadSafe = */ false>(const UObject* InObject, bool bIsTraceable);
+
+template<bool bForceThreadSafe>
+void FTraceFilter::MarkObjectTraceable(const UObject* InObject)
 {
 	ensure(InObject);
 
-	check(GObjectFilterAnnotations.IsLocked());
-	TMap<const UObjectBase*, FTraceFilterObjectAnnotation>& AnnotationMap = GObjectFilterAnnotations.GetAnnotationMap();
-	if (!bIsTraceable)
+	if constexpr (bForceThreadSafe)
 	{
-		AnnotationMap.FindOrAdd(InObject).bIsTraceable = false;
+		FTraceFilterObjectAnnotation Annotation;
+		Annotation.bIsTraceable = true;
+		GObjectFilterAnnotations.AddAnnotation(InObject, Annotation);
 	}
 	else
 	{
-		AnnotationMap.Remove(InObject);
-		TRACE_OBJECT(InObject);
-	}	
-}
-
-template<>
-ENGINE_API void FTraceFilter::MarkObjectTraceable</*bForceThreadSafe = */ true>(const UObject* InObject)
-{
-	ensure(InObject);	
-	FTraceFilterObjectAnnotation Annotation;
-	Annotation.bIsTraceable = true;
-	GObjectFilterAnnotations.AddAnnotation(InObject, Annotation);
+		check(GObjectFilterAnnotations.IsLocked());
+		SetObjectIsTraceable(InObject, true);
+	}
 }
 
-template<>
-ENGINE_API void FTraceFilter::MarkObjectTraceable</*bForceThreadSafe = */ false>(const UObject* InObject)
-{
-	ensure(InObject);
-	check(GObjectFilterAnnotations.IsLocked());
-	SetObjectIsTraceable(InObject, true);
-}
+template void FTraceFilter::MarkObjectTraceable</*bForceThreadSafe = */ true>(const UObject* InObject);
+template void FTraceFilter::MarkObjectTraceable</*bForceThreadSafe = */ false>(const UObject* InObject);
 
 void FTraceFilter::Init()
 {
diff --git a/Engine/Source/Runtime/Engine/Public/TraceFilter.h b/Engine/Source/Runtime/Engine/Public/TraceFilter.h
--- a/Engine/Source/Runtime/Engine/Public/TraceFilter.h
+++ b/Engine/Source/Runtime/Engine/Public/TraceFilter.h
@@ -90,6 +90,13 @@ struct FTraceFilter
 	static ENGINE_API void Unlock();
 };
 
+extern template ENGINE_API void FTraceFilter::SetObjectIsTraceable<true>(const UObject* InObject, bool bIsTraceable);
+extern template ENGINE_API void FTraceFilter::SetObjectIsTraceable<false>(const UObject* InObject, bool bIsTraceable);
+extern template ENGINE_API void FTraceFilter::MarkObjectTraceable<true>(const UObject* InObject);
+extern template ENGINE_API void FTraceFilter::MarkObjectTraceable<false>(const UObject* InObject);
+extern template ENGINE_API bool FTraceFilter::IsObjectTraceable<true>(const UObject* InObject);
+extern template ENGINE_API bool FTraceFilter::IsObjectTraceable<false>(const UObject* InObject);
+
 #define CAN_TRACE_OBJECT(Object) \
 	FTraceFilter::TObjectFilter::CanTrace(Object)
 
