HWRT: Update how skinned geometry update is performed:
* Deprecate old Commit function and related cvars
* Simplify the new Commit to run updated for all BLAS at once.
* Allow setting limit to number of primitives updated per frame with r.SkinCache.RayTracing.MaxUpdatePrimitivesPerFrame.
* Keep track of when the last update happened so the next frame we can update the oldest BLAS first.
* Always use transient scratch buffer when building.
#tests CitySample and using r.SkinCache.RayTracing.MaxUpdatePrimitivesPerFrame

diff --git a/Engine/Source/Runtime/Engine/Private/RayTracingSkinnedGeometry.cpp b/Engine/Source/Runtime/Engine/Private/RayTracingSkinnedGeometry.cpp
--- a/Engine/Source/Runtime/Engine/Private/RayTracingSkinnedGeometry.cpp
+++ b/Engine/Source/Runtime/Engine/Private/RayTracingSkinnedGeometry.cpp
@@ -10,19 +10,19 @@
 DECLARE_GPU_STAT(SkinnedGeometryBuildBLAS);
 DECLARE_GPU_STAT(SkinnedGeometryUpdateBLAS);
 
-static int32 GMemoryLimitForBatchedRayTracingGeometryUpdates = 512;
-FAutoConsoleVariableRef CVarSkinnedGeometryMemoryLimitForBatchedRayTracingGeometryUpdates(
-	TEXT("r.SkinCache.MemoryLimitForBatchedRayTracingGeometryUpdates"),
-	GMemoryLimitForBatchedRayTracingGeometryUpdates,
-	TEXT(""),
+DECLARE_DWORD_COUNTER_STAT(TEXT("Ray tracing skinned update primitives"), STAT_RayTracingSkinnedUpdatePrimitives, STATGROUP_SceneRendering);
+
+static TAutoConsoleVariable<int32> CVarSkinCacheRayTracingMaxUpdatePrimitivesPerFrame(
+	TEXT("r.SkinCache.RayTracing.MaxUpdatePrimitivesPerFrame"),
+	-1,
+	TEXT("Sets the skinned ray tracing acceleration structure build budget in terms of maximum number of updated triangles per frame (<= 0 then disabled and all acceleration structures are updated - default)"),
 	ECVF_RenderThreadSafe
 );
 
-static int32 GRayTracingUseTransientForScratch = 0;
-FAutoConsoleVariableRef CVarSkinnedGeometryRayTracingUseTransientForScratch(
-	TEXT("r.SkinCache.RayTracingUseTransientForScratch"),
-	GRayTracingUseTransientForScratch,
-	TEXT("Use Transient memory for BLAS scratch allocation to reduce memory footprint and allocation overhead."),
+static TAutoConsoleVariable<int32> CVarSkinCacheRayTracingAsyncBuild(
+	TEXT("r.SkinCache.RayTracing.AsyncBuild"),
+	0,
+	TEXT("Whether to build skinned ray tracing acceleration structures on async compute queue (default = 0).\n"),
 	ECVF_RenderThreadSafe
 );
 
@@ -72,7 +72,7 @@ uint32 FRayTracingSkinnedGeometryUpdateQueue::ComputeScratchBufferSize() const
 	const uint64 ScratchAlignment = GRHIRayTracingScratchBufferAlignment;
 	uint32 ScratchBLASSize = 0;
 
-	if (ToUpdate.Num() && GRayTracingUseTransientForScratch > 0)
+	if (ToUpdate.Num())
 	{
 		for (TMap<FRayTracingGeometry*, FRayTracingUpdateInfo>::TRangedForConstIterator Iter = ToUpdate.begin(); Iter != ToUpdate.end(); ++Iter)
 		{			
@@ -91,7 +91,6 @@ void FRayTracingSkinnedGeometryUpdateQueue::Commit(FRHICommandList& RHICmdList,
 	if (ToUpdate.Num())
 	{
 		FScopeLock Lock(&CS);
-
 		// Track the amount of primitives which need to be build/updated in a single batch
 		uint64 PrimitivesToUpdates = 0;
 		TArray<FRayTracingGeometryBuildParams> BatchedBuildParams;
@@ -203,14 +202,109 @@ END_SHADER_PARAMETER_STRUCT()
 
 void FRayTracingSkinnedGeometryUpdateQueue::Commit(FRDGBuilder& GraphBuilder)
 {
-	// Find out the total BLAS scratch size and allocate transient RDG buffer.
+	TRACE_CPUPROFILER_EVENT_SCOPE(FRayTracingSkinnedGeometryUpdateQueue::Commit);
+
+	TArray<FRayTracingGeometryBuildParams> GeometryBuildRequests;
+	GeometryBuildRequests.Reserve(ToUpdate.Num());
+
+	TArray<FRayTracingGeometry*> GeometriesToUpdate;
+	GeometriesToUpdate.Reserve(ToUpdate.Num());
+
+	uint32 BLASScratchSize = 0;
+
+	{
+		FScopeLock Lock(&CS);
+
+		for (TMap<FRayTracingGeometry*, FRayTracingUpdateInfo>::TRangedForIterator Iter = ToUpdate.begin(); Iter != ToUpdate.end(); ++Iter)
+		{
+			FRayTracingGeometry* RayTracingGeometry = Iter.Key();
+			FRayTracingUpdateInfo& UpdateInfo = Iter.Value();
+
+			FRayTracingGeometryBuildParams BuildParams;
+			BuildParams.Geometry = RayTracingGeometry->GetRHI();
+			BuildParams.BuildMode = UpdateInfo.BuildMode;
+			BuildParams.Segments = RayTracingGeometry->Initializer.Segments;
+
+			if (BuildParams.BuildMode == EAccelerationStructureBuildMode::Build)
+			{
+				GeometryBuildRequests.Add(BuildParams);
+
+				BLASScratchSize += UpdateInfo.ScratchSize;
+
+				RayTracingGeometry->LastUpdatedFrame = GFrameCounterRenderThread;
+			}
+			else
+			{
+				GeometriesToUpdate.Add(RayTracingGeometry);
+			}
+		}
+
+		// Clear working data
+		ToUpdate.Reset();
+		EstimatedMemoryPendingRelease = 0;
+	}	
+
+	const int32 MaxUpdatePrimitivesPerFrame = CVarSkinCacheRayTracingMaxUpdatePrimitivesPerFrame.GetValueOnRenderThread();
+
+	int32 NumUpdatedPrimitives = 0;
+
+	if (MaxUpdatePrimitivesPerFrame <= 0)
+	{		
+		for (FRayTracingGeometry* RayTracingGeometry : GeometriesToUpdate)
+		{
+			FRayTracingGeometryBuildParams BuildParams;
+			BuildParams.Geometry = RayTracingGeometry->GetRHI();
+			BuildParams.BuildMode = EAccelerationStructureBuildMode::Update;
+			BuildParams.Segments = RayTracingGeometry->Initializer.Segments;
+
+			GeometryBuildRequests.Add(BuildParams);
+
+			RayTracingGeometry->LastUpdatedFrame = GFrameCounterRenderThread;
+
+			BLASScratchSize += RayTracingGeometry->GetRHI()->GetSizeInfo().UpdateScratchSize;
+
+			NumUpdatedPrimitives += RayTracingGeometry->Initializer.TotalPrimitiveCount;			
+		}		
+	}
+	else
+	{
+		GeometriesToUpdate.Sort([](const FRayTracingGeometry& InLHS, const FRayTracingGeometry& InRHS)
+			{
+				return InLHS.LastUpdatedFrame < InRHS.LastUpdatedFrame;
+			});			
+
+		for (FRayTracingGeometry* RayTracingGeometry : GeometriesToUpdate)
+		{
+			const int32 NumPrimitives = RayTracingGeometry->Initializer.TotalPrimitiveCount;
+
+			if (NumUpdatedPrimitives + NumPrimitives > MaxUpdatePrimitivesPerFrame)
+			{
+				break;
+			}
+						
+			FRayTracingGeometryBuildParams BuildParams;
+			BuildParams.Geometry = RayTracingGeometry->GetRHI();
+			BuildParams.BuildMode = EAccelerationStructureBuildMode::Update;
+			BuildParams.Segments = RayTracingGeometry->Initializer.Segments;
+
+			GeometryBuildRequests.Add(BuildParams);
+
+			RayTracingGeometry->LastUpdatedFrame = GFrameCounterRenderThread;
+
+			BLASScratchSize += RayTracingGeometry->GetRHI()->GetSizeInfo().UpdateScratchSize;
+
+			NumUpdatedPrimitives += NumPrimitives;
+		}
+	}
+
+	INC_DWORD_STAT_BY(STAT_RayTracingSkinnedUpdatePrimitives, NumUpdatedPrimitives);
+
 	FRDGBufferRef SharedScratchBuffer = nullptr;
-	
-	const uint32 BLASScratchSize = ComputeScratchBufferSize();
+
 	if (BLASScratchSize > 0)
 	{
 		const uint32 ScratchAlignment = GRHIRayTracingScratchBufferAlignment;
-		
+
 		FRDGBufferDesc ScratchBufferDesc;
 		ScratchBufferDesc.Usage = EBufferUsageFlags::RayTracingScratch | EBufferUsageFlags::StructuredBuffer;
 		ScratchBufferDesc.BytesPerElement = ScratchAlignment;
@@ -224,11 +318,26 @@ void FRayTracingSkinnedGeometryUpdateQueue::Commit(FRDGBuilder& GraphBuilder)
 
 	RDG_GPU_MASK_SCOPE(GraphBuilder, FRHIGPUMask::All());
 
-	GraphBuilder.AddPass(RDG_EVENT_NAME("CommitRayTracingSkinnedGeometryUpdates"), BLASUpdateParams, ERDGPassFlags::Compute | ERDGPassFlags::NeverCull,
-		[this, SharedScratchBuffer](FRHICommandList& RHICmdList)
-		{
-			Commit(RHICmdList, SharedScratchBuffer ? SharedScratchBuffer->GetRHI() : nullptr);
-		});
+	if (GeometryBuildRequests.Num())
+	{
+		const bool bRayTracingAsyncBuild = CVarSkinCacheRayTracingAsyncBuild.GetValueOnRenderThread() != 0 && GRHISupportsRayTracingAsyncBuildAccelerationStructure;
+		const ERDGPassFlags ComputePassFlags = bRayTracingAsyncBuild ? ERDGPassFlags::AsyncCompute : ERDGPassFlags::Compute;
+
+		GraphBuilder.AddPass(RDG_EVENT_NAME("CommitRayTracingSkinnedGeometryUpdates"), BLASUpdateParams, ComputePassFlags | ERDGPassFlags::NeverCull,
+			[
+				BuildRequests = MoveTemp(GeometryBuildRequests),
+				SharedScratchBuffer
+			](FRHICommandList& RHICmdList)
+			{
+				SCOPED_GPU_STAT(RHICmdList, SkinnedGeometryBuildBLAS);
+				SCOPED_DRAW_EVENT(RHICmdList, SkinnedGeometryBuildBLAS);
+
+				FRHIBufferRange ScratchBufferRange;
+				ScratchBufferRange.Buffer = SharedScratchBuffer->GetRHI();
+				ScratchBufferRange.Offset = 0;
+				RHICmdList.BuildAccelerationStructures(BuildRequests, ScratchBufferRange);
+			});
+	}
 }
 
 #endif // RHI_RAYTRACING
diff --git a/Engine/Source/Runtime/Engine/Public/RayTracingSkinnedGeometry.h b/Engine/Source/Runtime/Engine/Public/RayTracingSkinnedGeometry.h
--- a/Engine/Source/Runtime/Engine/Public/RayTracingSkinnedGeometry.h
+++ b/Engine/Source/Runtime/Engine/Public/RayTracingSkinnedGeometry.h
@@ -31,6 +31,7 @@ public:
 	ENGINE_API uint32 ComputeScratchBufferSize() const;
 	
 	/** Commit all pending work. Requires a scratch buffer that is at least as big as the size provided by ComputeScratchBufferSize(). */
+	UE_DEPRECATED(5.5, "This function has been deprecated. Use the version that takes GraphBuilder as input.")
 	ENGINE_API void Commit(FRHICommandList& RHICmdList, FRHIBuffer* ScratchBuffer);
 
 	/** Commit all pending work using render graph. This allocates a transient scratch buffer internally. */
diff --git a/Engine/Source/Runtime/RenderCore/Public/RayTracingGeometry.h b/Engine/Source/Runtime/RenderCore/Public/RayTracingGeometry.h
--- a/Engine/Source/Runtime/RenderCore/Public/RayTracingGeometry.h
+++ b/Engine/Source/Runtime/RenderCore/Public/RayTracingGeometry.h
@@ -54,6 +54,9 @@ public:
 	*/
 	int64 DynamicGeometrySharedBufferGenerationID = NonSharedVertexBuffers;
 
+	// Last frame when geometry was updated (only skinned geometry)
+	uint64 LastUpdatedFrame = 0;	
+
 	FRayTracingGeometryInitializer Initializer;
 
 	UE_DEPRECATED(5.5, "Use GetRHI() instead.")
