Fix mesh space fixup to always exclude root bone

diff --git a/Engine/Source/Runtime/Engine/Private/Animation/AnimationDecompression.cpp b/Engine/Source/Runtime/Engine/Private/Animation/AnimationDecompression.cpp
--- a/Engine/Source/Runtime/Engine/Private/Animation/AnimationDecompression.cpp
+++ b/Engine/Source/Runtime/Engine/Private/Animation/AnimationDecompression.cpp
@@ -318,8 +318,8 @@ void DecompressPose(FCompactPose& OutPose,
 		// and the output pose will contain an identity delta which isn't what we want. As such, bones missing from
 		// the sequence have their rotation set to their parent.
 
-		// If the first track is the root, we skip it since it has no parent (its delta value is fine as the identity)
-		for (FCompactPoseBoneIndex CompactBoneIndex(bFirstTrackIsRootBone ? 1 : 0); CompactBoneIndex < NumCompactBones; ++CompactBoneIndex)
+		// We always skip the root since it has no parent (its delta value is fine as the identity)
+		for (FCompactPoseBoneIndex CompactBoneIndex(1); CompactBoneIndex < NumCompactBones; ++CompactBoneIndex)
 		{
 			if (!AnimatedCompactRotations[CompactBoneIndex.GetInt()])
 			{
