[PrimitiveComponent] Prevent extra work and actor hydration in ShouldComponentIgnoreHitResult when MoveIgnoreActors list is empty

diff --git a/Engine/Source/Runtime/Engine/Private/Components/PrimitiveComponent.cpp b/Engine/Source/Runtime/Engine/Private/Components/PrimitiveComponent.cpp
--- a/Engine/Source/Runtime/Engine/Private/Components/PrimitiveComponent.cpp
+++ b/Engine/Source/Runtime/Engine/Private/Components/PrimitiveComponent.cpp
@@ -2654,8 +2654,11 @@ void UPrimitiveComponent::SetMoveIgnoreMask(FMaskFilter InMoveIgnoreMask)
 bool UPrimitiveComponent::ShouldComponentIgnoreHitResult(FHitResult const& TestHit, EMoveComponentFlags MoveFlags)
 {
 	// Check if the hit actors root actor is in the ignore array
-	if (MoveFlags & MOVECOMP_CheckBlockingRootActorInIgnoreList)
+	if (MoveFlags & MOVECOMP_CheckBlockingRootActorInIgnoreList && MoveIgnoreActors.Num())
 	{
+		// Ideally we should consider using GetCachedActor here to prevent actor hydration
+		// but that could have side effects and will require more investigation.
+		// Note that if the given instance is not hydrated yet, then it can't be in MoveIgnoreActors.
 		AActor const* const HitActor = TestHit.HitObjectHandle.FetchActor();
 		if (HitActor)
 		{
