[InstancedStaticMesh] When an ISM has no PerInstanceSMData, use the mesh bounds to test for 'ComponentIsTouching' tests (editor selection mechanism)

diff --git a/Engine/Source/Runtime/Engine/Private/InstancedStaticMesh.cpp b/Engine/Source/Runtime/Engine/Private/InstancedStaticMesh.cpp
--- a/Engine/Source/Runtime/Engine/Private/InstancedStaticMesh.cpp
+++ b/Engine/Source/Runtime/Engine/Private/InstancedStaticMesh.cpp
@@ -4783,6 +4783,11 @@ bool UInstancedStaticMeshComponent::ComponentIsTouchingSelectionBox(const FBox&
 		}
 	}
 
+	if (PerInstanceSMData.Num() == 0)
+	{
+		return Super::ComponentIsTouchingSelectionBox(InSelBBox, bConsiderOnlyBSP, bMustEncompassEntireComponent);
+	}
+
 	return bMustEncompassEntireComponent;
 }
 
@@ -4805,6 +4810,11 @@ bool UInstancedStaticMeshComponent::ComponentIsTouchingSelectionFrustum(const FC
 		}
 	}
 
+	if (PerInstanceSMData.Num() == 0)
+	{
+		return Super::ComponentIsTouchingSelectionFrustum(InFrustum, bConsiderOnlyBSP, bMustEncompassEntireComponent);		
+	}
+
 	return bMustEncompassEntireComponent;
 }
 
