[AsyncLoading] Provide better API for AssetManager and StreamableManager.
Let user to pass in TArray instead of the TSet which let them avoid creating of unnessesary array copy in the process.

diff --git a/Engine/Source/Runtime/Engine/Classes/Engine/AssetManager.h b/Engine/Source/Runtime/Engine/Classes/Engine/AssetManager.h
--- a/Engine/Source/Runtime/Engine/Classes/Engine/AssetManager.h
+++ b/Engine/Source/Runtime/Engine/Classes/Engine/AssetManager.h
@@ -311,7 +311,8 @@ public:
 	ENGINE_API void GetPrimaryAssetBundleStateMap(TMap<FPrimaryAssetId, TArray<FName>>& BundleStateMap, bool bForceCurrent = false) const;
 
 	/**
-	 * Fills in a set of object paths with the assets that need to be loaded, for a given Primary Asset and bundle list
+	 * Fills in a set of object paths with the assets that need to be loaded, for a given Primary Asset and bundle list.
+	 * Prefer to use GetPrimaryAssetLoadList instead, because this overload creates additional allocation and is here just for backward compatibility.
 	 *
 	 * @param OutAssetLoadSet	Set that will have asset paths added to it
 	 * @param PrimaryAssetId	Asset that would be loaded
@@ -319,7 +320,18 @@ public:
 	 * @param bLoadRecursive	If true, this will call RecursivelyExpandBundleData and recurse into sub bundles of other primary assets loaded by a bundle reference
 	 * @return					True if primary asset id was found
 	 */
-	ENGINE_API bool GetPrimaryAssetLoadSet(TArray<FSoftObjectPath>& OutAssetLoadSet, const FPrimaryAssetId& PrimaryAssetId, const TArray<FName>& LoadBundles, bool bLoadRecursive) const;
+	ENGINE_API bool GetPrimaryAssetLoadSet(TSet<FSoftObjectPath>& OutAssetLoadSet, const FPrimaryAssetId& PrimaryAssetId, const TArray<FName>& LoadBundles, bool bLoadRecursive) const;
+
+	/**
+	 * Fills in a array of unique object paths with the assets that need to be loaded, for a given Primary Asset and bundle list
+	 *
+	 * @param OutAssetLoadList	TArray that will have asset paths added to it
+	 * @param PrimaryAssetId	Asset that would be loaded
+	 * @param LoadBundles		List of bundles to load for those assets
+	 * @param bLoadRecursive	If true, this will call RecursivelyExpandBundleData and recurse into sub bundles of other primary assets loaded by a bundle reference
+	 * @return					True if primary asset id was found
+	 */
+	ENGINE_API bool GetPrimaryAssetLoadList(TArray<FSoftObjectPath>& OutAssetLoadList, const FPrimaryAssetId& PrimaryAssetId, const TArray<FName>& LoadBundles, bool bLoadRecursive) const;
 
 	/**
 	 * Preloads data for a set of assets in a specific bundle state, and returns a handle you must keep active.
@@ -335,7 +347,8 @@ public:
 	ENGINE_API virtual TSharedPtr<FStreamableHandle> PreloadPrimaryAssets(const TArray<FPrimaryAssetId>& AssetsToLoad, const TArray<FName>& LoadBundles, bool bLoadRecursive, FStreamableDelegate DelegateToCall = FStreamableDelegate(), TAsyncLoadPriority Priority = FStreamableManager::DefaultAsyncLoadPriority);
 
 	/** Quick wrapper to async load some non primary assets with the primary streamable manager. This will not auto release the handle, release it if needed */
-	ENGINE_API virtual TSharedPtr<FStreamableHandle> LoadAssetList(TArray<FSoftObjectPath> AssetList, FStreamableDelegate DelegateToCall = FStreamableDelegate(), TAsyncLoadPriority Priority = FStreamableManager::DefaultAsyncLoadPriority, const FString& DebugName = FStreamableHandle::HandleDebugName_AssetList);
+	ENGINE_API virtual TSharedPtr<FStreamableHandle> LoadAssetList(const TArray<FSoftObjectPath>& AssetList, FStreamableDelegate DelegateToCall = FStreamableDelegate(), TAsyncLoadPriority Priority = FStreamableManager::DefaultAsyncLoadPriority, const FString& DebugName = FStreamableHandle::HandleDebugName_AssetList);
+	ENGINE_API virtual TSharedPtr<FStreamableHandle> LoadAssetList(TArray<FSoftObjectPath>&& AssetList, FStreamableDelegate DelegateToCall = FStreamableDelegate(), TAsyncLoadPriority Priority = FStreamableManager::DefaultAsyncLoadPriority, const FString& DebugName = FStreamableHandle::HandleDebugName_AssetList);
 
 	/** Returns a single AssetBundleInfo, matching Scope and Name */
 	ENGINE_API virtual FAssetBundleEntry GetAssetBundleEntry(const FPrimaryAssetId& BundleScope, FName BundleName) const;
diff --git a/Engine/Source/Runtime/Engine/Classes/Engine/StreamableManager.h b/Engine/Source/Runtime/Engine/Classes/Engine/StreamableManager.h
--- a/Engine/Source/Runtime/Engine/Classes/Engine/StreamableManager.h
+++ b/Engine/Source/Runtime/Engine/Classes/Engine/StreamableManager.h
@@ -487,11 +487,13 @@ struct FStreamableManager : public FGCObject
 	 * @param bStartStalled			If true, the handle will start in a stalled state and will not attempt to actually async load until StartStalledHandle is called on it
 	 * @param DebugName				Name of this handle, will be reported in debug tools
 	 */
-	ENGINE_API TSharedPtr<FStreamableHandle> RequestAsyncLoad(TArray<FSoftObjectPath> TargetsToStream, FStreamableDelegate DelegateToCall = FStreamableDelegate(), TAsyncLoadPriority Priority = DefaultAsyncLoadPriority, bool bManageActiveHandle = false, bool bStartStalled = false, FString DebugName = TEXT("ArrayDelegate"));
+	ENGINE_API TSharedPtr<FStreamableHandle> RequestAsyncLoad(const TArray<FSoftObjectPath>& TargetsToStream, FStreamableDelegate DelegateToCall = FStreamableDelegate(), TAsyncLoadPriority Priority = DefaultAsyncLoadPriority, bool bManageActiveHandle = false, bool bStartStalled = false, FString DebugName = TEXT("ArrayDelegate"));
+	ENGINE_API TSharedPtr<FStreamableHandle> RequestAsyncLoad(TArray<FSoftObjectPath>&& TargetsToStream, FStreamableDelegate DelegateToCall = FStreamableDelegate(), TAsyncLoadPriority Priority = DefaultAsyncLoadPriority, bool bManageActiveHandle = false, bool bStartStalled = false, FString DebugName = TEXT("ArrayDelegate"));
 	ENGINE_API TSharedPtr<FStreamableHandle> RequestAsyncLoad(const FSoftObjectPath& TargetToStream, FStreamableDelegate DelegateToCall = FStreamableDelegate(), TAsyncLoadPriority Priority = DefaultAsyncLoadPriority, bool bManageActiveHandle = false, bool bStartStalled = false, FString DebugName = TEXT("SingleDelegate"));
 
 	/** Lambda Wrappers. Be aware that Callback may go off multiple seconds in the future. */
-	ENGINE_API TSharedPtr<FStreamableHandle> RequestAsyncLoad(TArray<FSoftObjectPath> TargetsToStream, TFunction<void()>&& Callback, TAsyncLoadPriority Priority = DefaultAsyncLoadPriority, bool bManageActiveHandle = false, bool bStartStalled = false, FString DebugName = TEXT("ArrayLambda"));
+	ENGINE_API TSharedPtr<FStreamableHandle> RequestAsyncLoad(const TArray<FSoftObjectPath>& TargetsToStream, TFunction<void()>&& Callback, TAsyncLoadPriority Priority = DefaultAsyncLoadPriority, bool bManageActiveHandle = false, bool bStartStalled = false, FString DebugName = TEXT("ArrayLambda"));
+	ENGINE_API TSharedPtr<FStreamableHandle> RequestAsyncLoad(TArray<FSoftObjectPath>&& TargetsToStream, TFunction<void()>&& Callback, TAsyncLoadPriority Priority = DefaultAsyncLoadPriority, bool bManageActiveHandle = false, bool bStartStalled = false, FString DebugName = TEXT("ArrayLambda"));
 	ENGINE_API TSharedPtr<FStreamableHandle> RequestAsyncLoad(const FSoftObjectPath& TargetToStream, TFunction<void()>&& Callback, TAsyncLoadPriority Priority = DefaultAsyncLoadPriority, bool bManageActiveHandle = false, bool bStartStalled = false, FString DebugName = TEXT("SingleLambda"));
 
 	/** 
@@ -501,7 +503,8 @@ struct FStreamableManager : public FGCObject
 	 * @param bManageActiveHandle	If true, the manager will keep the streamable handle active until explicitly released
 	 * @param DebugName				Name of this handle, will be reported in debug tools
 	 */
-	ENGINE_API TSharedPtr<FStreamableHandle> RequestSyncLoad(TArray<FSoftObjectPath> TargetsToStream, bool bManageActiveHandle = false, FString DebugName = TEXT("RequestSyncLoad Array"));
+	ENGINE_API TSharedPtr<FStreamableHandle> RequestSyncLoad(const TArray<FSoftObjectPath>& TargetsToStream, bool bManageActiveHandle = false, FString DebugName = TEXT("RequestSyncLoad Array"));
+	ENGINE_API TSharedPtr<FStreamableHandle> RequestSyncLoad(TArray<FSoftObjectPath>&& TargetsToStream, bool bManageActiveHandle = false, FString DebugName = TEXT("RequestSyncLoad Array"));
 	ENGINE_API TSharedPtr<FStreamableHandle> RequestSyncLoad(const FSoftObjectPath& TargetToStream, bool bManageActiveHandle = false, FString DebugName = TEXT("RequestSyncLoad Single"));
 
 	/** 
diff --git a/Engine/Source/Runtime/Engine/Private/AssetManager.cpp b/Engine/Source/Runtime/Engine/Private/AssetManager.cpp
--- a/Engine/Source/Runtime/Engine/Private/AssetManager.cpp
+++ b/Engine/Source/Runtime/Engine/Private/AssetManager.cpp
@@ -1905,7 +1905,15 @@ TSharedPtr<FStreamableHandle> UAssetManager::ChangeBundleStateForMatchingPrimary
 	return nullptr;
 }
 
-bool UAssetManager::GetPrimaryAssetLoadSet(TArray<FSoftObjectPath>& OutAssetLoadSet, const FPrimaryAssetId& PrimaryAssetId, const TArray<FName>& LoadBundles, bool bLoadRecursive) const
+bool UAssetManager::GetPrimaryAssetLoadSet(TSet<FSoftObjectPath>& OutAssetLoadSet, const FPrimaryAssetId& PrimaryAssetId, const TArray<FName>& LoadBundles, bool bLoadRecursive) const
+{
+	TArray<FSoftObjectPath> Array;
+	bool bReturnValue = GetPrimaryAssetLoadList(Array, PrimaryAssetId, LoadBundles, bLoadRecursive);
+	OutAssetLoadSet.Append(Array);
+	return bReturnValue;
+}
+
+bool UAssetManager::GetPrimaryAssetLoadList(TArray<FSoftObjectPath>& OutAssetLoadList, const FPrimaryAssetId& PrimaryAssetId, const TArray<FName>& LoadBundles, bool bLoadRecursive) const
 {
 	const FPrimaryAssetData* NameData = GetNameData(PrimaryAssetId);
 	if (NameData)
@@ -1915,7 +1923,7 @@ bool UAssetManager::GetPrimaryAssetLoadSet(TArray<FSoftObjectPath>& OutAssetLoad
 		if (!AssetPath.IsNull())
 		{
 			// Dynamic types can have no base asset path
-			OutAssetLoadSet.AddUnique(AssetPath);
+			OutAssetLoadList.AddUnique(AssetPath);
 		}
 
 		// Construct a temporary bundle data with the bundles specified
@@ -1939,13 +1947,13 @@ bool UAssetManager::GetPrimaryAssetLoadSet(TArray<FSoftObjectPath>& OutAssetLoad
 		{
 			for (const FTopLevelAssetPath& Path : Entry.AssetPaths)
 			{
-				OutAssetLoadSet.AddUnique(FSoftObjectPath(Path));
+				OutAssetLoadList.AddUnique(FSoftObjectPath(Path));
 			}
 		}
 	}
 	else
 	{
-		WarnAboutInvalidPrimaryAsset(PrimaryAssetId, TEXT("GetPrimaryAssetLoadSet failed to find NameData"));
+		WarnAboutInvalidPrimaryAsset(PrimaryAssetId, TEXT("GetPrimaryAssetLoadList failed to find NameData"));
 	}
 	return NameData != nullptr;
 }
@@ -1962,7 +1970,7 @@ TSharedPtr<FStreamableHandle> UAssetManager::PreloadPrimaryAssets(const TArray<F
 
 	for (const FPrimaryAssetId& PrimaryAssetId : AssetsToLoad)
 	{
-		if (GetPrimaryAssetLoadSet(PathsToLoad, PrimaryAssetId, LoadBundles, bLoadRecursive))
+		if (GetPrimaryAssetLoadList(PathsToLoad, PrimaryAssetId, LoadBundles, bLoadRecursive))
 		{
 			if (DebugValid.Len() < MaxDebugWarningLen)
 			{
@@ -2198,7 +2206,13 @@ int32 UAssetManager::UnloadPrimaryAssetsWithType(FPrimaryAssetType PrimaryAssetT
 	return UnloadPrimaryAssets(Assets);
 }
 
-TSharedPtr<FStreamableHandle> UAssetManager::LoadAssetList(TArray<FSoftObjectPath> AssetList, FStreamableDelegate DelegateToCall, TAsyncLoadPriority Priority, const FString& DebugName)
+TSharedPtr<FStreamableHandle> UAssetManager::LoadAssetList(const TArray<FSoftObjectPath>& AssetList, FStreamableDelegate DelegateToCall, TAsyncLoadPriority Priority, const FString& DebugName)
+{
+	// Make an explicit Copy of the input array
+	return LoadAssetList(TArray<FSoftObjectPath>(AssetList), DelegateToCall, Priority, DebugName);
+}
+
+TSharedPtr<FStreamableHandle> UAssetManager::LoadAssetList(TArray<FSoftObjectPath>&& AssetList, FStreamableDelegate DelegateToCall, TAsyncLoadPriority Priority, const FString& DebugName)
 {
 	TSharedPtr<FStreamableHandle> NewHandle;
 	TArray<int32> MissingChunks, ErrorChunks;
@@ -2218,12 +2232,12 @@ TSharedPtr<FStreamableHandle> UAssetManager::LoadAssetList(TArray<FSoftObjectPat
 	// SynchronousLoad doesn't make sense if chunks are missing
 	if (bShouldUseSynchronousLoad && MissingChunks.Num() == 0)
 	{
-		NewHandle = StreamableManager.RequestSyncLoad(MoveTemp(AssetList), false, DebugName);
+		NewHandle = StreamableManager.RequestSyncLoad(Forward<TArray<FSoftObjectPath>>(AssetList), false, DebugName);
 		FStreamableHandle::ExecuteDelegate(MoveTemp(DelegateToCall));
 	}
 	else
 	{
-		NewHandle = StreamableManager.RequestAsyncLoad(MoveTemp(AssetList), MoveTemp(DelegateToCall), Priority, false, MissingChunks.Num() > 0, DebugName);
+		NewHandle = StreamableManager.RequestAsyncLoad(Forward<TArray<FSoftObjectPath>>(AssetList), MoveTemp(DelegateToCall), Priority, false, MissingChunks.Num() > 0, DebugName);
 
 		if (MissingChunks.Num() > 0 && NewHandle.IsValid())
 		{
diff --git a/Engine/Source/Runtime/Engine/Private/StreamableManager.cpp b/Engine/Source/Runtime/Engine/Private/StreamableManager.cpp
--- a/Engine/Source/Runtime/Engine/Private/StreamableManager.cpp
+++ b/Engine/Source/Runtime/Engine/Private/StreamableManager.cpp
@@ -1390,7 +1390,13 @@ FStreamable* FStreamableManager::StreamInternal(const FSoftObjectPath& InTargetN
 	return Existing;
 }
 
-TSharedPtr<FStreamableHandle> FStreamableManager::RequestAsyncLoad(TArray<FSoftObjectPath> TargetsToStream, FStreamableDelegate DelegateToCall, TAsyncLoadPriority Priority, bool bManageActiveHandle, bool bStartStalled, FString DebugName)
+TSharedPtr<FStreamableHandle> FStreamableManager::RequestAsyncLoad(const TArray<FSoftObjectPath>& TargetsToStream, FStreamableDelegate DelegateToCall, TAsyncLoadPriority Priority, bool bManageActiveHandle, bool bStartStalled, FString DebugName)
+{
+	// Make an explicit Copy of the input array
+	return RequestAsyncLoad(TArray<FSoftObjectPath>(TargetsToStream), MoveTemp(DelegateToCall), Priority, bManageActiveHandle, bStartStalled, MoveTemp(DebugName));
+}
+
+TSharedPtr<FStreamableHandle> FStreamableManager::RequestAsyncLoad(TArray<FSoftObjectPath>&& TargetsToStream, FStreamableDelegate DelegateToCall, TAsyncLoadPriority Priority, bool bManageActiveHandle, bool bStartStalled, FString DebugName)
 {
 	LLM_SCOPE(ELLMTag::StreamingManager);
 
@@ -1398,7 +1404,7 @@ TSharedPtr<FStreamableHandle> FStreamableManager::RequestAsyncLoad(TArray<FSoftO
 	TSharedRef<FStreamableHandle> NewRequest = MakeShareable(new FStreamableHandle());
 	NewRequest->CompleteDelegate = MoveTemp(DelegateToCall);
 	NewRequest->OwningManager = this;
-	NewRequest->RequestedAssets = MoveTemp(TargetsToStream);
+	NewRequest->RequestedAssets = Forward<TArray<FSoftObjectPath>>(TargetsToStream);
 
 	bool bShouldStripDebugName = false;
 #if (PLATFORM_IOS || PLATFORM_ANDROID)
@@ -1519,9 +1525,15 @@ TSharedPtr<FStreamableHandle> FStreamableManager::RequestAsyncLoad(const FSoftOb
 	return RequestAsyncLoad(TArray<FSoftObjectPath>{TargetToStream}, MoveTemp(DelegateToCall), Priority, bManageActiveHandle, bStartStalled, MoveTemp(DebugName));
 }
 
-TSharedPtr<FStreamableHandle> FStreamableManager::RequestAsyncLoad(TArray<FSoftObjectPath> TargetsToStream, TFunction<void()>&& Callback, TAsyncLoadPriority Priority, bool bManageActiveHandle, bool bStartStalled, FString DebugName)
+TSharedPtr<FStreamableHandle> FStreamableManager::RequestAsyncLoad(const TArray<FSoftObjectPath>& TargetsToStream, TFunction<void()>&& Callback, TAsyncLoadPriority Priority, bool bManageActiveHandle, bool bStartStalled, FString DebugName)
 {
-	return RequestAsyncLoad(MoveTemp(TargetsToStream), FStreamableDelegate::CreateLambda( MoveTemp( Callback ) ), Priority, bManageActiveHandle, bStartStalled, MoveTemp(DebugName));
+	// Make an explicit Copy of the input array
+	return RequestAsyncLoad(TArray<FSoftObjectPath>(TargetsToStream), Forward<TFunction<void()>>(Callback), Priority, bManageActiveHandle, bStartStalled, MoveTemp(DebugName));
+}
+
+TSharedPtr<FStreamableHandle> FStreamableManager::RequestAsyncLoad(TArray<FSoftObjectPath>&& TargetsToStream, TFunction<void()>&& Callback, TAsyncLoadPriority Priority, bool bManageActiveHandle, bool bStartStalled, FString DebugName)
+{
+	return RequestAsyncLoad(Forward<TArray<FSoftObjectPath>>(TargetsToStream), FStreamableDelegate::CreateLambda( MoveTemp( Callback ) ), Priority, bManageActiveHandle, bStartStalled, MoveTemp(DebugName));
 }
 
 TSharedPtr<FStreamableHandle> FStreamableManager::RequestAsyncLoad(const FSoftObjectPath& TargetToStream, TFunction<void()>&& Callback, TAsyncLoadPriority Priority, bool bManageActiveHandle, bool bStartStalled, FString DebugName)
@@ -1529,7 +1541,13 @@ TSharedPtr<FStreamableHandle> FStreamableManager::RequestAsyncLoad(const FSoftOb
 	return RequestAsyncLoad(TargetToStream, FStreamableDelegate::CreateLambda( MoveTemp( Callback ) ), Priority, bManageActiveHandle, bStartStalled, MoveTemp(DebugName));
 }
 
-TSharedPtr<FStreamableHandle> FStreamableManager::RequestSyncLoad(TArray<FSoftObjectPath> TargetsToStream, bool bManageActiveHandle, FString DebugName)
+TSharedPtr<FStreamableHandle> FStreamableManager::RequestSyncLoad(const TArray<FSoftObjectPath>& TargetsToStream, bool bManageActiveHandle, FString DebugName)
+{
+	// Make an explicit Copy of the input array
+	return RequestSyncLoad(TArray<FSoftObjectPath>(TargetsToStream), bManageActiveHandle, MoveTemp(DebugName));
+}
+
+TSharedPtr<FStreamableHandle> FStreamableManager::RequestSyncLoad(TArray<FSoftObjectPath>&& TargetsToStream, bool bManageActiveHandle, FString DebugName)
 {
 	// If in async loading thread or from callback always do sync as recursive tick is unsafe
 	// If in EDL always do sync as EDL internally avoids flushing
@@ -1537,7 +1555,7 @@ TSharedPtr<FStreamableHandle> FStreamableManager::RequestSyncLoad(TArray<FSoftOb
 	bForceSynchronousLoads = IsInAsyncLoadingThread() || IsEventDrivenLoaderEnabled() || !IsAsyncLoading();
 
 	// Do an async load and wait to complete. In some cases this will do a sync load due to safety issues
-	TSharedPtr<FStreamableHandle> Request = RequestAsyncLoad(MoveTemp(TargetsToStream), FStreamableDelegate(), AsyncLoadHighPriority, bManageActiveHandle, false, MoveTemp(DebugName));
+	TSharedPtr<FStreamableHandle> Request = RequestAsyncLoad(Forward<TArray<FSoftObjectPath>>(TargetsToStream), FStreamableDelegate(), AsyncLoadHighPriority, bManageActiveHandle, false, MoveTemp(DebugName));
 
 	bForceSynchronousLoads = false;
 
