[NetPhysics] Resim improvements
- Only compare predicted state for local player
- Rework the debug log flow for NetworkPhysicsComponent to make it less cluttery and implement custom input and state debugging along with it.
- Don't create a new input/state array each replication call, instead just copy over the data from the history to the replication array
- Don't extrapolate data longer than the history size via ExtractData
- Don't assign bIsLocalHistory in TDataRewindHistory since it's not used and we also don't have that data available in the place where we initialize the rewind history.
- Relocated UNetworkPhysicsSettingsComponent overrides to UNetworkPhysicsComponent::InitializeComponent since the component isn't initialized during UNetworkPhysicsComponent constructor.

diff --git a/Engine/Plugins/Experimental/ChaosVehiclesPlugin/Source/ChaosVehicles/Private/ChaosVehicleManagerAsyncCallback.cpp b/Engine/Plugins/Experimental/ChaosVehiclesPlugin/Source/ChaosVehicles/Private/ChaosVehicleManagerAsyncCallback.cpp
--- a/Engine/Plugins/Experimental/ChaosVehiclesPlugin/Source/ChaosVehicles/Private/ChaosVehicleManagerAsyncCallback.cpp
+++ b/Engine/Plugins/Experimental/ChaosVehiclesPlugin/Source/ChaosVehicles/Private/ChaosVehicleManagerAsyncCallback.cpp
@@ -197,29 +197,6 @@ void FNetworkVehicleInputs::ApplyData(UActorComponent* NetworkComponent) const
 				Transmission.SetCurrentGearChangeTime(TransmissionChangeTime);
 			}
 		}
-#if DEBUG_NETWORK_PHYSICS
-		int32 SolverFrame = INDEX_NONE;
-		if (NetworkComponent->GetWorld() && NetworkComponent->GetWorld()->GetPhysicsScene())
-		{
-			if (Chaos::FPhysicsSolver* PhysicsSolver = NetworkComponent->GetWorld()->GetPhysicsScene()->GetSolver())
-			{
-				SolverFrame = PhysicsSolver->GetCurrentFrame();
-			}
-		}
-
-		if (NetworkComponent->GetWorld()->IsNetMode(NM_ListenServer) || NetworkComponent->GetWorld()->IsNetMode(NM_DedicatedServer))
-		{
-			UE_LOG(LogTemp, Log, TEXT("SERVER | PT | ApplyData | Report replicated inputs at frame %d %d: Throttle = %f Brake = %f Roll = %f Pitch = %f Yaw = %f Steering = %f Handbrake = %f Gear = %d | VehicleInputs size = %d | ControlInputs size = %d | NetworkInputs = %d"),
-				LocalFrame, SolverFrame, VehicleInputs.ThrottleInput, VehicleInputs.BrakeInput, VehicleInputs.RollInput, VehicleInputs.PitchInput,
-				VehicleInputs.YawInput, VehicleInputs.SteeringInput, VehicleInputs.HandbrakeInput, TransmissionTargetGear, sizeof(FVehicleInputs) * 8, sizeof(FControlInputs) * 8, sizeof(FNetworkVehicleInputs) * 8);
-		}
-		else
-		{
-			UE_LOG(LogTemp, Log, TEXT("CLIENT | PT | ApplyData | Report replicated inputs at frame %d %d: Throttle = %f Brake = %f Roll = %f Pitch = %f Yaw = %f Steering = %f Handbrake = %f Gear = %d"),
-				LocalFrame, SolverFrame, VehicleInputs.ThrottleInput, VehicleInputs.BrakeInput, VehicleInputs.RollInput, VehicleInputs.PitchInput,
-				VehicleInputs.YawInput, VehicleInputs.SteeringInput, VehicleInputs.HandbrakeInput, TransmissionTargetGear);
-		}
-#endif
 	}
 }
 
@@ -241,20 +218,6 @@ void FNetworkVehicleInputs::BuildData(const UActorComponent* NetworkComponent)
 					TransmissionChangeTime = Transmission.GetCurrentGearChangeTime();
 				}
 			}
-#if DEBUG_NETWORK_PHYSICS
-			if(NetworkComponent->GetWorld()->IsNetMode(NM_ListenServer) || NetworkComponent->GetWorld()->IsNetMode(NM_DedicatedServer))
-			{
-				UE_LOG(LogTemp, Log, TEXT("SERVER | PT | BuildData | Extract local inputs at frame %d : Throttle = %f Brake = %f Roll = %f Pitch = %f Yaw = %f Steering = %f Handbrake = %f Gear = %d"),
-					LocalFrame, VehicleInputs.ThrottleInput, VehicleInputs.BrakeInput, VehicleInputs.RollInput, VehicleInputs.PitchInput,
-					VehicleInputs.YawInput, VehicleInputs.SteeringInput, VehicleInputs.HandbrakeInput, TransmissionTargetGear);
-			}
-			else
-			{
-				UE_LOG(LogTemp, Log, TEXT("CLIENT | PT | BuildData | Extract local inputs at frame %d : Throttle = %f Brake = %f Roll = %f Pitch = %f Yaw = %f Steering = %f Handbrake = %f Gear = %d"),
-					LocalFrame, VehicleInputs.ThrottleInput, VehicleInputs.BrakeInput, VehicleInputs.RollInput, VehicleInputs.PitchInput,
-					VehicleInputs.YawInput, VehicleInputs.SteeringInput, VehicleInputs.HandbrakeInput, TransmissionTargetGear);
-			}
-#endif
 		}
 	}
 }
@@ -304,6 +267,13 @@ void FNetworkVehicleInputs::DecayData(float DecayAmount)
 	VehicleInputs.YawInput = FMath::Lerp(VehicleInputs.YawInput, 0.0f, DecayAmount);
 }
 
+const FString FNetworkVehicleInputs::DebugData()
+{
+	return FString::Printf(TEXT("FNetworkVehicleInputs | Throttle = %f Brake = %f Roll = %f Pitch = %f Yaw = %f Steering = %f Handbrake = %f Gear = %d | VehicleInputs size = %d | ControlInputs size = %d | NetworkInputs = %d"),
+			VehicleInputs.ThrottleInput, VehicleInputs.BrakeInput, VehicleInputs.RollInput, VehicleInputs.PitchInput,
+			VehicleInputs.YawInput, VehicleInputs.SteeringInput, VehicleInputs.HandbrakeInput, TransmissionTargetGear, sizeof(FVehicleInputs) * 8, sizeof(FControlInputs) * 8, sizeof(FNetworkVehicleInputs) * 8);
+}
+
 bool FNetworkVehicleStates::NetSerialize(FArchive& Ar, class UPackageMap* Map, bool& bOutSuccess)
 {
 	FNetworkPhysicsData::SerializeFrames(Ar);
@@ -354,19 +324,6 @@ void FNetworkVehicleStates::ApplyData(UActorComponent* NetworkComponent) const
 		if (TUniquePtr<Chaos::FSimpleWheeledVehicle>& Vehicle = VehicleSimulation->PVehicle)
 		{
 			Vehicle->GetEngine().SetEngineOmega(EngineOmega);
-
-#if DEBUG_NETWORK_PHYSICS
-			if (NetworkComponent->GetWorld()->IsNetMode(NM_ListenServer) || NetworkComponent->GetWorld()->IsNetMode(NM_DedicatedServer))
-			{
-				UE_LOG(LogTemp, Log, TEXT("SERVER | PT | ApplyData | Report replicated states at frame %d %d : Omega = %f"),
-					LocalFrame, ServerFrame, EngineOmega);
-			}
-			else
-			{
-				UE_LOG(LogTemp, Log, TEXT("CLIENT | PT | ApplyData| Report replicated states at frame %d %d : Omega = %f"),
-					LocalFrame, ServerFrame, EngineOmega);
-			}
-#endif
 		
 			int32 LengthCount = 0;
 			for (int32 WheelIdx = 0, NumWheels = Vehicle->Wheels.Num(); WheelIdx < NumWheels; ++WheelIdx)
@@ -400,18 +357,6 @@ void FNetworkVehicleStates::BuildData(const UActorComponent* NetworkComponent)
 			if (const TUniquePtr<Chaos::FSimpleWheeledVehicle>& Vehicle = VehicleSimulation->PVehicle)
 			{
 				EngineOmega = Vehicle->GetEngine().GetEngineOmega();
-#if DEBUG_NETWORK_PHYSICS
-				if (NetworkComponent->GetWorld()->IsNetMode(NM_ListenServer) || NetworkComponent->GetWorld()->IsNetMode(NM_DedicatedServer))
-				{
-					UE_LOG(LogTemp, Log, TEXT("SERVER | PT | BuildData | Extract local states at frame %d %d : Omega = %f"),
-						LocalFrame, ServerFrame, EngineOmega);
-				}
-				else
-				{
-					UE_LOG(LogTemp, Log, TEXT("CLIENT | PT | BuildData| Extract local states at frame %d %d : Omega = %f"),
-						LocalFrame, ServerFrame, EngineOmega);
-				}
-#endif
 				
 				const int32 NumWheels = Vehicle->Wheels.Num();
 				
@@ -491,4 +436,7 @@ void FNetworkVehicleStates::InterpolateData(const FNetworkPhysicsData& MinData,
 	}
 }
 
-
+const FString FNetworkVehicleStates::DebugData()
+{
+	return FString::Printf(TEXT("FNetworkVehicleStates | EngineOmega = %f"), EngineOmega);
+}
\ No newline at end of file
diff --git a/Engine/Plugins/Experimental/ChaosVehiclesPlugin/Source/ChaosVehicles/Public/ChaosVehicleManagerAsyncCallback.h b/Engine/Plugins/Experimental/ChaosVehiclesPlugin/Source/ChaosVehicles/Public/ChaosVehicleManagerAsyncCallback.h
--- a/Engine/Plugins/Experimental/ChaosVehiclesPlugin/Source/ChaosVehicles/Public/ChaosVehicleManagerAsyncCallback.h
+++ b/Engine/Plugins/Experimental/ChaosVehiclesPlugin/Source/ChaosVehicles/Public/ChaosVehicleManagerAsyncCallback.h
@@ -139,6 +139,9 @@ struct CHAOSVEHICLES_API FNetworkVehicleInputs : public FNetworkPhysicsData
 
 	/** Merge data into this input */
 	virtual void MergeData(const FNetworkPhysicsData& FromData) override;
+
+	/** Return string with input debug data */
+	virtual const FString DebugData() override;
 };
 
 template<>
@@ -203,6 +206,9 @@ struct CHAOSVEHICLES_API FNetworkVehicleStates : public FNetworkPhysicsData
 
 	/** Interpolate the data in between two inputs data */
 	virtual void InterpolateData(const FNetworkPhysicsData& MinData, const FNetworkPhysicsData& MaxData) override;
+
+	/** Return string with state debug data */
+	virtual const FString DebugData() override;
 };
 
 template<>
diff --git a/Engine/Source/Runtime/Engine/Private/PhysicsEngine/NetworkPhysicsComponent.cpp b/Engine/Source/Runtime/Engine/Private/PhysicsEngine/NetworkPhysicsComponent.cpp
--- a/Engine/Source/Runtime/Engine/Private/PhysicsEngine/NetworkPhysicsComponent.cpp
+++ b/Engine/Source/Runtime/Engine/Private/PhysicsEngine/NetworkPhysicsComponent.cpp
@@ -404,14 +404,6 @@ void UNetworkPhysicsComponent::InitPhysics()
 	AActor* Owner = GetOwner();
 	if (Owner)
 	{
-		if (UNetworkPhysicsSettingsComponent* PhysicsSettings = Owner->GetComponentByClass<UNetworkPhysicsSettingsComponent>())
-		{
-			InputRedundancy = PhysicsSettings->ResimulationSettings.bOverrideRedundantInputs ? PhysicsSettings->ResimulationSettings.RedundantInputs : InputRedundancy;
-			StateRedundancy = PhysicsSettings->ResimulationSettings.bOverrideRedundantStates ? PhysicsSettings->ResimulationSettings.RedundantStates : StateRedundancy;
-			bCompareInputToTriggerRewind = PhysicsSettings->ResimulationSettings.GetCompareInputToTriggerRewind(bCompareInputToTriggerRewind);
-			bCompareStateToTriggerRewind = PhysicsSettings->ResimulationSettings.GetCompareStateToTriggerRewind(bCompareStateToTriggerRewind);
-		}
-
 		FRepMovement& RepMovement = Owner->GetReplicatedMovement_Mutable();
 		RepMovement.LocationQuantizationLevel = EVectorQuantization::RoundTwoDecimals;
 		RepMovement.RotationQuantizationLevel = ERotatorQuantization::ShortComponents;
@@ -462,6 +454,27 @@ void UNetworkPhysicsComponent::BeginPlay()
 void UNetworkPhysicsComponent::InitializeComponent()
 {
 	Super::InitializeComponent();
+
+	if (AActor* Owner = GetOwner())
+	{
+		if (UNetworkPhysicsSettingsComponent* PhysicsSettings = Owner->FindComponentByClass<UNetworkPhysicsSettingsComponent>())
+		{
+			InputRedundancy = PhysicsSettings->ResimulationSettings.bOverrideRedundantInputs ? PhysicsSettings->ResimulationSettings.RedundantInputs : InputRedundancy;
+			StateRedundancy = PhysicsSettings->ResimulationSettings.bOverrideRedundantStates ? PhysicsSettings->ResimulationSettings.RedundantStates : StateRedundancy;
+			bCompareInputToTriggerRewind = PhysicsSettings->ResimulationSettings.GetCompareInputToTriggerRewind(bCompareInputToTriggerRewind);
+			bCompareStateToTriggerRewind = PhysicsSettings->ResimulationSettings.GetCompareStateToTriggerRewind(bCompareStateToTriggerRewind);
+
+			if (ReplicatedInputs.History)
+			{
+				ReplicatedInputs.History->ResizeDataHistory(InputRedundancy);
+			}
+			if (ReplicatedStates.History)
+			{
+				ReplicatedStates.History->ResizeDataHistory(StateRedundancy);
+			}
+		}
+	}
+
 	if (UWorld* World = GetWorld())
 	{
 		if (UNetworkPhysicsSystem* NetworkManager = World->GetSubsystem<UNetworkPhysicsSystem>())
@@ -501,22 +514,6 @@ void UNetworkPhysicsComponent::AsyncPhysicsTickComponent(float DeltaTime, float
 	QUICK_SCOPE_CYCLE_COUNTER(NetworkPhysicsComponent_AsyncPhysicsTick);
 
 	Super ::AsyncPhysicsTickComponent(DeltaTime, SimTime);
-#if DEBUG_NETWORK_PHYSICS
-	if(HasServerWorld() && !IsLocallyControlled() && InputHistory)
-	{
-		TArray<int32> LocalFrames, ServerFrames, InputFrames;
-		PRAGMA_DISABLE_DEPRECATION_WARNINGS // TODO: Change to DebugData() in UE 5.6 and remove deprecation pragma
-		InputHistory->DebugDatas(*ReplicatedInputs.History, LocalFrames, ServerFrames, InputFrames);
-		PRAGMA_ENABLE_DEPRECATION_WARNINGS
-
-		UE_LOG(LogChaos, Log, TEXT("SERVER | PT | AsyncPhysicsTickComponent | Receiving %d inputs from CLIENT | Component = %s"), LocalFrames.Num(), *GetFullName());
-		for (int32 FrameIndex = 0; FrameIndex < LocalFrames.Num(); ++FrameIndex)
-		{
-			UE_LOG(LogChaos, Log, TEXT("		Debugging replicated inputs at local frame = %d | server frame = %d | Component = %s"),
-				LocalFrames[FrameIndex], ServerFrames[FrameIndex], *GetFullName());
-		}
-	}
-#endif
 
 	// Record the received states from the server into the history for future use
 	if (UWorld* World = GetWorld())
@@ -549,37 +546,34 @@ void UNetworkPhysicsComponent::SendInputData()
 		// Check that client and server have synced the physics tick offset
 		if (PlayerController && PlayerController->GetNetworkPhysicsTickOffsetAssigned())
 		{
-			const int32 LocalOffset = HasServerWorld() ? 0 : PlayerController->GetNetworkPhysicsTickOffset();
-
 			// Send latest N frames from history
-			const int32 ToFrame = InputHistory->GetLatestFrame() + 1; // ToFrame is excluded so add 1.
+			const int32 ToFrame = InputHistory->GetLatestFrame();
 			const int32 FromFrame = ToFrame - InputRedundancy;
 
-			ReplicatedInputs.History = InputHistory->CopyFramesWithOffset(FromFrame, ToFrame, LocalOffset);
-			MARK_PROPERTY_DIRTY_FROM_NAME(UNetworkPhysicsComponent, ReplicatedInputs, this);
+			InputHistory->CopyData(*ReplicatedInputs.History, FromFrame, ToFrame);
 
-			if (!HasServerWorld())
+			if (HasServerWorld())
 			{
-#if DEBUG_NETWORK_PHYSICS
-				FAsyncPhysicsTimestamp Timestamp = const_cast<APlayerController*>(PlayerController)->GetPhysicsTimestamp();
-
-				TArray<int32> LocalFrames, ServerFrames, InputFrames;
-				PRAGMA_DISABLE_DEPRECATION_WARNINGS // TODO: Change to DebugData() in UE 5.6 and remove deprecation pragma
-				InputHistory->DebugDatas(*ReplicatedInputs.History, LocalFrames, ServerFrames, InputFrames);
-				PRAGMA_ENABLE_DEPRECATION_WARNINGS
-
-				UE_LOG(LogChaos, Log, TEXT("CLIENT | GT | SendInputData | Sending %d inputs from CLIENT | Component = %s"), LocalFrames.Num(), *GetFullName());
-				for (int32 FrameIndex = 0; FrameIndex < LocalFrames.Num(); ++FrameIndex)
-				{
-					UE_LOG(LogChaos, Log, TEXT("		Debugging local inputs at local frame = %d | server frame = %d | Current local frame = %d | Current server frame = %d"),
-						LocalFrames[FrameIndex], ServerFrames[FrameIndex], Timestamp.LocalFrame, Timestamp.ServerFrame);
-				}
-#endif
-
-				// if on the client we should first send the replicated inputs onto the server
-				// the RPC will then resend them onto all the other clients (except the local ones)
+				// Server sends inputs through property replication
+				MARK_PROPERTY_DIRTY_FROM_NAME(UNetworkPhysicsComponent, ReplicatedInputs, this);
+			}
+			else
+			{
+				// Clients send inputs through an RPC to the server
 				ServerReceiveInputData(ReplicatedInputs);
 			}
+
+#if DEBUG_NETWORK_PHYSICS
+			FAsyncPhysicsTimestamp Timestamp = const_cast<APlayerController*>(PlayerController)->GetPhysicsTimestamp();
+			if (HasServerWorld())
+			{
+				ReplicatedInputs.History->DebugData(FString::Printf(TEXT("SERVER | PT | SendInputData | CurrentLocalFrame = %d | CurrentServerFrame = %d | Component = %s"), Timestamp.LocalFrame, Timestamp.ServerFrame, *GetFullName()));
+			}
+			else
+			{
+				ReplicatedInputs.History->DebugData(FString::Printf(TEXT("CLIENT | PT | SendInputData | CurrentLocalFrame = %d | CurrentServerFrame = %d | Component = %s"), Timestamp.LocalFrame, Timestamp.ServerFrame, *GetFullName()));
+			}
+#endif
 		}
 	}
 }
@@ -589,12 +583,16 @@ void UNetworkPhysicsComponent::SendStateData()
 	if (HasServerWorld() && StateHistory)
 	{
 		// Send latest N frames from history
-		const int32 ToFrame = StateHistory->GetLatestFrame() + 1; // ToFrame is excluded so add 1.
+		const int32 ToFrame = StateHistory->GetLatestFrame();
 		const int32 FromFrame = ToFrame - StateRedundancy;
 
-		// if on server we should send the states onto all the clients through repnotify
-		ReplicatedStates.History = StateHistory->CopyFramesWithOffset(FromFrame, ToFrame, 0);
+		// If on server we should send the states onto all the clients through repnotify
+		StateHistory->CopyData(*ReplicatedStates.History, FromFrame, ToFrame);
 		MARK_PROPERTY_DIRTY_FROM_NAME(UNetworkPhysicsComponent, ReplicatedStates, this);
+		
+#if DEBUG_NETWORK_PHYSICS
+		ReplicatedStates.History->DebugData(FString::Printf(TEXT("SERVER | PT | SendStateData | Component = %s"), *GetFullName()));
+#endif
 	}
 }
 
@@ -612,7 +610,7 @@ void UNetworkPhysicsComponent::CorrectServerToLocalOffset(const int32 LocalToSer
 		for (int32 FrameIndex = 0; FrameIndex < LocalFrames.Num(); ++FrameIndex)
 		{
 #if DEBUG_NETWORK_PHYSICS || DEBUG_REWIND_DATA
-			UE_LOG(LogChaos, Log, TEXT("CLIENT | GT | CorrectServerToLocalOffset | Server frame = %d | Client Frame = %d"), ServerFrames[FrameIndex], InputFrames[FrameIndex]);
+			UE_LOG(LogChaos, Log, TEXT("CLIENT | GT | CorrectServerToLocalOffset | Server frame = %d | Client Frame = %d | [NOTE: DEPRECATED logic, deactivate by setting CVar: p.net.CmdOffsetEnabled 0]"), ServerFrames[FrameIndex], InputFrames[FrameIndex]);
 #endif
 			ServerToLocalOffset = FMath::Min(ServerToLocalOffset, ServerFrames[FrameIndex] - InputFrames[FrameIndex]);
 		}
@@ -621,7 +619,7 @@ void UNetworkPhysicsComponent::CorrectServerToLocalOffset(const int32 LocalToSer
 		PRAGMA_ENABLE_DEPRECATION_WARNINGS
 
 #if DEBUG_NETWORK_PHYSICS || DEBUG_REWIND_DATA
-		UE_LOG(LogChaos, Log, TEXT("CLIENT | GT | CorrectServerToLocalOffset | Server to local offset = %d | Local to server offset = %d"), ServerToLocalOffset, LocalToServerOffset);
+		UE_LOG(LogChaos, Log, TEXT("CLIENT | GT | CorrectServerToLocalOffset | Server to local offset = %d | Local to server offset = %d | [NOTE: DEPRECATED logic, deactivate by setting CVar: p.net.CmdOffsetEnabled 0]"), ServerToLocalOffset, LocalToServerOffset);
 #endif
 	}
 }
@@ -656,7 +654,7 @@ void UNetworkPhysicsComponent::OnRep_SetReplicatedStates()
 				TSharedPtr<Chaos::FBaseRewindHistory> ReceivedStates = MakeShareable(ReplicatedStates.History->Clone().Release());
 				PhysScene->EnqueueAsyncPhysicsCommand(0, this, [this, PhysScene, ReceivedStates, LocalOffset]()
 				{
-					if (bCompareStateToTriggerRewind)
+					if (bCompareStateToTriggerRewind && IsLocallyControlled())
 					{
 						int32 ResimFrame = StateHistory->ReceiveNewData(*ReceivedStates, LocalOffset, /*CompareDataForRewind*/ bCompareStateToTriggerRewind);
 						if (ResimFrame != INDEX_NONE)
@@ -688,21 +686,8 @@ void UNetworkPhysicsComponent::OnRep_SetReplicatedStates()
 
 #if DEBUG_NETWORK_PHYSICS
 					{
-						TArray<int32> LocalFrames, ServerFrames, InputFrames;
-						PRAGMA_DISABLE_DEPRECATION_WARNINGS // TODO: Change to DebugData() in UE 5.6 and remove deprecation pragma
-						StateHistory->DebugDatas(*ReceivedStates, LocalFrames, ServerFrames, InputFrames);
-						PRAGMA_ENABLE_DEPRECATION_WARNINGS
-
-						UE_LOG(LogChaos, Log, TEXT("CLIENT | PT | OnRep_SetReplicatedStates | Receiving %d states from SERVER | Local offset = %d | Component = %s "), LocalFrames.Num(), LocalOffset, *GetFullName());
-						for (int32 FrameIndex = 0; FrameIndex < LocalFrames.Num(); ++FrameIndex)
-						{
-							UE_LOG(LogChaos, Log, TEXT("		Recording replicated states at local frame = %d | server frame = %d | life time = %d | Component = %s"), ServerFrames[FrameIndex] - LocalOffset, ServerFrames[FrameIndex], InputFrames[FrameIndex], *GetFullName());
-
-							if (IsLocallyControlled() && (InputFrames[FrameIndex] != (ServerFrames[FrameIndex] - LocalOffset)))
-							{
-								UE_LOG(LogChaos, Log, TEXT("		Bad local frame compared to input frame!!!"));
-							}
-						}
+						const int32 CurrentFrame = PhysScene->GetSolver()->GetCurrentFrame();
+						ReceivedStates->DebugData(FString::Printf(TEXT("CLIENT | PT | OnRep_SetReplicatedStates | CurrentLocalFrame = %d | LocalOffset = %d | Component = %s"), CurrentFrame, LocalOffset, *GetFullName()));
 					}
 #endif
 				}, false);
@@ -732,48 +717,40 @@ void UNetworkPhysicsComponent::OnRep_SetReplicatedInputs()
 				TSharedPtr<Chaos::FBaseRewindHistory> ReceivedInputs = MakeShareable(ReplicatedInputs.History->Clone().Release());
 				PhysScene->EnqueueAsyncPhysicsCommand(0, this, [this, PhysScene, ReceivedInputs, LocalOffset]()
 				{
-						if (bCompareInputToTriggerRewind && IsLocallyControlled())
+					if (bCompareInputToTriggerRewind && IsLocallyControlled())
+					{
+						int32 ResimFrame = InputHistory->ReceiveNewData(*ReceivedInputs, LocalOffset, /*CompareDataForRewind*/ bCompareInputToTriggerRewind);
+						if (ResimFrame != INDEX_NONE)
 						{
-							int32 ResimFrame = InputHistory->ReceiveNewData(*ReceivedInputs, LocalOffset, /*CompareDataForRewind*/ bCompareInputToTriggerRewind);
-							if (ResimFrame != INDEX_NONE)
+							if (Chaos::FPhysicsSolver* Solver = PhysScene->GetSolver())
 							{
-								if (Chaos::FPhysicsSolver* Solver = PhysScene->GetSolver())
+								if (Chaos::FRewindData* RewindData = Solver->GetRewindData())
 								{
-									if (Chaos::FRewindData* RewindData = Solver->GetRewindData())
+									// Mark particle/island as resim
+									Chaos::FReadPhysicsObjectInterface_Internal Interface = Chaos::FPhysicsObjectInternalInterface::GetRead();
+									if (Chaos::FPBDRigidParticleHandle* POHandle = Interface.GetRigidParticle(RootPhysicsObject))
 									{
-										// Mark particle/island as resim
-										Chaos::FReadPhysicsObjectInterface_Internal Interface = Chaos::FPhysicsObjectInternalInterface::GetRead();
-										if (Chaos::FPBDRigidParticleHandle* POHandle = Interface.GetRigidParticle(RootPhysicsObject))
-										{
-											Solver->GetEvolution()->GetIslandManager().SetParticleResimFrame(POHandle, ResimFrame);
-										}
-
-										// Set resim frame in rewind data
-										ResimFrame = (RewindData->GetResimFrame() == INDEX_NONE) ? ResimFrame : FMath::Min(ResimFrame, RewindData->GetResimFrame());
-										RewindData->SetResimFrame(ResimFrame);
+										Solver->GetEvolution()->GetIslandManager().SetParticleResimFrame(POHandle, ResimFrame);
 									}
+
+									// Set resim frame in rewind data
+									ResimFrame = (RewindData->GetResimFrame() == INDEX_NONE) ? ResimFrame : FMath::Min(ResimFrame, RewindData->GetResimFrame());
+									RewindData->SetResimFrame(ResimFrame);
 								}
 							}
 						}
-						else
-						{
-							PRAGMA_DISABLE_DEPRECATION_WARNINGS // TODO: Change to ReceiveNewData() in UE 5.6 and remove deprecation pragma
-							InputHistory->ReceiveNewDatas(*ReceivedInputs, LocalOffset);
-							PRAGMA_ENABLE_DEPRECATION_WARNINGS
-						}
-
-#if DEBUG_NETWORK_PHYSICS
+					}
+					else
 					{
-						TArray<int32> LocalFrames, ServerFrames, InputFrames;
-						PRAGMA_DISABLE_DEPRECATION_WARNINGS // TODO: Change to DebugData() in UE 5.6 and remove deprecation pragma
-						InputHistory->DebugDatas(*ReceivedInputs, LocalFrames, ServerFrames, InputFrames);
+						PRAGMA_DISABLE_DEPRECATION_WARNINGS // TODO: Change to ReceiveNewData() in UE 5.6 and remove deprecation pragma
+						InputHistory->ReceiveNewDatas(*ReceivedInputs, LocalOffset);
 						PRAGMA_ENABLE_DEPRECATION_WARNINGS
+					}
 
-						UE_LOG(LogChaos, Log, TEXT("CLIENT | PT | OnRep_SetReplicatedInputs | Receiving %d inputs from SERVER | Local offset = %d | Component = %s"), LocalFrames.Num(), LocalOffset, *GetFullName());
-						for (int32 FrameIndex = 0; FrameIndex < LocalFrames.Num(); ++FrameIndex)
-						{
-							UE_LOG(LogChaos, Log, TEXT("		Recording replicated inputs at local frame = %d | server frame = %d | Component = %s"), ServerFrames[FrameIndex] - LocalOffset, ServerFrames[FrameIndex], *GetFullName());
-						}
+#if DEBUG_NETWORK_PHYSICS
+					{
+						const int32 CurrentFrame = PhysScene->GetSolver()->GetCurrentFrame();
+						ReceivedInputs->DebugData(FString::Printf(TEXT("CLIENT | PT | OnRep_SetReplicatedInputs | CurrentLocalFrame = %d | LocalOffset = %d | Component = %s"), CurrentFrame, LocalOffset, *GetFullName()));
 					}
 #endif
 				}, false);
@@ -810,24 +787,12 @@ void UNetworkPhysicsComponent::ServerReceiveInputData_Implementation(const FNetw
 					InputHistory->ReceiveNewDatas(*ReceivedInputs, 0);
 					PRAGMA_ENABLE_DEPRECATION_WARNINGS
 
-	#if DEBUG_NETWORK_PHYSICS
+#if DEBUG_NETWORK_PHYSICS
 					{
-						TArray<int32> LocalFrames, ServerFrames, InputFrames;
-						PRAGMA_DISABLE_DEPRECATION_WARNINGS // TODO: Change to DebugData() in UE 5.6 and remove deprecation pragma
-						InputHistory->DebugDatas(*ReceivedInputs, LocalFrames, ServerFrames, InputFrames);
-						PRAGMA_ENABLE_DEPRECATION_WARNINGS
-
 						const int32 CurrentFrame = PhysScene->GetSolver()->GetCurrentFrame();
-
-						const int32 EvalOffset = CurrentFrame - InputFrames[InputFrames.Num()-1] + 4;
-						UE_LOG(LogChaos, Log, TEXT("SERVER | PT | ServerReceiveInputData | Receiving %d inputs from CLIENT | Inputs frame = %d | Server frame = %d | Eval Offset = %d | Component = %s"), LocalFrames.Num(), InputFrames[InputFrames.Num() - 1], CurrentFrame, EvalOffset, *GetFullName());
-						for (int32 FrameIndex = 0; FrameIndex < LocalFrames.Num(); ++FrameIndex)
-						{
-							UE_LOG(LogChaos, Log, TEXT("		Recording replicated inputs at local frame = %d | server frame = %d | Solver offset = %d | Component = %s"), 
-								LocalFrames[FrameIndex], ServerFrames[FrameIndex], ServerFrames[FrameIndex] - LocalFrames[FrameIndex], *GetFullName());
-						}
+						ReceivedInputs->DebugData(FString::Printf(TEXT("SERVER | PT | ServerReceiveInputData | CurrentLocalFrame = %d | Component = %s"), CurrentFrame, *GetFullName()));
 					}
-	#endif
+#endif
 				}, false);
 			}
 		}
@@ -868,16 +833,24 @@ void UNetworkPhysicsComponent::OnPreProcessInputsInternal(const int32 PhysicsSte
 		{
 			FNetworkPhysicsData* PhysicsData = StateData.Get();
 			PhysicsData->LocalFrame = PhysicsStep;
-#if DEBUG_NETWORK_PHYSICS
-			UE_LOG(LogChaos, Log, TEXT("		Extracting history states at frame %d | Component = %s"), PhysicsStep, *GetFullName());
-#endif
 			const bool bExactFrame = PhysicsReplicationCVars::ResimulationCVars::bAllowRewindToClosestState ? !bIsSolverReset : true;
 			PRAGMA_DISABLE_DEPRECATION_WARNINGS // TODO: Change to ExtractData() in UE 5.6 and remove deprecation pragma
 			if (StateHistory->ExtractDatas(PhysicsStep, bIsSolverReset, PhysicsData, bExactFrame) && PhysicsData->bReceivedData)
 			PRAGMA_ENABLE_DEPRECATION_WARNINGS
 			{
 				PhysicsData->ApplyData(ActorComponent);
+#if DEBUG_NETWORK_PHYSICS
+				UE_LOG(LogChaos, Log, TEXT("		Applying extracted state from history | bExactFrame = %d | LocalFrame = %d | ServerFrame = %d | InputFrame = %d | ")
+					, bExactFrame, PhysicsData->LocalFrame, PhysicsData->ServerFrame, PhysicsData->InputFrame);
+#endif
 			}
+#if DEBUG_NETWORK_PHYSICS
+			else
+			{
+				UE_LOG(LogChaos, Log, TEXT("		FAILED to extract and apply state from history | bExactFrame = %d | -- Printing history --"), bExactFrame);
+				StateHistory->DebugData(FString::Printf(TEXT("StateHistory | Component = %s"), *GetFullName()));
+			}
+#endif
 		}
 
 		// Apply replicated inputs on server and simulated proxies (and on local player if we are resimulating)
@@ -887,10 +860,6 @@ void UNetworkPhysicsComponent::OnPreProcessInputsInternal(const int32 PhysicsSte
 			const int32 NextExpectedLocalFrame = PhysicsData->LocalFrame + 1;
 			PhysicsData->LocalFrame = PhysicsStep;
 
-	#if DEBUG_NETWORK_PHYSICS
-			UE_LOG(LogChaos, Log, TEXT("		Extracting history inputs at frame %d | Component = %s"), PhysicsStep, *GetFullName());
-	#endif
-
 			PRAGMA_DISABLE_DEPRECATION_WARNINGS // TODO: Change to ExtractData() in UE 5.6 and remove deprecation pragma
 			if (InputHistory->ExtractDatas(PhysicsStep, bIsSolverReset, PhysicsData))
 			PRAGMA_ENABLE_DEPRECATION_WARNINGS
@@ -919,7 +888,19 @@ void UNetworkPhysicsComponent::OnPreProcessInputsInternal(const int32 PhysicsSte
 				}
 
 				PhysicsData->ApplyData(ActorComponent);
+
+#if DEBUG_NETWORK_PHYSICS
+				UE_LOG(LogChaos, Log, TEXT("		Applying extracted input from history | LocalFrame = %d | ServerFrame = %d | InputFrame = %d | IsResim = %d | IsLocallyControlled = %d | InputDecay = %f")
+					, PhysicsData->LocalFrame, PhysicsData->ServerFrame, PhysicsData->InputFrame, bIsSolverResim, IsLocallyControlled(), GetCurrentInputDecay(PhysicsData));
+#endif
 			}
+#if DEBUG_NETWORK_PHYSICS
+			else
+			{
+				UE_LOG(LogChaos, Log, TEXT("		FAILED to extract and apply input from history | IsResim = %d | IsLocallyControlled = %d | -- Printing history --"), bIsSolverResim, IsLocallyControlled());
+				InputHistory->DebugData(FString::Printf(TEXT("InputHistory | Component = %s"), *GetFullName()));
+			}
+#endif
 		}
 	}
 }
@@ -975,7 +956,8 @@ void UNetworkPhysicsComponent::OnPostProcessInputsInternal(const int32 PhysicsSt
 			PRAGMA_ENABLE_DEPRECATION_WARNINGS
 
 #if DEBUG_NETWORK_PHYSICS
-			UE_LOG(LogChaos, Log, TEXT("		Recording local inputs at frame %d | Component = %s"), PhysicsData->LocalFrame, *GetFullName());
+				UE_LOG(LogChaos, Log, TEXT("		Recording input into history | LocalFrame = %d | ServerFrame = %d | InputFrame = %d | Input: %s ")
+					, PhysicsData->LocalFrame, PhysicsData->ServerFrame, PhysicsData->InputFrame, *PhysicsData->DebugData());
 #endif
 		}
 
@@ -1007,7 +989,8 @@ void UNetworkPhysicsComponent::OnPostProcessInputsInternal(const int32 PhysicsSt
 			PRAGMA_ENABLE_DEPRECATION_WARNINGS
 
 #if DEBUG_NETWORK_PHYSICS
-			UE_LOG(LogChaos, Log, TEXT("		Recording local states at frame %d | from input frame = %d | Component = %s"), PhysicsData->LocalFrame, PhysicsData->InputFrame, *GetFullName());
+				UE_LOG(LogChaos, Log, TEXT("		Recording state into history | LocalFrame = %d | ServerFrame = %d | InputFrame = %d | State: %s ")
+					, PhysicsData->LocalFrame, PhysicsData->ServerFrame, PhysicsData->InputFrame, *PhysicsData->DebugData());
 #endif
 		}
 	}
diff --git a/Engine/Source/Runtime/Engine/Public/Physics/NetworkPhysicsComponent.h b/Engine/Source/Runtime/Engine/Public/Physics/NetworkPhysicsComponent.h
--- a/Engine/Source/Runtime/Engine/Public/Physics/NetworkPhysicsComponent.h
+++ b/Engine/Source/Runtime/Engine/Public/Physics/NetworkPhysicsComponent.h
@@ -29,6 +29,11 @@ struct TNetRewindHistory : public Chaos::TDataRewindHistory<DataType>
 	{
 	}
 
+	FORCEINLINE TNetRewindHistory(const int32 FrameCount) :
+		Super(FrameCount)
+	{
+	}
+
 	FORCEINLINE virtual ~TNetRewindHistory() {}
 
 	virtual TUniquePtr<Chaos::FBaseRewindHistory> CreateNew() const
@@ -53,6 +58,22 @@ struct TNetRewindHistory : public Chaos::TDataRewindHistory<DataType>
 		}
 	}
 
+	virtual void CopyData(Chaos::FBaseRewindHistory& OutHistory, const uint32 StartFrame, const uint32 EndFrame) override
+	{
+		TNetRewindHistory& OutNetHistory = static_cast<TNetRewindHistory&>(OutHistory);
+
+		DataType FrameData;
+		for (uint32 CopyFrame = StartFrame; CopyFrame <= EndFrame; ++CopyFrame)
+		{
+			const int32 CopyIndex = Super::GetFrameIndex(CopyFrame);
+			if (CopyFrame == Super::DataHistory[CopyIndex].LocalFrame)
+			{
+				FrameData = Super::DataHistory[CopyIndex];
+				OutNetHistory.RecordData(CopyFrame, &FrameData);
+			}
+		}
+	}
+
 	virtual TUniquePtr<Chaos::FBaseRewindHistory> CopyFramesWithOffset(const uint32 StartFrame, const uint32 EndFrame, const int32 FrameOffset) override
 	{
 		uint32 FramesCount = (uint32)Super::NumValidData(StartFrame, EndFrame);
@@ -151,6 +172,26 @@ struct TNetRewindHistory : public Chaos::TDataRewindHistory<DataType>
 		}
 	}
 
+	/** Print custom string along with values for each entry in history */
+	FORCEINLINE virtual void DebugData(const FString& DebugText) override
+	{
+		UE_LOG(LogChaos, Log, TEXT("%s"), *DebugText);
+		UE_LOG(LogChaos, Log, TEXT("	NumFrames in data collection: %d"), Super::NumFrames);
+
+		if (Super::NumFrames >= 0)
+		{
+			for (int32 FrameIndex = 0; FrameIndex < Super::NumFrames; ++FrameIndex)
+			{
+				UE_LOG(LogChaos, Log, TEXT("		Index: %d || LocalFrame = %d || ServerFrame = %d || InputFrame = %d  ||  Data: %s")
+				, FrameIndex
+				, Super::DataHistory[FrameIndex].LocalFrame
+				, Super::DataHistory[FrameIndex].ServerFrame
+				, Super::DataHistory[FrameIndex].InputFrame
+				, *Super::DataHistory[FrameIndex].DebugData());
+			}
+		}
+	}
+
 private :
 
 	/** Serialized array size limit to guard against invalid network data */
@@ -444,6 +485,9 @@ struct FNetworkPhysicsData
 	virtual void BuildDatas(const UActorComponent* NetworkComponent) { }
 	UE_DEPRECATED(5.4, "Deprecated, use InterpolateData instead")
 	virtual void InterpolateDatas(const FNetworkPhysicsData& MinData, const FNetworkPhysicsData& MaxData) { InterpolateData(MinData, MaxData); }
+	
+	/** Return string with custom debug data */
+	virtual const FString DebugData() { return FString(" - DebugData() not implemented - "); }
 
 	friend UNetworkPhysicsComponent;
 };
@@ -660,21 +704,18 @@ FORCEINLINE void UNetworkPhysicsComponent::CreateDataHistory(UActorComponent* Hi
 {
 	const int32 NumFrames = SetupRewindData();
 
-	APlayerController* Controller = GetPlayerController();
-	const bool bIsLocalHistory = (Controller && Controller->IsLocalController()); // FIXME: The controller is null at this point, but bIsLocalHistory isn't currently used so doesn't create an issue.
-
-	InputHistory = MakeShared<TNetRewindHistory<typename PhysicsTraits::InputsType>>(NumFrames, bIsLocalHistory);
-	StateHistory = MakeShared<TNetRewindHistory<typename PhysicsTraits::StatesType>>(NumFrames, bIsLocalHistory);
+	InputHistory = MakeShared<TNetRewindHistory<typename PhysicsTraits::InputsType>>(NumFrames);
+	StateHistory = MakeShared<TNetRewindHistory<typename PhysicsTraits::StatesType>>(NumFrames);
 
 	InputData = MakeUnique<typename PhysicsTraits::InputsType>();
 	StateData = MakeUnique<typename PhysicsTraits::StatesType>();
 
-	ReplicatedInputs.History = MakeUnique<TNetRewindHistory<typename PhysicsTraits::InputsType>>(InputRedundancy, bIsLocalHistory);
+	ReplicatedInputs.History = MakeUnique<TNetRewindHistory<typename PhysicsTraits::InputsType>>(InputRedundancy);
 	ReplicatedInputs.Owner = this;
 
-	ReplicatedStates.History = MakeUnique<TNetRewindHistory<typename PhysicsTraits::StatesType>>(StateRedundancy, bIsLocalHistory);
+	ReplicatedStates.History = MakeUnique<TNetRewindHistory<typename PhysicsTraits::StatesType>>(StateRedundancy);
 	ReplicatedStates.Owner = this;
-	
+
 	ActorComponent = HistoryComponent;
 	
 	AddDataHistory();
@@ -685,14 +726,11 @@ FORCEINLINE void UNetworkPhysicsComponent::CreateInputHistory(UActorComponent* H
 {
 	const int32 NumFrames = SetupRewindData();
 
-	APlayerController* Controller = GetPlayerController();
-	const bool bIsLocalHistory = (Controller && Controller->IsLocalController()); // FIXME: The controller is null at this point, but bIsLocalHistory isn't currently used so doesn't create an issue.
-
-	InputHistory = MakeShared<TNetRewindHistory<InputsType>>(NumFrames, bIsLocalHistory);
+	InputHistory = MakeShared<TNetRewindHistory<InputsType>>(NumFrames);
 
 	InputData = MakeUnique<InputsType>();
 
-	ReplicatedInputs.History = MakeUnique<TNetRewindHistory<InputsType>>(InputRedundancy, bIsLocalHistory);
+	ReplicatedInputs.History = MakeUnique<TNetRewindHistory<InputsType>>(InputRedundancy);
 	ReplicatedInputs.Owner = this;
 
 	ActorComponent = HistoryComponent;
diff --git a/Engine/Source/Runtime/Experimental/Chaos/Public/RewindData.h b/Engine/Source/Runtime/Experimental/Chaos/Public/RewindData.h
--- a/Engine/Source/Runtime/Experimental/Chaos/Public/RewindData.h
+++ b/Engine/Source/Runtime/Experimental/Chaos/Public/RewindData.h
@@ -60,10 +60,14 @@ struct FBaseRewindHistory
 	* Can for example block received data from client from overriding server authoritative data */
 	FORCEINLINE virtual bool ShouldRecordReceivedDataOnFrame(const int32 RecordFrame) { return true; }
 
+	/** Copy data from local history into @param OutHistory
+	* @param StartFrame = Included
+	* @param EndFrame = Included */
+	virtual void CopyData(Chaos::FBaseRewindHistory& OutHistory, const uint32 StartFrame, const uint32 EndFrame) {};
+
 	/** Create a polymorphic copy of only a range of frames, applying the frame offset to the copies
 	* @param StartFrame = Included
-	* @param EndFrame = Excluded
-	*/
+	* @param EndFrame = Excluded */
 	virtual TUniquePtr<FBaseRewindHistory> CopyFramesWithOffset(const uint32 StartFrame, const uint32 EndFrame, const int32 FrameOffset) = 0;
 
 	/** Copy new data (received from the network) into this history, returns frame to resimulate from if @param CompareDataForRewind is set to true and compared data differ enough */
@@ -80,7 +84,10 @@ struct FBaseRewindHistory
 	/** Validate data in history buffer received from clients on the server */
 	virtual void ValidateDataInHistory(const void* ActorComponent) {}
 
-	/** Debug the data from the array of uint8 that will be transferred from client to server */
+	/** Print custom string along with values for each entry in history */
+	FORCEINLINE virtual void DebugData(const FString& DebugText) { }
+
+	/** Get arrays of frame values for each entry in the history */
 	FORCEINLINE virtual void DebugData(const Chaos::FBaseRewindHistory& NewData, TArray<int32>& LocalFrames, TArray<int32>& ServerFrames, TArray<int32>& InputFrames) { }
 	UE_DEPRECATED(5.4, "Deprecated, use DebugData() instead")
 	FORCEINLINE virtual void DebugDatas(const Chaos::FBaseRewindHistory& NewDatas, TArray<int32>& LocalFrames, TArray<int32>& ServerFrames, TArray<int32>& InputFrames) { DebugData(NewDatas, LocalFrames, ServerFrames, InputFrames); }
@@ -93,6 +100,9 @@ struct FBaseRewindHistory
 
 	/** Return the most up to date frame entry in history */
 	virtual const int32 GetLatestFrame() const { return INDEX_NONE; }
+
+	/** Resize the history */
+	virtual void ResizeDataHistory(const int32 FrameCount, const EAllowShrinking AllowShrinking = EAllowShrinking::Yes) {}
 };
 
 /** Templated data history holding a data buffer */
@@ -104,6 +114,12 @@ struct TDataRewindHistory : public FBaseRewindHistory
 	{
 		DataHistory.SetNum(NumFrames);
 	}
+
+	FORCEINLINE TDataRewindHistory(const int32 FrameCount) :
+		DataHistory(), LatestFrame(0), CurrentFrame(0), CurrentIndex(0), NumFrames(FrameCount)
+	{
+		DataHistory.SetNum(NumFrames);
+	}
 	FORCEINLINE virtual ~TDataRewindHistory() {}
 
 protected:
@@ -140,6 +156,12 @@ public :
 	/** Extract states at a given time */
 	FORCEINLINE virtual bool ExtractData(const int32 ExtractFrame, const bool bResetSolver, void* HistoryData, const bool bExactFrame = false) override
 	{
+		// Early out if we are trying to extract data but the latest data is more than the whole buffer size old
+		if (ExtractFrame - NumFrames > GetLatestFrame())
+		{
+			return false;
+		}
+
 		const int32 ExtractIndex = GetFrameIndex(ExtractFrame);
 		if (ExtractFrame == DataHistory[ExtractIndex].LocalFrame)
 		{
@@ -321,6 +343,17 @@ public :
 		return LatestFrame;
 	}
 
+	/** Resize the history */
+	FORCEINLINE void ResizeDataHistory(const int32 FrameCount, const EAllowShrinking AllowShrinking = EAllowShrinking::Yes) override
+	{
+		if (FrameCount > 0 && NumFrames != FrameCount)
+		{
+			NumFrames = FrameCount;
+			DataHistory.SetNum(NumFrames, AllowShrinking);
+			CurrentIndex = GetFrameIndex(CurrentFrame);
+		}
+	}
+
 	FORCEINLINE const uint32 GetFrameIndex(const int32 Frame) const
 	{
 		return FMath::Abs(Frame % NumFrames);
@@ -331,7 +364,7 @@ protected :
 	/** Check if the history is on the local/remote client*/
 	bool bIsLocalHistory;
 
-	/**  Data buffer holding the history */
+	/** Data buffer holding the history */
 	TArray<DataType> DataHistory;
 
 	/** The most up to date frame entry in history */
