[Material Editor] Additional failsafe code to avoid crashing when a cached refernece to a ParameterCollection fails to load propertly during material translation.

diff --git a/Engine/Source/Runtime/Engine/Private/Materials/HLSLMaterialTranslator.cpp b/Engine/Source/Runtime/Engine/Private/Materials/HLSLMaterialTranslator.cpp
--- a/Engine/Source/Runtime/Engine/Private/Materials/HLSLMaterialTranslator.cpp
+++ b/Engine/Source/Runtime/Engine/Private/Materials/HLSLMaterialTranslator.cpp
@@ -2597,7 +2597,7 @@ void FHLSLMaterialTranslator::GetMaterialEnvironment(EShaderPlatform InPlatform,
 
 		// Check that the parameter collection loaded succesfully.
 		UMaterialParameterCollection* ParameterCollection = EnvironmentDefines->ParameterCollections[CollectionIndex];
-		if (!ParameterCollection)
+		if (!ParameterCollection || !ParameterCollection->HasValidUniformBufferStruct())
 		{
 			UE_LOG(LogMaterial, Warning, TEXT("Null parameter collection found in environment defines while translating material."));
 			continue;
diff --git a/Engine/Source/Runtime/Engine/Public/Materials/MaterialParameterCollection.h b/Engine/Source/Runtime/Engine/Public/Materials/MaterialParameterCollection.h
--- a/Engine/Source/Runtime/Engine/Public/Materials/MaterialParameterCollection.h
+++ b/Engine/Source/Runtime/Engine/Public/Materials/MaterialParameterCollection.h
@@ -173,6 +173,12 @@ class UMaterialParameterCollection : public UObject
 		return *UniformBufferStruct;
 	}
 
+	/** Returns whether this instance loaded correctly and is in valid state. */
+	bool HasValidUniformBufferStruct() const
+	{
+		return UniformBufferStruct.IsValid();
+	}
+
 	/** Create an instance for this collection in every world. */
 	ENGINE_API void SetupWorldParameterCollectionInstances();
 
@@ -195,6 +201,3 @@ private:
 
 	void UpdateDefaultResource(bool bRecreateUniformBuffer);
 };
-
-
-
