Fix crash creating ray tracing geometry in FSkeletalMeshStreamIn.
- Initializer was being incorrecly cleared in FSkeletalMeshStreamOut.

diff --git a/Engine/Source/Runtime/Engine/Internal/Streaming/SkeletalMeshUpdate.cpp b/Engine/Source/Runtime/Engine/Internal/Streaming/SkeletalMeshUpdate.cpp
--- a/Engine/Source/Runtime/Engine/Internal/Streaming/SkeletalMeshUpdate.cpp
+++ b/Engine/Source/Runtime/Engine/Internal/Streaming/SkeletalMeshUpdate.cpp
@@ -377,7 +377,7 @@ void FSkeletalMeshStreamOut::ReleaseBuffers(const FContext& Context)
 				{
 					if (RenderData->LODRenderData[LODIndex].bReferencedByStaticSkeletalMeshObjects_RenderThread)
 					{
-						LODResource.StaticRayTracingGeometry.ReleaseResource();
+						LODResource.StaticRayTracingGeometry.ReleaseRHIForStreaming(Batcher);
 					}
 				}
 #endif
