MergeActor - More include cleanup in order to avoid rebuilding the whole engine when editing merge actor settings

diff --git a/Engine/Source/Developer/HierarchicalLODUtilities/Private/HierarchicalLODUtilities.cpp b/Engine/Source/Developer/HierarchicalLODUtilities/Private/HierarchicalLODUtilities.cpp
--- a/Engine/Source/Developer/HierarchicalLODUtilities/Private/HierarchicalLODUtilities.cpp
+++ b/Engine/Source/Developer/HierarchicalLODUtilities/Private/HierarchicalLODUtilities.cpp
@@ -53,6 +53,7 @@
 #include "Materials/Material.h"
 #include "Materials/MaterialInstance.h"
 #include "UObject/ICookInfo.h"
+#include "HLOD/HLODSetup.h"
 
 DEFINE_LOG_CATEGORY_STATIC(LogHierarchicalLODUtilities, Verbose, All);
 
diff --git a/Engine/Source/Developer/HierarchicalLODUtilities/Public/HierarchicalLODProxyProcessor.h b/Engine/Source/Developer/HierarchicalLODUtilities/Public/HierarchicalLODProxyProcessor.h
--- a/Engine/Source/Developer/HierarchicalLODUtilities/Public/HierarchicalLODProxyProcessor.h
+++ b/Engine/Source/Developer/HierarchicalLODUtilities/Public/HierarchicalLODProxyProcessor.h
@@ -5,7 +5,7 @@
 #include "CoreMinimal.h"
 #include "Misc/Guid.h"
 #include "Containers/Ticker.h"
-#include "GameFramework/WorldSettings.h"
+#include "HLOD/HLODSetup.h"
 #include "MeshUtilities.h"
 #include "IMeshReductionInterfaces.h"
 #include "UObject/StrongObjectPtr.h"
@@ -73,3 +73,7 @@ private:
 	/** Critical section to keep JobActorMap/ToProcessJobs access thread-safe */
 	FCriticalSection StateLock;	
 };
+
+#if UE_ENABLE_INCLUDE_ORDER_DEPRECATED_IN_5_5
+#include "GameFramework/WorldSettings.h"
+#endif
diff --git a/Engine/Source/Editor/DetailCustomizations/Private/HierarchicalSimplificationCustomizations.h b/Engine/Source/Editor/DetailCustomizations/Private/HierarchicalSimplificationCustomizations.h
--- a/Engine/Source/Editor/DetailCustomizations/Private/HierarchicalSimplificationCustomizations.h
+++ b/Engine/Source/Editor/DetailCustomizations/Private/HierarchicalSimplificationCustomizations.h
@@ -2,7 +2,7 @@
 
 #pragma once
 
-#include "GameFramework/WorldSettings.h"
+#include "HLOD/HLODSetup.h"
 #include "IPropertyTypeCustomization.h"
 #include "Layout/Visibility.h"
 #include "OverrideResetToDefault.h"
@@ -29,3 +29,7 @@ protected:
 
 	TSharedPtr< IPropertyHandle > SimplificationMethodPropertyHandle;
 };
+
+#if UE_ENABLE_INCLUDE_ORDER_DEPRECATED_IN_5_5
+#include "GameFramework/WorldSettings.h"
+#endif
diff --git a/Engine/Source/Editor/HierarchicalLODOutliner/Private/HLODOutliner.cpp b/Engine/Source/Editor/HierarchicalLODOutliner/Private/HLODOutliner.cpp
--- a/Engine/Source/Editor/HierarchicalLODOutliner/Private/HLODOutliner.cpp
+++ b/Engine/Source/Editor/HierarchicalLODOutliner/Private/HLODOutliner.cpp
@@ -21,6 +21,7 @@
 #include "Framework/Notifications/NotificationManager.h"
 #include "GameFramework/Actor.h"
 #include "GameFramework/WorldSettings.h"
+#include "HLOD/HLODSetup.h"
 #include "HLODOutlinerDragDrop.h"
 #include "HLODTreeWidgetItem.h"
 #include "HierarchicalLOD.h"
diff --git a/Engine/Source/Editor/HierarchicalLODOutliner/Private/HierarchicalLODOutlinerModule.cpp b/Engine/Source/Editor/HierarchicalLODOutliner/Private/HierarchicalLODOutlinerModule.cpp
--- a/Engine/Source/Editor/HierarchicalLODOutliner/Private/HierarchicalLODOutlinerModule.cpp
+++ b/Engine/Source/Editor/HierarchicalLODOutliner/Private/HierarchicalLODOutlinerModule.cpp
@@ -12,6 +12,7 @@
 #include "Engine/EngineTypes.h"
 #include "Engine/World.h"
 #include "GameFramework/WorldSettings.h"
+#include "HLOD/HLODSetup.h"
 #include "HLODOutliner.h"
 #include "HierarchicalLODUtilitiesModule.h"
 #include "IHierarchicalLODUtilities.h"
diff --git a/Engine/Source/Editor/UnrealEd/Private/HierarchicalLOD.cpp b/Engine/Source/Editor/UnrealEd/Private/HierarchicalLOD.cpp
--- a/Engine/Source/Editor/UnrealEd/Private/HierarchicalLOD.cpp
+++ b/Engine/Source/Editor/UnrealEd/Private/HierarchicalLOD.cpp
@@ -13,7 +13,7 @@
 #include "Logging/MessageLog.h"
 #include "Misc/UObjectToken.h"
 #include "Misc/MapErrors.h"
-#include "GameFramework/WorldSettings.h"
+#include "HLOD/HLODSetup.h"
 
 #include "Algo/Transform.h"
 #include "EditorLevelUtils.h"
diff --git a/Engine/Source/Runtime/Engine/Classes/GameFramework/WorldSettings.h b/Engine/Source/Runtime/Engine/Classes/GameFramework/WorldSettings.h
--- a/Engine/Source/Runtime/Engine/Classes/GameFramework/WorldSettings.h
+++ b/Engine/Source/Runtime/Engine/Classes/GameFramework/WorldSettings.h
@@ -9,9 +9,6 @@
 #include "GameFramework/Actor.h"
 #include "GameFramework/DamageType.h"
 #include "GameFramework/Info.h"
-#include "MeshMerge/MeshProxySettings.h"
-#include "MeshMerge/MeshMergingSettings.h"
-#include "MeshMerge/MeshApproximationSettings.h"
 #include "Sound/AudioVolume.h"
 #include "UObject/ConstructorHelpers.h"
 #include "WorldGridPreviewer.h"
@@ -23,6 +20,8 @@ class UAssetUserData;
 class UNetConnection;
 class UNavigationSystemConfig;
 class UAISystemBase;
+class UHierarchicalLODSetup;
+struct FHierarchicalSimplification;
 
 UENUM()
 enum EVisibilityAggressiveness : int
@@ -279,138 +278,6 @@ struct FNetViewer
 	ENGINE_API FNetViewer(AController* InController);
 };
 
-UENUM()
-enum class EHierarchicalSimplificationMethod : uint8
-{
-	None = 0			UMETA(hidden),
-	Merge = 1,
-	Simplify = 2,
-	Approximate = 3
-};
-
-
-USTRUCT()
-struct FHierarchicalSimplification
-{
-	GENERATED_USTRUCT_BODY()
-
-	/** The screen radius an mesh object should reach before swapping to the LOD actor, once one of parent displays, it won't draw any of children. */
-	UPROPERTY(Category = FHierarchicalSimplification, EditAnywhere, meta = (UIMin = "0.00001", ClampMin = "0.000001", UIMax = "1.0", ClampMax = "1.0"))
-	float TransitionScreenSize;
-
-	UPROPERTY(Category = FHierarchicalSimplification, EditAnywhere, AdvancedDisplay, meta = (UIMin = "1.0", ClampMin = "1.0", UIMax = "50000.0", editcondition="bUseOverrideDrawDistance"))
-	float OverrideDrawDistance;
-
-	UPROPERTY(Category = FHierarchicalSimplification, EditAnywhere, AdvancedDisplay, meta = (InlineEditConditionToggle))
-	uint8 bUseOverrideDrawDistance:1;
-
-	UPROPERTY(Category = FHierarchicalSimplification, EditAnywhere, AdvancedDisplay)
-	uint8 bAllowSpecificExclusion : 1;
-
-	/** Only generate clusters for HLOD volumes */
-	UPROPERTY(EditAnywhere, Category = FHierarchicalSimplification, AdvancedDisplay, meta = (editcondition = "!bReusePreviousLevelClusters", DisplayAfter="MinNumberOfActorsToBuild"))
-	uint8 bOnlyGenerateClustersForVolumes:1;
-
-	/** Will reuse the clusters generated for the previous (lower) HLOD level */
-	UPROPERTY(EditAnywhere, Category = FHierarchicalSimplification, AdvancedDisplay, meta=(DisplayAfter="bOnlyGenerateClustersForVolumes"))
-	uint8 bReusePreviousLevelClusters:1;
-
-	UPROPERTY(Category = FHierarchicalSimplification, EditAnywhere)
-	EHierarchicalSimplificationMethod SimplificationMethod;
-
-	/** Simplification settings, used if SimplificationMethod is Simplify */
-	UPROPERTY(Category = FHierarchicalSimplification, EditAnywhere, AdvancedDisplay)
-	FMeshProxySettings ProxySetting;
-
-	/** Merge settings, used if SimplificationMethod is Merge */
-	UPROPERTY(Category = FHierarchicalSimplification, EditAnywhere, AdvancedDisplay)
-	FMeshMergingSettings MergeSetting;
-
-	/** Approximate settings, used if SimplificationMethod is Approximate */
-	UPROPERTY(Category = FHierarchicalSimplification, EditAnywhere, AdvancedDisplay)
-	FMeshApproximationSettings ApproximateSettings;
-
-	/** Desired Bounding Radius for clustering - this is not guaranteed but used to calculate filling factor for auto clustering */
-	UPROPERTY(EditAnywhere, Category=FHierarchicalSimplification, AdvancedDisplay, meta=(UIMin=10.f, ClampMin=10.f, editcondition = "!bReusePreviousLevelClusters"))
-	float DesiredBoundRadius;
-
-	/** Desired Filling Percentage for clustering - this is not guaranteed but used to calculate filling factor  for auto clustering */
-	UPROPERTY(EditAnywhere, Category=FHierarchicalSimplification, AdvancedDisplay, meta=(ClampMin = "0", ClampMax = "100", UIMin = "0", UIMax = "100", editcondition = "!bReusePreviousLevelClusters"))
-	float DesiredFillingPercentage;
-
-	/** Min number of actors to build LODActor */
-	UPROPERTY(EditAnywhere, Category=FHierarchicalSimplification, AdvancedDisplay, meta=(ClampMin = "1", UIMin = "1", editcondition = "!bReusePreviousLevelClusters"))
-	int32 MinNumberOfActorsToBuild;
-
-#if WITH_EDITORONLY_DATA
-	UPROPERTY(meta = (DeprecatedProperty))
-	uint8 bSimplifyMesh_DEPRECATED:1;
-#endif
-
-	FHierarchicalSimplification()
-		: TransitionScreenSize(0.315f)
-		, OverrideDrawDistance(10000)
-		, bUseOverrideDrawDistance(false)
-		, bAllowSpecificExclusion(false)
-		, bOnlyGenerateClustersForVolumes(false)
-		, bReusePreviousLevelClusters(false)
-		, SimplificationMethod(EHierarchicalSimplificationMethod::Merge)
-		, DesiredBoundRadius(2000)
-		, DesiredFillingPercentage(50)
-		, MinNumberOfActorsToBuild(2)
-	{
-		MergeSetting.bMergeMaterials = true;
-		MergeSetting.bGenerateLightMapUV = true;
-		ProxySetting.MaterialSettings.MaterialMergeType = EMaterialMergeType::MaterialMergeType_Simplygon;
-		ProxySetting.bCreateCollision = false;
-	}
-
-#if WITH_EDITORONLY_DATA
-	ENGINE_API bool Serialize(FArchive& Ar);
-
-	/** Handles deprecated properties */
-	ENGINE_API void PostSerialize(const FArchive& Ar);
-#endif
-
-	/** Retrieve the correct material proxy settings based on the simplification method. */
-	ENGINE_API FMaterialProxySettings* GetSimplificationMethodMaterialSettings();
-};
-
-template<>
-struct TStructOpsTypeTraits<FHierarchicalSimplification> : public TStructOpsTypeTraitsBase2<FHierarchicalSimplification>
-{
-#if WITH_EDITORONLY_DATA
-	enum
-	{
-		WithSerializer = true,
-		WithPostSerialize = true,
-	};
-#endif
-};
-
-UCLASS(Blueprintable, MinimalAPI)
-class UHierarchicalLODSetup : public UObject
-{
-	GENERATED_BODY()
-public:
-	UHierarchicalLODSetup()
-	{
-		HierarchicalLODSetup.AddDefaulted();
-		OverrideBaseMaterial = nullptr;
-	}
-
-	/** Hierarchical LOD Setup */
-	UPROPERTY(EditAnywhere, Category = HLODSystem)
-	TArray<struct FHierarchicalSimplification> HierarchicalLODSetup;
-
-	UPROPERTY(EditAnywhere, Category = HLODSystem)
-	TSoftObjectPtr<UMaterialInterface> OverrideBaseMaterial;
-
-#if WITH_EDITOR
-	ENGINE_API virtual void PostEditChangeProperty(struct FPropertyChangedEvent& PropertyChangedEvent) override;
-#endif // WITH_EDITOR
-};
-
 USTRUCT()
 struct FNaniteSettings
 {
@@ -789,7 +656,7 @@ public:
 #if WITH_EDITORONLY_DATA
 	/** If set overrides the level settings and global project settings */
 	UPROPERTY(EditAnywhere, config, Category = HLODSystem, meta=(EditConditionHides, EditCondition = "WorldPartition == nullptr"))
-	TSoftClassPtr<class UHierarchicalLODSetup> HLODSetupAsset;
+	TSoftClassPtr<UHierarchicalLODSetup> HLODSetupAsset;
 
 	/** If set overrides the project-wide base material used for Proxy Materials */
 	UPROPERTY(EditAnywhere, config, Category = HLODSystem, meta=(EditConditionHides, EditCondition = "WorldPartition == nullptr && HLODSetupAsset == nullptr"))
@@ -798,7 +665,7 @@ public:
 protected:
 	/** Hierarchical LOD Setup */
 	UPROPERTY(EditAnywhere, Category = HLODSystem, config, meta=(EditConditionHides, EditCondition = "WorldPartition == nullptr && HLODSetupAsset == nullptr"))
-	TArray<struct FHierarchicalSimplification> HierarchicalLODSetup;
+	TArray<FHierarchicalSimplification> HierarchicalLODSetup;
 
 public:
 	UPROPERTY()
@@ -977,8 +844,8 @@ public:
 	//~ End IInterface_AssetUserData Interface
 
 #if WITH_EDITOR
-	ENGINE_API const TArray<struct FHierarchicalSimplification>& GetHierarchicalLODSetup() const;
-	ENGINE_API TArray<struct FHierarchicalSimplification>& GetHierarchicalLODSetup();
+	ENGINE_API const TArray<FHierarchicalSimplification>& GetHierarchicalLODSetup() const;
+	ENGINE_API TArray<FHierarchicalSimplification>& GetHierarchicalLODSetup();
 	ENGINE_API int32 GetNumHierarchicalLODLevels() const;
 	ENGINE_API UMaterialInterface* GetHierarchicalLODBaseMaterial() const;
 	ENGINE_API void ResetHierarchicalLODSetup();
@@ -1127,4 +994,8 @@ private: //DEPRECATED
 
 #if UE_ENABLE_INCLUDE_ORDER_DEPRECATED_IN_5_5
 #include "Engine/MeshMerging.h"
+
+#include "MeshMerge/MeshProxySettings.h"
+#include "MeshMerge/MeshMergingSettings.h"
+#include "MeshMerge/MeshApproximationSettings.h"
 #endif
diff --git a/Engine/Source/Runtime/Engine/Private/Components/PrimitiveComponent.cpp b/Engine/Source/Runtime/Engine/Private/Components/PrimitiveComponent.cpp
--- a/Engine/Source/Runtime/Engine/Private/Components/PrimitiveComponent.cpp
+++ b/Engine/Source/Runtime/Engine/Private/Components/PrimitiveComponent.cpp
@@ -18,6 +18,7 @@
 #include "EngineStats.h"
 #include "GameFramework/Pawn.h"
 #include "HLOD/HLODBatchingPolicy.h"
+#include "HLOD/HLODSetup.h"
 #include "PSOPrecache.h"
 #include "AI/NavigationSystemBase.h"
 #include "AI/Navigation/NavigationRelevantData.h"
@@ -4897,7 +4898,7 @@ const bool UPrimitiveComponent::ShouldGenerateAutoLOD(const int32 HierarchicalLe
 	bool bExcluded = false;
 	if (HierarchicalLevelIndex < CHAR_BIT && IsExcludedFromHLODLevel(EHLODLevelExclusion(1 << HierarchicalLevelIndex)))
 	{
-		const TArray<struct FHierarchicalSimplification>& HLODSetup = GetOwner()->GetLevel()->GetWorldSettings()->GetHierarchicalLODSetup();
+		const TArray<FHierarchicalSimplification>& HLODSetup = GetOwner()->GetLevel()->GetWorldSettings()->GetHierarchicalLODSetup();
 		if (HLODSetup.IsValidIndex(HierarchicalLevelIndex))
 		{
 			if (HLODSetup[HierarchicalLevelIndex].bAllowSpecificExclusion)
diff --git a/Engine/Source/Runtime/Engine/Private/HLOD/HLODSetup.cpp b/Engine/Source/Runtime/Engine/Private/HLOD/HLODSetup.cpp
new file mode 100644
--- /dev/null
+++ b/Engine/Source/Runtime/Engine/Private/HLOD/HLODSetup.cpp
@@ -0,0 +1,69 @@
+// Copyright Epic Games, Inc. All Rights Reserved.
+
+#include "HLOD/HLODSetup.h"
+#include "UObject/FortniteMainBranchObjectVersion.h"
+
+
+FHierarchicalSimplification::FHierarchicalSimplification()
+	: TransitionScreenSize(0.315f)
+	, OverrideDrawDistance(10000)
+	, bUseOverrideDrawDistance(false)
+	, bAllowSpecificExclusion(false)
+	, bOnlyGenerateClustersForVolumes(false)
+	, bReusePreviousLevelClusters(false)
+	, SimplificationMethod(EHierarchicalSimplificationMethod::Merge)
+	, DesiredBoundRadius(2000)
+	, DesiredFillingPercentage(50)
+	, MinNumberOfActorsToBuild(2)
+{
+	MergeSetting.bMergeMaterials = true;
+	MergeSetting.bGenerateLightMapUV = true;
+	ProxySetting.MaterialSettings.MaterialMergeType = EMaterialMergeType::MaterialMergeType_Simplygon;
+	ProxySetting.bCreateCollision = false;
+}
+
+#if WITH_EDITORONLY_DATA
+
+bool FHierarchicalSimplification::Serialize(FArchive& Ar)
+{
+	Ar.UsingCustomVersion(FFortniteMainBranchObjectVersion::GUID);
+
+	// Don't actually serialize, just write the custom version for PostSerialize
+	return false;
+}
+
+void FHierarchicalSimplification::PostSerialize(const FArchive& Ar)
+{
+	if (Ar.IsLoading())
+	{
+		if (Ar.CustomVer(FFortniteMainBranchObjectVersion::GUID) < FFortniteMainBranchObjectVersion::HierarchicalSimplificationMethodEnumAdded)
+		{
+			SimplificationMethod = bSimplifyMesh_DEPRECATED ? EHierarchicalSimplificationMethod::Simplify : EHierarchicalSimplificationMethod::Merge;
+		}
+	}
+}
+
+#endif
+
+FMaterialProxySettings* FHierarchicalSimplification::GetSimplificationMethodMaterialSettings()
+{
+	switch (SimplificationMethod)
+	{
+	case EHierarchicalSimplificationMethod::Merge:
+		return &MergeSetting.MaterialSettings;
+
+	case EHierarchicalSimplificationMethod::Simplify:
+		return &ProxySetting.MaterialSettings;
+
+	case EHierarchicalSimplificationMethod::Approximate:
+		return &ApproximateSettings.MaterialSettings;
+	}
+
+	return nullptr;
+}
+
+UHierarchicalLODSetup::UHierarchicalLODSetup()
+{
+	HierarchicalLODSetup.AddDefaulted();
+	OverrideBaseMaterial = nullptr;
+}
diff --git a/Engine/Source/Runtime/Engine/Private/HLODProxy.cpp b/Engine/Source/Runtime/Engine/Private/HLODProxy.cpp
--- a/Engine/Source/Runtime/Engine/Private/HLODProxy.cpp
+++ b/Engine/Source/Runtime/Engine/Private/HLODProxy.cpp
@@ -8,6 +8,7 @@
 #include "Engine/World.h"
 #include "GameFramework/WorldSettings.h"
 #include "HLOD/HLODProxyDesc.h"
+#include "HLOD/HLODSetup.h"
 #include "Materials/Material.h"
 
 #if WITH_EDITOR
diff --git a/Engine/Source/Runtime/Engine/Private/LODActor.cpp b/Engine/Source/Runtime/Engine/Private/LODActor.cpp
--- a/Engine/Source/Runtime/Engine/Private/LODActor.cpp
+++ b/Engine/Source/Runtime/Engine/Private/LODActor.cpp
@@ -10,6 +10,7 @@
 #include "UObject/UObjectIterator.h"
 #include "Engine/CollisionProfile.h"
 #include "HLOD/HLODProxyDesc.h"
+#include "HLOD/HLODSetup.h"
 #include "Misc/MapErrors.h"
 #include "Logging/MessageLog.h"
 #include "Materials/MaterialInterface.h"
diff --git a/Engine/Source/Runtime/Engine/Private/WorldSettings.cpp b/Engine/Source/Runtime/Engine/Private/WorldSettings.cpp
--- a/Engine/Source/Runtime/Engine/Private/WorldSettings.cpp
+++ b/Engine/Source/Runtime/Engine/Private/WorldSettings.cpp
@@ -22,7 +22,6 @@
 #include "PhysicsEngine/PhysicsSettings.h"
 #include "UObject/ReleaseObjectVersion.h"
 #include "UObject/EnterpriseObjectVersion.h"
-#include "UObject/FortniteMainBranchObjectVersion.h"
 #include "SceneManagement.h"
 #include "AI/AISystemBase.h"
 #include "AI/NavigationSystemConfig.h"
@@ -36,6 +35,7 @@
 #include "Editor.h"
 #include "Misc/TransactionObjectEvent.h"
 #include "HierarchicalLOD.h"
+#include "HLOD/HLODSetup.h"
 #include "WorldPartition/DataLayer/DataLayerManager.h"
 #endif 
 
@@ -1129,44 +1129,4 @@ void AWorldSettings::RewindForReplay()
 	bHighPriorityLoading = false;
 }
 
-#if WITH_EDITORONLY_DATA
-
-bool FHierarchicalSimplification::Serialize(FArchive& Ar)
-{
-	Ar.UsingCustomVersion(FFortniteMainBranchObjectVersion::GUID);
-
-	// Don't actually serialize, just write the custom version for PostSerialize
-	return false;
-}
-
-void FHierarchicalSimplification::PostSerialize(const FArchive& Ar)
-{
-	if (Ar.IsLoading())
-	{
-		if (Ar.CustomVer(FFortniteMainBranchObjectVersion::GUID) < FFortniteMainBranchObjectVersion::HierarchicalSimplificationMethodEnumAdded)
-		{
-			SimplificationMethod = bSimplifyMesh_DEPRECATED ? EHierarchicalSimplificationMethod::Simplify : EHierarchicalSimplificationMethod::Merge;
-		}
-	}
-}
-
-#endif
-
-FMaterialProxySettings* FHierarchicalSimplification::GetSimplificationMethodMaterialSettings()
-{
-	switch (SimplificationMethod)
-	{
-	case EHierarchicalSimplificationMethod::Merge:
-		return &MergeSetting.MaterialSettings;
-
-	case EHierarchicalSimplificationMethod::Simplify:
-		return &ProxySetting.MaterialSettings;
-
-	case EHierarchicalSimplificationMethod::Approximate:
-		return &ApproximateSettings.MaterialSettings;
-	}
-
-	return nullptr;
-}
-
 #undef LOCTEXT_NAMESPACE
diff --git a/Engine/Source/Runtime/Engine/Public/HLOD/HLODSetup.h b/Engine/Source/Runtime/Engine/Public/HLOD/HLODSetup.h
new file mode 100644
--- /dev/null
+++ b/Engine/Source/Runtime/Engine/Public/HLOD/HLODSetup.h
@@ -0,0 +1,122 @@
+// Copyright Epic Games, Inc. All Rights Reserved.
+
+#pragma once
+
+#include "MeshMerge/MeshMergingSettings.h"
+#include "MeshMerge/MeshProxySettings.h"
+#include "MeshMerge/MeshApproximationSettings.h"
+#include "HLODSetup.generated.h"
+
+
+UENUM()
+enum class EHierarchicalSimplificationMethod : uint8
+{
+	None = 0			UMETA(hidden),
+	Merge = 1,
+	Simplify = 2,
+	Approximate = 3
+};
+
+
+USTRUCT()
+struct FHierarchicalSimplification
+{
+	GENERATED_USTRUCT_BODY()
+
+	/** The screen radius an mesh object should reach before swapping to the LOD actor, once one of parent displays, it won't draw any of children. */
+	UPROPERTY(Category = FHierarchicalSimplification, EditAnywhere, meta = (UIMin = "0.00001", ClampMin = "0.000001", UIMax = "1.0", ClampMax = "1.0"))
+	float TransitionScreenSize;
+
+	UPROPERTY(Category = FHierarchicalSimplification, EditAnywhere, AdvancedDisplay, meta = (UIMin = "1.0", ClampMin = "1.0", UIMax = "50000.0", editcondition="bUseOverrideDrawDistance"))
+	float OverrideDrawDistance;
+
+	UPROPERTY(Category = FHierarchicalSimplification, EditAnywhere, AdvancedDisplay, meta = (InlineEditConditionToggle))
+	uint8 bUseOverrideDrawDistance:1;
+
+	UPROPERTY(Category = FHierarchicalSimplification, EditAnywhere, AdvancedDisplay)
+	uint8 bAllowSpecificExclusion : 1;
+
+	/** Only generate clusters for HLOD volumes */
+	UPROPERTY(EditAnywhere, Category = FHierarchicalSimplification, AdvancedDisplay, meta = (editcondition = "!bReusePreviousLevelClusters", DisplayAfter="MinNumberOfActorsToBuild"))
+	uint8 bOnlyGenerateClustersForVolumes:1;
+
+	/** Will reuse the clusters generated for the previous (lower) HLOD level */
+	UPROPERTY(EditAnywhere, Category = FHierarchicalSimplification, AdvancedDisplay, meta=(DisplayAfter="bOnlyGenerateClustersForVolumes"))
+	uint8 bReusePreviousLevelClusters:1;
+
+	UPROPERTY(Category = FHierarchicalSimplification, EditAnywhere)
+	EHierarchicalSimplificationMethod SimplificationMethod;
+
+	/** Simplification settings, used if SimplificationMethod is Simplify */
+	UPROPERTY(Category = FHierarchicalSimplification, EditAnywhere, AdvancedDisplay)
+	FMeshProxySettings ProxySetting;
+
+	/** Merge settings, used if SimplificationMethod is Merge */
+	UPROPERTY(Category = FHierarchicalSimplification, EditAnywhere, AdvancedDisplay)
+	FMeshMergingSettings MergeSetting;
+
+	/** Approximate settings, used if SimplificationMethod is Approximate */
+	UPROPERTY(Category = FHierarchicalSimplification, EditAnywhere, AdvancedDisplay)
+	FMeshApproximationSettings ApproximateSettings;
+
+	/** Desired Bounding Radius for clustering - this is not guaranteed but used to calculate filling factor for auto clustering */
+	UPROPERTY(EditAnywhere, Category=FHierarchicalSimplification, AdvancedDisplay, meta=(UIMin=10.f, ClampMin=10.f, editcondition = "!bReusePreviousLevelClusters"))
+	float DesiredBoundRadius;
+
+	/** Desired Filling Percentage for clustering - this is not guaranteed but used to calculate filling factor  for auto clustering */
+	UPROPERTY(EditAnywhere, Category=FHierarchicalSimplification, AdvancedDisplay, meta=(ClampMin = "0", ClampMax = "100", UIMin = "0", UIMax = "100", editcondition = "!bReusePreviousLevelClusters"))
+	float DesiredFillingPercentage;
+
+	/** Min number of actors to build LODActor */
+	UPROPERTY(EditAnywhere, Category=FHierarchicalSimplification, AdvancedDisplay, meta=(ClampMin = "1", UIMin = "1", editcondition = "!bReusePreviousLevelClusters"))
+	int32 MinNumberOfActorsToBuild;
+
+#if WITH_EDITORONLY_DATA
+	UPROPERTY(meta = (DeprecatedProperty))
+	uint8 bSimplifyMesh_DEPRECATED:1;
+#endif
+
+	ENGINE_API FHierarchicalSimplification();
+
+#if WITH_EDITORONLY_DATA
+	ENGINE_API bool Serialize(FArchive& Ar);
+
+	/** Handles deprecated properties */
+	ENGINE_API void PostSerialize(const FArchive& Ar);
+#endif
+
+	/** Retrieve the correct material proxy settings based on the simplification method. */
+	ENGINE_API FMaterialProxySettings* GetSimplificationMethodMaterialSettings();
+};
+
+template<>
+struct TStructOpsTypeTraits<FHierarchicalSimplification> : public TStructOpsTypeTraitsBase2<FHierarchicalSimplification>
+{
+#if WITH_EDITORONLY_DATA
+	enum
+	{
+		WithSerializer = true,
+		WithPostSerialize = true,
+	};
+#endif
+};
+
+UCLASS(Blueprintable, MinimalAPI)
+class UHierarchicalLODSetup : public UObject
+{
+	GENERATED_BODY()
+
+public:
+	ENGINE_API UHierarchicalLODSetup();
+
+	/** Hierarchical LOD Setup */
+	UPROPERTY(EditAnywhere, Category = HLODSystem)
+	TArray<struct FHierarchicalSimplification> HierarchicalLODSetup;
+
+	UPROPERTY(EditAnywhere, Category = HLODSystem)
+	TSoftObjectPtr<UMaterialInterface> OverrideBaseMaterial;
+
+#if WITH_EDITOR
+	ENGINE_API virtual void PostEditChangeProperty(struct FPropertyChangedEvent& PropertyChangedEvent) override;
+#endif // WITH_EDITOR
+};
