[Scene Outliner] Optimize update time by:
- Storing cached strings into actor descriptors to avoid FName::ToString.
- Getting actor names from the descriptor's stored FSoftObjectPath substring.
- Using move operator instead of copying strings when possible.

diff --git a/Engine/Source/Editor/DataLayerEditor/Private/DataLayer/SDataLayerBrowser.cpp b/Engine/Source/Editor/DataLayerEditor/Private/DataLayer/SDataLayerBrowser.cpp
--- a/Engine/Source/Editor/DataLayerEditor/Private/DataLayer/SDataLayerBrowser.cpp
+++ b/Engine/Source/Editor/DataLayerEditor/Private/DataLayer/SDataLayerBrowser.cpp
@@ -102,7 +102,7 @@ void SDataLayerBrowser::Construct(const FArguments& InArgs)
 		{
 			if (const FWorldPartitionActorDescInstance* ActorDescInstance = *ActorDescItem->ActorDescHandle)
 			{
-				return ActorDescInstance->GetActorName().ToString();
+				return ActorDescInstance->GetActorNameString();
 			}
 		}
 		return FString();
diff --git a/Engine/Source/Editor/SceneOutliner/Private/ActorDescTreeItem.cpp b/Engine/Source/Editor/SceneOutliner/Private/ActorDescTreeItem.cpp
--- a/Engine/Source/Editor/SceneOutliner/Private/ActorDescTreeItem.cpp
+++ b/Engine/Source/Editor/SceneOutliner/Private/ActorDescTreeItem.cpp
@@ -233,7 +233,7 @@ FActorDescTreeItem::FActorDescTreeItem(const FGuid& InActorGuid, UActorDescConta
 {
 	if (const FWorldPartitionActorDescInstance* const ActorDescInstance = *ActorDescHandle)
 	{
-		DisplayString = ActorDescInstance->GetActorLabel().ToString();
+		DisplayString = ActorDescInstance->GetActorLabelString();
 	}
 	else
 	{
@@ -251,7 +251,7 @@ FActorDescTreeItem::FActorDescTreeItem(const FWorldPartitionActorDescInstance* I
 {
 	if (const FWorldPartitionActorDescInstance* const ActorDescInstance = *ActorDescHandle)
 	{
-		DisplayString = ActorDescInstance->GetActorLabel().ToString();
+		DisplayString = ActorDescInstance->GetActorLabelString();
 	}
 	else
 	{
diff --git a/Engine/Source/Editor/SceneOutliner/Private/SceneOutlinerActorInfoColumn.cpp b/Engine/Source/Editor/SceneOutliner/Private/SceneOutlinerActorInfoColumn.cpp
--- a/Engine/Source/Editor/SceneOutliner/Private/SceneOutlinerActorInfoColumn.cpp
+++ b/Engine/Source/Editor/SceneOutliner/Private/SceneOutlinerActorInfoColumn.cpp
@@ -63,7 +63,7 @@ struct FGetInfo
 		{
 			if (const FWorldPartitionActorDescInstance* ActorDescInstance = *ActorDescItem->ActorDescHandle)
 			{
-				return ActorDescInstance->GetDisplayClassName().ToString();
+				return ActorDescInstance->GetDisplayClassNameString();
 			}
 		}
 
@@ -171,7 +171,7 @@ void FTypeInfoColumn::PopulateSearchStrings( const ISceneOutlinerTreeItem& Item,
 		FString String = FGetInfo()(Item);
 		if (String.Len())
 		{
-			OutSearchStrings.Add(String);
+			OutSearchStrings.Add(MoveTemp(String));
 		}
 	}
 }
diff --git a/Engine/Source/Editor/SceneOutliner/Private/SceneOutlinerModule.cpp b/Engine/Source/Editor/SceneOutliner/Private/SceneOutlinerModule.cpp
--- a/Engine/Source/Editor/SceneOutliner/Private/SceneOutlinerModule.cpp
+++ b/Engine/Source/Editor/SceneOutliner/Private/SceneOutlinerModule.cpp
@@ -471,28 +471,28 @@ void FSceneOutlinerModule::CreateActorInfoColumns(FSceneOutlinerInitializationOp
 		{
 			if (AActor* Actor = ActorItem->Actor.Get())
 			{
-				return Actor->GetFName().ToString();
+				return Actor->GetName();
 			}
 		}
 		else if (const FComponentTreeItem* ComponentItem = Item.CastTo<FComponentTreeItem>())
 		{
 			if (UActorComponent* Component = ComponentItem->Component.Get())
 			{
-				return Component->GetFName().ToString();
+				return Component->GetName();
 			}
 		}
 		else if (const FActorDescTreeItem* ActorDescItem = Item.CastTo<FActorDescTreeItem>())
 		{
 			if (const FWorldPartitionActorDescInstance* ActorDescInstance = *ActorDescItem->ActorDescHandle)
 			{
-				return ActorDescInstance->GetActorName().ToString();
+				return ActorDescInstance->GetActorNameString();
 			}
 		}
 		else if (const FActorFolderTreeItem* ActorFolderItem = Item.CastTo<FActorFolderTreeItem>())
 		{
 			if (const UActorFolder* ActorFolder = Cast<UActorFolder>(ActorFolderItem->GetActorFolder()))
 			{
-				return ActorFolder->GetFName().ToString();
+				return ActorFolder->GetName();
 			}
 		}
 
diff --git a/Engine/Source/Editor/SceneOutliner/Private/SceneOutlinerTextInfoColumn.cpp b/Engine/Source/Editor/SceneOutliner/Private/SceneOutlinerTextInfoColumn.cpp
--- a/Engine/Source/Editor/SceneOutliner/Private/SceneOutlinerTextInfoColumn.cpp
+++ b/Engine/Source/Editor/SceneOutliner/Private/SceneOutlinerTextInfoColumn.cpp
@@ -65,7 +65,7 @@ void FTextInfoColumn::PopulateSearchStrings(const ISceneOutlinerTreeItem& Item,
 	FString String = GetTextForItem.Execute(Item);
 	if (String.Len())
 	{
-		OutSearchStrings.Add(String);
+		OutSearchStrings.Add(MoveTemp(String));
 	}
 }
 
diff --git a/Engine/Source/Editor/UnrealEd/Private/Commandlets/WorldPartitionDataLayerToAssetCommandlet.cpp b/Engine/Source/Editor/UnrealEd/Private/Commandlets/WorldPartitionDataLayerToAssetCommandlet.cpp
--- a/Engine/Source/Editor/UnrealEd/Private/Commandlets/WorldPartitionDataLayerToAssetCommandlet.cpp
+++ b/Engine/Source/Editor/UnrealEd/Private/Commandlets/WorldPartitionDataLayerToAssetCommandlet.cpp
@@ -510,7 +510,7 @@ bool UDataLayerToAssetCommandlet::RemapActorDataLayersToAssets(TStrongObjectPtr<
 				FString DataLayerString = FString::JoinBy(ActDescDataLayers.ToArray(), TEXT(", "), [](const FName& DataLayerName) { return DataLayerName.ToString(); });
 
 				UE_LOG(LogDataLayerToAssetCommandlet, Error, TEXT("Actor %s failed to load. Its data layers %s will not be remapped to a data layer asset."),
-					*ActorDescInstance->GetActorName().ToString(), *DataLayerString);
+					*ActorDescInstance->GetActorNameString(), *DataLayerString);
 				if (!bIgnoreActorLoadingErrors)
 				{
 					ActorConversionErrors++;
diff --git a/Engine/Source/Editor/UnrealEd/Private/WorldPartition/WorldPartitionFoliageBuilder.cpp b/Engine/Source/Editor/UnrealEd/Private/WorldPartition/WorldPartitionFoliageBuilder.cpp
--- a/Engine/Source/Editor/UnrealEd/Private/WorldPartition/WorldPartitionFoliageBuilder.cpp
+++ b/Engine/Source/Editor/UnrealEd/Private/WorldPartition/WorldPartitionFoliageBuilder.cpp
@@ -114,7 +114,7 @@ bool UWorldPartitionFoliageBuilder::RunInternal(UWorld* World, const FCellInfo&
 		AInstancedFoliageActor* IFA = Cast<AInstancedFoliageActor>(ActorDescInstance->GetActor());
 		if (!IFA)
 		{
-			UE_LOG(LogWorldPartitionFoliageBuilder, Error, TEXT("Foliage actor failed to load: %s (%s)"), *PartitionActorDesc->GetActorName().ToString(), *PartitionActorDesc->GetActorPackage().ToString());
+			UE_LOG(LogWorldPartitionFoliageBuilder, Error, TEXT("Foliage actor failed to load: %s (%s)"), *PartitionActorDesc->GetActorNameString(), *PartitionActorDesc->GetActorPackage().ToString());
 			return false;
 		}
 
@@ -134,7 +134,7 @@ bool UWorldPartitionFoliageBuilder::RunInternal(UWorld* World, const FCellInfo&
 				UPackage* PackageToDelete = nullptr;
 				UPackage* PackageToAdd = nullptr;
 
-				UE_LOG(LogWorldPartitionFoliageBuilder, Display, TEXT("Repairing invalid foliage actor: %s->%s "), *PartitionActorDesc->GetActorName().ToString(), *ExpectedActorName);
+				UE_LOG(LogWorldPartitionFoliageBuilder, Display, TEXT("Repairing invalid foliage actor: %s->%s "), *PartitionActorDesc->GetActorNameString(), *ExpectedActorName);
 
 				const FString ExpectedActorPath = FString::Printf(TEXT("%s.%s"), *World->PersistentLevel->GetPathName(), *ExpectedActorName);
 				
@@ -146,7 +146,7 @@ bool UWorldPartitionFoliageBuilder::RunInternal(UWorld* World, const FCellInfo&
 					AInstancedFoliageActor* DupIFA = Cast<AInstancedFoliageActor>(DupActorDescInstance->GetActor());
 					if (!DupIFA)
 					{
-						UE_LOG(LogWorldPartitionFoliageBuilder, Error, TEXT("Foliage actor failed to load: %s (%s)"), *DupActorDescInstance->GetActorName().ToString(), *DupActorDescInstance->GetActorPackage().ToString());
+						UE_LOG(LogWorldPartitionFoliageBuilder, Error, TEXT("Foliage actor failed to load: %s (%s)"), *DupActorDescInstance->GetActorNameString(), *DupActorDescInstance->GetActorPackage().ToString());
 						return false;
 					}
 
@@ -225,7 +225,7 @@ bool UWorldPartitionFoliageBuilder::RunInternal(UWorld* World, const FCellInfo&
 		{	
 			check(IFA->GetGridSize() != NewGridSize);
 			// Harvest Instances from existing IFA and build up instance bounds
-			UE_LOG(LogWorldPartitionFoliageBuilder, Display, TEXT("Processing existing foliage actor: %s (%s)"), *PartitionActorDesc->GetActorName().ToString(), *PartitionActorDesc->GetActorPackage().ToString());
+			UE_LOG(LogWorldPartitionFoliageBuilder, Display, TEXT("Processing existing foliage actor: %s (%s)"), *PartitionActorDesc->GetActorNameString(), *PartitionActorDesc->GetActorPackage().ToString());
 		
 			IFA->ForEachFoliageInfo([IFA, &FoliageToAdd, &InstanceBounds, &NumInstances](UFoliageType* FoliageType, FFoliageInfo& FoliageInfo)
 			{
diff --git a/Engine/Source/Runtime/Engine/Private/WorldPartition/ActorDescContainer.cpp b/Engine/Source/Runtime/Engine/Private/WorldPartition/ActorDescContainer.cpp
--- a/Engine/Source/Runtime/Engine/Private/WorldPartition/ActorDescContainer.cpp
+++ b/Engine/Source/Runtime/Engine/Private/WorldPartition/ActorDescContainer.cpp
@@ -119,7 +119,7 @@ void UActorDescContainer::Initialize(const FInitializeParams& InitParams)
 			else if (!ActorDesc->GetNativeClass().IsValid())
 			{
 				UE_LOG(LogWorldPartition, Warning, TEXT("Invalid actor native class: Actor: '%s' (guid '%s') from package '%s'"),
-					*ActorDesc->GetActorName().ToString(),
+					*ActorDesc->GetActorNameString(),
 					*ActorDesc->GetGuid().ToString(),
 					*ActorDesc->GetActorPackage().ToString());
 				InvalidActors.Emplace(Asset);
@@ -128,7 +128,7 @@ void UActorDescContainer::Initialize(const FInitializeParams& InitParams)
 			{
 				UE_LOG(LogWorldPartition, Warning, TEXT("Unknown actor base class `%s`: Actor: '%s' (guid '%s') from package '%s'"),
 					*ActorDesc->GetBaseClass().ToString(),
-					*ActorDesc->GetActorName().ToString(),
+					*ActorDesc->GetActorNameString(),
 					*ActorDesc->GetGuid().ToString(),
 					*ActorDesc->GetActorPackage().ToString());
 				InvalidActors.Emplace(Asset);
@@ -143,8 +143,8 @@ void UActorDescContainer::Initialize(const FInitializeParams& InitParams)
 			{
 				UE_LOG(LogWorldPartition, Warning, TEXT("Duplicate actor descriptor in package `%s`: Actor: '%s' -> Existing actor '%s'"), 
 					*ActorDesc->GetActorPackage().ToString(), 
-					*ActorDesc->GetActorName().ToString(), 
-					*ExistingDescPackage->GetActorName().ToString());
+					*ActorDesc->GetActorNameString(), 
+					*ExistingDescPackage->GetActorNameString());
 
 				// No need to add all actors in the same package several times as we only want to open the package for delete when repairing
 				if (ValidActorDescs.Contains(ActorDesc->GetGuid()))
@@ -159,9 +159,9 @@ void UActorDescContainer::Initialize(const FInitializeParams& InitParams)
 				check(ExistingActorDesc->GetGuid() == ActorDesc->GetGuid());
 				UE_LOG(LogWorldPartition, Warning, TEXT("Duplicate actor descriptor guid `%s`: Actor: '%s' from package '%s' -> Existing actor '%s' from package '%s'"), 
 					*ActorDesc->GetGuid().ToString(), 
-					*ActorDesc->GetActorName().ToString(), 
+					*ActorDesc->GetActorNameString(), 
 					*ActorDesc->GetActorPackage().ToString(),
-					*ExistingActorDesc->GetActorName().ToString(),
+					*ExistingActorDesc->GetActorNameString(),
 					*ExistingActorDesc->GetActorPackage().ToString());
 				InvalidActors.Emplace(Asset);
 			}
@@ -269,7 +269,7 @@ void UActorDescContainer::UnregisterActorDescriptor(FWorldPartitionActorDesc* Ac
 {
 	FActorDescList::RemoveActorDescriptor(ActorDesc);
 	ActorDesc->SetContainer(nullptr);
-	verifyf(ActorsByName.Remove(ActorDesc->GetActorName()), TEXT("Missing actor '%s' from container '%s'"), *ActorDesc->GetActorName().ToString(), *ContainerPackageName.ToString());
+	verifyf(ActorsByName.Remove(ActorDesc->GetActorName()), TEXT("Missing actor '%s' from container '%s'"), *ActorDesc->GetActorNameString(), *ContainerPackageName.ToString());
 }
 
 bool UActorDescContainer::ShouldHandleActorEvent(const AActor* Actor, bool bInUseLoadedPath) const
diff --git a/Engine/Source/Runtime/Engine/Private/WorldPartition/HLOD/HLODRuntimeSubsystem.cpp b/Engine/Source/Runtime/Engine/Private/WorldPartition/HLOD/HLODRuntimeSubsystem.cpp
--- a/Engine/Source/Runtime/Engine/Private/WorldPartition/HLOD/HLODRuntimeSubsystem.cpp
+++ b/Engine/Source/Runtime/Engine/Private/WorldPartition/HLOD/HLODRuntimeSubsystem.cpp
@@ -756,7 +756,7 @@ bool UWorldPartitionHLODRuntimeSubsystem::WriteHLODStatsCSV(UWorld* InWorld, con
 	TArray<TPair<FName, FGetStatFunc>> StatsToWrite =
 	{	
 		{ "WorldPackage",		[InWorld](FWorldPartitionActorDescInstance* InActorDescInstance, const FHLODActorDesc& InActorDesc) { return InWorld->GetPackage()->GetName(); } },
-		{ "Name",				[](FWorldPartitionActorDescInstance* InActorDescInstance, const FHLODActorDesc& InActorDesc) { return InActorDescInstance->GetActorLabel().ToString(); } },
+		{ "Name",				[](FWorldPartitionActorDescInstance* InActorDescInstance, const FHLODActorDesc& InActorDesc) { return InActorDescInstance->GetActorLabelString(); } },
 		{ "HLODLayer",			[](FWorldPartitionActorDescInstance* InActorDescInstance, const FHLODActorDesc& InActorDesc) { return InActorDesc.GetSourceHLODLayer().GetAssetName().ToString(); }},
 		{ "SpatiallyLoaded",	[](FWorldPartitionActorDescInstance* InActorDescInstance, const FHLODActorDesc& InActorDesc) { return InActorDescInstance->GetIsSpatiallyLoaded() ? TEXT("true") : TEXT("false"); } },
 		{ "DataLayers",			[&GetDataLayerShortName](FWorldPartitionActorDescInstance* InActorDescInstance, const FHLODActorDesc& InActorDesc) { return FString::JoinBy(InActorDescInstance->GetDataLayerInstanceNames().ToArray(), TEXT(" | "), GetDataLayerShortName); }},
diff --git a/Engine/Source/Runtime/Engine/Private/WorldPartition/WorldPartition.cpp b/Engine/Source/Runtime/Engine/Private/WorldPartition/WorldPartition.cpp
--- a/Engine/Source/Runtime/Engine/Private/WorldPartition/WorldPartition.cpp
+++ b/Engine/Source/Runtime/Engine/Private/WorldPartition/WorldPartition.cpp
@@ -1988,7 +1988,7 @@ UActorDescContainerInstance* UWorldPartition::RegisterActorDescContainerInstance
 			{
 				UE_LOG(LogWorldPartition, Warning, TEXT("Found existing actor descriptor guid `%s`: Actor: '%s' from package '%s'"),
 					*ActorDesc->GetGuid().ToString(),
-					*ActorDesc->GetActorName().ToString(),
+					*ActorDesc->GetActorNameString(),
 					*ActorDesc->GetActorPackage().ToString());
 				return false;
 			}
diff --git a/Engine/Source/Runtime/Engine/Private/WorldPartition/WorldPartitionActorDesc.cpp b/Engine/Source/Runtime/Engine/Private/WorldPartition/WorldPartitionActorDesc.cpp
--- a/Engine/Source/Runtime/Engine/Private/WorldPartition/WorldPartitionActorDesc.cpp
+++ b/Engine/Source/Runtime/Engine/Private/WorldPartition/WorldPartitionActorDesc.cpp
@@ -200,6 +200,9 @@ void FWorldPartitionActorDesc::Init(const AActor* InActor)
 		ActorLabel = *InActor->GetActorLabel(false);
 	}
 
+	ActorLabelString = ActorLabel.ToString();
+	ActorDisplayClassNameString = GetDisplayClassName().ToString();
+
 	Container = nullptr;
 }
 
@@ -209,7 +212,7 @@ void FWorldPartitionActorDesc::Init(const FWorldPartitionActorDescInitData& Desc
 	ActorPath = DescData.ActorPath;
 	ActorNativeClass = DescData.NativeClass;
 	NativeClass = *DescData.NativeClass->GetPathName();
-	ActorName = *FPaths::GetExtension(ActorPath.ToString());
+	ActorName = *GetActorNameString();
 
 	auto DeprecateClass = [this](FArchive& Archive)
 	{
@@ -249,6 +252,9 @@ void FWorldPartitionActorDesc::Init(const FWorldPartitionActorDescInitData& Desc
 		DeprecateClass(*DescData.GetArchive());
 	}
 
+	ActorLabelString = ActorLabel.ToString();
+	ActorDisplayClassNameString = GetDisplayClassName().ToString();
+
 	Container = nullptr;
 }
 
@@ -453,7 +459,7 @@ FString FWorldPartitionActorDesc::ToString(EToStringMode Mode) const
 			TEXT(" BaseClass:%s NativeClass:%s Name:%s Label:%s SpatiallyLoaded:%s Bounds:%s RuntimeGrid:%s EditorOnly:%s RuntimeOnly:%s HLODRelevant:%s ListedInSceneOutliner:%s IsMainWorldOnly:%s"),
 			*BaseClass.ToString(), 
 			*NativeClass.ToString(), 
-			*GetActorName().ToString(),
+			*GetActorNameString(),
 			*GetActorLabel().ToString(),
 			GetBoolStr(bIsSpatiallyLoaded),
 			*BoundsStr,
@@ -720,7 +726,7 @@ PRAGMA_ENABLE_DEPRECATION_WARNINGS
 			FGuid FixupContentBundleGuid = ContentBundlePaths::GetContentBundleGuidFromExternalActorPackagePath(ActorPackage.ToString());
 			if (ContentBundleGuid != FixupContentBundleGuid)
 			{
-				UE_LOG(LogWorldPartition, Log, TEXT("ActorDesc ContentBundleGuid was fixed up: %s"), *GetActorName().ToString());
+				UE_LOG(LogWorldPartition, Log, TEXT("ActorDesc ContentBundleGuid was fixed up: %s"), *GetActorNameString());
 				ContentBundleGuid = FixupContentBundleGuid;
 			}
 		}
@@ -766,6 +772,21 @@ FName FWorldPartitionActorDesc::GetActorName() const
 	return ActorName;
 }
 
+const FString& FWorldPartitionActorDesc::GetActorNameString() const
+{
+	return ActorPath.GetSubPathString();
+}
+
+const FString& FWorldPartitionActorDesc::GetActorLabelString() const
+{
+	return ActorLabelString;
+}
+
+const FString& FWorldPartitionActorDesc::GetDisplayClassNameString() const
+{
+	return ActorDisplayClassNameString;
+}
+
 FName FWorldPartitionActorDesc::GetActorLabelOrName() const
 {
 	return GetActorLabel().IsNone() ? GetActorName() : GetActorLabel();
diff --git a/Engine/Source/Runtime/Engine/Private/WorldPartition/WorldPartitionActorDescInstance.cpp b/Engine/Source/Runtime/Engine/Private/WorldPartition/WorldPartitionActorDescInstance.cpp
--- a/Engine/Source/Runtime/Engine/Private/WorldPartition/WorldPartitionActorDescInstance.cpp
+++ b/Engine/Source/Runtime/Engine/Private/WorldPartition/WorldPartitionActorDescInstance.cpp
@@ -131,7 +131,7 @@ bool FWorldPartitionActorDescInstance::StartAsyncLoad()
 
 			if ((Result != EAsyncLoadingResult::Succeeded) || !Package)
 			{
-				UE_LOG(LogWorldPartition, Warning, TEXT("Can't load actor guid `%s` ('%s') from package '%s'"), *GetGuid().ToString(), *GetActorName().ToString(), *ActorPackage.ToString());
+				UE_LOG(LogWorldPartition, Warning, TEXT("Can't load actor guid `%s` ('%s') from package '%s'"), *GetGuid().ToString(), *GetActorNameString(), *ActorPackage.ToString());
 				UnloadedReason = &FailedToLoad;
 				return;
 			}
@@ -140,7 +140,7 @@ bool FWorldPartitionActorDescInstance::StartAsyncLoad()
 
 			if (!ActorPtr.IsValid())
 			{
-				UE_LOG(LogWorldPartition, Warning, TEXT("Can't find actor guid `%s` ('%s') in package '%s'"), *GetGuid().ToString(), *GetActorName().ToString(), *ActorPackage.ToString());
+				UE_LOG(LogWorldPartition, Warning, TEXT("Can't find actor guid `%s` ('%s') in package '%s'"), *GetGuid().ToString(), *GetActorNameString(), *ActorPackage.ToString());
 				UnloadedReason = &FailedToLoad;
 				return;
 			}
@@ -216,6 +216,21 @@ const FText& FWorldPartitionActorDescInstance::GetUnloadedReason() const
 	return UnloadedReason ? *UnloadedReason : Unloaded;
 }
 
+const FString& FWorldPartitionActorDescInstance::GetActorNameString() const
+{
+	return ActorDesc->GetActorNameString();
+}
+
+const FString& FWorldPartitionActorDescInstance::GetActorLabelString() const
+{
+	return ActorDesc->GetActorLabelString();
+}
+
+const FString& FWorldPartitionActorDescInstance::GetDisplayClassNameString() const
+{
+	return ActorDesc->GetDisplayClassNameString();
+}
+
 FString FWorldPartitionActorDescInstance::ToString(FWorldPartitionActorDesc::EToStringMode Mode) const
 {
 	return ActorDesc->ToString(Mode);
diff --git a/Engine/Source/Runtime/Engine/Private/WorldPartition/WorldPartitionActorDescUtils.cpp b/Engine/Source/Runtime/Engine/Private/WorldPartition/WorldPartitionActorDescUtils.cpp
--- a/Engine/Source/Runtime/Engine/Private/WorldPartition/WorldPartitionActorDescUtils.cpp
+++ b/Engine/Source/Runtime/Engine/Private/WorldPartition/WorldPartitionActorDescUtils.cpp
@@ -102,7 +102,7 @@ TUniquePtr<FWorldPartitionActorDesc> FWorldPartitionActorDescUtils::GetActorDesc
 			
 		if (!ActorDescInitData.NativeClass)
 		{
-			UE_LOG(LogWorldPartition, Warning, TEXT("Invalid class for actor guid `%s` ('%s') from package '%s'"), *NewActorDesc->GetGuid().ToString(), *NewActorDesc->GetActorName().ToString(), *NewActorDesc->GetActorPackage().ToString());
+			UE_LOG(LogWorldPartition, Warning, TEXT("Invalid class for actor guid `%s` ('%s') from package '%s'"), *NewActorDesc->GetGuid().ToString(), *NewActorDesc->GetActorNameString(), *NewActorDesc->GetActorPackage().ToString());
 			NewActorDesc->NativeClass.Reset();
 		}
 
diff --git a/Engine/Source/Runtime/Engine/Public/WorldPartition/ActorDescList.h b/Engine/Source/Runtime/Engine/Public/WorldPartition/ActorDescList.h
--- a/Engine/Source/Runtime/Engine/Public/WorldPartition/ActorDescList.h
+++ b/Engine/Source/Runtime/Engine/Public/WorldPartition/ActorDescList.h
@@ -207,7 +207,7 @@ void TActorDescList<DescType>::AddActorDescriptor(DescType* ActorDesc)
 {
 	check(!bIsProxy);
 	check(ActorDesc);
-	checkf(!ActorsByGuid.Contains(ActorDesc->GetGuid()), TEXT("Duplicated actor descriptor guid '%s' detected: `%s`"), *ActorDesc->GetGuid().ToString(), *ActorDesc->GetActorName().ToString());
+	checkf(!ActorsByGuid.Contains(ActorDesc->GetGuid()), TEXT("Duplicated actor descriptor guid '%s' detected: `%s`"), *ActorDesc->GetGuid().ToString(), *ActorDesc->GetActorNameString());
 
 	TUniquePtr<DescType>* NewActorDesc = new(ActorDescList) TUniquePtr<DescType>(ActorDesc);
 	ActorsByGuid.Add(ActorDesc->GetGuid(), NewActorDesc);
diff --git a/Engine/Source/Runtime/Engine/Public/WorldPartition/WorldPartitionActorDesc.h b/Engine/Source/Runtime/Engine/Public/WorldPartition/WorldPartitionActorDesc.h
--- a/Engine/Source/Runtime/Engine/Public/WorldPartition/WorldPartitionActorDesc.h
+++ b/Engine/Source/Runtime/Engine/Public/WorldPartition/WorldPartitionActorDesc.h
@@ -198,6 +198,11 @@ public:
 	ENGINE_API FName GetActorLabelOrName() const;
 	ENGINE_API FName GetDisplayClassName() const;
 
+	// Faster accessors for names as strings
+	ENGINE_API const FString& GetActorNameString() const;
+	ENGINE_API const FString& GetActorLabelString() const;
+	ENGINE_API const FString& GetDisplayClassNameString() const;
+
 	inline bool IsDefaultActorDesc() const { return bIsDefaultActorDesc; }
 		
 	virtual bool IsChildContainerInstance() const { return false; }
@@ -435,6 +440,8 @@ protected:
 	// Transient
 	UClass*							ActorNativeClass;
 	FName							ActorName;
+	FString							ActorLabelString;
+	FString							ActorDisplayClassNameString;
 	UActorDescContainer*			Container;
 	bool							bIsDefaultActorDesc;
 
diff --git a/Engine/Source/Runtime/Engine/Public/WorldPartition/WorldPartitionActorDescInstance.h b/Engine/Source/Runtime/Engine/Public/WorldPartition/WorldPartitionActorDescInstance.h
--- a/Engine/Source/Runtime/Engine/Public/WorldPartition/WorldPartitionActorDescInstance.h
+++ b/Engine/Source/Runtime/Engine/Public/WorldPartition/WorldPartitionActorDescInstance.h
@@ -117,6 +117,10 @@ public:
 	FName GetDisplayClassName() const { return ActorDesc->GetDisplayClassName(); }
 	
 	ENGINE_API const FText& GetUnloadedReason() const;
+
+	ENGINE_API const FString& GetActorNameString() const;
+	ENGINE_API const FString& GetActorLabelString() const;
+	ENGINE_API const FString& GetDisplayClassNameString() const;
 		
 protected:
 	ENGINE_API virtual FBox GetLocalEditorBounds() const { return ActorDesc->GetEditorBounds(); }
