Texture.cpp : fix check when RHI max dimension exceeds 16384

diff --git a/Engine/Source/Runtime/Engine/Private/Texture.cpp b/Engine/Source/Runtime/Engine/Private/Texture.cpp
--- a/Engine/Source/Runtime/Engine/Private/Texture.cpp
+++ b/Engine/Source/Runtime/Engine/Private/Texture.cpp
@@ -578,10 +578,11 @@ void UTexture::ValidateSettingsAfterImportOrEdit(bool * pRequiresNotifyMaterials
 #if WITH_EDITORONLY_DATA
 
 	// GetMaximumDimension is virtual, for the current texture type, on the current (host) RHI
-	//	not really right to ever be using it, quereis should be about the target platform
-	const int32 RHIMaximumDimension = GetMaximumDimension();
+	//	not really right to ever be using it, queries should be about the target platform
+	// GetMaximumDimensionOfNonVT is just a constant 16384
+	// beware GetMaximumDimension() can be over 16384 but we don't support that; it should always be clamepd
+	const int32 RHIMaximumDimension = FMath::Min<int32>(GetMaximumDimension(),GetMaximumDimensionOfNonVT());
 	check( RHIMaximumDimension > 0 );
-	check( RHIMaximumDimension <= GetMaximumDimensionOfNonVT() ); // GetMaximumDimensionOfNonVT is just a constant (16384)
 
 	if (Source.IsValid()) // we can have an empty source if the last source in a texture2d array is removed via the editor.
 	{
