Fix MarkComponentsAsGarbage to also mark all nested subobjects as garbage, this is needed for GC safety as valid objects cannot have garbage outers.
Doing this in the engine code allows removing game-specific workarounds.

diff --git a/Engine/Source/Runtime/Engine/Classes/GameFramework/Actor.h b/Engine/Source/Runtime/Engine/Classes/GameFramework/Actor.h
--- a/Engine/Source/Runtime/Engine/Classes/GameFramework/Actor.h
+++ b/Engine/Source/Runtime/Engine/Classes/GameFramework/Actor.h
@@ -3179,9 +3179,11 @@ public:
 	UE_DEPRECATED(5.3, "Use MarkComponentsAsGarbage instead.")
 	virtual void MarkComponentsAsPendingKill() { MarkComponentsAsGarbage(); }
 
-	/** Called to mark all components as garbage when the actor is being destroyed
+	/**
+	 * Called to mark all components and nested subobjects as garbage when the actor is being destroyed.
+	 * This will call OnComponentDestroyed on ActorComponents before marking them as garbage.
 	 *
-	 *  @param bModify if True, Modify will be called on actor before marking components
+	 * @param bModify if True, Modify will be called on ActorComponents before OnComponentDestroyed
 	 */
 	ENGINE_API virtual void MarkComponentsAsGarbage(bool bModify = true);
 	
diff --git a/Engine/Source/Runtime/Engine/Private/Actor.cpp b/Engine/Source/Runtime/Engine/Private/Actor.cpp
--- a/Engine/Source/Runtime/Engine/Private/Actor.cpp
+++ b/Engine/Source/Runtime/Engine/Private/Actor.cpp
@@ -5563,6 +5563,12 @@ void AActor::MarkComponentsAsGarbage(bool bModify)
 		Component->OnComponentDestroyed(true);
 		Component->MarkAsGarbage();
 	}
+
+	// Do a second pass to make sure all non-component subobjects are garbage as it is not valid for an object to have a garbage outer
+	ForEachObjectWithOuter(this, [](UObject* Object)
+	{
+		Object->MarkAsGarbage();
+	}, true, RF_MirroredGarbage);
 }
 
 void AActor::ReregisterAllComponents()
