Add net.SendDormantDestructionOnRemoval as a workaround for issues with net.SerializeNewActorOverrideLevel...

* Sometimes world streaming will stream out levels even though some players can still see them.
* Usually this is not a problem since the clients will remove the actors when it's told the level is gone.
* But when SerializeNewActorOverrideLevel is off, the client thinks the actors are in the main level, so they don't get removed.
* If the client streams the world back in the actors start to double-up D:
* Ideally world streaming would never stream out worlds until all clients can't see them, but that change is too risky now.
* So for now, notify all the removed actors to make sure the client cleans them up.

#tests Ensured in PIE actors that stuck around after streaming out levels are now correctly cleaned up when streaming out levels, quick smoke testing in standalone

diff --git a/Engine/Source/Runtime/Engine/Private/NetDriver.cpp b/Engine/Source/Runtime/Engine/Private/NetDriver.cpp
--- a/Engine/Source/Runtime/Engine/Private/NetDriver.cpp
+++ b/Engine/Source/Runtime/Engine/Private/NetDriver.cpp
@@ -317,6 +317,12 @@ namespace UE::Net::Private
 		SerializeNewActorOverrideLevel,
 		TEXT("If true, servers will serialize a spawned, replicated actor's level so the client attempts to spawn it into that level too. If false, clients will spawn all these actors into the persistent level."));
 
+	int32 SendDormantDestructionOnRemoval = 1;
+	static FAutoConsoleVariableRef CVarNetSendDormantDestructionOnRemoval(
+		TEXT("net.SendDormantDestructionOnRemoval"),
+		SendDormantDestructionOnRemoval,
+		TEXT("If true and net.SerializeNewActorOverrideLevel = 0, servers will send destruction info for dormant destroyed actors in invisible levels, regardless of how they were removed. This ensures they destroy on the client, in-case the they (or the levels they are in) were removed unexpectedly."));
+
 } //namespace UE::Net::Private
 
 namespace UE::Net
@@ -3900,6 +3906,28 @@ void UNetDriver::NotifySubObjectDestroyed(UObject* SubObject)
 
 void UNetDriver::RemoveNetworkActor(AActor* Actor)
 {
+	// This is a hack: Sometimes world streaming will stream out levels even though some players can still see them
+	// Usually this is not a problem since the clients will remove the actors when it's told the level is gone. 
+	// But when SerializeNewActorOverrideLevel is off, the client thinks the actors are in the main level, so they don't get removed.
+	// If the client streams the world back in the actors start to double-up D: 
+	// Ideally world streaming would never stream out worlds until all clients can't see them, but that change is too risky now. 
+	// So for now, notify all the removed actors to make sure the client cleans them up. 
+	if (IsServer() 
+		&& !UE::Net::Private::SerializeNewActorOverrideLevel && UE::Net::Private::SendDormantDestructionOnRemoval
+		&& IsValid(Actor) && !Actor->IsActorBeingDestroyed()
+		&& IsValid(Actor->GetLevel()) && !Actor->GetLevel()->bIsVisible
+		&& !Actor->IsNetStartupActor())
+	{
+		for (int32 i = ClientConnections.Num() - 1; i >= 0; i--)
+		{
+			UNetConnection* Connection = ClientConnections[i];
+			if (Connection && Connection->ClientHasInitializedLevel(Actor->GetLevel()))
+			{
+				SendDestructionInfoForLevelUnloadIfDormant(Actor, Connection);
+			}
+		}
+	}
+
 	GetNetworkObjectList().Remove(Actor);
 
 	// Remove from renamed list if destroyed
