Cancel texture I/O requests asynchronous and lock the update object.

diff --git a/Engine/Source/Runtime/Engine/Private/Streaming/Texture2DStreamIn_IO.cpp b/Engine/Source/Runtime/Engine/Private/Streaming/Texture2DStreamIn_IO.cpp
--- a/Engine/Source/Runtime/Engine/Private/Streaming/Texture2DStreamIn_IO.cpp
+++ b/Engine/Source/Runtime/Engine/Private/Streaming/Texture2DStreamIn_IO.cpp
@@ -215,7 +215,11 @@ void FTexture2DStreamIn_IO::CancelIORequests()
 	if (bBatchIORequest)
 	{
 		check(IORequests.IsEmpty());
-		return BatchRequest.Cancel();
+		if (BatchRequest.IsPending())
+		{
+			BatchRequest.Cancel();
+		}
+		return;
 	}
 
 	check(BatchRequest.IsNone());
@@ -312,13 +316,6 @@ void FTexture2DStreamIn_IO::Abort()
 
 		if (HasPendingIORequests())
 		{
-			if (bBatchIORequest)
-			{
-				check(IORequests.IsEmpty());
-				return BatchRequest.Cancel();
-			}
-
-			check(BatchRequest.IsNone());
 			// Prevent the update from being considered done before this is finished.
 			// By checking that it was not already canceled, we make sure this doesn't get called twice.
 			(new FAsyncCancelIORequestsTask(this))->StartBackgroundTask();
