Fix regression from 33237422.

diff --git a/Engine/Source/Runtime/Engine/Private/WorldPartition/WorldPartitionActorDesc.cpp b/Engine/Source/Runtime/Engine/Private/WorldPartition/WorldPartitionActorDesc.cpp
--- a/Engine/Source/Runtime/Engine/Private/WorldPartition/WorldPartitionActorDesc.cpp
+++ b/Engine/Source/Runtime/Engine/Private/WorldPartition/WorldPartitionActorDesc.cpp
@@ -159,6 +159,7 @@ void FWorldPartitionActorDesc::Init(const AActor* InActor)
 	ActorPackage = InActor->GetPackage()->GetFName();
 	ActorPath = bIsDefaultActorDesc ? *InActor->GetClass()->GetPathName() : *InActor->GetPathName();
 	ActorName = InActor->GetFName();
+	ActorNameString = ActorName.ToString();
 
 	ContentBundleGuid = InActor->GetContentBundleGuid();
 
@@ -212,7 +213,8 @@ void FWorldPartitionActorDesc::Init(const FWorldPartitionActorDescInitData& Desc
 	ActorPath = DescData.ActorPath;
 	ActorNativeClass = DescData.NativeClass;
 	NativeClass = *DescData.NativeClass->GetPathName();
-	ActorName = *GetActorNameString();
+	ActorName = *FPaths::GetExtension(ActorPath.ToString());
+	ActorNameString = ActorName.ToString();
 
 	auto DeprecateClass = [this](FArchive& Archive)
 	{
@@ -774,7 +776,7 @@ FName FWorldPartitionActorDesc::GetActorName() const
 
 const FString& FWorldPartitionActorDesc::GetActorNameString() const
 {
-	return ActorPath.GetSubPathString();
+	return ActorNameString;
 }
 
 const FString& FWorldPartitionActorDesc::GetActorLabelString() const
diff --git a/Engine/Source/Runtime/Engine/Public/WorldPartition/WorldPartitionActorDesc.h b/Engine/Source/Runtime/Engine/Public/WorldPartition/WorldPartitionActorDesc.h
--- a/Engine/Source/Runtime/Engine/Public/WorldPartition/WorldPartitionActorDesc.h
+++ b/Engine/Source/Runtime/Engine/Public/WorldPartition/WorldPartitionActorDesc.h
@@ -440,6 +440,7 @@ protected:
 	// Transient
 	UClass*							ActorNativeClass;
 	FName							ActorName;
+	FString							ActorNameString;
 	FString							ActorLabelString;
 	FString							ActorDisplayClassNameString;
 	UActorDescContainer*			Container;
