Chaos - Add physics thread query

diff --git a/Engine/Source/Runtime/Engine/Private/Collision/CollisionConversions.cpp b/Engine/Source/Runtime/Engine/Private/Collision/CollisionConversions.cpp
--- a/Engine/Source/Runtime/Engine/Private/Collision/CollisionConversions.cpp
+++ b/Engine/Source/Runtime/Engine/Private/Collision/CollisionConversions.cpp
@@ -16,6 +16,7 @@
 
 #include "PhysicsEngine/CollisionQueryFilterCallback.h"
 #include "Engine/ActorInstanceManagerInterface.h"
+#include "PhysicsProxy/ClusterUnionPhysicsProxy.h"
 #include "PhysicsProxy/SingleParticlePhysicsProxy.h"
 #include "PhysicsProxy/GeometryCollectionPhysicsProxy.h"
 
@@ -232,11 +233,24 @@ const FPhysicsShape* GetGTShape<false>(const FPhysicsShape* PTShape, const FPhys
 
 const FPhysicsActor* GetGTActor(const Chaos::FGeometryParticleHandle* PTActor)
 {
-	//TODO: need to pass in context so that in PT we always return null
-	//In frozen GT this is ok because object can't be unregistered while we're holding on to this
-	//On PT this is only true if we acquire a gt lock which is not great
-	auto Proxy = static_cast<const Chaos::FSingleParticlePhysicsProxy*>(PTActor->PhysicsProxy());
-	return Proxy->GetParticle_LowLevel();
+	const IPhysicsProxyBase* Proxy = PTActor->PhysicsProxy();
+	if (Proxy)
+	{
+		const EPhysicsProxyType ProxyType = Proxy->GetType();
+		switch (ProxyType)
+		{
+		case EPhysicsProxyType::ClusterUnionProxy:
+			return static_cast<const Chaos::FClusterUnionPhysicsProxy*>(Proxy)->GetParticle_External();
+		case EPhysicsProxyType::GeometryCollectionType:
+			return static_cast<const FGeometryCollectionPhysicsProxy*>(Proxy)->GetInitialRootParticle_External();
+		case EPhysicsProxyType::SingleParticleProxy:
+			return static_cast<const Chaos::FSingleParticlePhysicsProxy*>(Proxy)->GetParticle_LowLevel();
+		default:
+			break;
+		}
+	}
+
+	return nullptr;
 }
 
 template <typename THitLocation>
diff --git a/Engine/Source/Runtime/Engine/Private/Collision/SceneQuery.cpp b/Engine/Source/Runtime/Engine/Private/Collision/SceneQuery.cpp
--- a/Engine/Source/Runtime/Engine/Private/Collision/SceneQuery.cpp
+++ b/Engine/Source/Runtime/Engine/Private/Collision/SceneQuery.cpp
@@ -1210,3 +1210,31 @@ template struct FGenericGeomPhysicsInterfaceUsingSpatialAcceleration<TACCEL, FPh
 
 DEFINE_GENERIC_GEOM_FOR_ACCEL(Chaos::IDefaultChaosSpatialAcceleration)
 DEFINE_GENERIC_GEOM_FOR_ACCEL(IExternalSpatialAcceleration)
+
+namespace Chaos::Private
+{
+	bool FGenericPhysicsInterface_Internal::SpherecastMulti(const UWorld* World, float QueryRadius, TArray<FHitResult>& OutHits, FVector Start, FVector End, ECollisionChannel TraceChannel, const FCollisionQueryParams& Params, const FCollisionResponseParams& ResponseParams, const FCollisionObjectQueryParams& ObjectParams)
+	{
+		SCOPE_CYCLE_COUNTER(STAT_Collision_SceneQueryTotal);
+		SCOPE_CYCLE_COUNTER(STAT_Collision_SpherecastMultiple_Internal);
+		CSV_SCOPED_TIMING_STAT(SceneQuery, SpherecastMultiple_Internal);
+
+		Chaos::EnsureIsInPhysicsThreadContext();
+
+		if ((World == NULL) || (World->GetPhysicsScene() == NULL))
+		{
+			return false;
+		}
+
+		if (QueryRadius > UE_SMALL_NUMBER)
+		{
+			using PTTraits = TSQTraits<ChaosInterface::FPTSweepHit, ESweepOrRay::Sweep, ESingleMultiOrTest::Multi>;
+			return TSceneCastCommonImp<PTTraits, FGeomSQAdditionalInputs>(World, OutHits, FGeomSQAdditionalInputs(FCollisionShape::MakeSphere(QueryRadius), FQuat::Identity), Start, End, TraceChannel, Params, ResponseParams, ObjectParams, FDefaultAccelContainer{});
+		}
+		else
+		{
+			using PTTraits = TSQTraits<ChaosInterface::FPTRaycastHit, ESweepOrRay::Raycast, ESingleMultiOrTest::Multi>;
+			return TSceneCastCommonImp<PTTraits, FRaycastSQAdditionalInputs>(World, OutHits, FRaycastSQAdditionalInputs(), Start, End, TraceChannel, Params, ResponseParams, ObjectParams, FDefaultAccelContainer{});
+		}
+	};
+}
diff --git a/Engine/Source/Runtime/Engine/Private/Collision/WorldCollision.cpp b/Engine/Source/Runtime/Engine/Private/Collision/WorldCollision.cpp
--- a/Engine/Source/Runtime/Engine/Private/Collision/WorldCollision.cpp
+++ b/Engine/Source/Runtime/Engine/Private/Collision/WorldCollision.cpp
@@ -36,6 +36,7 @@ DEFINE_STAT(STAT_Collision_FBodyInstance_OverlapMulti);
 DEFINE_STAT(STAT_Collision_FBodyInstance_OverlapTest);
 DEFINE_STAT(STAT_Collision_FBodyInstance_LineTrace);
 DEFINE_STAT(STAT_Collision_PreFilter);
+DEFINE_STAT(STAT_Collision_SpherecastMultiple_Internal);
 
 /** default collision response container - to be used without reconstructing every time**/
 FCollisionResponseContainer FCollisionResponseContainer::DefaultResponseContainer(ECR_Block);
diff --git a/Engine/Source/Runtime/Engine/Public/Collision.h b/Engine/Source/Runtime/Engine/Public/Collision.h
--- a/Engine/Source/Runtime/Engine/Public/Collision.h
+++ b/Engine/Source/Runtime/Engine/Public/Collision.h
@@ -35,6 +35,7 @@ DECLARE_CYCLE_STAT_EXTERN(TEXT("BodyInstanceOverlapMulti"), STAT_Collision_FBody
 DECLARE_CYCLE_STAT_EXTERN(TEXT("BodyInstanceOverlapTest"), STAT_Collision_FBodyInstance_OverlapTest, STATGROUP_Collision, );
 DECLARE_CYCLE_STAT_EXTERN(TEXT("BodyInstanceLineTrace"), STAT_Collision_FBodyInstance_LineTrace, STATGROUP_Collision, );
 DECLARE_CYCLE_STAT_EXTERN(TEXT("PreFilter"), STAT_Collision_PreFilter, STATGROUP_CollisionVerbose, );
+DECLARE_CYCLE_STAT_EXTERN(TEXT("SpherecastMultipleInternal"), STAT_Collision_SpherecastMultiple_Internal, STATGROUP_Collision, );
 
 /** Enable collision analyzer support */
 #if (1 && !(UE_BUILD_SHIPPING || UE_BUILD_TEST) && WITH_EDITOR && WITH_UNREAL_DEVELOPER_TOOLS)
diff --git a/Engine/Source/Runtime/Engine/Public/Physics/GenericPhysicsInterface.h b/Engine/Source/Runtime/Engine/Public/Physics/GenericPhysicsInterface.h
--- a/Engine/Source/Runtime/Engine/Public/Physics/GenericPhysicsInterface.h
+++ b/Engine/Source/Runtime/Engine/Public/Physics/GenericPhysicsInterface.h
@@ -93,3 +93,19 @@ ENGINE_API bool FGenericPhysicsInterface::GeomSweepMulti(const UWorld* World, co
 
 template<>
 ENGINE_API bool FGenericPhysicsInterface::GeomOverlapMulti(const UWorld* World, const FPhysicsGeometryCollection& InGeom, const FVector& InPosition, const FQuat& InRotation, TArray<FOverlapResult>& OutOverlaps, ECollisionChannel TraceChannel, const FCollisionQueryParams& Params, const FCollisionResponseParams& ResponseParams, const FCollisionObjectQueryParams& ObjectParams);
+
+namespace Chaos::Private
+{
+	struct FGenericPhysicsInterface_Internal
+	{
+		/**
+		*  INTERNAL USE ONLY
+		*  Physics thread sphere query:
+		*  Trace a sphere against the world and return touching hits and then first blocking hit
+		*  Results are sorted, so a blocking hit (if found) will be the last element of the array
+		*  Only the single closest blocking result will be generated, no tests will be done after that.
+		*  Falls back to a raycast is the query radius is less than or equal to zero.
+		*/
+		static ENGINE_API bool SpherecastMulti(const UWorld* World, float QueryRadius, TArray<FHitResult>& OutHits, FVector Start, FVector End, ECollisionChannel TraceChannel, const FCollisionQueryParams& Params, const FCollisionResponseParams& ResponseParams, const FCollisionObjectQueryParams& ObjectParams = FCollisionObjectQueryParams::DefaultObjectQueryParam);
+	};
+}
