ReplicatedSubObjects
* Added some utility functions to count the amount of registered replicated objects of an actor

#rnx

diff --git a/Engine/Source/Runtime/Engine/Classes/Engine/NetDriver.h b/Engine/Source/Runtime/Engine/Classes/Engine/NetDriver.h
--- a/Engine/Source/Runtime/Engine/Classes/Engine/NetDriver.h
+++ b/Engine/Source/Runtime/Engine/Classes/Engine/NetDriver.h
@@ -372,6 +372,9 @@ namespace UE::Net
 	};
 
 	bool ShouldIgnoreStaticActorDestruction();
+
+	/** Utility function that counts the number of replicated subobjects owned by a given actor. Be aware that this function is not cheap because it needs to iterate and count all replicated components and their subobjects too. */
+	ENGINE_API uint32 CountReplicatedSubObjectsOfActor(AActor* ReplicatedActor);
 }
 
 using FConnectionMap = TMap<TSharedRef<const FInternetAddr>, TObjectPtr<UNetConnection>, FDefaultSetAllocator, FInternetAddrConstKeyMapFuncs<TObjectPtr<UNetConnection>>>;
diff --git a/Engine/Source/Runtime/Engine/Private/Net/NetSubObjectRegistryGetter.h b/Engine/Source/Runtime/Engine/Private/Net/NetSubObjectRegistryGetter.h
--- a/Engine/Source/Runtime/Engine/Private/Net/NetSubObjectRegistryGetter.h
+++ b/Engine/Source/Runtime/Engine/Private/Net/NetSubObjectRegistryGetter.h
@@ -70,6 +70,22 @@ public:
 	{
 		InActor->BuildReplicatedComponentsInfo();
 	}
+
+	static uint32 CountReplicatedSubObjectsOfActor(AActor* InActor)
+	{
+		// Start with the actor's subobjects;
+		uint32 NumReplicatedObjects = (uint32)InActor->ReplicatedSubObjects.Num();
+		// Add the replicated components
+		NumReplicatedObjects += (uint32)InActor->ReplicatedComponentsInfo.Num();
+
+		// Finally add the subobjects of the replicated components	
+		for (const UE::Net::FReplicatedComponentInfo& ReplicatedComponentInfo : InActor->ReplicatedComponentsInfo)
+		{
+			NumReplicatedObjects += (uint32)ReplicatedComponentInfo.SubObjects.Num();
+		}
+
+		return NumReplicatedObjects;
+	}
 };
 
 }
diff --git a/Engine/Source/Runtime/Engine/Private/NetDriver.cpp b/Engine/Source/Runtime/Engine/Private/NetDriver.cpp
--- a/Engine/Source/Runtime/Engine/Private/NetDriver.cpp
+++ b/Engine/Source/Runtime/Engine/Private/NetDriver.cpp
@@ -343,13 +343,18 @@ namespace UE::Net
 		return UE::Net::Private::bIgnoreStaticActorDestruction;
 	}
 
+	uint32 CountReplicatedSubObjectsOfActor(AActor* ReplicatedActor)
+	{
+		return UE::Net::FSubObjectRegistryGetter::CountReplicatedSubObjectsOfActor(ReplicatedActor);
+	}
+
 	bool bDiscardTornOffActorRPCs = true;
 	FAutoConsoleVariableRef CVarNetDiscardTornOffActorRPCs(
 		TEXT("net.DiscardTornOffActorRPCs"),
 		bDiscardTornOffActorRPCs,
 		TEXT("If enabled, discard RPCs if the actor has been torn off."),
 		ECVF_Default);
-}
+} // end namespace UE::Net
 
 #if UE_BUILD_SHIPPING
 #define DEBUG_REMOTEFUNCTION(Format, ...)
diff --git a/Engine/Source/Runtime/Net/Core/Public/Net/Core/Misc/NetSubObjectRegistry.h b/Engine/Source/Runtime/Net/Core/Public/Net/Core/Misc/NetSubObjectRegistry.h
--- a/Engine/Source/Runtime/Net/Core/Public/Net/Core/Misc/NetSubObjectRegistry.h
+++ b/Engine/Source/Runtime/Net/Core/Public/Net/Core/Misc/NetSubObjectRegistry.h
@@ -46,6 +46,7 @@ public:
 	NETCORE_API bool RemoveSubObject(UObject* InSubObject);
 
 	bool IsEmpty() const { return Registry.Num() == 0; }
+	int32 Num() const	 { return Registry.Num(); }
 
 	/**
 	* Remove all the indexes specified by the passed array
