[World Partition] Make sure to intersect with fixed-size cells insetad of the content bounds to reduce actor counts.

diff --git a/Engine/Source/Runtime/Engine/Private/WorldPartition/RuntimeHashSet/RuntimePartitionLHGrid.cpp b/Engine/Source/Runtime/Engine/Private/WorldPartition/RuntimeHashSet/RuntimePartitionLHGrid.cpp
--- a/Engine/Source/Runtime/Engine/Private/WorldPartition/RuntimeHashSet/RuntimePartitionLHGrid.cpp
+++ b/Engine/Source/Runtime/Engine/Private/WorldPartition/RuntimeHashSet/RuntimePartitionLHGrid.cpp
@@ -166,6 +166,7 @@ bool URuntimePartitionLHGrid::GenerateStreaming(const FGenerateStreamingParams&
 
 		if (bIsSpatiallyLoaded)
 		{
+			CellDesc.Bounds = FCellCoord::GetCellBounds(CellCoord, CellSize);
 			CellDesc.CellBounds = FCellCoord::GetCellBounds(CellCoord, CellSize);
 		}
 	}
@@ -177,4 +178,4 @@ FArchive& URuntimePartitionLHGrid::AppendCellGuid(FArchive& InAr)
 {
 	return Super::AppendCellGuid(InAr) << CellSize;
 }
-#endif
\ No newline at end of file
+#endif
