Add additional assertions for validity of mesh description for each lod level with name of bad asset

diff --git a/Engine/Source/Runtime/Engine/Private/StaticMesh.cpp b/Engine/Source/Runtime/Engine/Private/StaticMesh.cpp
--- a/Engine/Source/Runtime/Engine/Private/StaticMesh.cpp
+++ b/Engine/Source/Runtime/Engine/Private/StaticMesh.cpp
@@ -3122,6 +3122,7 @@ void FStaticMeshRenderData::Cache(const ITargetPlatform* TargetPlatform, UStatic
 				AllocateLODResources(NumSourceModels);
 				for (int32 LodIndex = 0; LodIndex < NumSourceModels; LodIndex++)
 				{
+					checkf(Owner->IsMeshDescriptionValid(LodIndex), TEXT("Bad MeshDescription at lod index %d on %s"), LodIndex, *GetPathNameSafe(Owner));
 					Owner->BuildFromMeshDescription(*Owner->GetMeshDescription(LodIndex), LODResources[LodIndex]);
 				}
 			}
