[SkeletalMeshGPUSkin] added index checking before accessing

diff --git a/Engine/Source/Runtime/Engine/Private/SkeletalRenderGPUSkin.cpp b/Engine/Source/Runtime/Engine/Private/SkeletalRenderGPUSkin.cpp
--- a/Engine/Source/Runtime/Engine/Private/SkeletalRenderGPUSkin.cpp
+++ b/Engine/Source/Runtime/Engine/Private/SkeletalRenderGPUSkin.cpp
@@ -1617,7 +1617,11 @@ const FVertexFactory* FSkeletalMeshObjectGPUSkin::GetStaticSkinVertexFactory(int
 
 	if (LOD.GPUSkinTechnique != ESkeletalMeshGPUSkinTechnique::Inline || VFMode == ESkinVertexFactoryMode::RayTracing)
 	{
-		return LOD.GPUSkinVertexFactories.PassthroughVertexFactories[ChunkIdx].Get();
+		if (LOD.GPUSkinVertexFactories.PassthroughVertexFactories.IsValidIndex(ChunkIdx))
+		{
+			return LOD.GPUSkinVertexFactories.PassthroughVertexFactories[ChunkIdx].Get();
+		}
+		return nullptr;
 	}
 
 	const FGPUBaseSkinVertexFactory* VertexFactory = LOD.GPUSkinVertexFactories.VertexFactories[ChunkIdx].Get();
