- Added helper scope-object to handle Rename operations for (generic) external packages/objects (similar to world/level behaviour)
- Added helper function to handle Duplication of external packages (replaced ULevel behaviour)
- Added FExternalAssetDependencyGatherer to handle external packages used by [REDACTED_PERSON]
- Also DetachLinker for external objects not part of a UWorld when doing a reload

diff --git a/Engine/Source/Developer/SourceControl/Private/SourceControlHelpers.cpp b/Engine/Source/Developer/SourceControl/Private/SourceControlHelpers.cpp
--- a/Engine/Source/Developer/SourceControl/Private/SourceControlHelpers.cpp
+++ b/Engine/Source/Developer/SourceControl/Private/SourceControlHelpers.cpp
@@ -978,13 +978,15 @@ bool USourceControlHelpers::ApplyOperationAndReloadPackages(const TArray<FString
 			{
 				if (Asset->IsPackageExternal())
 				{
-					if (bReloadWorld)
+					const bool bIsAssetPartOfWorld = Asset->GetWorld() && Asset->GetWorld()->GetPackage();
+					// In case this is an explicit world object/actor reload or just an external object unload/detach it 
+					if (bReloadWorld || !bIsAssetPartOfWorld)
 					{
 						// detach linker on the object
 						DetachLinker(Package);
 
-						// but track its world for reloading - not the object package itself
-						if (Asset->GetWorld() && Asset->GetWorld()->GetPackage())
+						// track its world for reloading - not the object package itself
+						if (bIsAssetPartOfWorld)
 						{
 							UniqueLoadedPackages.Add(Asset->GetWorld()->GetPackage());
 						}
diff --git a/Engine/Source/Runtime/Engine/Classes/Engine/ExternalAssetDependencyGatherer.h b/Engine/Source/Runtime/Engine/Classes/Engine/ExternalAssetDependencyGatherer.h
new file mode 100644
--- /dev/null
+++ b/Engine/Source/Runtime/Engine/Classes/Engine/ExternalAssetDependencyGatherer.h
@@ -0,0 +1,20 @@
+﻿// Copyright Epic Games, Inc. All Rights Reserved.
+
+#pragma once
+
+#if WITH_EDITOR
+
+#include "AssetRegistry/AssetDependencyGatherer.h"
+#include "AssetRegistry/AssetData.h"
+#include "Engine/World.h"
+
+class ENGINE_API FExternalAssetDependencyGatherer : public IAssetDependencyGatherer
+{
+public:	
+	FExternalAssetDependencyGatherer() = default;
+	virtual ~FExternalAssetDependencyGatherer() = default;
+
+	virtual void GatherDependencies(const FAssetData& AssetData, const FAssetRegistryState& AssetRegistryState, TFunctionRef<FARCompiledFilter(const FARFilter&)> CompileFilterFunc, TArray<IAssetDependencyGatherer::FGathereredDependency>& OutDependencies, TArray<FString>& OutDependencyDirectories) const override;
+};
+
+#endif
\ No newline at end of file
diff --git a/Engine/Source/Runtime/Engine/Private/ExternalAssetDependencyGatherer.cpp b/Engine/Source/Runtime/Engine/Private/ExternalAssetDependencyGatherer.cpp
new file mode 100644
--- /dev/null
+++ b/Engine/Source/Runtime/Engine/Private/ExternalAssetDependencyGatherer.cpp
@@ -0,0 +1,29 @@
+﻿// Copyright Epic Games, Inc. All Rights Reserved.
+#include "Engine/ExternalAssetDependencyGatherer.h"
+
+#if WITH_EDITOR
+
+#include "AssetRegistry/ARFilter.h"
+#include "AssetRegistry/AssetRegistryState.h"
+#include "ExternalPackageHelper.h"
+
+void FExternalAssetDependencyGatherer::GatherDependencies(const FAssetData& AssetData, const FAssetRegistryState& AssetRegistryState, TFunctionRef<FARCompiledFilter(const FARFilter&)> CompileFilterFunc, TArray<IAssetDependencyGatherer::FGathereredDependency>& OutDependencies, TArray<FString>& OutDependencyDirectories) const
+{		
+	const FString ExternalObjectsPath = FExternalPackageHelper::GetExternalObjectsPath(AssetData.PackageName.ToString());
+	OutDependencyDirectories.Add(ExternalObjectsPath);
+
+	FARFilter Filter;
+	Filter.bRecursivePaths = true;
+	Filter.bIncludeOnlyOnDiskAssets = true;
+	Filter.PackagePaths.Add(*ExternalObjectsPath);
+
+	TArray<FAssetData> FilteredAssets;
+	AssetRegistryState.GetAssets(CompileFilterFunc(Filter), {}, FilteredAssets, true);
+
+	for (const FAssetData& FilteredAsset : FilteredAssets)
+	{
+		OutDependencies.Emplace(IAssetDependencyGatherer::FGathereredDependency{ FilteredAsset.PackageName, UE::AssetRegistry::EDependencyProperty::Game | UE::AssetRegistry::EDependencyProperty::Build });
+	}
+}
+
+#endif
diff --git a/Engine/Source/Runtime/Engine/Private/ExternalPackageHelper.cpp b/Engine/Source/Runtime/Engine/Private/ExternalPackageHelper.cpp
--- a/Engine/Source/Runtime/Engine/Private/ExternalPackageHelper.cpp
+++ b/Engine/Source/Runtime/Engine/Private/ExternalPackageHelper.cpp
@@ -8,6 +8,8 @@
 #include "Misc/ArchiveMD5.h"
 #include "Misc/Paths.h"
 #include "HAL/PlatformApplicationMisc.h"
+#include "UObject/MetaData.h"
+#include "GameFramework/Actor.h"
 
 FExternalPackageHelper::FOnObjectPackagingModeChanged FExternalPackageHelper::OnObjectPackagingModeChanged;
 
@@ -16,7 +18,79 @@ EPackageFlags FExternalPackageHelper::GetDefaultExternalPackageFlags()
 	return (PKG_EditorOnly | PKG_ContainsMapData | PKG_NewlyCreated);
 }
 
-UPackage* FExternalPackageHelper::CreateExternalPackage(UObject* InObjectOuter, const FString& InObjectPath, EPackageFlags InFlags, const UExternalDataLayerAsset* InExternalDataLayerAsset)
+FExternalPackageHelper::FRenameExternalObjectsHelperContext::FRenameExternalObjectsHelperContext(const UObject* SourceObject, ERenameFlags Flags)
+{
+	if (GIsEditor && ((Flags & REN_Test) == 0))
+	{
+		check(SourceObject);
+		const UPackage* Package = SourceObject->GetPackage();
+		check(Package);
+		SourcePackage = Package;
+		OldObject = SourceObject;
+	}
+}
+
+FExternalPackageHelper::FRenameExternalObjectsHelperContext::~FRenameExternalObjectsHelperContext()
+{	
+	if (GIsEditor && OldObject && (OldObject->GetPackage() != SourcePackage))
+	{		
+		// Get external packages
+		const TArray<UPackage*>& ExternalObjectPackages = OldObject->GetPackage()->GetExternalPackages();
+		for (const UPackage* ExternalPackage : ExternalObjectPackages)
+		{
+			TArray<UObject*> DependantObjects;
+			ForEachObjectWithPackage(ExternalPackage, [&DependantObjects](UObject* Object)
+			{
+				check(Object)
+				if (!Cast<UMetaData>(Object))
+				{
+					DependantObjects.Add(Object);
+				}
+				return true;
+			}, false);
+
+			for (UObject* Object : DependantObjects)
+            {
+			    FExternalPackageHelper::SetPackagingMode(Object, nullptr, false);
+			    FExternalPackageHelper::SetPackagingMode(Object, OldObject, true);
+            }
+		}
+	}
+}
+
+void FExternalPackageHelper::DuplicateExternalPackages(const UObject* InObject, FObjectDuplicationParameters& InDuplicationParameters, EActorPackagingScheme ActorPackagingScheme /*= EActorPackagingScheme::Reduced*/)
+{
+	if (InDuplicationParameters.DuplicateMode != EDuplicateMode::PIE && InDuplicationParameters.bAssignExternalPackages)
+	{
+		const UPackage* SourcePackage = InObject->GetPackage();
+		UPackage* DestinationPackage = InDuplicationParameters.DestOuter->GetPackage();
+		
+		FString ReplaceFrom = FPaths::GetBaseFilename(*SourcePackage->GetName());
+		ReplaceFrom = FString::Printf(TEXT("%s.%s:"), *ReplaceFrom, *ReplaceFrom);
+
+		FString ReplaceTo = FPaths::GetBaseFilename(*DestinationPackage->GetName());
+		ReplaceTo = FString::Printf(TEXT("%s.%s:"), *ReplaceTo, *ReplaceTo);
+			
+		ForEachObjectWithOuter(InObject, [&ReplaceFrom, &ReplaceTo, &InDuplicationParameters, SourcePackage, DestinationPackage, ActorPackagingScheme](const UObject* Object)
+		{
+			if (UPackage* Package = Object ? Object->GetExternalPackage() : nullptr)
+			{
+				FString Path = Object->GetPathName();
+				if (DestinationPackage != SourcePackage)
+				{
+					Path = Path.Replace(*ReplaceFrom, *ReplaceTo);
+				}
+				UPackage* DupPackage = Object->IsA<AActor>() ? ULevel::CreateActorPackage(DestinationPackage, ActorPackagingScheme, Path, Object) : FExternalPackageHelper::CreateExternalPackage(DestinationPackage, Path);
+				DupPackage->MarkAsFullyLoaded();
+				DupPackage->MarkPackageDirty();
+				
+				InDuplicationParameters.DuplicationSeed.Add(Package, DupPackage);
+			}
+		}, /*bIncludeNestedObjects*/ true);
+	}
+}
+
+UPackage* FExternalPackageHelper::CreateExternalPackage(const UObject* InObjectOuter, const FString& InObjectPath, EPackageFlags InFlags, const UExternalDataLayerAsset* InExternalDataLayerAsset)
 {
 	const UPackage* OutermostPackage = InObjectOuter->IsA<UPackage>() ? CastChecked<UPackage>(InObjectOuter) : InObjectOuter->GetOutermostObject()->GetPackage();
 	const FString RootPath = InExternalDataLayerAsset ? FExternalDataLayerHelper::GetExternalDataLayerLevelRootPath(InExternalDataLayerAsset, OutermostPackage->GetName()) : OutermostPackage->GetName();
@@ -26,7 +100,7 @@ UPackage* FExternalPackageHelper::CreateExternalPackage(UObject* InObjectOuter,
 	return Package;
 }
 
-void FExternalPackageHelper::SetPackagingMode(UObject* InObject, UObject* InObjectOuter, bool bInIsPackageExternal, bool bInShouldDirty, EPackageFlags InExternalPackageFlags)
+void FExternalPackageHelper::SetPackagingMode(UObject* InObject, const UObject* InObjectOuter, bool bInIsPackageExternal, bool bInShouldDirty, EPackageFlags InExternalPackageFlags)
 {
 	if (bInIsPackageExternal == InObject->IsPackageExternal())
 	{
diff --git a/Engine/Source/Runtime/Engine/Private/Level.cpp b/Engine/Source/Runtime/Engine/Private/Level.cpp
--- a/Engine/Source/Runtime/Engine/Private/Level.cpp
+++ b/Engine/Source/Runtime/Engine/Private/Level.cpp
@@ -1385,33 +1385,7 @@ void ULevel::PreDuplicate(FObjectDuplicationParameters& DupParams)
 		}
 	}
 
-	if (DupParams.DuplicateMode != EDuplicateMode::PIE && DupParams.bAssignExternalPackages)
-	{
-		UPackage* SrcPackage = GetPackage();
-		UPackage* DstPackage = DupParams.DestOuter->GetPackage();
-
-		FString ReplaceFrom = FPaths::GetBaseFilename(*SrcPackage->GetName());
-		ReplaceFrom = FString::Printf(TEXT("%s.%s:"), *ReplaceFrom, *ReplaceFrom);
-
-		FString ReplaceTo = FPaths::GetBaseFilename(*DstPackage->GetName());
-		ReplaceTo = FString::Printf(TEXT("%s.%s:"), *ReplaceTo, *ReplaceTo);
-
-		ForEachObjectWithOuter(this, [this, &SrcPackage, &DstPackage, &ReplaceFrom, &ReplaceTo, &DupParams](UObject* Object)
-		{
-			if (UPackage* Package = Object ? Object->GetExternalPackage() : nullptr)
-			{
-				FString Path = Object->GetPathName();
-				if (DstPackage != SrcPackage)
-				{
-					Path = Path.Replace(*ReplaceFrom, *ReplaceTo);
-				}
-				UPackage* DupPackage = Object->IsA<AActor>() ? ULevel::CreateActorPackage(DstPackage, GetActorPackagingScheme(), Path, Object) : FExternalPackageHelper::CreateExternalPackage(DstPackage, Path);
-				DupPackage->MarkAsFullyLoaded();
-				DupPackage->MarkPackageDirty();
-				DupParams.DuplicationSeed.Add(Package, DupPackage);
-			}
-		}, /*bIncludeNestedObjects*/ true);
-	}
+	FExternalPackageHelper::DuplicateExternalPackages(this, DupParams, GetActorPackagingScheme());
 #endif
 }
 
diff --git a/Engine/Source/Runtime/Engine/Public/ExternalPackageHelper.h b/Engine/Source/Runtime/Engine/Public/ExternalPackageHelper.h
--- a/Engine/Source/Runtime/Engine/Public/ExternalPackageHelper.h
+++ b/Engine/Source/Runtime/Engine/Public/ExternalPackageHelper.h
@@ -18,11 +18,24 @@
 #include "WorldPartition/DataLayer/ExternalDataLayerInstance.h"
 #include "WorldPartition/DataLayer/ExternalDataLayerAsset.h"
 #include "WorldPartition/DataLayer/ExternalDataLayerHelper.h"
+#include "Engine/Level.h"
 
 class FExternalPackageHelper
 {
 public:
-
+	class ENGINE_API FRenameExternalObjectsHelperContext
+	{
+	public:
+		FRenameExternalObjectsHelperContext() = delete;
+		FRenameExternalObjectsHelperContext(FRenameExternalObjectsHelperContext&) = delete;
+		FRenameExternalObjectsHelperContext& operator=(FRenameExternalObjectsHelperContext&) = delete;
+		
+		explicit FRenameExternalObjectsHelperContext(const UObject* SourceObject, ERenameFlags Flags);
+		~FRenameExternalObjectsHelperContext();
+	private:		
+		const UObject* OldObject = nullptr;
+		const UPackage* SourcePackage = nullptr;
+	};
 	DECLARE_EVENT_TwoParams(FExternalPackageHelper, FOnObjectPackagingModeChanged, UObject*, bool /* bExternal */);
 	static ENGINE_API FOnObjectPackagingModeChanged OnObjectPackagingModeChanged;
 
@@ -33,7 +46,7 @@ public:
 	 * @param InFlags the package flags to apply
 	 * @return the created package
 	 */
-	static ENGINE_API UPackage* CreateExternalPackage(UObject* InObjectOuter, const FString& InObjectPath, EPackageFlags InFlags = FExternalPackageHelper::GetDefaultExternalPackageFlags(), const UExternalDataLayerAsset* InExternalDataLayerAsset = nullptr);
+	static ENGINE_API UPackage* CreateExternalPackage(const UObject* InObjectOuter, const FString& InObjectPath, EPackageFlags InFlags = FExternalPackageHelper::GetDefaultExternalPackageFlags(), const UExternalDataLayerAsset* InExternalDataLayerAsset = nullptr);
 
 	/** Returns default external package flags used to create external packages. */
 	static ENGINE_API EPackageFlags GetDefaultExternalPackageFlags();
@@ -46,7 +59,7 @@ public:
 	 * @param bInShouldDirty should dirty or not the object's outer package
 	 * @param InExternalPackageFlags the flags to apply to the external package if bInIsPackageExternal is true
 	 */
-	static ENGINE_API void SetPackagingMode(UObject* InObject, UObject* InObjectOuter, bool bInIsPackageExternal, bool bInShouldDirty = true, EPackageFlags InExternalPackageFlags = FExternalPackageHelper::GetDefaultExternalPackageFlags());
+	static ENGINE_API void SetPackagingMode(UObject* InObject, const UObject* InObjectOuter, bool bInIsPackageExternal, bool bInShouldDirty = true, EPackageFlags InExternalPackageFlags = FExternalPackageHelper::GetDefaultExternalPackageFlags());
 
 	/**
 	 * Get the path containing the external objects for this path
@@ -113,6 +126,11 @@ public:
 	 * Call AssetRegistry.GetAssets and sort the results for deterministic use in cooked data.
 	 */
 	static ENGINE_API void GetSortedAssets(const FARFilter& Filter, TArray<FAssetData>& OutAssets);
+
+    /**
+    * Duplicates all ExternalPackage for any UObject outered to InObject (UObject themselves are handled by regular DuplicateObject behavior)
+    */
+	static ENGINE_API void DuplicateExternalPackages(const UObject* InObject, FObjectDuplicationParameters& InDuplicationParameters, EActorPackagingScheme ActorPackagingScheme = EActorPackagingScheme::Reduced);
 private:
 	/** Get the external object package instance name. */
 	static ENGINE_API FString GetExternalObjectPackageInstanceName(const FString& OuterPackageName, const FString& ObjectPackageName);
