Changed USkinnedMeshComponent::LODInfos to use a non-blocking method to delete the buffers. Also fixed a use-after-free crash resulting from not marking render state dirty after updating vertex factories, resulting in the use of stale UB/SRVs.

diff --git a/Engine/Source/Runtime/Engine/Classes/Components/SkinnedMeshComponent.h b/Engine/Source/Runtime/Engine/Classes/Components/SkinnedMeshComponent.h
--- a/Engine/Source/Runtime/Engine/Classes/Components/SkinnedMeshComponent.h
+++ b/Engine/Source/Runtime/Engine/Classes/Components/SkinnedMeshComponent.h
@@ -168,16 +168,34 @@ struct FSkelMeshComponentLODInfo
 	ENGINE_API FSkelMeshComponentLODInfo();
 	ENGINE_API ~FSkelMeshComponentLODInfo();
 
+	UE_DEPRECATED(5.5, "Use ReleaseOverrideVertexColors which does not require a flush before proceeding.")
 	ENGINE_API void ReleaseOverrideVertexColorsAndBlock();
+
+	UE_DEPRECATED(5.5, "Use ReleaseOverrideVertexColors which does not require a flush before proceeding.")
 	ENGINE_API bool BeginReleaseOverrideVertexColors();
+
+	UE_DEPRECATED(5.5, "Use ReleaseOverrideVertexColors which does not require a flush before proceeding.")
 	ENGINE_API void EndReleaseOverrideVertexColors();
+
+	/** Releases the vertex color override buffers without requiring a flush of the render thread */
+	ENGINE_API void ReleaseOverrideVertexColors();
+
 private:
 	void CleanUpOverrideVertexColors();
 
 public:
+	UE_DEPRECATED(5.5, "Use ReleaseOverrideSkinWeights which does not require a flush before proceeding.")
 	ENGINE_API void ReleaseOverrideSkinWeightsAndBlock();
+
+	UE_DEPRECATED(5.5, "Use ReleaseOverrideSkinWeights which does not require a flush before proceeding.")
 	ENGINE_API bool BeginReleaseOverrideSkinWeights();
+
+	UE_DEPRECATED(5.5, "Use ReleaseOverrideSkinWeights which does not require a flush before proceeding.")
 	ENGINE_API void EndReleaseOverrideSkinWeights();
+
+	/** Releases the skin weight override buffers without requiring a flush of the render thread */
+	ENGINE_API void ReleaseOverrideSkinWeights();
+
 private:
 	void CleanUpOverrideSkinWeights();
 };
@@ -1492,7 +1510,7 @@ public:
 	 */
 	ENGINE_API virtual void FinalizeBoneTransform();
 
-	/** Initialize the LOD entries for the component */
+	/** Initialize the LOD entries for the component, freeing the previous ones. */
 	ENGINE_API void InitLODInfos();
 
 	/**
@@ -1980,7 +1998,7 @@ public:
 
 	ENGINE_API bool ShouldUseUpdateRateOptimizations() const;
 
-	/** Release any rendering resources owned by this component */
+	/** Release rendering resources without requiring a block. */
 	ENGINE_API void ReleaseResources();
 
 #if WITH_EDITOR
diff --git a/Engine/Source/Runtime/Engine/Private/Components/SkinnedMeshComponent.cpp b/Engine/Source/Runtime/Engine/Private/Components/SkinnedMeshComponent.cpp
--- a/Engine/Source/Runtime/Engine/Private/Components/SkinnedMeshComponent.cpp
+++ b/Engine/Source/Runtime/Engine/Private/Components/SkinnedMeshComponent.cpp
@@ -83,14 +83,6 @@ FAutoConsoleVariableRef CVarUpdateBoundsNotifyStreamingRadiusChangeRatio(
 	ECVF_Default
 );
 
-static bool GReleasePreviousLODInfoOnInitialization = true;
-static FAutoConsoleVariableRef CVarReleasePreviousLODInfoOnInitialization(
-	TEXT("r.SkinnedMesh.ReleasePreviousLODInfoOnInitialization"),
-	GReleasePreviousLODInfoOnInitialization,
-	TEXT("Whether to flush the render thread (incurring a game thread stall) and clean existing LOD info when re-initializating."),
-	ECVF_Default
-);
-
 int32 GSkinnedMeshRenderNanite = 1;
 static FAutoConsoleVariableRef CVarSkinnedMeshRenderNanite(
 	TEXT("r.SkinnedMesh.RenderNanite"),
@@ -1358,52 +1350,14 @@ void USkinnedMeshComponent::PostEditChangeProperty(FPropertyChangedEvent& Proper
 
 void USkinnedMeshComponent::InitLODInfos()
 {
-	if (GetSkinnedAsset() != nullptr)
+	ReleaseResources();
+	if (const USkinnedAsset* LocalSkinnedAsset = GetSkinnedAsset())
 	{
-		if (GetSkinnedAsset()->GetLODNum() != LODInfo.Num())
-		{
-			// Perform cleanup if LOD infos have been initialized before 
-			if (!LODInfo.IsEmpty() && GReleasePreviousLODInfoOnInitialization)
-			{
-				// Batch release all resources of LOD infos we're about to destruct.
-				// This is relevant when overrides have been set but LODInfos are
-				// re-initialized, for example by changing the mesh at runtime.
-				bool bNeedRenderFlush = false;
-				
-				for (int32 Idx = 0; Idx < LODInfo.Num(); ++Idx)
-				{
-					if (LODInfo[Idx].BeginReleaseOverrideSkinWeights())
-					{
-						bNeedRenderFlush = true;
-					}
-					
-					if (LODInfo[Idx].BeginReleaseOverrideVertexColors())
-					{
-						bNeedRenderFlush = true;
-					}
-				}
-
-				if (bNeedRenderFlush)
-				{
-					FlushRenderingCommands();
-					for (int32 Idx = 0; Idx < LODInfo.Num(); ++Idx)
-					{
-						LODInfo[Idx].EndReleaseOverrideSkinWeights();
-						LODInfo[Idx].EndReleaseOverrideVertexColors();
-					}
-				}
-			}
-			
-			LODInfo.Empty(GetSkinnedAsset()->GetLODNum());
-			for (int32 Idx=0; Idx < GetSkinnedAsset()->GetLODNum(); Idx++)
-			{
-				new(LODInfo) FSkelMeshComponentLODInfo();
-			}
-		}
-	}	
+		const int32 NumLODs = LocalSkinnedAsset->GetLODNum();
+		LODInfo.SetNum(NumLODs, EAllowShrinking::No);
+	}
 }
 
-
 bool USkinnedMeshComponent::ShouldTickPose() const
 {
 	if (LeaderPoseComponent.IsValid() && !bFollowerShouldTickPose)
@@ -4367,13 +4321,11 @@ bool USkinnedMeshComponent::HasValidNaniteData() const
 
 void USkinnedMeshComponent::ReleaseResources()
 {
-	for (int32 LODIndex = 0; LODIndex < LODInfo.Num(); LODIndex++)
+	for (FSkelMeshComponentLODInfo& CurrentLODInfo : LODInfo)
 	{
-		LODInfo[LODIndex].BeginReleaseOverrideVertexColors();
-		LODInfo[LODIndex].BeginReleaseOverrideSkinWeights();
+		CurrentLODInfo.ReleaseOverrideVertexColors();
+		CurrentLODInfo.ReleaseOverrideSkinWeights();
 	}
-
-	DetachFence.BeginFence();
 }
 
 void USkinnedMeshComponent::RegisterLODStreamingCallback(FLODStreamingCallback&& Callback, int32 LODIdx, float TimeoutSecs, bool bOnStreamIn)
@@ -4505,6 +4457,21 @@ void FSkelMeshComponentLODInfo::CleanUpOverrideVertexColors()
 	}
 }
 
+void FSkelMeshComponentLODInfo::ReleaseOverrideVertexColors()
+{
+	if (OverrideVertexColors)
+	{
+		ENQUEUE_RENDER_COMMAND(ReleaseOverrideVertexColors)(&UE::RenderCommandPipe::SkeletalMesh,
+			[VertexColors = OverrideVertexColors]
+			{
+				VertexColors->ReleaseResource();
+				delete VertexColors;
+			});
+
+		OverrideVertexColors = nullptr;
+	}
+}
+
 void FSkelMeshComponentLODInfo::ReleaseOverrideSkinWeightsAndBlock()
 {
 	if (OverrideSkinWeights)
@@ -4549,6 +4516,23 @@ void FSkelMeshComponentLODInfo::CleanUpOverrideSkinWeights()
 	}
 }
 
+void FSkelMeshComponentLODInfo::ReleaseOverrideSkinWeights()
+{
+	if (OverrideSkinWeights)
+	{
+		ENQUEUE_RENDER_COMMAND(ReleaseOverrideSkinWeights)(&UE::RenderCommandPipe::SkeletalMesh,
+			[SkinWeights = OverrideSkinWeights]
+			{
+				SkinWeights->ReleaseResources();
+				delete SkinWeights;
+			});
+
+		OverrideSkinWeights = nullptr;
+	}
+
+	OverrideProfileSkinWeights = nullptr;
+}
+
 //////////////////////////////////////////////////////////////////////////
 
 void USkinnedMeshComponent::SetVertexColorOverride_LinearColor(int32 LODIndex, const TArray<FLinearColor>& VertexColors)
@@ -4583,7 +4567,7 @@ void USkinnedMeshComponent::SetVertexColorOverride(int32 LODIndex, const TArray<
 		FSkelMeshComponentLODInfo& Info = LODInfo[LODIndex];
 		if (Info.OverrideVertexColors != nullptr)
 		{
-			Info.ReleaseOverrideVertexColorsAndBlock();
+			Info.ReleaseOverrideVertexColors();
 		}
 
 		const TArray<FColor>* UseColors;
@@ -4639,7 +4623,7 @@ void USkinnedMeshComponent::ClearVertexColorOverride(int32 LODIndex)
 		FSkelMeshComponentLODInfo& Info = LODInfo[LODIndex];
 		if (Info.OverrideVertexColors != nullptr)
 		{
-			Info.ReleaseOverrideVertexColorsAndBlock();
+			Info.ReleaseOverrideVertexColors();
 			MarkRenderStateDirty();
 		}
 	}
@@ -4779,7 +4763,7 @@ void USkinnedMeshComponent::SetSkinWeightOverride(int32 LODIndex, const TArray<F
 		FSkelMeshComponentLODInfo& Info = LODInfo[LODIndex];
 		if (Info.OverrideSkinWeights != nullptr)
 		{
-			Info.ReleaseOverrideSkinWeightsAndBlock();
+			Info.ReleaseOverrideSkinWeights();
 		}
 
 		FSkeletalMeshLODRenderData& LODData = SkelMeshRenderData->LODRenderData[LODIndex];
@@ -4831,8 +4815,8 @@ void USkinnedMeshComponent::ClearSkinWeightOverride(int32 LODIndex)
 		FSkelMeshComponentLODInfo& Info = LODInfo[LODIndex];
 		if (Info.OverrideSkinWeights != nullptr)
 		{
-			Info.ReleaseOverrideSkinWeightsAndBlock();
-			MarkRenderStateDirty();
+			Info.ReleaseOverrideSkinWeights();
+			UpdateSkinWeightOverrideBuffer();
 		}
 	}
 }
@@ -4891,7 +4875,7 @@ bool USkinnedMeshComponent::SetSkinWeightProfile(FName InProfileName)
 
 						if (FSkeletalMeshRenderData * RenderData = WeakMesh->GetResourceForRendering())
 						{
-							const int32 NumLODs = RenderData->LODRenderData.Num();
+							const int32 NumLODs = FMath::Min(RenderData->LODRenderData.Num(), Component->LODInfo.Num());
 							for (int32 Index = 0; Index < NumLODs; ++Index)
 							{
 								FSkeletalMeshLODRenderData& LODRenderData = RenderData->LODRenderData[Index];
@@ -4929,23 +4913,10 @@ void USkinnedMeshComponent::ClearSkinWeightProfile()
 {
 	if (FSkeletalMeshRenderData* SkelMeshRenderData = GetSkeletalMeshRenderData())
 	{	
-		bool bCleared = false;
-
 		if (bSkinWeightProfileSet)
 		{
 			InitLODInfos();
-			// Clear skin weight buffer set for all of the LODs
-			for (int32 LODIndex = 0; LODIndex < LODInfo.Num(); ++LODIndex)
-			{
-				FSkelMeshComponentLODInfo& Info = LODInfo[LODIndex];
-				bCleared |= (Info.OverrideProfileSkinWeights != nullptr);
-				Info.OverrideProfileSkinWeights = nullptr;
-			}
-
-			if (bCleared)
-			{
-				UpdateSkinWeightOverrideBuffer();
-			}
+			UpdateSkinWeightOverrideBuffer();
 		}
 
 		if (bSkinWeightProfilePending)
@@ -5051,7 +5022,7 @@ void USkinnedMeshComponent::UpdateSkinWeightOverrideBuffer()
 {
 	// Force a mesh update to ensure bone buffers are up to date
 	bForceMeshObjectUpdate = true;
-	MarkRenderDynamicDataDirty();
+	MarkRenderStateDirty();
 
 	// Queue an update of the skin weight buffer used by the current Mesh Object
 	if (MeshObject)
