StaticMesh Buffers: Fix crash in FStaticMeshVertexBuffer::Serialize and FPositionVertexBuffer::Serialize when loading if they were saved with NumVertices == 0.
#rnx

diff --git a/Engine/Source/Runtime/Engine/Private/PositionVertexBuffer.cpp b/Engine/Source/Runtime/Engine/Private/PositionVertexBuffer.cpp
--- a/Engine/Source/Runtime/Engine/Private/PositionVertexBuffer.cpp
+++ b/Engine/Source/Runtime/Engine/Private/PositionVertexBuffer.cpp
@@ -165,7 +165,7 @@ void FPositionVertexBuffer::Serialize( FArchive& Ar, bool bInNeedsCPUAccess )
 
 	SerializeMetaData(Ar);
 
-	if(Ar.IsLoading())
+	if(Ar.IsLoading() && NumVertices > 0)
 	{
 		// Allocate the vertex data storage type.
 		AllocateData( bInNeedsCPUAccess );
diff --git a/Engine/Source/Runtime/Engine/Private/Rendering/StaticMeshVertexBuffer.cpp b/Engine/Source/Runtime/Engine/Private/Rendering/StaticMeshVertexBuffer.cpp
--- a/Engine/Source/Runtime/Engine/Private/Rendering/StaticMeshVertexBuffer.cpp
+++ b/Engine/Source/Runtime/Engine/Private/Rendering/StaticMeshVertexBuffer.cpp
@@ -210,7 +210,7 @@ void FStaticMeshVertexBuffer::Serialize(FArchive& Ar, bool bNeedsCPUAccess)
 
 	SerializeMetaData(Ar);
 
-	if (Ar.IsLoading())
+	if (Ar.IsLoading() && NumVertices > 0)
 	{
 		// Allocate the vertex data storage type.
 		AllocateData(bNeedsCPUAccess);
