Add const evaluation to the Lerp function.

diff --git a/Engine/Source/Runtime/Engine/Private/Materials/HLSLMaterialTranslator.cpp b/Engine/Source/Runtime/Engine/Private/Materials/HLSLMaterialTranslator.cpp
--- a/Engine/Source/Runtime/Engine/Private/Materials/HLSLMaterialTranslator.cpp
+++ b/Engine/Source/Runtime/Engine/Private/Materials/HLSLMaterialTranslator.cpp
@@ -9822,7 +9822,7 @@ int32 FHLSLMaterialTranslator::InvLerp(int32 X, int32 Y, int32 A)
 
 int32 FHLSLMaterialTranslator::Lerp(int32 X,int32 Y,int32 A)
 {
-	if(X == INDEX_NONE || Y == INDEX_NONE || A == INDEX_NONE)
+	if (X == INDEX_NONE || Y == INDEX_NONE || A == INDEX_NONE)
 	{
 		return INDEX_NONE;
 	}
@@ -9859,23 +9859,60 @@ int32 FHLSLMaterialTranslator::Lerp(int32 X,int32 Y,int32 A)
 	}
 
 	EMaterialValueType ResultType = GetArithmeticResultType(X,Y);
-	EMaterialValueType AlphaType = ResultType == (*CurrentScopeChunks)[A].Type ? ResultType : MCT_Float1;
+	EMaterialValueType AlphaType = (*CurrentScopeChunks)[A].Type;
 
-	if (AlphaType == MCT_Float1 && ExpressionA && ExpressionA->IsConstant())
+	if ((AlphaType & MCT_Float) != 0 && ExpressionA && ExpressionA->IsConstant())
 	{
 		// Skip over interpolations that explicitly select an input
-		FLinearColor Value;
+		FLinearColor AlphaValue;
 		FMaterialRenderContext DummyContext(nullptr, *Material, nullptr);
-		ExpressionA->GetNumberValue(DummyContext, Value);
+		ExpressionA->GetNumberValue(DummyContext, AlphaValue);
+
+		const EMaterialValueType ResultFloatType = MakeNonLWCType(ResultType);
+		const EMaterialValueType AlphaCastType = (ResultFloatType == AlphaType) ? ResultFloatType : MCT_Float1;
+
+		if (!CastConstantType(AlphaValue, AlphaCastType, ResultFloatType, EMaterialCastFlags::ReplicateScalar | EMaterialCastFlags::AllowTruncate, AlphaValue))
+		{
+			return INDEX_NONE;
+		}
+
+		// Cast will not change the value if the types are the same, but we want the true masked value for comparisons below
+		AlphaValue = GetTypeMaskedValue(ResultFloatType, AlphaValue, nullptr);
+		if (AlphaCastType == MCT_Float1 && ResultFloatType != MCT_Float1)
+		{
+			AlphaValue = GetTypeMaskedValue(ResultFloatType, FLinearColor(AlphaValue.R, AlphaValue.R, AlphaValue.R, AlphaValue.R), nullptr);
+		}
 
-		if (Value.R == 0.0f)
+		const FLinearColor ZeroValue = GetTypeMaskedValue(ResultFloatType, FLinearColor(0.0f, 0.0f, 0.0f, 0.0f), nullptr);
+		const FLinearColor OneValue = GetTypeMaskedValue(ResultFloatType, FLinearColor(1.0f, 1.0f, 1.0f, 1.0f), nullptr);
+		
+		if (AlphaValue == ZeroValue)
 		{
 			return X;
 		}
-		else if (Value.R == 1.f)
+		else if (AlphaValue == OneValue)
 		{
 			return Y;
 		}
+
+		// If all inputs are constant and we have a float result type, produce a constant result
+		if ((ResultType & MCT_Float) != 0 && ExpressionX && ExpressionY && ExpressionX->IsConstant() && ExpressionY->IsConstant())
+		{
+			FLinearColor ValueX, ValueY;
+
+			ExpressionX->GetNumberValue(DummyContext, ValueX);
+			ExpressionY->GetNumberValue(DummyContext, ValueY);
+			if (!CastConstantType(ValueX, (*CurrentScopeChunks)[X].Type, ResultType, EMaterialCastFlags::ReplicateScalar, ValueX))
+			{
+				return INDEX_NONE;
+			}
+			if (!CastConstantType(ValueY, (*CurrentScopeChunks)[Y].Type, ResultType, EMaterialCastFlags::ReplicateScalar, ValueY))
+			{
+				return INDEX_NONE;
+			}
+			FLinearColor ConstLerp = FMath::Lerp(ValueX, ValueY, AlphaValue);
+			return ConstResultValue(ResultType, ConstLerp);
+		}
 	}
 
 	if (IsAnalyticDerivEnabled())
