Expose IsCultureRightToLeft to BP
- a better option than checking specific cultures when we want to update some UI layouts based on RtL reading
Note: When using this for UI, this will ignore widget settings to inherit flow direction.

diff --git a/Engine/Source/Runtime/Engine/Classes/Kismet/KismetInternationalizationLibrary.h b/Engine/Source/Runtime/Engine/Classes/Kismet/KismetInternationalizationLibrary.h
--- a/Engine/Source/Runtime/Engine/Classes/Kismet/KismetInternationalizationLibrary.h
+++ b/Engine/Source/Runtime/Engine/Classes/Kismet/KismetInternationalizationLibrary.h
@@ -160,4 +160,12 @@ public:
 	 */
 	UFUNCTION(BlueprintPure, Category="Utilities|Internationalization", meta=(AdvancedDisplay="1"))
 	static ENGINE_API FString GetCultureDisplayName(const FString& Culture, const bool Localized = true);
+
+	/**
+	 * Returns if the given culture reads left to right
+	 * @param Culture The culture to get the display name of, as an IETF language tag (eg, "zh-Hans-CN")
+	 * @return True if the given culture reads left to right.
+	 */
+	UFUNCTION(BlueprintPure, Category="Utilities|Internationalization")
+	static ENGINE_API bool IsCultureRightToLeft(const FString& Culture);
 };
diff --git a/Engine/Source/Runtime/Engine/Private/KismetInternationalizationLibrary.cpp b/Engine/Source/Runtime/Engine/Private/KismetInternationalizationLibrary.cpp
--- a/Engine/Source/Runtime/Engine/Private/KismetInternationalizationLibrary.cpp
+++ b/Engine/Source/Runtime/Engine/Private/KismetInternationalizationLibrary.cpp
@@ -164,5 +164,11 @@ FString UKismetInternationalizationLibrary::GetCultureDisplayName(const FString&
 	return Culture;
 }
 
+bool UKismetInternationalizationLibrary::IsCultureRightToLeft(const FString& Culture)
+{
+	const FCulturePtr CulturePtr = FInternationalization::Get().GetCulture(Culture);
+	return CulturePtr && CulturePtr->IsRightToLeft();
+}
+
 #undef LOCTEXT_NAMESPACE
 
