Bindless specific shader debug info
- Adding a new cvar r.DumpShaderDebugBindlessNames, set to 0 by default
- Will change the shader info dump path to a different folder based on the bindless configurations.

diff --git a/Engine/Source/Runtime/Engine/Private/ShaderCompiler/ShaderCompiler.cpp b/Engine/Source/Runtime/Engine/Private/ShaderCompiler/ShaderCompiler.cpp
--- a/Engine/Source/Runtime/Engine/Private/ShaderCompiler/ShaderCompiler.cpp
+++ b/Engine/Source/Runtime/Engine/Private/ShaderCompiler/ShaderCompiler.cpp
@@ -2109,6 +2109,14 @@ static FAutoConsoleVariableRef CVarDumpShaderDebugShortNames(
 	TEXT("When set to 1, will shorten names factory and shader type folder names to avoid issues with long paths.")
 	);
 
+static int32 GDumpShaderDebugInfoBindless = 0;
+static FAutoConsoleVariableRef CVarDumpShaderDebugBindlessNames(
+	TEXT("r.DumpShaderDebugBindlessNames"),
+	GDumpShaderDebugInfoBindless,
+	TEXT("Only valid when r.DumpShaderDebugInfo > 0.\n")
+	TEXT("When set to 1, will add bindless folder names.")
+	);
+
 static int32 GDumpShaderDebugInfoSCWCommandLine = 0;
 static FAutoConsoleVariableRef CVarDumpShaderDebugSCWCommandLine(
 	TEXT("r.DumpShaderDebugWorkerCommandLine"),
@@ -7611,6 +7619,27 @@ void GlobalBeginCompileShader(
 	Input.RootParametersStructure = ShaderType->GetRootParametersMetadata();
 	Input.ShaderName = ShaderType->GetName();
 
+	if (GDumpShaderDebugInfoBindless)
+	{
+		auto GetBindlessString = [](ERHIBindlessConfiguration InConfig)
+			{
+				switch (InConfig)
+				{
+				default:
+				case ERHIBindlessConfiguration::Disabled:          return TEXT("Off");
+				case ERHIBindlessConfiguration::AllShaders:        return TEXT("On");
+				case ERHIBindlessConfiguration::RayTracingShaders: return TEXT("RT");
+				}
+			};
+
+		FStringBuilderBase Builder;
+		Builder.Append("BindlessR").Append(GetBindlessString(UE::ShaderCompiler::GetBindlessResourcesConfiguration(ShaderFormatName)));
+		Builder.Append("_");
+		Builder.Append("BindlessS").Append(GetBindlessString(UE::ShaderCompiler::GetBindlessSamplersConfiguration(ShaderFormatName)));
+
+		Input.DebugGroupName = Builder.ToString() / Input.DebugGroupName;
+	}
+
 	// Verify FShaderCompilerInput's file paths are consistent. 
 	#if DO_CHECK
 		check(CheckVirtualShaderFilePath(Input.VirtualSourceFilePath));
