Fixed Client-Only Data Layers
- Added validation to detect a mismatch when an actor is part of data layers with a different type of LoadFilter (None,client-only,server-only).
- Properly flag streaming Cell bClientOnlyVisible so that AddToWorld of cell's Level doesn't wait for server acknowledge before making the level visible.

diff --git a/Engine/Plugins/Editor/DataValidation/Source/DataValidation/Private/WorldPartitionChangelistValidator.cpp b/Engine/Plugins/Editor/DataValidation/Source/DataValidation/Private/WorldPartitionChangelistValidator.cpp
--- a/Engine/Plugins/Editor/DataValidation/Source/DataValidation/Private/WorldPartitionChangelistValidator.cpp
+++ b/Engine/Plugins/Editor/DataValidation/Source/DataValidation/Private/WorldPartitionChangelistValidator.cpp
@@ -387,6 +387,17 @@ void UWorldPartitionChangelistValidator::OnInvalidReferenceRuntimeGrid(const IWo
 	}
 }
 
+void UWorldPartitionChangelistValidator::OnDataLayersLoadFilterMismatch(const IWorldPartitionActorDescInstanceView& ActorDescView)
+{
+	if (Filter(ActorDescView))
+	{
+		FText CurrentError = FText::Format(LOCTEXT("DataValidation.Changelist.WorldPartition.DataLayersLoadFilterMismatchError", "{0} contains runtime data layers with different types of Load Filter"),
+			FText::FromString(GetFullActorName(ActorDescView)));
+
+		AssetFails(CurrentAsset, CurrentError);
+	}
+}
+
 void UWorldPartitionChangelistValidator::OnInvalidWorldReference(const IWorldPartitionActorDescInstanceView& ActorDescView, EWorldReferenceInvalidReason Reason)
 {
 	if (Filter(ActorDescView))
diff --git a/Engine/Plugins/Editor/DataValidation/Source/DataValidation/Private/WorldPartitionChangelistValidator.h b/Engine/Plugins/Editor/DataValidation/Source/DataValidation/Private/WorldPartitionChangelistValidator.h
--- a/Engine/Plugins/Editor/DataValidation/Source/DataValidation/Private/WorldPartitionChangelistValidator.h
+++ b/Engine/Plugins/Editor/DataValidation/Source/DataValidation/Private/WorldPartitionChangelistValidator.h
@@ -41,6 +41,7 @@ protected:
 	virtual void OnInvalidReferenceDataLayers(const IWorldPartitionActorDescInstanceView& ActorDescView, const IWorldPartitionActorDescInstanceView& ReferenceActorDescView, EDataLayerInvalidReason Reason) override;
 	virtual void OnInvalidWorldReference(const IWorldPartitionActorDescInstanceView& ActorDescView, EWorldReferenceInvalidReason Reason) override;
 	virtual void OnInvalidReferenceRuntimeGrid(const IWorldPartitionActorDescInstanceView& ActorDescView, const IWorldPartitionActorDescInstanceView& ReferenceActorDescView) override;
+	virtual void OnDataLayersLoadFilterMismatch(const IWorldPartitionActorDescInstanceView& ActorDescView) override;
 	virtual void OnInvalidReferenceDataLayerAsset(const UDataLayerInstanceWithAsset* DataLayerInstance) override;
 	virtual void OnInvalidDataLayerAssetType(const UDataLayerInstanceWithAsset* DataLayerInstance, const UDataLayerAsset* DataLayerAsset) override;
 	virtual void OnDataLayerHierarchyTypeMismatch(const UDataLayerInstance* DataLayerInstance, const UDataLayerInstance* Parent, EDataLayerHierarchyInvalidReason Reason) override;
diff --git a/Engine/Source/Runtime/Engine/Private/WorldPartition/ErrorHandling/WorldPartitionStreamingGenerationLogErrorHandler.cpp b/Engine/Source/Runtime/Engine/Private/WorldPartition/ErrorHandling/WorldPartitionStreamingGenerationLogErrorHandler.cpp
--- a/Engine/Source/Runtime/Engine/Private/WorldPartition/ErrorHandling/WorldPartitionStreamingGenerationLogErrorHandler.cpp
+++ b/Engine/Source/Runtime/Engine/Private/WorldPartition/ErrorHandling/WorldPartitionStreamingGenerationLogErrorHandler.cpp
@@ -48,6 +48,11 @@ void FStreamingGenerationLogErrorHandler::OnInvalidReferenceRuntimeGrid(const IW
 	UE_ASSET_LOG_ACTORDESCVIEW(LogWorldPartition, Log, ActorDescView, TEXT("Actor %s references an actor in a different runtime grid %s"), *GetActorName(ActorDescView), *GetActorName(ReferenceActorDescView));
 }
 
+void FStreamingGenerationLogErrorHandler::OnDataLayersLoadFilterMismatch(const IWorldPartitionActorDescInstanceView& ActorDescView)
+{
+	UE_ASSET_LOG_ACTORDESCVIEW(LogWorldPartition, Log, ActorDescView, TEXT("Actor %s contains runtime data layers with different types of Load Filter"), *GetActorName(ActorDescView));
+}
+
 void FStreamingGenerationLogErrorHandler::OnInvalidWorldReference(const IWorldPartitionActorDescInstanceView& ActorDescView, EWorldReferenceInvalidReason Reason)
 {
 	switch(Reason)
diff --git a/Engine/Source/Runtime/Engine/Private/WorldPartition/ErrorHandling/WorldPartitionStreamingGenerationProxyErrorHandler.cpp b/Engine/Source/Runtime/Engine/Private/WorldPartition/ErrorHandling/WorldPartitionStreamingGenerationProxyErrorHandler.cpp
--- a/Engine/Source/Runtime/Engine/Private/WorldPartition/ErrorHandling/WorldPartitionStreamingGenerationProxyErrorHandler.cpp
+++ b/Engine/Source/Runtime/Engine/Private/WorldPartition/ErrorHandling/WorldPartitionStreamingGenerationProxyErrorHandler.cpp
@@ -28,6 +28,11 @@ void FStreamingGenerationProxyErrorHandler::OnInvalidReferenceRuntimeGrid(const
 	InnerErrorHandler->OnInvalidReferenceRuntimeGrid(ActorDescView, ReferenceActorDescView);
 }
 
+void FStreamingGenerationProxyErrorHandler::OnDataLayersLoadFilterMismatch(const IWorldPartitionActorDescInstanceView& ActorDescView)
+{
+	InnerErrorHandler->OnDataLayersLoadFilterMismatch(ActorDescView);
+}
+
 void FStreamingGenerationProxyErrorHandler::OnInvalidWorldReference(const IWorldPartitionActorDescInstanceView& ActorDescView, EWorldReferenceInvalidReason Reason)
 {
 	InnerErrorHandler->OnInvalidWorldReference(ActorDescView, Reason);
diff --git a/Engine/Source/Runtime/Engine/Private/WorldPartition/ErrorHandling/WorldPartitionStreamingGenerationTokenizedMessageErrorHandler.cpp b/Engine/Source/Runtime/Engine/Private/WorldPartition/ErrorHandling/WorldPartitionStreamingGenerationTokenizedMessageErrorHandler.cpp
--- a/Engine/Source/Runtime/Engine/Private/WorldPartition/ErrorHandling/WorldPartitionStreamingGenerationTokenizedMessageErrorHandler.cpp
+++ b/Engine/Source/Runtime/Engine/Private/WorldPartition/ErrorHandling/WorldPartitionStreamingGenerationTokenizedMessageErrorHandler.cpp
@@ -97,6 +97,15 @@ void ITokenizedMessageErrorHandler::OnInvalidReferenceRuntimeGrid(const IWorldPa
 	HandleTokenizedMessage(MoveTemp(Message));
 }
 
+void ITokenizedMessageErrorHandler::OnDataLayersLoadFilterMismatch(const IWorldPartitionActorDescInstanceView& ActorDescView)
+{
+	TSharedRef<FTokenizedMessage> Message = FTokenizedMessage::Create(EMessageSeverity::Error);
+	Message->AddToken(FTextToken::Create(LOCTEXT("TokenMessage_WorldPartition_Actor", "Actor")))
+		->AddToken(FActorToken::Create(ActorDescView.GetActorSoftPath().ToString(), ActorDescView.GetGuid(), FText::FromString(GetActorName(ActorDescView))))
+		->AddToken(FTextToken::Create(LOCTEXT("TokenMessage_WorldPartition_ContainsRuntimeDataLayersWithDifferentTypesOfLoadFilter", "contains runtime data layers with different types of Load Filter")))
+		->AddToken(FMapErrorToken::Create(TEXT("WorldPartition_DataLayersLoadFilterMismatch_CheckForErrors")));
+}
+
 void ITokenizedMessageErrorHandler::OnInvalidWorldReference(const IWorldPartitionActorDescInstanceView& ActorDescView, EWorldReferenceInvalidReason Reason)
 {
 	TSharedRef<FTokenizedMessage> Message = FTokenizedMessage::Create(EMessageSeverity::Error);
diff --git a/Engine/Source/Runtime/Engine/Private/WorldPartition/RuntimeHashSet/RuntimePartition.cpp b/Engine/Source/Runtime/Engine/Private/WorldPartition/RuntimeHashSet/RuntimePartition.cpp
--- a/Engine/Source/Runtime/Engine/Private/WorldPartition/RuntimeHashSet/RuntimePartition.cpp
+++ b/Engine/Source/Runtime/Engine/Private/WorldPartition/RuntimeHashSet/RuntimePartition.cpp
@@ -2,6 +2,8 @@
 
 #include "WorldPartition/RuntimeHashSet/RuntimePartition.h"
 #include "WorldPartition/RuntimeHashSet/WorldPartitionRuntimeHashSet.h"
+#include "WorldPartition/DataLayer/DataLayerInstance.h"
+#include "Algo/AllOf.h"
 
 URuntimePartition::URuntimePartition(const FObjectInitializer& ObjectInitializer)
 	: Super(ObjectInitializer)
@@ -78,4 +80,15 @@ URuntimePartition::FCellDesc URuntimePartition::CreateCellDesc(const FString& In
 
 	return CellDesc;
 }
+
+URuntimePartition::FCellDescInstance::FCellDescInstance(const FCellDesc& InCellDesc, URuntimePartition* InSourcePartition, const TArray<const UDataLayerInstance*>& InDataLayerInstances, const FGuid& InContentBundleID)
+	: FCellDesc(InCellDesc)
+	, SourcePartition(InSourcePartition)
+	, DataLayerInstances(InDataLayerInstances)
+	, ContentBundleID(InContentBundleID)
+{
+	const bool bAreClientOnlyDataLayers = DataLayerInstances.Num() && Algo::AllOf(DataLayerInstances, [](const UDataLayerInstance* DataLayerInstance) { return DataLayerInstance->IsClientOnly(); });
+	bClientOnlyVisible = InCellDesc.bClientOnlyVisible || bAreClientOnlyDataLayers;
+}
+
 #endif
\ No newline at end of file
diff --git a/Engine/Source/Runtime/Engine/Private/WorldPartition/WorldPartitionRuntimeSpatialHash.cpp b/Engine/Source/Runtime/Engine/Private/WorldPartition/WorldPartitionRuntimeSpatialHash.cpp
--- a/Engine/Source/Runtime/Engine/Private/WorldPartition/WorldPartitionRuntimeSpatialHash.cpp
+++ b/Engine/Source/Runtime/Engine/Private/WorldPartition/WorldPartitionRuntimeSpatialHash.cpp
@@ -39,6 +39,7 @@
 #include "UObject/Package.h"
 #include "UObject/FortniteMainBranchObjectVersion.h"
 #include "Misc/HashBuilder.h"
+#include "Algo/AllOf.h"
 
 #include UE_INLINE_GENERATED_CPP_BY_NAME(WorldPartitionRuntimeSpatialHash)
 
@@ -1574,10 +1575,13 @@ bool UWorldPartitionRuntimeSpatialHash::CreateStreamingGrid(const FSpatialHashRu
 					UWorldPartitionRuntimeCell* StreamingCell = CreateRuntimeCell(StreamingPolicy->GetRuntimeCellClass(), UWorldPartitionRuntimeCellDataSpatialHash::StaticClass(), CellName, WorldInstanceSuffix);
 					UWorldPartitionRuntimeCellDataSpatialHash* CellDataSpatialHash = CastChecked<UWorldPartitionRuntimeCellDataSpatialHash>(StreamingCell->RuntimeCellData);
 
+					const TArray<const UDataLayerInstance*>& DataLayers = GridCellDataChunk.GetDataLayers();
+					const bool bAreClientOnlyDataLayers = DataLayers.Num() && Algo::AllOf(DataLayers, [](const UDataLayerInstance* DataLayerInstance) { return DataLayerInstance->IsClientOnly(); });
+
 					StreamingCell->SetIsAlwaysLoaded(bIsCellAlwaysLoaded);
-					StreamingCell->SetDataLayers(GridCellDataChunk.GetDataLayers());
+					StreamingCell->SetDataLayers(DataLayers);
 					StreamingCell->SetContentBundleUID(GridCellDataChunk.GetContentBundleID());
-					StreamingCell->SetClientOnlyVisible(CurrentStreamingGrid.bClientOnlyVisible);
+					StreamingCell->SetClientOnlyVisible(CurrentStreamingGrid.bClientOnlyVisible || bAreClientOnlyDataLayers);
 					StreamingCell->SetBlockOnSlowLoading(CurrentStreamingGrid.bBlockOnSlowStreaming);
 					StreamingCell->SetIsHLOD(RuntimeGrid.HLODLayer ? true : false);
 					StreamingCell->SetGuid(CellGuid);
diff --git a/Engine/Source/Runtime/Engine/Private/WorldPartition/WorldPartitionStreamingGeneration.cpp b/Engine/Source/Runtime/Engine/Private/WorldPartition/WorldPartitionStreamingGeneration.cpp
--- a/Engine/Source/Runtime/Engine/Private/WorldPartition/WorldPartitionStreamingGeneration.cpp
+++ b/Engine/Source/Runtime/Engine/Private/WorldPartition/WorldPartitionStreamingGeneration.cpp
@@ -1190,7 +1190,33 @@ class FWorldPartitionStreamingGenerator
 					return RefererActorDescView.GetRuntimeDataLayerInstanceNames().GetExternalDataLayer() == ReferenceActorDescView.GetRuntimeDataLayerInstanceNames().GetExternalDataLayer();
 				};
 
-				// Validate data layers
+				auto GetDataLayerLoadFilter = [](const UDataLayerInstance* DataLayerInstance)
+				{
+					EDataLayerLoadFilter LoadFilter = DataLayerInstance->IsClientOnly() ? EDataLayerLoadFilter::ClientOnly : DataLayerInstance->IsServerOnly() ? EDataLayerLoadFilter::ServerOnly : EDataLayerLoadFilter::None;
+					return LoadFilter;
+				};
+
+				// Validate that all runtime data layers have the same DataLayerLoadFilter
+				auto AreDataLayersLoadFilterValid = [this, GetDataLayerLoadFilter](const FStreamingGenerationActorDescView& ActorDescView)
+				{
+					TArrayView<const FName> DataLayers = ActorDescView.GetRuntimeDataLayerInstanceNames().GetNonExternalDataLayers();
+					if (DataLayers.Num() > 1)
+					{
+						TArray<const UDataLayerInstance*> RuntimeDataLayerInstances = GetRuntimeDataLayerInstances(TArray<FName>(DataLayers));
+						EDataLayerLoadFilter LoadFilter = GetDataLayerLoadFilter(RuntimeDataLayerInstances[0]);
+						for (int32 Index = 1; Index < RuntimeDataLayerInstances.Num(); ++Index)
+						{
+							EDataLayerLoadFilter Current = GetDataLayerLoadFilter(RuntimeDataLayerInstances[Index]);
+							if (LoadFilter != Current)
+							{
+								return false;
+							}
+						}
+					}
+					return true;
+				};
+
+				// Validate reference actor data layers
 				auto IsReferenceDataLayersValid = [](const FStreamingGenerationActorDescView& RefererActorDescView, const FStreamingGenerationActorDescView& ReferenceActorDescView)
 				{
 					if (RefererActorDescView.GetRuntimeDataLayerInstanceNames().GetNonExternalDataLayers().Num() == ReferenceActorDescView.GetRuntimeDataLayerInstanceNames().GetNonExternalDataLayers().Num())
@@ -1221,6 +1247,19 @@ class FWorldPartitionStreamingGenerator
 					FStreamingGenerationActorDescView* ReferenceActorDesc;
 				};
 
+				if (!AreDataLayersLoadFilterValid(ActorDescView))
+				{
+					if (PassType == EPassType::ErrorReporting)
+					{
+						ErrorHandler->OnDataLayersLoadFilterMismatch(ActorDescView);
+					}
+					else
+					{
+						ActorDescView.SetForcedNoDataLayers();
+					}
+					NbErrorsDetected++;
+				}
+
 				TArray<FActorReferenceInfo> References;
 
 				// Add normal actor references
diff --git a/Engine/Source/Runtime/Engine/Public/WorldPartition/ErrorHandling/WorldPartitionStreamingGenerationErrorHandler.h b/Engine/Source/Runtime/Engine/Public/WorldPartition/ErrorHandling/WorldPartitionStreamingGenerationErrorHandler.h
--- a/Engine/Source/Runtime/Engine/Public/WorldPartition/ErrorHandling/WorldPartitionStreamingGenerationErrorHandler.h
+++ b/Engine/Source/Runtime/Engine/Public/WorldPartition/ErrorHandling/WorldPartitionStreamingGenerationErrorHandler.h
@@ -49,6 +49,11 @@ public:
 	 */
 	virtual void OnInvalidReferenceRuntimeGrid(const IWorldPartitionActorDescInstanceView& ActorDescView, const IWorldPartitionActorDescInstanceView& ReferenceActorDescView) = 0;
 
+	/**
+	 * Called when an actor contains runtime data layers with different types of LoadFilter.
+	 */
+	virtual void OnDataLayersLoadFilterMismatch(const IWorldPartitionActorDescInstanceView& ActorDescView) = 0;
+
 	/**
 	 * Called when the world references a streamed actor.
 	 */
diff --git a/Engine/Source/Runtime/Engine/Public/WorldPartition/ErrorHandling/WorldPartitionStreamingGenerationLogErrorHandler.h b/Engine/Source/Runtime/Engine/Public/WorldPartition/ErrorHandling/WorldPartitionStreamingGenerationLogErrorHandler.h
--- a/Engine/Source/Runtime/Engine/Public/WorldPartition/ErrorHandling/WorldPartitionStreamingGenerationLogErrorHandler.h
+++ b/Engine/Source/Runtime/Engine/Public/WorldPartition/ErrorHandling/WorldPartitionStreamingGenerationLogErrorHandler.h
@@ -13,6 +13,7 @@ public:
 	ENGINE_API virtual void OnInvalidReferenceDataLayers(const IWorldPartitionActorDescInstanceView& ActorDescView, const IWorldPartitionActorDescInstanceView& ReferenceActorDescView, EDataLayerInvalidReason Reason) override;
 	ENGINE_API virtual void OnInvalidWorldReference(const IWorldPartitionActorDescInstanceView& ActorDescView, EWorldReferenceInvalidReason Reason) override;
 	ENGINE_API virtual void OnInvalidReferenceRuntimeGrid(const IWorldPartitionActorDescInstanceView& ActorDescView, const IWorldPartitionActorDescInstanceView& ReferenceActorDescView) override;
+	ENGINE_API virtual void OnDataLayersLoadFilterMismatch(const IWorldPartitionActorDescInstanceView& ActorDescView) override;
 	ENGINE_API virtual void OnInvalidReferenceDataLayerAsset(const UDataLayerInstanceWithAsset* DataLayerInstance) override;
 	ENGINE_API virtual void OnInvalidDataLayerAssetType(const UDataLayerInstanceWithAsset* DataLayerInstance, const UDataLayerAsset* DataLayerAsset) override;
 	ENGINE_API virtual void OnDataLayerHierarchyTypeMismatch(const UDataLayerInstance* DataLayerInstance, const UDataLayerInstance* Parent, EDataLayerHierarchyInvalidReason Reason) override;
diff --git a/Engine/Source/Runtime/Engine/Public/WorldPartition/ErrorHandling/WorldPartitionStreamingGenerationNullErrorHandler.h b/Engine/Source/Runtime/Engine/Public/WorldPartition/ErrorHandling/WorldPartitionStreamingGenerationNullErrorHandler.h
--- a/Engine/Source/Runtime/Engine/Public/WorldPartition/ErrorHandling/WorldPartitionStreamingGenerationNullErrorHandler.h
+++ b/Engine/Source/Runtime/Engine/Public/WorldPartition/ErrorHandling/WorldPartitionStreamingGenerationNullErrorHandler.h
@@ -13,6 +13,7 @@ public:
 	virtual void OnInvalidReferenceDataLayers(const IWorldPartitionActorDescInstanceView& ActorDescView, const IWorldPartitionActorDescInstanceView& ReferenceActorDescView, EDataLayerInvalidReason Reason) override {}
 	virtual void OnInvalidWorldReference(const IWorldPartitionActorDescInstanceView& ActorDescView, EWorldReferenceInvalidReason Reason) override {}
 	virtual void OnInvalidReferenceRuntimeGrid(const IWorldPartitionActorDescInstanceView& ActorDescView, const IWorldPartitionActorDescInstanceView& ReferenceActorDescView) override {}
+	virtual void OnDataLayersLoadFilterMismatch(const IWorldPartitionActorDescInstanceView& ActorDescView) override {}
 	virtual void OnInvalidReferenceDataLayerAsset(const UDataLayerInstanceWithAsset* DataLayerInstance) override {}
 	virtual void OnInvalidDataLayerAssetType(const UDataLayerInstanceWithAsset* DataLayerInstance, const UDataLayerAsset* DataLayerAsset) override {}
 	virtual void OnDataLayerHierarchyTypeMismatch(const UDataLayerInstance* DataLayerInstance, const UDataLayerInstance* Parent, EDataLayerHierarchyInvalidReason Reason) override {};
diff --git a/Engine/Source/Runtime/Engine/Public/WorldPartition/ErrorHandling/WorldPartitionStreamingGenerationProxyErrorHandler.h b/Engine/Source/Runtime/Engine/Public/WorldPartition/ErrorHandling/WorldPartitionStreamingGenerationProxyErrorHandler.h
--- a/Engine/Source/Runtime/Engine/Public/WorldPartition/ErrorHandling/WorldPartitionStreamingGenerationProxyErrorHandler.h
+++ b/Engine/Source/Runtime/Engine/Public/WorldPartition/ErrorHandling/WorldPartitionStreamingGenerationProxyErrorHandler.h
@@ -17,6 +17,7 @@ public:
 	ENGINE_API virtual void OnInvalidReferenceDataLayers(const IWorldPartitionActorDescInstanceView& ActorDescView, const IWorldPartitionActorDescInstanceView& ReferenceActorDescView, EDataLayerInvalidReason Reason) override;
 	ENGINE_API virtual void OnInvalidWorldReference(const IWorldPartitionActorDescInstanceView& ActorDescView, EWorldReferenceInvalidReason Reason) override;
 	ENGINE_API virtual void OnInvalidReferenceRuntimeGrid(const IWorldPartitionActorDescInstanceView& ActorDescView, const IWorldPartitionActorDescInstanceView& ReferenceActorDescView) override;
+	ENGINE_API virtual void OnDataLayersLoadFilterMismatch(const IWorldPartitionActorDescInstanceView& ActorDescView) override;
 	ENGINE_API virtual void OnInvalidReferenceDataLayerAsset(const UDataLayerInstanceWithAsset* DataLayerInstance) override;
 	ENGINE_API virtual void OnInvalidDataLayerAssetType(const UDataLayerInstanceWithAsset* DataLayerInstance, const UDataLayerAsset* DataLayerAsset) override;
 	ENGINE_API virtual void OnDataLayerHierarchyTypeMismatch(const UDataLayerInstance* DataLayerInstance, const UDataLayerInstance* Parent, EDataLayerHierarchyInvalidReason Reason) override;
diff --git a/Engine/Source/Runtime/Engine/Public/WorldPartition/ErrorHandling/WorldPartitionStreamingGenerationTokenizedMessageErrorHandler.h b/Engine/Source/Runtime/Engine/Public/WorldPartition/ErrorHandling/WorldPartitionStreamingGenerationTokenizedMessageErrorHandler.h
--- a/Engine/Source/Runtime/Engine/Public/WorldPartition/ErrorHandling/WorldPartitionStreamingGenerationTokenizedMessageErrorHandler.h
+++ b/Engine/Source/Runtime/Engine/Public/WorldPartition/ErrorHandling/WorldPartitionStreamingGenerationTokenizedMessageErrorHandler.h
@@ -14,6 +14,7 @@ class ITokenizedMessageErrorHandler : public IStreamingGenerationErrorHandler
 	ENGINE_API virtual void OnInvalidReferenceDataLayers(const IWorldPartitionActorDescInstanceView& ActorDescView, const IWorldPartitionActorDescInstanceView& ReferenceActorDescView, EDataLayerInvalidReason Reason) override;
 	ENGINE_API virtual void OnInvalidWorldReference(const IWorldPartitionActorDescInstanceView& ActorDescView, EWorldReferenceInvalidReason Reason) override;
 	ENGINE_API virtual void OnInvalidReferenceRuntimeGrid(const IWorldPartitionActorDescInstanceView& ActorDescView, const IWorldPartitionActorDescInstanceView& ReferenceActorDescView) override;
+	ENGINE_API virtual void OnDataLayersLoadFilterMismatch(const IWorldPartitionActorDescInstanceView& ActorDescView) override;
 	ENGINE_API virtual void OnInvalidReferenceDataLayerAsset(const UDataLayerInstanceWithAsset* DataLayerInstance) override;
 	ENGINE_API virtual void OnInvalidDataLayerAssetType(const UDataLayerInstanceWithAsset* DataLayerInstance, const UDataLayerAsset* DataLayerAsset) override;
 	ENGINE_API virtual void OnDataLayerHierarchyTypeMismatch(const UDataLayerInstance* DataLayerInstance, const UDataLayerInstance* Parent, EDataLayerHierarchyInvalidReason Reason) override;
diff --git a/Engine/Source/Runtime/Engine/Public/WorldPartition/RuntimeHashSet/RuntimePartition.h b/Engine/Source/Runtime/Engine/Public/WorldPartition/RuntimeHashSet/RuntimePartition.h
--- a/Engine/Source/Runtime/Engine/Public/WorldPartition/RuntimeHashSet/RuntimePartition.h
+++ b/Engine/Source/Runtime/Engine/Public/WorldPartition/RuntimeHashSet/RuntimePartition.h
@@ -44,12 +44,7 @@ public:
 	 */
 	struct FCellDescInstance : public FCellDesc
 	{
-		FCellDescInstance(const FCellDesc& InCellDesc, URuntimePartition* InSourcePartition, const TArray<const UDataLayerInstance*>& InDataLayerInstances, const FGuid& InContentBundleID)
-			: FCellDesc(InCellDesc)
-			, SourcePartition(InSourcePartition)
-			, DataLayerInstances(InDataLayerInstances)
-			, ContentBundleID(InContentBundleID)
-		{}
+		FCellDescInstance(const FCellDesc& InCellDesc, URuntimePartition* InSourcePartition, const TArray<const UDataLayerInstance*>& InDataLayerInstances, const FGuid& InContentBundleID);
 
 		URuntimePartition* SourcePartition;
 		TArray<const UDataLayerInstance*> DataLayerInstances;
