Upgrading a check() on a condition to a Fatal as the relevant code can corrupt an iterator if it is allowed to run. Crashing here at least gives us a better idea of what went wrong.

diff --git a/Engine/Source/Runtime/Engine/Private/TextureCompiler.cpp b/Engine/Source/Runtime/Engine/Private/TextureCompiler.cpp
--- a/Engine/Source/Runtime/Engine/Private/TextureCompiler.cpp
+++ b/Engine/Source/Runtime/Engine/Private/TextureCompiler.cpp
@@ -255,8 +255,17 @@ void FTextureCompilingManager::AddTextures(TArrayView<UTexture* const> InTexture
 	// to strcmp on the name of the repro texture and just log the full key suffix. Then you should immediately see the changed
 	// keys right before the crash and you can backsolve what value changed. Once you have that, you can set a data breakpoint on
 	// the property and see who is poking it.
-	checkf(bIsRoutingPostCompilation == false,
-		TEXT("Registering a texture to the compile manager from inside a texture postcompilation is not supported and usually indicate that the previous async operation wasn't completed (i.e. missing call to PreEditChange) before modifying a texture property."));
+	if (bIsRoutingPostCompilation)
+	{
+		// This has been updated to Fatal because it potentially modifies RegisteredTextureBuckets below which is iterated upon
+		// during PostCompilation routing. That modification can put us in an unstable state and crash in unexpected and rather
+		// undebuggable ways.
+		UE_LOG(LogTexture, Fatal, 
+			TEXT("Registering a texture to the compile manager from inside a texture postcompilation is not supported and usually")
+			TEXT(" indicates that the previous async operation wasn't completed (i.e. missing call to PreEditChange) before modifying a texture property.")
+			);
+	}
+		
 
 	TRACE_CPUPROFILER_EVENT_SCOPE(FTextureCompilingManager::AddTextures)
 
