[Actor Coloration] Allow for actor coloration handlers to be available only in game, not editor.

diff --git a/Engine/Source/Editor/LevelEditor/Private/SLevelViewportToolBar.cpp b/Engine/Source/Editor/LevelEditor/Private/SLevelViewportToolBar.cpp
--- a/Engine/Source/Editor/LevelEditor/Private/SLevelViewportToolBar.cpp
+++ b/Engine/Source/Editor/LevelEditor/Private/SLevelViewportToolBar.cpp
@@ -1864,6 +1864,11 @@ void SLevelViewportToolBar::FillViewMenu(UToolMenu* Menu)
 
 			for (const FActorPrimitiveColorHandler::FPrimitiveColorHandler& PrimitiveColorHandler : PrimitiveColorHandlers)
 			{
+				if (!PrimitiveColorHandler.bAvailalbleInEditor)
+				{
+					continue;
+				}
+
 				SubMenuSection.AddMenuEntry(
 					NAME_None,
 					PrimitiveColorHandler.HandlerText,
diff --git a/Engine/Source/Runtime/Engine/Classes/GameFramework/ActorPrimitiveColorHandler.h b/Engine/Source/Runtime/Engine/Classes/GameFramework/ActorPrimitiveColorHandler.h
--- a/Engine/Source/Runtime/Engine/Classes/GameFramework/ActorPrimitiveColorHandler.h
+++ b/Engine/Source/Runtime/Engine/Classes/GameFramework/ActorPrimitiveColorHandler.h
@@ -22,15 +22,17 @@ class ENGINE_API FActorPrimitiveColorHandler
 public:
 	struct FPrimitiveColorHandler
 	{
-		FPrimitiveColorHandler(FName InHandlerName, FText InHandlerText, const FGetColorFunc& InGetColorFunc, const FActivateFunc& InActivateFunc)
+		FPrimitiveColorHandler(FName InHandlerName, FText InHandlerText, bool bInAvailalbleInEditor, const FGetColorFunc& InGetColorFunc, const FActivateFunc& InActivateFunc)
 			: HandlerName(InHandlerName)
 			, HandlerText(InHandlerText)
+			, bAvailalbleInEditor(bInAvailalbleInEditor)
 			, GetColorFunc(InGetColorFunc)
 			, ActivateFunc(InActivateFunc)
 		{}
 
 		FName HandlerName;
 		FText HandlerText;
+		bool bAvailalbleInEditor;
 		FGetColorFunc GetColorFunc;
 		FActivateFunc ActivateFunc;
 	};	
@@ -39,6 +41,7 @@ public:
 	static FActorPrimitiveColorHandler& Get();
 
 	void RegisterPrimitiveColorHandler(FName InHandlerName, FText InHandlerText, const FGetColorFunc& InHandlerFunc, const FActivateFunc& InActivateFunc = []() {});
+	void RegisterPrimitiveColorHandler(FName InHandlerName, FText InHandlerText, bool bInAvailalbleInEditor, const FGetColorFunc& InHandlerFunc, const FActivateFunc& InActivateFunc = []() {});
 	void UnregisterPrimitiveColorHandler(FName InHandlerName);
 	void GetRegisteredPrimitiveColorHandlers(TArray<FPrimitiveColorHandler>& OutPrimitiveColorHandlers) const;
 
diff --git a/Engine/Source/Runtime/Engine/Private/ActorPrimitiveColorHandler.cpp b/Engine/Source/Runtime/Engine/Private/ActorPrimitiveColorHandler.cpp
--- a/Engine/Source/Runtime/Engine/Private/ActorPrimitiveColorHandler.cpp
+++ b/Engine/Source/Runtime/Engine/Private/ActorPrimitiveColorHandler.cpp
@@ -25,11 +25,16 @@ FActorPrimitiveColorHandler& FActorPrimitiveColorHandler::Get()
 	return TLazySingleton<FActorPrimitiveColorHandler>::Get();
 }
 
-void FActorPrimitiveColorHandler::RegisterPrimitiveColorHandler(FName InHandlerName, FText InHandlerText, const FGetColorFunc& InGetColorFunc, const FActivateFunc& InActivateFunc)
+void FActorPrimitiveColorHandler::RegisterPrimitiveColorHandler(FName InHandlerName, FText InHandlerText, const FGetColorFunc& InHandlerFunc, const FActivateFunc& InActivateFunc)
+{
+	RegisterPrimitiveColorHandler(InHandlerName, InHandlerText, true, InHandlerFunc, InActivateFunc);
+}
+
+void FActorPrimitiveColorHandler::RegisterPrimitiveColorHandler(FName InHandlerName, FText InHandlerText, bool bInAvailalbleInEditor, const FGetColorFunc& InHandlerFunc, const FActivateFunc& InActivateFunc)
 {
 #if ENABLE_ACTOR_PRIMITIVE_COLOR_HANDLER
 	check(!Handlers.Contains(InHandlerName));
-	Handlers.Add(InHandlerName, { InHandlerName, InHandlerText, InGetColorFunc, InActivateFunc });
+	Handlers.Add(InHandlerName, { InHandlerName, InHandlerText, bInAvailalbleInEditor, InHandlerFunc, InActivateFunc });
 	ActivePrimitiveColorHandler = &Handlers.FindChecked(ActivePrimitiveColorHandlerName);
 #endif
 }
