Fixed crash caused by [REDACTED_PERSON] handling an unsaved actor not in this Content Bundle.

diff --git a/Engine/Source/Runtime/Engine/Private/WorldPartition/ActorDescContainerInstance.cpp b/Engine/Source/Runtime/Engine/Private/WorldPartition/ActorDescContainerInstance.cpp
--- a/Engine/Source/Runtime/Engine/Private/WorldPartition/ActorDescContainerInstance.cpp
+++ b/Engine/Source/Runtime/Engine/Private/WorldPartition/ActorDescContainerInstance.cpp
@@ -343,7 +343,10 @@ bool UActorDescContainerInstance::IsActorDescHandled(const AActor* Actor) const
 	}
 
 	// Special case of Newly created maps where the Container might point to a template map but our map path is different
-	if (!Actor->GetContentBundleGuid().IsValid() && GetPackage()->HasAnyPackageFlags(PKG_NewlyCreated) && GetPackage()->GetName() != GetContainerPackage())
+	if (!Actor->GetContentBundleGuid().IsValid() && 
+		!GetContentBundleGuid().IsValid() &&
+		GetPackage()->HasAnyPackageFlags(PKG_NewlyCreated) && 
+		GetPackage()->GetName() != GetContainerPackage())
 	{
 		const FString ActorPackageName = Actor->GetPackage()->GetName();
 		const FString ExternalActorPath = ULevel::GetExternalActorsPath(GetPackage()->GetName()) / TEXT("");
