Add CVar to control whether we precache PSOs again on component render state recreation.

diff --git a/Engine/Source/Runtime/Engine/Private/Components/ActorComponent.cpp b/Engine/Source/Runtime/Engine/Private/Components/ActorComponent.cpp
--- a/Engine/Source/Runtime/Engine/Private/Components/ActorComponent.cpp
+++ b/Engine/Source/Runtime/Engine/Private/Components/ActorComponent.cpp
@@ -110,6 +110,14 @@ FAutoConsoleVariableRef CVarEnableDeferredPhysicsCreation(
 	TEXT("Enables/Disables deferred physics creation.")
 );
 
+// Allows for CreatePhysicsState to be deferred, to batch work and parallelize.
+int32 GPrecachePSOsOnComponentRecreateRenderContext = 1;
+FAutoConsoleVariableRef CVarPrecachePSOsOnComponentRecreateRenderContext(
+	TEXT("r.PrecachePSOsOnComponentRecreateRenderContext"),
+	GPrecachePSOsOnComponentRecreateRenderContext,
+	TEXT("If > 0, re-creating a component's render context will also re-precache the component's PSOs.")
+);
+
 void FRegisterComponentContext::Process()
 {
 	TRACE_CPUPROFILER_EVENT_SCOPE(FRegisterComponentContext::Process)
@@ -299,6 +307,28 @@ void FGlobalComponentReregisterContext::UpdateAllPrimitiveSceneInfos()
 	check(ScenesToUpdateAllPrimitiveSceneInfos.Num() == 0);
 }
 
+
+FComponentRecreateRenderStateContext::~FComponentRecreateRenderStateContext()
+{
+	if (Component && !Component->IsRenderStateCreated() && Component->IsRegistered())
+	{
+		if (GPrecachePSOsOnComponentRecreateRenderContext)
+		{
+			Component->PrecachePSOs();
+		}
+		Component->CreateRenderState_Concurrent(nullptr);
+
+		UpdateAllPrimitiveSceneInfosForSingleComponent(Component, ScenesToUpdateAllPrimitiveSceneInfos);
+	}
+
+	if (ComponentInterface && !ComponentInterface->IsRenderStateCreated() && ComponentInterface->IsRegistered())
+	{
+		ComponentInterface->CreateRenderState(nullptr);
+
+		UpdateAllPrimitiveSceneInfosForSingleComponentInterface(ComponentInterface, ScenesToUpdateAllPrimitiveSceneInfos);
+	}
+}
+
 FGlobalComponentRecreateRenderStateContext::FGlobalComponentRecreateRenderStateContext()
 {
 	if (FApp::CanEverRender())
diff --git a/Engine/Source/Runtime/Engine/Public/ComponentRecreateRenderStateContext.h b/Engine/Source/Runtime/Engine/Public/ComponentRecreateRenderStateContext.h
--- a/Engine/Source/Runtime/Engine/Public/ComponentRecreateRenderStateContext.h
+++ b/Engine/Source/Runtime/Engine/Public/ComponentRecreateRenderStateContext.h
@@ -73,23 +73,7 @@ public:
 		return *this;
 	}
 
-	~FComponentRecreateRenderStateContext()
-	{
-		if (Component && !Component->IsRenderStateCreated() && Component->IsRegistered())
-		{
-			Component->PrecachePSOs();
-			Component->CreateRenderState_Concurrent(nullptr);
-
-			UpdateAllPrimitiveSceneInfosForSingleComponent(Component, ScenesToUpdateAllPrimitiveSceneInfos);
-		}
-
-		if (ComponentInterface && !ComponentInterface ->IsRenderStateCreated() && ComponentInterface ->IsRegistered())
-		{
-			ComponentInterface ->CreateRenderState(nullptr);
-
-			UpdateAllPrimitiveSceneInfosForSingleComponentInterface(ComponentInterface, ScenesToUpdateAllPrimitiveSceneInfos);
-		}
-	}
+	ENGINE_API ~FComponentRecreateRenderStateContext();
 };
 
 /** Destroys render states for all components or for a provided list of components and then recreates them when this object is destroyed */
