Fix for warning on 4gb not available for DDC wave compiler

diff --git a/Engine/Source/Runtime/Engine/Classes/Sound/SoundWave.h b/Engine/Source/Runtime/Engine/Classes/Sound/SoundWave.h
--- a/Engine/Source/Runtime/Engine/Classes/Sound/SoundWave.h
+++ b/Engine/Source/Runtime/Engine/Classes/Sound/SoundWave.h
@@ -971,7 +971,6 @@ public:
 		{
 			return RawData.GetPayloadId();
 		}
-		UE_DEPRECATED(5.4, "GetPayloadSize is provided just for API backwards compatibility.")
 		int64 GetPayloadSize() const
 		{
 			return RawData.GetPayloadSize();
diff --git a/Engine/Source/Runtime/Engine/Private/AudioDerivedData.cpp b/Engine/Source/Runtime/Engine/Private/AudioDerivedData.cpp
--- a/Engine/Source/Runtime/Engine/Private/AudioDerivedData.cpp
+++ b/Engine/Source/Runtime/Engine/Private/AudioDerivedData.cpp
@@ -1141,7 +1141,12 @@ void FStreamedAudioPlatformData::Cache(USoundWave& InSoundWave, const FPlatformA
 			FWriteScopeLock AsyncTaskScope(AsyncTaskLock.Get());
 			check(AsyncTask == nullptr);
 			AsyncTask = new FStreamedAudioAsyncCacheDerivedDataTask(this, &InSoundWave, CompressionOverrides, AudioFormatName, Flags, InTargetPlatform);
-			int64 RequiredMemory = -1; // @todo RequiredMemory
+
+			// Use the size of the Uncompressed data x3 as guestmate of how much memory will be used.
+			const int64 PayloadSize = InSoundWave.RawData.GetPayloadSize() * 3;
+			// If there is no payoad size for some reason, use the default of -1.
+			const int64 RequiredMemory = PayloadSize > 0 ? PayloadSize : -1;
+
 			AsyncTask->StartBackgroundTask(SoundWaveThreadPool, BasePriority, EQueuedWorkFlags::DoNotRunInsideBusyWait, RequiredMemory, TEXT("AudioDerivedData") );
 		}
 
