Add minimum material quality as a filter for whether an RVT volume is instantiated in game.
#test In map with RVT enabled, apply the new setting and go into PIE at different material qualities to observe whether the RVT is being setup or not.

diff --git a/Engine/Source/Runtime/Engine/Classes/Components/RuntimeVirtualTextureComponent.h b/Engine/Source/Runtime/Engine/Classes/Components/RuntimeVirtualTextureComponent.h
--- a/Engine/Source/Runtime/Engine/Classes/Components/RuntimeVirtualTextureComponent.h
+++ b/Engine/Source/Runtime/Engine/Classes/Components/RuntimeVirtualTextureComponent.h
@@ -7,6 +7,7 @@
 #include "PerPlatformProperties.h"
 #include "RenderCommandFence.h"
 #include "SceneComponent.h"
+#include "SceneTypes.h"
 #include "RuntimeVirtualTextureComponent.generated.h"
 
 class URuntimeVirtualTexture;
@@ -14,6 +15,16 @@ class UTexture2D;
 class UVirtualTextureBuilder;
 enum class EShadingPath;
 
+/** Enumeration of material quality. Similar to EMaterialQualityLevel, but as a UENUM for appearance in UI. */
+UENUM()
+enum class ERuntimeVirtualTextureMaterialQuality : uint8
+{
+	Low = 0,
+	Medium = 1,
+	High = 2,
+	Epic = 3,
+};
+
 /** Component used to place a URuntimeVirtualTexture in the world. */
 UCLASS(Blueprintable, ClassGroup = Rendering, HideCategories = (Activation, Collision, Cooking, HLOD, Mobility, LOD, Navigation, Object, Physics), MinimalAPI)
 class URuntimeVirtualTextureComponent : public USceneComponent
@@ -49,6 +60,14 @@ protected:
 	UPROPERTY(EditAnywhere, AdvancedDisplay, Category = RuntimeVirtualTexture)
 	bool bEnableForNaniteOnly = false;
 
+	/** Use a minimum material quality to determine if we enable the virtual texture. */
+	UPROPERTY(EditAnywhere, AdvancedDisplay, Category = RuntimeVirtualTexture, Meta = (InlineEditConditionToggle))
+	bool bUseMinMaterialQuality = false;
+
+	/** The minimum material quality for which we enable the virtual texture. Only affects In-Game and PIE. */
+	UPROPERTY(EditAnywhere, AdvancedDisplay, Category = RuntimeVirtualTexture, Meta = (EditCondition = "bUseMinMaterialQuality"))
+	ERuntimeVirtualTextureMaterialQuality MinInGameMaterialQuality = ERuntimeVirtualTextureMaterialQuality::Low;
+
 	/** Set to true to enable scalability settings for the virtual texture. */
 	UPROPERTY(EditAnywhere, AdvancedDisplay, Category = RuntimeVirtualTexture, meta = (InlineEditConditionToggle))
 	bool bEnableScalability = false;
diff --git a/Engine/Source/Runtime/Engine/Private/Components/RuntimeVirtualTextureComponent.cpp b/Engine/Source/Runtime/Engine/Private/Components/RuntimeVirtualTextureComponent.cpp
--- a/Engine/Source/Runtime/Engine/Private/Components/RuntimeVirtualTextureComponent.cpp
+++ b/Engine/Source/Runtime/Engine/Private/Components/RuntimeVirtualTextureComponent.cpp
@@ -15,6 +15,7 @@
 #include "RHIGlobals.h"
 #include "SceneInterface.h"
 #include "SceneUtils.h"
+#include "UnrealEngine.h"
 #include "UObject/UObjectIterator.h"
 #include "UObject/UnrealType.h"
 #include "VT/RuntimeVirtualTexture.h"
@@ -107,9 +108,9 @@ void URuntimeVirtualTextureComponent::OnUnregister()
 
 #endif
 
-void URuntimeVirtualTextureComponent::SetVirtualTexture(URuntimeVirtualTexture* InVirtualTexture) 
+void URuntimeVirtualTextureComponent::SetVirtualTexture(URuntimeVirtualTexture* InVirtualTexture)
 {
-	VirtualTexture = InVirtualTexture; 
+	VirtualTexture = InVirtualTexture;
 	MarkRenderStateDirty();
 }
 
@@ -166,6 +167,20 @@ void URuntimeVirtualTextureComponent::DestroyRenderState_Concurrent()
 	Super::DestroyRenderState_Concurrent();
 }
 
+static ERuntimeVirtualTextureMaterialQuality ConvertMaterialQualityEnum(EMaterialQualityLevel::Type InMaterialQualityLevel)
+{
+	switch (InMaterialQualityLevel)
+	{
+	case EMaterialQualityLevel::Low: return ERuntimeVirtualTextureMaterialQuality::Low;
+	case EMaterialQualityLevel::Medium: return ERuntimeVirtualTextureMaterialQuality::Medium;
+	case EMaterialQualityLevel::High: return ERuntimeVirtualTextureMaterialQuality::High;
+	case EMaterialQualityLevel::Epic: return ERuntimeVirtualTextureMaterialQuality::Epic;
+	default: check(0);
+	}
+
+	return ERuntimeVirtualTextureMaterialQuality::Low;
+}
+
 bool URuntimeVirtualTextureComponent::IsEnabledInScene() const
 {
 	const bool bUseNanite = UseNanite(GetScene()->GetShaderPlatform());
@@ -182,6 +197,15 @@ bool URuntimeVirtualTextureComponent::IsEnabledInScene() const
 			{
 				return false;
 			}
+
+			if (bUseMinMaterialQuality)
+			{
+				ERuntimeVirtualTextureMaterialQuality CurrentQuality = ConvertMaterialQualityEnum(GetCachedScalabilityCVars().MaterialQualityLevel);
+ 				if (CurrentQuality < MinInGameMaterialQuality)
+ 				{
+ 					return false;
+ 				}
+			}
 		}
 	}
 
