Track game thread critical path time (GGameThreadTimeCriticalPath).
- only WorldTick and SendAllEndOfFrameUpdates are considered critical path.
- added `stat UnitCriticalPath` to show game and render thread critical path on-screen.
- added GameThreadTime_CriticalPath to CSV profiler output.
- added r.DynamicRes.UseGameThreadCriticalPath to control whether dynamic resolution uses game thread critical path time when determining whether game is CPU bound (also disabled by default for now).
- also added missing `stat UnitMax` to auto complete list.

diff --git a/Engine/Config/BaseInput.ini b/Engine/Config/BaseInput.ini
--- a/Engine/Config/BaseInput.ini
+++ b/Engine/Config/BaseInput.ini
@@ -104,6 +104,8 @@ MaxScrollbackSize=1024
 +ManualAutoCompleteList=(Command="Stat FPS",Desc="Shows FPS counter")
 +ManualAutoCompleteList=(Command="Stat UNIT",Desc="Shows hardware unit framerate")
 +ManualAutoCompleteList=(Command="Stat DrawCount",Desc="Shows draw counts broken down by category")
++ManualAutoCompleteList=(Command="Stat UnitCriticalPath",Desc="Shows thread critical path times")
++ManualAutoCompleteList=(Command="Stat UnitMax",Desc="Shows thread max times ")
 +ManualAutoCompleteList=(Command="Stat UnitGraph",Desc="Draws simple unit time graph")
 +ManualAutoCompleteList=(Command="Stat NamedEvents",Desc="Stat NamedEvents (Enables named events for external profilers)")
 +ManualAutoCompleteList=(Command="Stat VerboseNamedEvents",Desc="Stat VerboseNamedEvents (Enables verbose named events for external profilers)")
diff --git a/Engine/Source/Runtime/Engine/Classes/Engine/Engine.h b/Engine/Source/Runtime/Engine/Classes/Engine/Engine.h
--- a/Engine/Source/Runtime/Engine/Classes/Engine/Engine.h
+++ b/Engine/Source/Runtime/Engine/Classes/Engine/Engine.h
@@ -3824,6 +3824,7 @@ private:
 	ENGINE_API bool ToggleStatUnit(UWorld* World, FCommonViewportClient* ViewportClient, const TCHAR* Stream = nullptr);
 #if !UE_BUILD_SHIPPING
 	ENGINE_API bool PostStatSoundModulatorHelp(UWorld* World, FCommonViewportClient* ViewportClient, const TCHAR* Stream = nullptr);
+	ENGINE_API bool ToggleStatUnitCriticalPath(UWorld* World, FCommonViewportClient* ViewportClient, const TCHAR* Stream = nullptr);
 	ENGINE_API bool ToggleStatUnitMax(UWorld* World, FCommonViewportClient* ViewportClient, const TCHAR* Stream = nullptr);
 	ENGINE_API bool ToggleStatUnitGraph(UWorld* World, FCommonViewportClient* ViewportClient, const TCHAR* Stream = nullptr);
 	ENGINE_API bool ToggleStatUnitTime(UWorld* World, FCommonViewportClient* ViewportClient, const TCHAR* Stream = nullptr);
diff --git a/Engine/Source/Runtime/Engine/Private/DynamicResolution.cpp b/Engine/Source/Runtime/Engine/Private/DynamicResolution.cpp
--- a/Engine/Source/Runtime/Engine/Private/DynamicResolution.cpp
+++ b/Engine/Source/Runtime/Engine/Private/DynamicResolution.cpp
@@ -44,6 +44,11 @@ static TAutoConsoleVariable<float> CVarFrameTimeBudget(
 	TEXT("Frame's time budget in milliseconds."),
 	ECVF_RenderThreadSafe | ECVF_Default);
 
+static TAutoConsoleVariable<int32> CVarUseGameThreadCriticalPath(
+	TEXT("r.DynamicRes.UseGameThreadCriticalPath"), 0,
+	TEXT("Whether to use game thread critical path time when determining whether game is CPU bound."),
+	ECVF_RenderThreadSafe | ECVF_Default);
+
 #if COMPILE_DYNAMIC_FRAME_TIME
 
 static TAutoConsoleVariable<int32> CVarDynamicFrameTimeEnable(
@@ -1184,6 +1189,11 @@ public:
 			float PrevFrameTimeMs = (FApp::GetCurrentTime() - FApp::GetLastTime()) * 1000.0f;
 			float PrevGameThreadTimeMs = FPlatformTime::ToMilliseconds(GGameThreadTime);
 
+			if (CVarUseGameThreadCriticalPath.GetValueOnAnyThread())
+			{
+				PrevGameThreadTimeMs = FPlatformTime::ToMilliseconds(GGameThreadTimeCriticalPath);
+			}
+
 			FDefaultDynamicResolutionStateProxy* P = Proxy;
 			ENQUEUE_RENDER_COMMAND(DynamicResolutionBeginFrame)(
 				[PrevFrameTimeMs, PrevGameThreadTimeMs, P](class FRHICommandList& RHICmdList)
diff --git a/Engine/Source/Runtime/Engine/Private/LevelTick.cpp b/Engine/Source/Runtime/Engine/Private/LevelTick.cpp
--- a/Engine/Source/Runtime/Engine/Private/LevelTick.cpp
+++ b/Engine/Source/Runtime/Engine/Private/LevelTick.cpp
@@ -1286,6 +1286,8 @@ void UWorld::Tick( ELevelTick TickType, float DeltaSeconds )
 		return;
 	}
 
+	FThreadIdleStats::BeginCriticalPath();
+
 	SCOPED_DRAW_EVENT_GAMETHREAD(WorldTick);
 
 	FWorldDelegates::OnWorldTickStart.Broadcast(this, TickType, DeltaSeconds);
@@ -1785,6 +1787,8 @@ void UWorld::Tick( ELevelTick TickType, float DeltaSeconds )
 		});
 
 	FWorldDelegates::OnWorldTickEnd.Broadcast(this, TickType, DeltaSeconds);
+
+	FThreadIdleStats::EndCriticalPath();
 }
 
 void UWorld::CleanupActors()
diff --git a/Engine/Source/Runtime/Engine/Private/UnrealClient.cpp b/Engine/Source/Runtime/Engine/Private/UnrealClient.cpp
--- a/Engine/Source/Runtime/Engine/Private/UnrealClient.cpp
+++ b/Engine/Source/Runtime/Engine/Private/UnrealClient.cpp
@@ -360,10 +360,16 @@ int32 FStatUnitData::DrawStat(FViewport* InViewport, FCanvas* InCanvas, int32 In
 	RawGameThreadTime = FPlatformTime::ToMilliseconds(GGameThreadTime);
 	GameThreadTime = 0.9 * GameThreadTime + 0.1 * RawGameThreadTime;
 
+	RawGameThreadTimeCriticalPath = FPlatformTime::ToMilliseconds(GGameThreadTimeCriticalPath);
+	GameThreadTimeCriticalPath = 0.9 * GameThreadTimeCriticalPath + 0.1 * RawGameThreadTimeCriticalPath;
+
 	/** Number of milliseconds the renderthread was used last frame. */
 	RawRenderThreadTime = FPlatformTime::ToMilliseconds(GRenderThreadTime);
 	RenderThreadTime = 0.9 * RenderThreadTime + 0.1 * RawRenderThreadTime;
 
+	RawRenderThreadTimeCriticalPath = FPlatformTime::ToMilliseconds(GRenderThreadTimeCriticalPath);
+	RenderThreadTimeCriticalPath = 0.9 * RenderThreadTimeCriticalPath + 0.1 * RawRenderThreadTimeCriticalPath;
+
 	RawRHITTime = FPlatformTime::ToMilliseconds(GRHIThreadTime);
 	RHITTime = 0.9 * RHITTime + 0.1 * RawRHITTime;
 
@@ -411,6 +417,7 @@ int32 FStatUnitData::DrawStat(FViewport* InViewport, FCanvas* InCanvas, int32 In
 	float Max_RHITTime = 0.0f;
 	float Max_InputLatencyTime = 0.0f;
 
+	const bool bShowUnitCriticalPathTimes = InViewport->GetClient() ? InViewport->GetClient()->IsStatEnabled(TEXT("UnitCriticalPath")) : false;
 	const bool bShowUnitMaxTimes = InViewport->GetClient() ? InViewport->GetClient()->IsStatEnabled(TEXT("UnitMax")) : false;
 	const bool bShowTSRStatistics = InViewport->GetClient() ? InViewport->GetClient()->IsStatEnabled(TEXT("TSR")) : false;
 #if !UE_BUILD_SHIPPING
@@ -496,14 +503,24 @@ int32 FStatUnitData::DrawStat(FViewport* InViewport, FCanvas* InCanvas, int32 In
 		const FColor NoUnitGraphColor(220, 220, 220);
 
 		#define STATUNIT_FORMAT_AVGTIME TEXT("%3.2f ms")
+		#define STATUNIT_FORMAT_CRITICALPATHTIME TEXT("%3.2f ms")
 		#define STATUNIT_FORMAT_MAXTIME TEXT("%4.2f ms")
 		#define STATUNIT_FORMAT_PERCEMT TEXT("%3.2f %%")
 
 		const int32 AvgUnitColumnId = 0;
-		const int32 MaxUnitColumnId = 1;
+		const int32 CriticalPathUnitColumnId = 1;
+		const int32 MaxUnitColumnId = bShowUnitCriticalPathTimes ? 2 : 1;
 
 		const int32 ColumnWidth = Font->GetStringSize(TEXT(" 0000.00 ms "));
-		const int32 ColumnCount = bShowUnitMaxTimes ? 2 : 1;
+		int32 ColumnCount = 1;
+		if (bShowUnitCriticalPathTimes)
+		{
+			++ColumnCount;
+		}
+		if (bShowUnitMaxTimes)
+		{
+			++ColumnCount;
+		}
 
 		int32 X3 = InX * (bStereoRendering ? 0.5f : 1.0f);
 		int32 X2 = X3 - ColumnWidth * (ColumnCount - 1);
@@ -547,12 +564,33 @@ int32 FStatUnitData::DrawStat(FViewport* InViewport, FCanvas* InCanvas, int32 In
 		{
 			DrawCell(/* RowId = */ 0, AvgUnitColumnId, ColumnCount, CellText, UnitGraphColor);
 		};
+		auto DrawDefaultCriticalPathCell = [&](const FString& CellText, const FColor& UnitGraphColor)
+		{
+			check(bShowUnitCriticalPathTimes);
+			DrawCell(/* RowId = */ 0, CriticalPathUnitColumnId, ColumnCount, CellText, UnitGraphColor);
+		};
 		auto DrawDefaultMaxCell = [&](const FString& CellText, const FColor& UnitGraphColor)
 		{
 			check(bShowUnitMaxTimes);
 			DrawCell(/* RowId = */ 0, MaxUnitColumnId, ColumnCount, CellText, UnitGraphColor);
 		};
 
+		if (bShowUnitCriticalPathTimes || bShowUnitMaxTimes)
+		{
+			// if there are multiple columns, draw the name of each column
+
+			DrawDefaultAvgCell(TEXT("Avg"), FColor::White);
+			if (bShowUnitCriticalPathTimes)
+			{
+				DrawDefaultCriticalPathCell(TEXT("Critical"), FColor::White);
+			}
+			if (bShowUnitMaxTimes)
+			{
+				DrawDefaultMaxCell(TEXT("Max"), FColor::White);
+			}
+			PushRows(/* RowsCount = */ 1);
+		}
+
 		{
 			DrawTitleString(TEXT("Frame"), /* UnitGraphColor = */ FColor(100, 255, 100));
 			{
@@ -573,6 +611,11 @@ int32 FStatUnitData::DrawStat(FViewport* InViewport, FCanvas* InCanvas, int32 In
 				const FColor GameThreadAverageColor = GEngine->GetFrameTimeDisplayColor(GameThreadTime);
 				DrawDefaultAvgCell(FString::Printf(STATUNIT_FORMAT_AVGTIME, GameThreadTime), GameThreadAverageColor);
 			}
+			if (bShowUnitCriticalPathTimes)
+			{
+				const FColor CriticalPathFrameTimeColor = GEngine->GetFrameTimeDisplayColor(GameThreadTimeCriticalPath);
+				DrawDefaultCriticalPathCell(FString::Printf(STATUNIT_FORMAT_CRITICALPATHTIME, GameThreadTimeCriticalPath), CriticalPathFrameTimeColor);
+			}
 			if (bShowUnitMaxTimes)
 			{
 				const FColor GameThreadMaxColor = GEngine->GetFrameTimeDisplayColor(Max_GameThreadTime);
@@ -587,6 +630,11 @@ int32 FStatUnitData::DrawStat(FViewport* InViewport, FCanvas* InCanvas, int32 In
 				const FColor RenderThreadAverageColor = GEngine->GetFrameTimeDisplayColor(RenderThreadTime);
 				DrawDefaultAvgCell(FString::Printf(STATUNIT_FORMAT_AVGTIME, RenderThreadTime), RenderThreadAverageColor);
 			}
+			if (bShowUnitCriticalPathTimes)
+			{
+				const FColor CriticalPathFrameTimeColor = GEngine->GetFrameTimeDisplayColor(RenderThreadTimeCriticalPath);
+				DrawDefaultCriticalPathCell(FString::Printf(STATUNIT_FORMAT_CRITICALPATHTIME, RenderThreadTimeCriticalPath), CriticalPathFrameTimeColor);
+			}
 			if (bShowUnitMaxTimes)
 			{
 				const FColor RenderThreadMaxColor = GEngine->GetFrameTimeDisplayColor(Max_RenderThreadTime);
@@ -1730,6 +1778,10 @@ void FViewport::Draw( bool bShouldPresent /*= true */)
 							// add any stalls via sleep or fevent
 							GGameThreadTime		= (ThreadTime > GameThread.Waits) ? (ThreadTime - GameThread.Waits) : ThreadTime;
 							GGameThreadWaitTime = GameThread.Waits;
+
+							// Compute GGameThreadTimeCriticalPath
+							uint32 GameThreadNonCriticalPathIdle = GameThread.Waits - GameThread.WaitsCriticalPath;
+							GGameThreadTimeCriticalPath = (ThreadTime > GameThreadNonCriticalPathIdle) ? (ThreadTime - GameThreadNonCriticalPathIdle) : ThreadTime;
 						}
 						else
 						{
@@ -1739,6 +1791,7 @@ void FViewport::Draw( bool bShouldPresent /*= true */)
 						LastFrameUpdated = GFrameCounter;
 						Lastimestamp		= CurrentTime;
 						GameThread.Reset();
+						GameThread.IsCriticalPathCounter = 0;
 					}
 				}
 
diff --git a/Engine/Source/Runtime/Engine/Private/UnrealEngine.cpp b/Engine/Source/Runtime/Engine/Private/UnrealEngine.cpp
--- a/Engine/Source/Runtime/Engine/Private/UnrealEngine.cpp
+++ b/Engine/Source/Runtime/Engine/Private/UnrealEngine.cpp
@@ -2207,6 +2207,7 @@ void UEngine::Init(IEngineLoop* InEngineLoop)
 	EngineStats.Add(FEngineStatFuncs(TEXT("STAT_Levels"), TEXT("STATCAT_Engine"), FText::FromString(TEXT("Display a list of names of all the loaded levels.")), FEngineStatRender::CreateUObject(this, &UEngine::RenderStatLevels), FEngineStatToggle()));
 	EngineStats.Add(FEngineStatFuncs(TEXT("STAT_Detailed"), TEXT("STATCAT_Engine"), FText::FromString(TEXT("Display detailed frame and fps timings.")), FEngineStatRender(), FEngineStatToggle::CreateUObject(this, &UEngine::ToggleStatDetailed)));
 #if !UE_BUILD_SHIPPING
+	EngineStats.Add(FEngineStatFuncs(TEXT("STAT_UnitCriticalPath"), TEXT("STATCAT_Engine"), FText::FromString(TEXT("Same as stat unit, but with critical path values also displayed.")), FEngineStatRender(), FEngineStatToggle::CreateUObject(this, &UEngine::ToggleStatUnitCriticalPath)));
 	EngineStats.Add(FEngineStatFuncs(TEXT("STAT_UnitMax"), TEXT("STATCAT_Engine"), FText::FromString(TEXT("Same as stat unit, but with max values also displayed.")), FEngineStatRender(), FEngineStatToggle::CreateUObject(this, &UEngine::ToggleStatUnitMax)));
 	EngineStats.Add(FEngineStatFuncs(TEXT("STAT_UnitGraph"), TEXT("STATCAT_Engine"), FText::FromString(TEXT("Displays a frame time stats graphed over time.")), FEngineStatRender(), FEngineStatToggle::CreateUObject(this, &UEngine::ToggleStatUnitGraph)));
 	EngineStats.Add(FEngineStatFuncs(TEXT("STAT_UnitTime"), TEXT("STATCAT_Engine"), FText::GetEmpty(), FEngineStatRender(), FEngineStatToggle::CreateUObject(this, &UEngine::ToggleStatUnitTime)));
@@ -17912,6 +17913,7 @@ bool UEngine::ToggleStatDetailed(UWorld* World, FCommonViewportClient* ViewportC
 		bSetup = true;
 		DetailedStats.Add(TEXT("FPS"));
 		DetailedStats.Add(TEXT("Unit"));
+		DetailedStats.Add(TEXT("UnitCriticalPath"));
 		DetailedStats.Add(TEXT("UnitMax"));
 		DetailedStats.Add(TEXT("UnitGraph"));
 		DetailedStats.Add(TEXT("Raw"));
@@ -18428,8 +18430,38 @@ int32 UEngine::RenderStatDrawCount(UWorld* World, FViewport* Viewport, FCanvas*
 	return Y;
 }
 
-// UNITMAX
 #if !UE_BUILD_SHIPPING
+
+// UNITCRITICALPATH
+bool UEngine::ToggleStatUnitCriticalPath(UWorld* World, FCommonViewportClient* ViewportClient, const TCHAR* Stream)
+{
+	if (ViewportClient == nullptr)
+	{
+		// Ignore if all Viewports are closed.
+		return false;
+	}
+	const bool bShowUnitCriticalPathTimes = ViewportClient->IsStatEnabled(TEXT("UnitCriticalPath"));
+	if (bShowUnitCriticalPathTimes)
+	{
+		// Force Unit to Active
+		SetEngineStat(World, ViewportClient, TEXT("Unit"), true);
+
+		// Force UnitCriticalPath to true as Unit will have Toggled it back to false
+		SetEngineStat(World, ViewportClient, TEXT("UnitCriticalPath"), true);
+	}
+	else
+	{
+		const bool bShowDetailed = ViewportClient->IsStatEnabled(TEXT("Detailed"));
+		if (bShowDetailed)
+		{
+			// Since we're turning this off, we also need to toggle off detailed too
+			ExecEngineStat(World, ViewportClient, TEXT("Detailed -Skip"));
+		}
+	}
+	return true;
+}
+
+// UNITMAX
 bool UEngine::ToggleStatUnitMax(UWorld* World, FCommonViewportClient* ViewportClient, const TCHAR* Stream)
 {
 	if( ViewportClient == nullptr )
diff --git a/Engine/Source/Runtime/Engine/Public/UnrealClient.h b/Engine/Source/Runtime/Engine/Public/UnrealClient.h
--- a/Engine/Source/Runtime/Engine/Public/UnrealClient.h
+++ b/Engine/Source/Runtime/Engine/Public/UnrealClient.h
@@ -275,7 +275,9 @@ struct FStatUnitData
 {
 	/** Unit frame times filtered with a simple running average */
 	float RenderThreadTime;
+	float RenderThreadTimeCriticalPath;
 	float GameThreadTime;
+	float GameThreadTimeCriticalPath;
 	float GPUFrameTime[MAX_NUM_GPUS];
 	float GPUClockFraction[MAX_NUM_GPUS];
 	float GPUUsageFraction[MAX_NUM_GPUS];
@@ -286,7 +288,9 @@ struct FStatUnitData
 
 	/** Raw equivalents of the above variables */
 	float RawRenderThreadTime;
+	float RawRenderThreadTimeCriticalPath;
 	float RawGameThreadTime;
+	float RawGameThreadTimeCriticalPath;
 	float RawGPUFrameTime[MAX_NUM_GPUS];
 	float RawGPUClockFraction[MAX_NUM_GPUS];
 	float RawGPUUsageFraction[MAX_NUM_GPUS];
@@ -313,7 +317,9 @@ struct FStatUnitData
 
 	FStatUnitData()
 		: RenderThreadTime(0.0f)
+		, RenderThreadTimeCriticalPath(0.0f)
 		, GameThreadTime(0.0f)
+		, GameThreadTimeCriticalPath(0.0f)
 		, GPUFrameTime{ 0.0f }
 		, GPUClockFraction{ 0.0f }
 		, GPUUsageFraction{ 0.0f }
@@ -322,7 +328,9 @@ struct FStatUnitData
 		, RHITTime(0.0f)
 		, InputLatencyTime(0.0f)
 		, RawRenderThreadTime(0.0f)
+		, RawRenderThreadTimeCriticalPath(0.0f)
 		, RawGameThreadTime(0.0f)
+		, RawGameThreadTimeCriticalPath(0.0f)
 		, RawGPUFrameTime{ 0.0f }
 		, RawGPUClockFraction{ 0.0f }
 		, RawGPUUsageFraction{ 0.0f }
diff --git a/Engine/Source/Runtime/Launch/Private/LaunchEngineLoop.cpp b/Engine/Source/Runtime/Launch/Private/LaunchEngineLoop.cpp
--- a/Engine/Source/Runtime/Launch/Private/LaunchEngineLoop.cpp
+++ b/Engine/Source/Runtime/Launch/Private/LaunchEngineLoop.cpp
@@ -1727,6 +1727,7 @@ static void UpdateCoreCsvStats_EndFrame()
 	    CSV_CUSTOM_STAT_GLOBAL(GameThreadTime, FPlatformTime::ToMilliseconds(GGameThreadTime), ECsvCustomStatOp::Set);
 	    CSV_CUSTOM_STAT_GLOBAL(GPUTime, FPlatformTime::ToMilliseconds(GGPUFrameTime), ECsvCustomStatOp::Set);
 		CSV_CUSTOM_STAT_GLOBAL(RenderThreadTime_CriticalPath, FPlatformTime::ToMilliseconds(GRenderThreadTimeCriticalPath), ECsvCustomStatOp::Set);
+		CSV_CUSTOM_STAT_GLOBAL(GameThreadTime_CriticalPath, FPlatformTime::ToMilliseconds(GGameThreadTimeCriticalPath), ECsvCustomStatOp::Set);
 		if (IsRunningRHIInSeparateThread())
 	    {
 		    CSV_CUSTOM_STAT_GLOBAL(RHIThreadTime, FPlatformTime::ToMilliseconds(GRHIThreadTime), ECsvCustomStatOp::Set);
diff --git a/Engine/Source/Runtime/RenderCore/Private/RenderingThread.cpp b/Engine/Source/Runtime/RenderCore/Private/RenderingThread.cpp
--- a/Engine/Source/Runtime/RenderCore/Private/RenderingThread.cpp
+++ b/Engine/Source/Runtime/RenderCore/Private/RenderingThread.cpp
@@ -1157,6 +1157,9 @@ uint32 GGameThreadTime = 0;
 /** How much idle time on the game thread. It's set once per frame in FViewport::Draw. */
 uint32 GGameThreadWaitTime = 0;
 
+/** How many cycles the gamethread used, including dependent wait time. */
+uint32 GGameThreadTimeCriticalPath = 0;
+
 /** How many cycles it took to swap buffers to present the frame. */
 uint32 GSwapBufferTime = 0;
 
diff --git a/Engine/Source/Runtime/RenderCore/Public/RenderTimer.h b/Engine/Source/Runtime/RenderCore/Public/RenderTimer.h
--- a/Engine/Source/Runtime/RenderCore/Public/RenderTimer.h
+++ b/Engine/Source/Runtime/RenderCore/Public/RenderTimer.h
@@ -122,5 +122,8 @@ extern RENDERCORE_API uint32 GGameThreadWaitTime;
 /** How many cycles it took to swap buffers to present the frame. */
 extern RENDERCORE_API uint32 GSwapBufferTime;
 
+/** How many cycles the gamethread used, including dependent wait time. */
+extern RENDERCORE_API uint32 GGameThreadTimeCriticalPath;
+
 /** How many cycles the renderthread used, including dependent wait time. */
 extern RENDERCORE_API uint32 GRenderThreadTimeCriticalPath;
diff --git a/Engine/Source/Runtime/Renderer/Private/SceneRendering.cpp b/Engine/Source/Runtime/Renderer/Private/SceneRendering.cpp
--- a/Engine/Source/Runtime/Renderer/Private/SceneRendering.cpp
+++ b/Engine/Source/Runtime/Renderer/Private/SceneRendering.cpp
@@ -5226,8 +5226,12 @@ void FRendererModule::BeginRenderingViewFamilies(FCanvas* Canvas, TArrayView<FSc
 		World = Scene->GetWorld();
 		if (World)
 		{
-			// Guarantee that all render proxies are up to date before kicking off a BeginRenderViewFamily.
-			World->SendAllEndOfFrameUpdates();
+			FThreadIdleStats::BeginCriticalPath();
+			{
+				// Guarantee that all render proxies are up to date before kicking off a BeginRenderViewFamily.
+				World->SendAllEndOfFrameUpdates();
+			}
+			FThreadIdleStats::EndCriticalPath();
 
 			GetNaniteVisualizationData().Pick(World);
 		}
