Skeletal Mesh: Fix default color detection

#rnx

diff --git a/Engine/Plugins/Experimental/MeshModelingToolsetExp/Source/MeshModelingToolsEditorOnlyExp/Private/SkeletalMesh/SkeletonEditingTool.cpp b/Engine/Plugins/Experimental/MeshModelingToolsetExp/Source/MeshModelingToolsEditorOnlyExp/Private/SkeletalMesh/SkeletonEditingTool.cpp
--- a/Engine/Plugins/Experimental/MeshModelingToolsetExp/Source/MeshModelingToolsEditorOnlyExp/Private/SkeletalMesh/SkeletonEditingTool.cpp
+++ b/Engine/Plugins/Experimental/MeshModelingToolsetExp/Source/MeshModelingToolsEditorOnlyExp/Private/SkeletalMesh/SkeletonEditingTool.cpp
@@ -162,6 +162,10 @@ void USkeletonEditingTool::SetupPreviewMesh()
 		PreviewMesh->SetSecondaryRenderMaterial(SelectionMaterial);
 	}
 
+	// FIXME: This setting really belongs on the underlying mesh.
+	UDynamicMeshComponent* DynamicMeshComponent = Cast<UDynamicMeshComponent>(PreviewMesh->GetRootComponent());
+	DynamicMeshComponent->SetVertexColorSpaceTransformMode(EDynamicMeshVertexColorTransformMode::LinearToSRGB);
+
 	// hide the skeletal mesh component
 	UE::ToolTarget::HideSourceObject(Target);
 }
diff --git a/Engine/Source/Developer/SkeletalMeshUtilitiesCommon/Private/StaticToSkeletalMeshConverter.cpp b/Engine/Source/Developer/SkeletalMeshUtilitiesCommon/Private/StaticToSkeletalMeshConverter.cpp
--- a/Engine/Source/Developer/SkeletalMeshUtilitiesCommon/Private/StaticToSkeletalMeshConverter.cpp
+++ b/Engine/Source/Developer/SkeletalMeshUtilitiesCommon/Private/StaticToSkeletalMeshConverter.cpp
@@ -271,6 +271,32 @@ static bool AddLODFromStaticMeshSourceModel(
 			SkinWeights.Set(VertexID, RootBinding);
 		}
 
+		// Convert weird static mesh inverse sRGB gamma to linear.
+		// FIXME: Remove once static mesh color space has been fixed to be linear again.
+		TVertexInstanceAttributesRef<FVector4f> VertexInstanceColors = SkeletalMeshAttributes.GetVertexInstanceColors();
+		auto ConvertLinearToSRGBGamma = [](float V)
+		{
+			V = FMath::Clamp(V, 0.0f, 1.0f);
+			if (V <= 0.0031308)
+			{
+				return V * 12.92f;
+			}
+			else
+			{
+				return 1.055f * FMath::Pow(V, 1.0f / 2.4f) - 0.055f;
+			}
+		};
+		
+		for (FVertexInstanceID VertexInstanceID: SkeletalMeshGeometry.VertexInstances().GetElementIDs())
+		{
+			FLinearColor VertexColor = VertexInstanceColors.Get(VertexInstanceID);
+			VertexColor.R = ConvertLinearToSRGBGamma(VertexColor.R);
+			VertexColor.G = ConvertLinearToSRGBGamma(VertexColor.G);
+			VertexColor.B = ConvertLinearToSRGBGamma(VertexColor.B);
+			VertexInstanceColors.Set(VertexInstanceID, VertexColor);
+		}
+		
+
 		if (!AddLODFromMeshDescription(MoveTemp(SkeletalMeshGeometry), InSkeletalMesh, InMeshUtilities))
 		{
 			return false;
diff --git a/Engine/Source/Runtime/Engine/Private/Rendering/SkeletalMeshLODImporterData.cpp b/Engine/Source/Runtime/Engine/Private/Rendering/SkeletalMeshLODImporterData.cpp
--- a/Engine/Source/Runtime/Engine/Private/Rendering/SkeletalMeshLODImporterData.cpp
+++ b/Engine/Source/Runtime/Engine/Private/Rendering/SkeletalMeshLODImporterData.cpp
@@ -2220,6 +2220,11 @@ bool FSkeletalMeshImportData::GetMeshDescription(const USkeletalMesh* InSkeletal
 					// Don't perform sRGB conversion (which mirrors what CreateFromMeshDescription does).
 					VertexInstanceColors.Set(VertexInstanceID, Wedge.Color.ReinterpretAsLinear());
 				}
+				else
+				{
+					// The default color for import data is black. 
+					VertexInstanceColors.Set(VertexInstanceID, FLinearColor::Black);
+				}
 				for (int32 UVIndex = 0; UVIndex < static_cast<int32>(NumTexCoords); UVIndex++)
 				{
 					VertexInstanceUVs.Set(VertexInstanceID, UVIndex, Wedge.UVs[UVIndex]);
@@ -2764,18 +2769,22 @@ FSkeletalMeshImportData FSkeletalMeshImportData::CreateFromMeshDescription(const
 	SkelMeshImportData.bHasTangents = bHaveValidTangents;
 	
 	SkelMeshImportData.Wedges.SetNumZeroed(InMeshDescription.VertexInstances().GetArraySize());
+	
 	for (FVertexInstanceID VertexInstanceID: InMeshDescription.VertexInstances().GetElementIDs())
 	{
 		SkeletalMeshImportData::FVertex& Wedge = SkelMeshImportData.Wedges[VertexInstanceID.GetValue()];
 		
 		Wedge.VertexIndex = static_cast<uint32>(InMeshDescription.GetVertexInstanceVertex(VertexInstanceID).GetValue());
 		Wedge.MatIndex = 0;			// We set this later -- not that this is actually used by any internal process.
+		
 		constexpr bool bSRGB = false; //avoid linear to srgb conversion
-		Wedge.Color = FLinearColor(VertexInstanceColors[VertexInstanceID]).ToFColor(bSRGB);
-		if (Wedge.Color != FColor::White)
+		const FLinearColor VertexColor = VertexInstanceColors[VertexInstanceID];
+		if (VertexColor != FLinearColor::Black)
 		{
 			SkelMeshImportData.bHasVertexColors = true;
 		}
+		Wedge.Color = VertexColor.ToFColor(bSRGB);
+		
 		for (int32 UVChannelIndex = 0; UVChannelIndex < static_cast<int32>(SkelMeshImportData.NumTexCoords); ++UVChannelIndex)
 		{
 			Wedge.UVs[UVChannelIndex] = VertexInstanceUVs.Get(VertexInstanceID, UVChannelIndex);
