- Fixed LevelInstanceLevelStreaming to reset the RF_Transactional flag on objects of the external actor packages only after the level is made visible which guarantees that all objects are properly updated.
- Also run the same logic in the OnLoadedActorsAddedToLevelPostEvent instead of OnLoadedActorsAddedToLevelPreEvent because to wait for component registration to complete before clearing the RF_Transactional flag.
Because the unloading of a Level Instance resets the transaction buffer when it detects any unloading object with the RF_Transactional, all deleted actors are no longer kept alive by the transaction and the next GC deletes these actors. The save dialog won't show the actor packages to delete.

diff --git a/Engine/Source/Runtime/Engine/Private/LevelInstance/LevelInstanceLevelStreaming.cpp b/Engine/Source/Runtime/Engine/Private/LevelInstance/LevelInstanceLevelStreaming.cpp
--- a/Engine/Source/Runtime/Engine/Private/LevelInstance/LevelInstanceLevelStreaming.cpp
+++ b/Engine/Source/Runtime/Engine/Private/LevelInstance/LevelInstanceLevelStreaming.cpp
@@ -250,8 +250,6 @@ void ULevelStreamingLevelInstance::InitializeActors(const TArray<AActor*>& InAct
 								}
 							}, /*bIncludeNestedObjects*/ true);
 						}
-
-						FLevelInstanceLevelStreamingUtils::MarkObjectsInPackageAsTransientAndNonTransactional(Actor->GetExternalPackage());
 					}
 
 					// Must happen before the actors are registered with the world, which is the case for this delegate.
@@ -276,6 +274,11 @@ void ULevelStreamingLevelInstance::OnLoadedActorsAddedToLevelPostEvent(const TAr
 				{
 					if (IsValid(Actor))
 					{
+						if (Actor->IsPackageExternal())
+						{
+							FLevelInstanceLevelStreamingUtils::MarkObjectsInPackageAsTransientAndNonTransactional(Actor->GetExternalPackage());
+						}
+
 						Actor->PushSelectionToProxies();
 						if (LevelInstance)
 						{
@@ -448,8 +451,7 @@ void ULevelStreamingLevelInstance::OnCurrentStateChanged(ELevelStreamingState In
 
 		ForEachObjectWithOuter(Level, [](UObject* InObject)
 		{
-			// Skip actors as they are already handled in OnLoadedActorsAddedToLevelPreEvent
-			if (InObject && InObject->IsPackageExternal() && !InObject->IsA<AActor>())
+			if (InObject && InObject->IsPackageExternal())
 			{
 				FLevelInstanceLevelStreamingUtils::MarkObjectsInPackageAsTransientAndNonTransactional(InObject->GetPackage());
 			}
