Interchange: fix crash when querying skeletal mesh description.

#rnx

diff --git a/Engine/Source/Runtime/Engine/Private/SkeletalMesh.cpp b/Engine/Source/Runtime/Engine/Private/SkeletalMesh.cpp
--- a/Engine/Source/Runtime/Engine/Private/SkeletalMesh.cpp
+++ b/Engine/Source/Runtime/Engine/Private/SkeletalMesh.cpp
@@ -2572,7 +2572,10 @@ void USkeletalMesh::LoadLODImportedData(const int32 LODIndex, FSkeletalMeshImpor
 
 	if (const FMeshDescription* MeshDescription = GetMeshDescription(LODIndex))
 	{
-		OutMesh = FSkeletalMeshImportData::CreateFromMeshDescription(*MeshDescription);
+		if (!MeshDescription->IsEmpty())
+		{
+			OutMesh = FSkeletalMeshImportData::CreateFromMeshDescription(*MeshDescription);
+		}
 	}
 }
 
