Add missing ENGINE_API to new API.

diff --git a/Engine/Source/Runtime/Engine/Public/WorldPartition/RuntimeHashSet/WorldPartitionRuntimeHashSet.h b/Engine/Source/Runtime/Engine/Public/WorldPartition/RuntimeHashSet/WorldPartitionRuntimeHashSet.h
--- a/Engine/Source/Runtime/Engine/Public/WorldPartition/RuntimeHashSet/WorldPartitionRuntimeHashSet.h
+++ b/Engine/Source/Runtime/Engine/Public/WorldPartition/RuntimeHashSet/WorldPartitionRuntimeHashSet.h
@@ -191,10 +191,10 @@ public:
 	ENGINE_API virtual uint32 ComputeUpdateStreamingHash() const override;
 
 protected:
-	virtual bool SupportsWorldAssetStreaming(const FName& InTargetGrid) override;
-	virtual FGuid RegisterWorldAssetStreaming(const UWorldPartition::FRegisterWorldAssetStreamingParams& InParams) override;
-	virtual bool UnregisterWorldAssetStreaming(const FGuid& InWorldAssetStreamingGuid) override;
-	virtual TArray<UWorldPartitionRuntimeCell*> GetWorldAssetStreamingCells(const FGuid& InWorldAssetStreamingGuid) override;
+	ENGINE_API virtual bool SupportsWorldAssetStreaming(const FName& InTargetGrid) override;
+	ENGINE_API virtual FGuid RegisterWorldAssetStreaming(const UWorldPartition::FRegisterWorldAssetStreamingParams& InParams) override;
+	ENGINE_API virtual bool UnregisterWorldAssetStreaming(const FGuid& InWorldAssetStreamingGuid) override;
+	ENGINE_API virtual TArray<UWorldPartitionRuntimeCell*> GetWorldAssetStreamingCells(const FGuid& InWorldAssetStreamingGuid) override;
 
 private:
 	ENGINE_API virtual void OnBeginPlay() override;
diff --git a/Engine/Source/Runtime/Engine/Public/WorldPartition/WorldPartitionRuntimeHash.h b/Engine/Source/Runtime/Engine/Public/WorldPartition/WorldPartitionRuntimeHash.h
--- a/Engine/Source/Runtime/Engine/Public/WorldPartition/WorldPartitionRuntimeHash.h
+++ b/Engine/Source/Runtime/Engine/Public/WorldPartition/WorldPartitionRuntimeHash.h
@@ -268,10 +268,10 @@ protected:
 
 	TSet<TWeakObjectPtr<URuntimeHashExternalStreamingObjectBase>> InjectedExternalStreamingObjects;
 
-	virtual bool SupportsWorldAssetStreaming(const FName& InTargetGrid) { return false; }
-	virtual FGuid RegisterWorldAssetStreaming(const UWorldPartition::FRegisterWorldAssetStreamingParams& InParams) { return FGuid(); }
-	virtual bool UnregisterWorldAssetStreaming(const FGuid& InWorldAssetStreamingGuid) { return false; }
-	virtual TArray<UWorldPartitionRuntimeCell*> GetWorldAssetStreamingCells(const FGuid& InWorldAssetStreamingGuid) { return {}; }
+	ENGINE_API virtual bool SupportsWorldAssetStreaming(const FName& InTargetGrid) { return false; }
+	ENGINE_API virtual FGuid RegisterWorldAssetStreaming(const UWorldPartition::FRegisterWorldAssetStreamingParams& InParams) { return FGuid(); }
+	ENGINE_API virtual bool UnregisterWorldAssetStreaming(const FGuid& InWorldAssetStreamingGuid) { return false; }
+	ENGINE_API virtual TArray<UWorldPartitionRuntimeCell*> GetWorldAssetStreamingCells(const FGuid& InWorldAssetStreamingGuid) { return {}; }
 
 #if WITH_EDITOR
 private:
diff --git a/Engine/Source/Runtime/Engine/Public/WorldPartition/WorldPartitionRuntimeSpatialHash.h b/Engine/Source/Runtime/Engine/Public/WorldPartition/WorldPartitionRuntimeSpatialHash.h
--- a/Engine/Source/Runtime/Engine/Public/WorldPartition/WorldPartitionRuntimeSpatialHash.h
+++ b/Engine/Source/Runtime/Engine/Public/WorldPartition/WorldPartitionRuntimeSpatialHash.h
@@ -419,10 +419,10 @@ protected:
 	mutable TMap<FName, const FSpatialHashStreamingGrid*> NameToGridMapping;
 	mutable bool bIsNameToGridMappingDirty;
 
-	virtual bool SupportsWorldAssetStreaming(const FName& InTargetGrid) override;
-	virtual FGuid RegisterWorldAssetStreaming(const UWorldPartition::FRegisterWorldAssetStreamingParams& InParams) override;
-	virtual bool UnregisterWorldAssetStreaming(const FGuid& InWorldAssetStreamingGuid) override;
-	virtual TArray<UWorldPartitionRuntimeCell*> GetWorldAssetStreamingCells(const FGuid& InWorldAssetStreamingGuid) override;
+	ENGINE_API virtual bool SupportsWorldAssetStreaming(const FName& InTargetGrid) override;
+	ENGINE_API virtual FGuid RegisterWorldAssetStreaming(const UWorldPartition::FRegisterWorldAssetStreamingParams& InParams) override;
+	ENGINE_API virtual bool UnregisterWorldAssetStreaming(const FGuid& InWorldAssetStreamingGuid) override;
+	ENGINE_API virtual TArray<UWorldPartitionRuntimeCell*> GetWorldAssetStreamingCells(const FGuid& InWorldAssetStreamingGuid) override;
 
 	UPROPERTY(Transient)
 	TMap<FGuid, TObjectPtr<URuntimeSpatialHashExternalStreamingObject>> WorldAssetStreamingObjects;
