Don't allow setting MSAA with deferred shading in renderer settings as this combination is not supported.

diff --git a/Engine/Source/Runtime/Engine/Classes/Engine/RendererSettings.h b/Engine/Source/Runtime/Engine/Classes/Engine/RendererSettings.h
--- a/Engine/Source/Runtime/Engine/Classes/Engine/RendererSettings.h
+++ b/Engine/Source/Runtime/Engine/Classes/Engine/RendererSettings.h
@@ -303,7 +303,7 @@ class URendererSettings : public UDeveloperSettings
 
 	UPROPERTY(config, EditAnywhere, Category = Mobile, meta = (
 		ConsoleVariable = "r.Mobile.ShadingPath", DisplayName = "Mobile Shading",
-		ToolTip = "The shading path to use on mobile platforms. Changing this setting requires restarting the editor. Mobile HDR is required for Deferred Shading.",
+		ToolTip = "The shading path to use on mobile platforms. Changing this setting requires restarting the editor. Forward shading will force MSAA. Mobile HDR is required for Deferred Shading.",
 		ConfigRestartRequired = true))
 	TEnumAsByte<EMobileShadingPath::Type> MobileShadingPath;
 
@@ -321,7 +321,8 @@ class URendererSettings : public UDeveloperSettings
 
 	UPROPERTY(config, EditAnywhere, Category = Mobile, meta = (
 		ConsoleVariable = "r.Mobile.AntiAliasing", DisplayName = "Mobile Anti-Aliasing Method",
-		ToolTip = "The mobile default anti-aliasing method."))
+		ToolTip = "The mobile default anti-aliasing method.",
+		EditCondition = "MobileShadingPath == 1"))
 	TEnumAsByte<EMobileAntiAliasingMethod::Type> MobileAntiAliasing;
 
 	UPROPERTY(config, EditAnywhere, Category = Mobile, meta = (
@@ -601,7 +602,7 @@ class URendererSettings : public UDeveloperSettings
 	UPROPERTY(config, EditAnywhere, Category=ForwardRenderer, meta=(
 		ConsoleVariable="r.ForwardShading",
 		DisplayName = "Forward Shading",
-		ToolTip="Whether to use forward shading on desktop platforms, requires Shader Model 5 hardware.  Forward shading supports MSAA and has lower default cost, but fewer features supported overall.  Materials have to opt-in to more expensive features like high quality reflections.  Changing this setting requires restarting the editor.",
+		ToolTip="Whether to use forward shading on desktop platforms, requires Shader Model 5 hardware.  Forward shading requires MSAA and has lower default cost, but fewer features supported overall.  Materials have to opt-in to more expensive features like high quality reflections.  Deferred shading does not support MSAA.  Changing this setting requires restarting the editor.",
 		ConfigRestartRequired=true))
 	uint32 bForwardShading:1;
 
@@ -722,7 +723,8 @@ class URendererSettings : public UDeveloperSettings
 
 	UPROPERTY(config, EditAnywhere, Category = DefaultSettings, meta = (
 		ConsoleVariable = "r.AntiAliasingMethod", DisplayName = "Anti-Aliasing Method",
-		ToolTip = "Selects the anti-aliasing method to use."))
+		ToolTip = "Selects the anti-aliasing method to use.",
+		EditCondition = "!bForwardShading"))
 	TEnumAsByte<EAntiAliasingMethod> DefaultFeatureAntiAliasing;
 
 	UPROPERTY(config, EditAnywhere, Category = DefaultSettings, meta = (
@@ -1362,6 +1364,7 @@ private:
 	TWeakPtr<class SNotificationItem> ShaderModelNotificationPtr;
 
 	void CheckForMissingShaderModels();
+	void FixAntiAliasingOnShadingPathChange(FPropertyChangedEvent& PropertyChangedEvent);
 #endif // WITH_EDITOR
 
 	void SanatizeReflectionCaptureResolution();
diff --git a/Engine/Source/Runtime/Engine/Private/RendererSettings.cpp b/Engine/Source/Runtime/Engine/Private/RendererSettings.cpp
--- a/Engine/Source/Runtime/Engine/Private/RendererSettings.cpp
+++ b/Engine/Source/Runtime/Engine/Private/RendererSettings.cpp
@@ -162,6 +162,63 @@ void UpdateDependentPropertyInConfigFile(URendererSettings* RendererSettings, FN
 	}
 }
 
+void URendererSettings::FixAntiAliasingOnShadingPathChange(FPropertyChangedEvent& PropertyChangedEvent)
+{
+	// Going from deferred to forward we force MSAA. After that, the AA method is blocked by the edit condition in metadata.
+	if (PropertyChangedEvent.GetPropertyName() == GET_MEMBER_NAME_CHECKED(URendererSettings, MobileShadingPath))
+	{
+		if (MobileShadingPath == EMobileShadingPath::Forward)
+		{
+			MobileAntiAliasing = EMobileAntiAliasingMethod::MSAA;
+		}
+		else // When we disable forward shading we need to change AA method
+		{
+			MobileAntiAliasing = EMobileAntiAliasingMethod::None;
+			UE_LOG(LogTemp, Warning, TEXT("Disabling Mobile AA because MSAA is not compatible with Mobile Deferred Shading"));
+		}
+		UpdateDependentPropertyInConfigFile(this, GET_MEMBER_NAME_CHECKED(URendererSettings, MobileAntiAliasing));
+		ExportValuesToConsoleVariables(GetClass()->FindPropertyByName(GET_MEMBER_NAME_CHECKED(URendererSettings, MobileAntiAliasing)));
+	}
+	else if (PropertyChangedEvent.GetPropertyName() == GET_MEMBER_NAME_CHECKED(URendererSettings, bForwardShading))
+	{
+		if (bForwardShading)
+		{
+			DefaultFeatureAntiAliasing = EAntiAliasingMethod::AAM_MSAA;
+		}
+		else // When we disable forward shading we need to change AA method
+		{
+			DefaultFeatureAntiAliasing = EAntiAliasingMethod::AAM_None;
+			UE_LOG(LogTemp, Warning, TEXT("Disabling AA because MSAA is not compatible with Deferred Shading"));
+		}
+		UpdateDependentPropertyInConfigFile(this, GET_MEMBER_NAME_CHECKED(URendererSettings, DefaultFeatureAntiAliasing));
+		ExportValuesToConsoleVariables(GetClass()->FindPropertyByName(GET_MEMBER_NAME_CHECKED(URendererSettings, DefaultFeatureAntiAliasing)));
+	}
+
+	// MSAA can't be selected if deferred shading is active
+	if (PropertyChangedEvent.GetPropertyName() == GET_MEMBER_NAME_CHECKED(URendererSettings, MobileAntiAliasing))
+	{
+		if (MobileAntiAliasing == EMobileAntiAliasingMethod::MSAA && MobileShadingPath == EMobileShadingPath::Deferred)
+		{
+			MobileAntiAliasing = EMobileAntiAliasingMethod::None;
+			UpdateDependentPropertyInConfigFile(this, GET_MEMBER_NAME_CHECKED(URendererSettings, MobileAntiAliasing));
+			ExportValuesToConsoleVariables(GetClass()->FindPropertyByName(GET_MEMBER_NAME_CHECKED(URendererSettings, MobileAntiAliasing)));
+
+			UE_LOG(LogTemp, Error, TEXT("MSAA can't be used with Mobile Deferred Rendering. Resetting Mobile AA to None."));
+		}
+	}
+	else if (PropertyChangedEvent.GetPropertyName() == GET_MEMBER_NAME_CHECKED(URendererSettings, DefaultFeatureAntiAliasing))
+	{
+		if (DefaultFeatureAntiAliasing == EAntiAliasingMethod::AAM_MSAA && !bForwardShading)
+		{
+			DefaultFeatureAntiAliasing = EAntiAliasingMethod::AAM_None;
+			UpdateDependentPropertyInConfigFile(this, GET_MEMBER_NAME_CHECKED(URendererSettings, DefaultFeatureAntiAliasing));
+			ExportValuesToConsoleVariables(GetClass()->FindPropertyByName(GET_MEMBER_NAME_CHECKED(URendererSettings, DefaultFeatureAntiAliasing)));
+			
+			UE_LOG(LogTemp, Error, TEXT("MSAA can't be used with Deferred Rendering. Resetting AA to None."));
+		}
+	}
+}
+
 void URendererSettings::PostEditChangeProperty(FPropertyChangedEvent& PropertyChangedEvent)
 {
 	Super::PostEditChangeProperty(PropertyChangedEvent);
@@ -330,6 +387,8 @@ void URendererSettings::PostEditChangeProperty(FPropertyChangedEvent& PropertyCh
 		{
 			CheckForMissingShaderModels();
 		}
+
+		FixAntiAliasingOnShadingPathChange(PropertyChangedEvent);
 	}
 }
 
