FExternalDataLayerHelper::CanMoveActorsToExternalDataLayer now ignores Actor's CDO Assets when validating asset references.

diff --git a/Engine/Source/Runtime/Engine/Private/WorldPartition/DataLayer/ExternalDataLayerHelper.cpp b/Engine/Source/Runtime/Engine/Private/WorldPartition/DataLayer/ExternalDataLayerHelper.cpp
--- a/Engine/Source/Runtime/Engine/Private/WorldPartition/DataLayer/ExternalDataLayerHelper.cpp
+++ b/Engine/Source/Runtime/Engine/Private/WorldPartition/DataLayer/ExternalDataLayerHelper.cpp
@@ -194,15 +194,28 @@ bool FExternalDataLayerHelper::CanMoveActorsToExternalDataLayer(const TArray<AAc
 		}
 
 		// Gather actor asset references
-		TSet<UObject*> ActorReferencedAssets;
-		const bool bIncludeDefaultRefs = false;
 		const bool bOnlyDirectReferences = true;
-		FFindReferencedAssets::BuildAssetList(InActor, IgnoreClasses, IgnorePackages, ActorReferencedAssets, bIncludeDefaultRefs, bOnlyDirectReferences);
+		TSet<UObject*> ActorReferencedAssets;
+		{
+			const bool bIncludeDefaultRefs = false;
+			FFindReferencedAssets::BuildAssetList(InActor, IgnoreClasses, IgnorePackages, ActorReferencedAssets, bIncludeDefaultRefs, bOnlyDirectReferences);
+		}
 		TArray<UObject*> ReferencedContent;
 		InActor->GetReferencedContentObjects(ReferencedContent);
 		// Remove itself and its data layer assets from the list
 		ActorReferencedAssets.Append(ReferencedContent);
 		ActorReferencedAssets.Remove(InActor);
+		// Gather and remove CDO assets
+		{
+			const bool bIncludeDefaultRefs = true;
+			TSet<UObject*> CDOReferencedAssets;
+			FFindReferencedAssets::BuildAssetList(InActor->GetClass()->GetDefaultObject(), IgnoreClasses, IgnorePackages, CDOReferencedAssets, bIncludeDefaultRefs, bOnlyDirectReferences);
+			for (UObject* CDOAsset : CDOReferencedAssets)
+			{
+				ActorReferencedAssets.Remove(CDOAsset);
+			}
+		}
+		
 		for (auto It = ActorReferencedAssets.CreateIterator(); It; ++It)
 		{
 			if (!(*It)->IsAsset() || (*It)->IsA<UDataLayerAsset>())
