World Partition Runtime HashSet LHGrid: add support for setting a grid origin and import half cell size origin from runtime spatial hash when using UseAlignedGridLevels.

diff --git a/Engine/Source/Runtime/Engine/Private/WorldPartition/RuntimeHashSet/RuntimePartitionLHGrid.cpp b/Engine/Source/Runtime/Engine/Private/WorldPartition/RuntimeHashSet/RuntimePartitionLHGrid.cpp
--- a/Engine/Source/Runtime/Engine/Private/WorldPartition/RuntimeHashSet/RuntimePartitionLHGrid.cpp
+++ b/Engine/Source/Runtime/Engine/Private/WorldPartition/RuntimeHashSet/RuntimePartitionLHGrid.cpp
@@ -18,7 +18,7 @@ struct FCellCoord
 
 	inline FString ToString() const
 	{
-		return FString::Printf(TEXT("%" INT64_FMT "_%" INT64_FMT "_%" INT64_FMT "_%d"), X, Y, Z, Level);
+		return FString::Printf(TEXT("L%d_X%" INT64_FMT "_Y%" INT64_FMT "_Z%" INT64_FMT), Level, X, Y, Z);
 	}
 
 	inline bool operator==(const FCellCoord& Other) const
@@ -33,23 +33,23 @@ struct FCellCoord
 		return FMath::CeilToInt32(FMath::Max<FVector::FReal>(FMath::Log2(MaxLength / InCellSize), 0));
 	}
 
-	static inline FCellCoord GetCellCoords(const FVector& InPos, int32 InCellSize, int32 InLevel)
+	static inline FCellCoord GetCellCoords(const FVector& InPos, int32 InCellSize, int32 InLevel, const FVector& InOrigin)
 	{
 		check(InLevel >= 0);
 		const int64 CellSizeForLevel = (int64)InCellSize * (1LL << InLevel);
 		return FCellCoord(
-			FMath::FloorToInt(InPos.X / CellSizeForLevel),
-			FMath::FloorToInt(InPos.Y / CellSizeForLevel),
-			FMath::FloorToInt(InPos.Z / CellSizeForLevel),
+			FMath::FloorToInt((InPos.X - InOrigin.X) / CellSizeForLevel),
+			FMath::FloorToInt((InPos.Y - InOrigin.Y) / CellSizeForLevel),
+			FMath::FloorToInt((InPos.Z - InOrigin.Z) / CellSizeForLevel),
 			InLevel
 		);
 	}
 
-	static inline FBox GetCellBounds(const FCellCoord& InCellCoord, int32 InCellSize)
+	static inline FBox GetCellBounds(const FCellCoord& InCellCoord, int32 InCellSize, const FVector& InOrigin)
 	{
 		check(InCellCoord.Level >= 0);
 		const int64 CellSizeForLevel = (int64)InCellSize * (1LL << InCellCoord.Level);
-		const FVector Min = FVector(
+		const FVector Min = InOrigin + FVector(
 			static_cast<FVector::FReal>(InCellCoord.X * CellSizeForLevel), 
 			static_cast<FVector::FReal>(InCellCoord.Y * CellSizeForLevel), 
 			static_cast<FVector::FReal>(InCellCoord.Z * CellSizeForLevel)
@@ -57,6 +57,7 @@ struct FCellCoord
 		const FVector Max = Min + FVector(static_cast<double>(CellSizeForLevel));
 		return FBox(Min, Max);
 	}
+
 	friend uint32 GetTypeHash(const FCellCoord& CellCoord)
 	{
 		FHashBuilder HashBuilder;
@@ -110,7 +111,7 @@ void URuntimePartitionLHGrid::PostEditChangeProperty(FPropertyChangedEvent& InPr
 	{
 		WorldGridPreviewer->CellSize = CellSize;
 		WorldGridPreviewer->GridColor = DebugColor;
-		WorldGridPreviewer->GridOffset = FVector::ZeroVector;
+		WorldGridPreviewer->GridOffset = Origin;
 		WorldGridPreviewer->LoadingRange = LoadingRange;
 		WorldGridPreviewer->Update();
 	}
@@ -151,7 +152,7 @@ bool URuntimePartitionLHGrid::GenerateStreaming(const FGenerateStreamingParams&
 		if (ActorSetInstance->bIsSpatiallyLoaded)
 		{
 			const int32 GridLevel = FCellCoord::GetLevelForBox(ActorSetInstance->Bounds, CellSize);
-			const FCellCoord CellCoord = FCellCoord::GetCellCoords(ActorSetInstance->Bounds.GetCenter(), CellSize, GridLevel);
+			const FCellCoord CellCoord = FCellCoord::GetCellCoords(ActorSetInstance->Bounds.GetCenter(), CellSize, GridLevel, Origin);
 			CellsActorSetInstances.FindOrAdd(CellCoord).Add(ActorSetInstance);
 		}
 		else
@@ -168,8 +169,8 @@ bool URuntimePartitionLHGrid::GenerateStreaming(const FGenerateStreamingParams&
 
 		if (bIsSpatiallyLoaded)
 		{
-			CellDesc.Bounds = FCellCoord::GetCellBounds(CellCoord, CellSize);
-			CellDesc.CellBounds = FCellCoord::GetCellBounds(CellCoord, CellSize);
+			CellDesc.CellBounds = FCellCoord::GetCellBounds(CellCoord, CellSize, Origin);
+			CellDesc.Bounds = CellDesc.CellBounds.GetValue();
 		}
 	}
 
diff --git a/Engine/Source/Runtime/Engine/Private/WorldPartition/RuntimeHashSet/WorldPartitionRuntimeHashSetConversions.cpp b/Engine/Source/Runtime/Engine/Private/WorldPartition/RuntimeHashSet/WorldPartitionRuntimeHashSetConversions.cpp
--- a/Engine/Source/Runtime/Engine/Private/WorldPartition/RuntimeHashSet/WorldPartitionRuntimeHashSetConversions.cpp
+++ b/Engine/Source/Runtime/Engine/Private/WorldPartition/RuntimeHashSet/WorldPartitionRuntimeHashSetConversions.cpp
@@ -29,6 +29,7 @@ UWorldPartitionRuntimeHashSet* UWorldPartitionRuntimeHashSet::CreateFrom(const U
 		// Gather all HLOD layers into their corresponding grids
 		TMap<FName, TMap<FName, TSet<const UHLODLayer*>>> GridHLODLayersMap;
 		const UHLODLayer* DefaultHLODLayer = WorldPartition->GetDefaultHLODLayer();
+		const bool bUseAlignedGridLevels = SpatialHash->GetUseAlignedGridLevels();
 
 		for (const FSpatialHashRuntimeGrid& Grid : SpatialHash->Grids)
 		{
@@ -67,6 +68,7 @@ UWorldPartitionRuntimeHashSet* UWorldPartitionRuntimeHashSet::CreateFrom(const U
 				URuntimePartitionLHGrid* LHGrid = NewObject<URuntimePartitionLHGrid>(HashSet, NAME_None);
 				LHGrid->Name = RuntimePartitionDesc.Name;
 				LHGrid->CellSize = Grid.CellSize;
+				LHGrid->Origin = bUseAlignedGridLevels ? FVector::ZeroVector : FVector(LHGrid->CellSize * -0.5f);
 				LHGrid->bBlockOnSlowStreaming = Grid.bBlockOnSlowStreaming;
 				LHGrid->bClientOnlyVisible = Grid.bClientOnlyVisible;
 				LHGrid->Priority = Grid.Priority;
@@ -88,6 +90,7 @@ UWorldPartitionRuntimeHashSet* UWorldPartitionRuntimeHashSet::CreateFrom(const U
 					{
 						URuntimePartitionLHGrid* HLODLHGrid = NewObject<URuntimePartitionLHGrid>(HashSet, NAME_None);
 						HLODLHGrid->CellSize = HLODSetup.HLODLayers[0]->GetCellSize();
+						HLODLHGrid->Origin = bUseAlignedGridLevels ? FVector::ZeroVector : FVector(HLODLHGrid->CellSize * -0.5f);
 						HLODLHGrid->LoadingRange = HLODSetup.HLODLayers[0]->GetLoadingRange();
 						HLODSetup.PartitionLayer = HLODLHGrid;
 					}
diff --git a/Engine/Source/Runtime/Engine/Private/WorldPartition/WorldPartitionRuntimeSpatialHash.cpp b/Engine/Source/Runtime/Engine/Private/WorldPartition/WorldPartitionRuntimeSpatialHash.cpp
--- a/Engine/Source/Runtime/Engine/Private/WorldPartition/WorldPartitionRuntimeSpatialHash.cpp
+++ b/Engine/Source/Runtime/Engine/Private/WorldPartition/WorldPartitionRuntimeSpatialHash.cpp
@@ -1196,8 +1196,7 @@ FString UWorldPartitionRuntimeSpatialHash::GetCellCoordString(const FGridCellCoo
 void UWorldPartitionRuntimeSpatialHash::DrawPreview() const
 {
 	// Use latest settings value for Preview
-	const bool bUseAlignedGridLevels = (UseAlignedGridLevels == EWorldPartitionCVarProjectDefaultOverride::ProjectDefault) ? GRuntimeSpatialHashUseAlignedGridLevels : (UseAlignedGridLevels == EWorldPartitionCVarProjectDefaultOverride::Enabled);
-	GridPreviewer.Draw(GetWorld(), Grids, bPreviewGrids, PreviewGridLevel, bUseAlignedGridLevels);
+	GridPreviewer.Draw(GetWorld(), Grids, bPreviewGrids, PreviewGridLevel, GetUseAlignedGridLevels());
 }
 
 bool UWorldPartitionRuntimeSpatialHash::HasStreamingContent() const
@@ -1316,7 +1315,7 @@ void FSpatialHashSettings::UpdateSettings(const UWorldPartitionRuntimeSpatialHas
 {
 #if WITH_EDITORONLY_DATA
 	check(RuntimeSpatialHash.StreamingGrids.IsEmpty());
-	bUseAlignedGridLevels = (RuntimeSpatialHash.UseAlignedGridLevels == EWorldPartitionCVarProjectDefaultOverride::ProjectDefault) ? GRuntimeSpatialHashUseAlignedGridLevels : (RuntimeSpatialHash.UseAlignedGridLevels == EWorldPartitionCVarProjectDefaultOverride::Enabled);
+	bUseAlignedGridLevels = RuntimeSpatialHash.GetUseAlignedGridLevels();
 	bSnapNonAlignedGridLevelsToLowerLevels = (RuntimeSpatialHash.SnapNonAlignedGridLevelsToLowerLevels == EWorldPartitionCVarProjectDefaultOverride::ProjectDefault) ? GRuntimeSpatialHashSnapNonAlignedGridLevelsToLowerLevels : (RuntimeSpatialHash.SnapNonAlignedGridLevelsToLowerLevels == EWorldPartitionCVarProjectDefaultOverride::Enabled);
 	bPlaceSmallActorsUsingLocation = (RuntimeSpatialHash.PlaceSmallActorsUsingLocation == EWorldPartitionCVarProjectDefaultOverride::ProjectDefault) ? GRuntimeSpatialHashPlaceSmallActorsUsingLocation : (RuntimeSpatialHash.PlaceSmallActorsUsingLocation == EWorldPartitionCVarProjectDefaultOverride::Enabled);
 	bPlacePartitionActorsUsingLocation = (RuntimeSpatialHash.PlacePartitionActorsUsingLocation == EWorldPartitionCVarProjectDefaultOverride::ProjectDefault) ? GRuntimeSpatialHashPlacePartitionActorsUsingLocation : (RuntimeSpatialHash.PlacePartitionActorsUsingLocation == EWorldPartitionCVarProjectDefaultOverride::Enabled);
@@ -1621,6 +1620,11 @@ bool UWorldPartitionRuntimeSpatialHash::CreateStreamingGrid(const FSpatialHashRu
 	return true;
 }
 
+bool UWorldPartitionRuntimeSpatialHash::GetUseAlignedGridLevels() const
+{
+	return UseAlignedGridLevels == EWorldPartitionCVarProjectDefaultOverride::ProjectDefault ? GRuntimeSpatialHashUseAlignedGridLevels : (UseAlignedGridLevels == EWorldPartitionCVarProjectDefaultOverride::Enabled);
+}
+
 void UWorldPartitionRuntimeSpatialHash::FlushStreamingContent()
 {
 	Super::FlushStreamingContent();
diff --git a/Engine/Source/Runtime/Engine/Public/WorldPartition/RuntimeHashSet/RuntimePartitionLHGrid.h b/Engine/Source/Runtime/Engine/Public/WorldPartition/RuntimeHashSet/RuntimePartitionLHGrid.h
--- a/Engine/Source/Runtime/Engine/Public/WorldPartition/RuntimeHashSet/RuntimePartitionLHGrid.h
+++ b/Engine/Source/Runtime/Engine/Public/WorldPartition/RuntimeHashSet/RuntimePartitionLHGrid.h
@@ -36,7 +36,10 @@ public:
 private:
 #if WITH_EDITORONLY_DATA
 	UPROPERTY(EditAnywhere, Category = RuntimeSettings)
-	uint32 CellSize;
+	uint32 CellSize = 25600;
+
+	UPROPERTY(EditAnywhere, Category = RuntimeSettings)
+	FVector Origin = FVector::ZeroVector;
 
 	UPROPERTY(EditAnywhere, Category = RuntimeSettings, Transient, SkipSerialization)
 	bool bShowGridPreview = false;
diff --git a/Engine/Source/Runtime/Engine/Public/WorldPartition/WorldPartitionRuntimeSpatialHash.h b/Engine/Source/Runtime/Engine/Public/WorldPartition/WorldPartitionRuntimeSpatialHash.h
--- a/Engine/Source/Runtime/Engine/Public/WorldPartition/WorldPartitionRuntimeSpatialHash.h
+++ b/Engine/Source/Runtime/Engine/Public/WorldPartition/WorldPartitionRuntimeSpatialHash.h
@@ -438,6 +438,7 @@ private:
 	ENGINE_API const TMap<FName, const FSpatialHashStreamingGrid*>& GetNameToGridMapping() const;
 #if WITH_EDITOR
 	ENGINE_API bool CreateStreamingGrid(const FSpatialHashRuntimeGrid& RuntimeGrid, const FSquare2DGridHelper& PartitionedActors, UWorldPartitionStreamingPolicy* StreamingPolicy, TArray<FString>* OutPackagesToGenerate = nullptr);
+	ENGINE_API bool GetUseAlignedGridLevels() const;
 #endif
 	ENGINE_API TArray<const FSpatialHashStreamingGrid*> GetFilteredStreamingGrids() const;
 
