World Grid Preview: fix cells origin to be top-left corner instead of center, as our se cases are using that,

diff --git a/Engine/Source/Runtime/Engine/Private/WorldGridPreviewer.cpp b/Engine/Source/Runtime/Engine/Private/WorldGridPreviewer.cpp
--- a/Engine/Source/Runtime/Engine/Private/WorldGridPreviewer.cpp
+++ b/Engine/Source/Runtime/Engine/Private/WorldGridPreviewer.cpp
@@ -74,7 +74,8 @@ void FWorldGridPreviewer::Update()
 		MaterialInstance->SetScalarParameterValueEditorOnly(TEXT("Grid_CellSize"), (float)CellSize);
 		MaterialInstance->SetScalarParameterValueEditorOnly(TEXT("Grid_LoadingRange"), (float)LoadingRange);
 		MaterialInstance->SetVectorParameterValueEditorOnly(TEXT("Grid_Color"), GridColor);
-		//MaterialInstance->SetVectorParameterValueEditorOnly(TEXT("Grid_Offset"), GridOffset);		
+		// Grid material has cell's orinig at center, and our grid are using top-left
+		MaterialInstance->SetVectorParameterValueEditorOnly(TEXT("Grid_Offset"), FLinearColor(GridOffset + CellSize * 0.5f));
 	}
 }
 #endif
\ No newline at end of file
