Do not try to reconnect to ODSC server when exit is requested, causing a crash on shutdown

diff --git a/Engine/Source/Runtime/Engine/Private/ODSC/ODSCThread.cpp b/Engine/Source/Runtime/Engine/Private/ODSC/ODSCThread.cpp
--- a/Engine/Source/Runtime/Engine/Private/ODSC/ODSCThread.cpp
+++ b/Engine/Source/Runtime/Engine/Private/ODSC/ODSCThread.cpp
@@ -303,6 +303,12 @@ void FODSCThread::SendMessageToServer(IPlatformFile::IFileServerMessageHandler*
 		}
 		else
 		{
+			// Losing connection when exit is requested is expected, do not try to reconnect
+			if (ExitRequest.GetValue())
+			{
+				return;
+			}
+
 			UE_LOG(LogODSC, Display, TEXT("Detected that CookOnTheFlyServerConnection has been lost, trying again"));
 			if (!ConnectToODSCHost())
 			{
