[Material Editor] Removed incorrect validation code during DDC query.

diff --git a/Engine/Source/Runtime/Engine/Private/Materials/HLSLMaterialTranslator.cpp b/Engine/Source/Runtime/Engine/Private/Materials/HLSLMaterialTranslator.cpp
--- a/Engine/Source/Runtime/Engine/Private/Materials/HLSLMaterialTranslator.cpp
+++ b/Engine/Source/Runtime/Engine/Private/Materials/HLSLMaterialTranslator.cpp
@@ -15093,16 +15093,6 @@ void FHLSLMaterialTranslator::AsyncQueryDDC(
 			ResultsMemoryReader << MaterialSourceTemplateParams;
 
 			STAT(DDCRequestSerializeTime = FPlatformTime::Seconds() - SerializeTime);
-
-			// Verify that the Parameter Collections referenced by the Defines buffer are all valid.
-			for (UMaterialParameterCollection* ParameterCollection : EnvironmentDefines->ParameterCollections)
-			{
-				// If the parameter collection is null (it failed to load), do not use these cached results and continue translating the material
-				if (!ParameterCollection)
-				{
-					return;
-				}
-			}
 			
 			// The DDC request hit the cache and usable data was retrieved from it.
 			DDCQueryHit = true;
