Duplicate audio data when reading from memory mapped files

diff --git a/Engine/Source/Runtime/CoreUObject/Public/Serialization/BulkData.h b/Engine/Source/Runtime/CoreUObject/Public/Serialization/BulkData.h
--- a/Engine/Source/Runtime/CoreUObject/Public/Serialization/BulkData.h
+++ b/Engine/Source/Runtime/CoreUObject/Public/Serialization/BulkData.h
@@ -255,6 +255,10 @@ struct FOwnedBulkDataPtr
 	{
 		return MappedRegion;
 	}
+	bool IsDataMemoryMapped() const
+	{
+		return MappedHandle != nullptr;
+	}
 
 	void RelinquishOwnership()
 	{
diff --git a/Engine/Source/Runtime/Engine/Private/SoundWave.cpp b/Engine/Source/Runtime/Engine/Private/SoundWave.cpp
--- a/Engine/Source/Runtime/Engine/Private/SoundWave.cpp
+++ b/Engine/Source/Runtime/Engine/Private/SoundWave.cpp
@@ -162,6 +162,8 @@ namespace SoundWaveCookStats
 
 FSoundWaveData::~FSoundWaveData()
 {
+	ZerothChunkData.Empty();
+	ResourceData.Empty();
 }
 
 void FSoundWaveData::InitializeDataFromSoundWave(USoundWave& InWave)
@@ -2231,27 +2233,41 @@ void USoundWave::InitAudioResource(FByteBulkData& CompressedData)
 			CompressedData.GetCopy((void**)&TempDataPtr, true);
 			SoundWaveDataPtr->ResourceData.Reset(TempDataPtr, SoundWaveDataPtr->ResourceSize);
 #else
-			FOwnedBulkDataPtr* OwnedBulkDataPtr = nullptr;
-			if (!SoundWaveDataPtr->ResourceData.GetView().GetData())
+			if (SoundWaveDataPtr->ResourceData.GetView().GetData())
 			{
-				OwnedBulkDataPtr = CompressedData.StealFileMapping();
+				UE_LOG(LogAudio, Display, TEXT("Soundwave '%s' Has already had InitAudioResource() called, and taken ownership of it's compressed data.")
+					, *GetFullName());
+				return;
 			}
 			else
 			{
-				UE_LOG(LogAudio, Display, TEXT("Soundwave '%s' Has already had InitAudioResource() called, and taken ownership of it's compressed data.")
-					, *GetFullName());
 			}
 
-			check(OwnedBulkDataPtr);
-			uint8* TempDataPtr = (uint8*)(OwnedBulkDataPtr)->GetPointer();
+			uint8* TempDataPtr = nullptr;
+			if (CompressedData.IsDataMemoryMapped())
+			{
+				CompressedData.GetCopy((void**)&TempDataPtr, false);
+			}
+			else
+			{
+				FOwnedBulkDataPtr* OwnedBulkDataPtr = CompressedData.StealFileMapping();
+				TempDataPtr = (uint8*)(OwnedBulkDataPtr)->GetPointer();
+			}
 			SoundWaveDataPtr->ResourceData.Reset(TempDataPtr, SoundWaveDataPtr->ResourceSize);
 			if (!TempDataPtr)
 			{
 				UE_LOG(LogAudio, Error, TEXT("Soundwave '%s' was not loaded when it should have been, forcing a sync load."), *GetFullName());
 
 				CompressedData.ForceBulkDataResident();
-				OwnedBulkDataPtr = CompressedData.StealFileMapping();
-				TempDataPtr = (uint8*)(OwnedBulkDataPtr)->GetPointer();
+				if (CompressedData.IsDataMemoryMapped())
+				{
+					CompressedData.GetCopy((void**)&TempDataPtr, false);
+				}
+				else
+				{
+					FOwnedBulkDataPtr* OwnedBulkDataPtr = CompressedData.StealFileMapping();
+					TempDataPtr = (uint8*)(OwnedBulkDataPtr)->GetPointer();
+				}
 
 				SoundWaveDataPtr->ResourceData.Reset(TempDataPtr, SoundWaveDataPtr->ResourceSize);
 				if (!TempDataPtr)
