Fix CIS - static analysis following CL 33371703.

diff --git a/Engine/Source/Runtime/Engine/Private/WorldPartition/RuntimeHashSet/WorldPartitionRuntimeHashSet.cpp b/Engine/Source/Runtime/Engine/Private/WorldPartition/RuntimeHashSet/WorldPartitionRuntimeHashSet.cpp
--- a/Engine/Source/Runtime/Engine/Private/WorldPartition/RuntimeHashSet/WorldPartitionRuntimeHashSet.cpp
+++ b/Engine/Source/Runtime/Engine/Private/WorldPartition/RuntimeHashSet/WorldPartitionRuntimeHashSet.cpp
@@ -196,9 +196,14 @@ void UWorldPartitionRuntimeHashSet::FlushStreamingContent()
 	RuntimeStreamingData.Empty();
 	UpdateRuntimeDataGridMap();
 }
-#endif
 
 bool UWorldPartitionRuntimeHashSet::IsValidGrid(FName GridName, const UClass* ActorClass) const
+{
+	return IsValidGridInternal(GridName, ActorClass);
+}
+#endif
+
+bool UWorldPartitionRuntimeHashSet::IsValidGridInternal(FName GridName, const UClass* ActorClass) const
 {
 	TArray<FName> MainPartitionTokens;
 	TArray<FName> HLODPartitionTokens;
@@ -580,7 +585,7 @@ uint32 UWorldPartitionRuntimeHashSet::ComputeUpdateStreamingHash() const
 
 bool UWorldPartitionRuntimeHashSet::SupportsWorldAssetStreaming(const FName& InTargetGrid)
 {
-	return IsValidGrid(InTargetGrid, nullptr);
+	return IsValidGridInternal(InTargetGrid, nullptr);
 }
 
 FGuid UWorldPartitionRuntimeHashSet::RegisterWorldAssetStreaming(const UWorldPartition::FRegisterWorldAssetStreamingParams& InParams)
@@ -595,7 +600,7 @@ FGuid UWorldPartitionRuntimeHashSet::RegisterWorldAssetStreaming(const UWorldPar
 		return FGuid();
 	}
 
-	if (!IsValidGrid(InParams.TargetGrid, nullptr) || (!InParams.TargetGridHLOD.IsNone() && !IsValidGrid(InParams.TargetGridHLOD, nullptr)))
+	if (!IsValidGridInternal(InParams.TargetGrid, nullptr) || (!InParams.TargetGridHLOD.IsNone() && !IsValidGridInternal(InParams.TargetGridHLOD, nullptr)))
 	{
 		return FGuid();
 	}
@@ -699,7 +704,7 @@ FGuid UWorldPartitionRuntimeHashSet::RegisterWorldAssetStreaming(const UWorldPar
 		}		
 		else
 		{
-			UE_LOG(LogWorldPartition, Error, TEXT("Error creating streaming cell %s for world asset %s at %s"), *RuntimeCell->GetName(), *WorldAsset.ToString(), *InParams.Transform.ToString());StreamingData.CreatePartitionsSpatialIndex();
+			UE_LOG(LogWorldPartition, Error, TEXT("Error creating streaming cell %s for world asset %s at %s"), *CellName, *WorldAsset.ToString(), *InParams.Transform.ToString());StreamingData.CreatePartitionsSpatialIndex();
 			return FGuid();
 		}
 	}
diff --git a/Engine/Source/Runtime/Engine/Private/WorldPartition/WorldPartitionRuntimeSpatialHash.cpp b/Engine/Source/Runtime/Engine/Private/WorldPartition/WorldPartitionRuntimeSpatialHash.cpp
--- a/Engine/Source/Runtime/Engine/Private/WorldPartition/WorldPartitionRuntimeSpatialHash.cpp
+++ b/Engine/Source/Runtime/Engine/Private/WorldPartition/WorldPartitionRuntimeSpatialHash.cpp
@@ -2049,7 +2049,7 @@ FGuid UWorldPartitionRuntimeSpatialHash::RegisterWorldAssetStreaming(const UWorl
 		}
 		else
 		{
-			UE_LOG(LogWorldPartition, Error, TEXT("Error creating streaming cell %s for world asset %s at %s"), *StreamingCell->GetName(), *WorldAsset.ToString(), *InParams.Transform.ToString());
+			UE_LOG(LogWorldPartition, Error, TEXT("Error creating streaming cell %s for world asset %s at %s"), *CellName, *WorldAsset.ToString(), *InParams.Transform.ToString());
 			return FGuid();
 		}
 	}
diff --git a/Engine/Source/Runtime/Engine/Public/WorldPartition/RuntimeHashSet/WorldPartitionRuntimeHashSet.h b/Engine/Source/Runtime/Engine/Public/WorldPartition/RuntimeHashSet/WorldPartitionRuntimeHashSet.h
--- a/Engine/Source/Runtime/Engine/Public/WorldPartition/RuntimeHashSet/WorldPartitionRuntimeHashSet.h
+++ b/Engine/Source/Runtime/Engine/Public/WorldPartition/RuntimeHashSet/WorldPartitionRuntimeHashSet.h
@@ -162,10 +162,8 @@ public:
 	ENGINE_API virtual bool SupportsHLODs() const override;
 	ENGINE_API virtual bool SetupHLODActors(const IStreamingGenerationContext* StreamingGenerationContext, const UWorldPartition::FSetupHLODActorsParams& Params) const override;
 	ENGINE_API virtual bool GenerateStreaming(class UWorldPartitionStreamingPolicy* StreamingPolicy, const IStreamingGenerationContext* StreamingGenerationContext, TArray<FString>* OutPackagesToGenerate) override;
-#endif
-	ENGINE_API virtual bool IsValidGrid(FName GridName, const UClass* ActorClass) const;
-#if WITH_EDITOR
-	ENGINE_API virtual bool IsValidHLODLayer(FName GridName, const FSoftObjectPath& HLODLayerPath) const;
+	ENGINE_API virtual bool IsValidGrid(FName GridName, const UClass* ActorClass) const override;
+	ENGINE_API virtual bool IsValidHLODLayer(FName GridName, const FSoftObjectPath& HLODLayerPath) const override;
 	ENGINE_API virtual void DumpStateLog(FHierarchicalLogArchive& Ar) const override;
 #endif
 
@@ -185,7 +183,7 @@ public:
 	ENGINE_API virtual bool RemoveExternalStreamingObject(URuntimeHashExternalStreamingObjectBase* ExternalStreamingObject) override;
 
 	// Streaming interface
-	ENGINE_API virtual void ForEachStreamingCells(TFunctionRef<bool(const UWorldPartitionRuntimeCell*)> Func) const;
+	ENGINE_API virtual void ForEachStreamingCells(TFunctionRef<bool(const UWorldPartitionRuntimeCell*)> Func) const override;
 	ENGINE_API virtual void ForEachStreamingCellsQuery(const FWorldPartitionStreamingQuerySource& QuerySource, TFunctionRef<bool(const UWorldPartitionRuntimeCell*)> Func, FWorldPartitionQueryCache* QueryCache) const override;
 	ENGINE_API virtual void ForEachStreamingCellsSources(const TArray<FWorldPartitionStreamingSource>& Sources, TFunctionRef<bool(const UWorldPartitionRuntimeCell*, EStreamingSourceTargetState)> Func) const override;
 	ENGINE_API virtual uint32 ComputeUpdateStreamingHash() const override;
@@ -219,6 +217,7 @@ private:
 	ENGINE_API void ForEachStreamingData(TFunctionRef<bool(const FRuntimePartitionStreamingData&)> Func) const;
 
 	ENGINE_API void UpdateRuntimeDataGridMap();
+	ENGINE_API bool IsValidGridInternal(FName GridName, const UClass* ActorClass) const;
 
 	/** Array of runtime partition descriptors */
 	UPROPERTY(EditAnywhere, Category = RuntimeSettings, Meta = (TitleProperty = "Name"))
