Initial submission of Material Sample expression prototype for use in layering iteration. Restricted to the Internal only folder, with functionality disabled by an explicit macro to ensure internal projects do not use this node in the short term.

diff --git a/Engine/Source/Runtime/Engine/Internal/Materials/MaterialExpressionMaterialSample.cpp b/Engine/Source/Runtime/Engine/Internal/Materials/MaterialExpressionMaterialSample.cpp
new file mode 100644
--- /dev/null
+++ b/Engine/Source/Runtime/Engine/Internal/Materials/MaterialExpressionMaterialSample.cpp
@@ -0,0 +1,114 @@
+// Copyright Epic Games, Inc. All Rights Reserved.
+
+#include "Materials/MaterialExpressionMaterialSample.h"
+#include "MaterialTypes.h"
+#include "MaterialDomain.h"
+#include "MaterialShared.h"
+#include "MaterialCompiler.h"
+#include "Materials/Material.h"
+
+UMaterialExpressionMaterialSample::UMaterialExpressionMaterialSample(const FObjectInitializer& ObjectInitializer)
+	: Super(ObjectInitializer)
+{
+#if ENABLE_MATERIAL_LAYER_PROTOTYPE
+	bCollapsed = false;
+	bShowOutputNameOnPin = true;
+	Outputs.Empty();
+	if (Substrate::IsSubstrateEnabled())
+	{
+		Outputs.Add(FExpressionOutput(TEXT("Front Material")));
+	}
+#endif
+}
+
+#if WITH_EDITOR 
+#if ENABLE_MATERIAL_LAYER_PROTOTYPE
+int32 UMaterialExpressionMaterialSample::DynamicCompile(FMaterialCompiler* Compiler, int32 OutputIndex, bool bCompilePreview)
+{
+	int32 Result = INDEX_NONE;
+	if (MaterialReference)
+	{
+		if (MaterialReference->GetMaterial() != Material->GetMaterial())
+		{
+			if (UMaterialEditorOnlyData* EditorData = Cast<UMaterialEditorOnlyData>(MaterialReference->GetEditorOnlyData()))
+			{
+				UMaterialExpression* CompileExpression = nullptr;
+				if (Substrate::IsSubstrateEnabled())
+				{
+					CompileExpression = EditorData->FrontMaterial.Expression;
+				}
+				else
+				{
+					Result = Compiler->Errorf(TEXT("Material Sample is only compatible with Substrate materials."));
+				}
+				//else MatAttributes?
+
+				if (CompileExpression)
+				{
+					Result = bCompilePreview ? CompileExpression->CompilePreview(Compiler, OutputIndex) : CompileExpression->Compile(Compiler, OutputIndex);
+				}
+			}
+		}
+		else
+		{
+			Result = Compiler->Errorf(TEXT("Material Sample cannot reference this material or one of it's child instances."));
+		}
+	}
+	return Result;
+}
+
+int32 UMaterialExpressionMaterialSample::Compile(FMaterialCompiler* Compiler, int32 OutputIndex)
+{
+	return DynamicCompile(Compiler, OutputIndex, false);
+}
+
+int32 UMaterialExpressionMaterialSample::CompilePreview(FMaterialCompiler* Compiler, int32 OutputIndex)
+{
+	return DynamicCompile(Compiler, OutputIndex, true);
+}
+
+uint32 UMaterialExpressionMaterialSample::GetOutputType(int32 OutputIndex)
+{
+	if (Substrate::IsSubstrateEnabled())
+	{
+		return MCT_Substrate;
+	}
+	else 
+	{
+		return MCT_MaterialAttributes;
+	}
+}
+
+bool UMaterialExpressionMaterialSample::IsResultSubstrateMaterial(int32 OutputIndex)
+{
+	return Substrate::IsSubstrateEnabled() && MaterialReference != nullptr;
+}
+
+FSubstrateOperator* UMaterialExpressionMaterialSample::SubstrateGenerateMaterialTopologyTree(class FMaterialCompiler* Compiler, class UMaterialExpression* Parent, int32 OutputIndex)
+{
+	if (MaterialReference)
+	{
+		if (MaterialReference->GetMaterial() != Material->GetMaterial())
+		{
+			if (UMaterialEditorOnlyData* EditorData = Cast<UMaterialEditorOnlyData>(MaterialReference->GetEditorOnlyData()))
+			{
+				return EditorData->FrontMaterial.Expression->SubstrateGenerateMaterialTopologyTree(Compiler, Parent, OutputIndex);
+			}
+		}
+	}
+	return nullptr;
+}
+#endif //ENABLE_MATERIAL_LAYER_PROTOTYPE
+
+void UMaterialExpressionMaterialSample::GetCaption(TArray<FString>& OutCaptions) const
+{
+	if (MaterialReference)
+	{
+		OutCaptions.Add(FString(TEXT("MS ")) + MaterialReference->GetName());
+	}
+	else
+	{
+		OutCaptions.Add(TEXT("Material Sample"));
+	}
+}
+#endif // WITH_EDITOR
diff --git a/Engine/Source/Runtime/Engine/Internal/Materials/MaterialExpressionMaterialSample.h b/Engine/Source/Runtime/Engine/Internal/Materials/MaterialExpressionMaterialSample.h
new file mode 100644
--- /dev/null
+++ b/Engine/Source/Runtime/Engine/Internal/Materials/MaterialExpressionMaterialSample.h
@@ -0,0 +1,44 @@
+// Copyright Epic Games, Inc. All Rights Reserved.
+
+/**
+ * MaterialExpressionMaterialSample - an expression which allows a material to reference another's functional output
+ */
+
+#pragma once
+
+#include "CoreMinimal.h"
+#include "Materials/MaterialExpression.h"
+
+#include "MaterialExpressionMaterialSample.generated.h"
+
+#ifndef ENABLE_MATERIAL_LAYER_PROTOTYPE
+//WIP Internal Only - Please do not enable + use this in any projects other than for iterating the new layering prototype system until further notice
+#define ENABLE_MATERIAL_LAYER_PROTOTYPE 0
+#endif
+
+class FMaterialCompiler;
+
+UCLASS(MinimalAPI)
+class UMaterialExpressionMaterialSample: public UMaterialExpression
+{
+	GENERATED_UCLASS_BODY()
+
+	UPROPERTY(EditAnywhere, Category=MaterialExpressionMaterialSample)
+	TObjectPtr<class UMaterialInterface> MaterialReference;
+
+	//~ Begin UMaterialExpression Interface
+#if WITH_EDITOR 
+#if ENABLE_MATERIAL_LAYER_PROTOTYPE
+	ENGINE_API virtual int32 Compile(class FMaterialCompiler* Compiler, int32 OutputIndex) override;
+	ENGINE_API virtual int32 CompilePreview(FMaterialCompiler* Compiler, int32 OutputIndex) override;
+	ENGINE_API virtual uint32 GetOutputType(int32 OutputIndex) override;
+	virtual bool IsResultSubstrateMaterial(int32 OutputIndex) override;
+	virtual FSubstrateOperator * SubstrateGenerateMaterialTopologyTree(class FMaterialCompiler* Compiler, class UMaterialExpression* Parent, int32 OutputIndex) override; 
+
+	int32 DynamicCompile(FMaterialCompiler* Compiler, int32 OutputIndex, bool bCompilePreview);
+#endif
+	ENGINE_API virtual void GetCaption(TArray<FString>& OutCaptions) const override;
+#endif
+
+	//~ End UMaterialExpression Interface
+};
\ No newline at end of file
