Keep track of when a key state is flushed on the Player input functions. This can happen during several situations:

1. During a level transition
2. You manually call it in C++
3. You change the "input mode" of your player controller
4. The viewport focus is lost (i.e. you alt+Tab out of a window or otherwise click out of the game viewport window)
... probably several more ways as well.

When any of these things happen, if the user keeps holding down the key that has been flushed, then the Slate Application will send an IE_Repeat event, since it is what the actual operating system is reporting (and is correct). This means that your key state would get a repeat event before a "pressed" event. This puts the FKeyState in an incorrect state, and results in the improper evaluation of a "pressed" and "released" event firing every single frame, instead of a "repeat" event.

This has the side effect that any "timed" triggers in Enhanced Input (hold and release, hold, etc) would fire several delegates each frame because their state is basically re-evaluating the whole trigger set again, but it is correct because the starting time passes the actuation threshold that is specified.

This behavior is actually the same as we saw with UE-195988, and is behind the same CVar "Input.bAutoReconcilePressedEventsOnFirstRepeat" if you need to toggle it off. This is on by default, and this CVar will likely be removed in a future unless there are reports of broken behavior with it, as it is the "correct" behavior for the input system here from the engine standpoint.

Phew, das a lot of words... if you read this far... you get a cookie?

#pf 6631451f882a9e6121af7f8b

diff --git a/Engine/Source/Runtime/Engine/Private/UserInterface/PlayerInput.cpp b/Engine/Source/Runtime/Engine/Private/UserInterface/PlayerInput.cpp
--- a/Engine/Source/Runtime/Engine/Private/UserInterface/PlayerInput.cpp
+++ b/Engine/Source/Runtime/Engine/Private/UserInterface/PlayerInput.cpp
@@ -158,6 +158,8 @@ void UPlayerInput::FlushPressedKeys()
 		KeyState.bDown = false;
 		KeyState.bDownPrevious = false;
 		KeyState.LastUpDownTransitionTime = TimeSeconds;
+		// Flag each key as having been flushed this frame so that we can correctly evaluate it next time
+		KeyState.bWasJustFlushed = true;
 	}
 }
 
@@ -345,11 +347,15 @@ bool UPlayerInput::InputKey(const FInputKeyParams& Params)
 		UWorld* World = GetWorld();
 		check(World);
 
-		const bool bIsFirstEventForKey = (ExistingKeyState == nullptr);
+		const bool bIsFirstEventForKey = (ExistingKeyState == nullptr) || KeyState.bWasJustFlushed;
+
+		const float WorldRealTimeSeconds = World->GetRealTimeSeconds();
 
 		// If this is the first key press for us and it is a repeat, then we have missed the initial IE_Pressed event.
 		// This can happen if you are holding down a key between level transitions and the player controller gets recreated,
 		// which means that we will be using a new UPlayerInput object and the KeyState map is emptied.
+		// This can ALSO happen if "FlushPressedKeys" gets called, like when we change player controller input modes, and you keep holding down the key 
+		// in between those transitions
 		if (UE::Input::bAutoReconcilePressedEventsOnFirstRepeat && bIsFirstEventForKey && Params.Event == IE_Repeat && KeyState.EventAccumulator[IE_Pressed].IsEmpty())
 		{
 			// Mark as having received the IE_Pressed event already, so that we can correctly evaluate the 
@@ -359,8 +365,12 @@ bool UPlayerInput::InputKey(const FInputKeyParams& Params)
 			// every frame, even if you are just holding it
 			KeyState.RawValueAccumulator.X = Params.Delta.X;
 			KeyState.EventAccumulator[IE_Pressed].Add(++EventCount);
-			KeyState.LastUpDownTransitionTime = World->GetRealTimeSeconds();
+			KeyState.LastUpDownTransitionTime = WorldRealTimeSeconds;
 			KeyState.SampleCountAccumulator++;
+
+			// We can assume that if we are getting a "Repeat event" that this key was already down the previous frame and it is down this frame
+			KeyState.bDown = true;
+			KeyState.bDownPrevious = true;
 		}
 
 		switch(Params.Event)
@@ -373,7 +383,6 @@ bool UPlayerInput::InputKey(const FInputKeyParams& Params)
 			{
 				// check for doubleclick
 				// note, a tripleclick will currently count as a 2nd double click.
-				const float WorldRealTimeSeconds = World->GetRealTimeSeconds();
 				if ((WorldRealTimeSeconds - KeyState.LastUpDownTransitionTime) < GetDefault<UInputSettings>()->DoubleClickTime)
 				{
 					KeyState.EventAccumulator[IE_DoubleClick].Add(++EventCount);
@@ -410,6 +419,9 @@ bool UPlayerInput::InputKey(const FInputKeyParams& Params)
 			return IsKeyHandledByAction( Params.Key);
 		}
 
+		// We have now processed this key's state, so we can clear its "just flushed" flag and treat it normally
+		KeyState.bWasJustFlushed = false;
+
 		return true;
 	}
 }
diff --git a/Engine/Source/Runtime/Engine/Public/KeyState.h b/Engine/Source/Runtime/Engine/Public/KeyState.h
--- a/Engine/Source/Runtime/Engine/Public/KeyState.h
+++ b/Engine/Source/Runtime/Engine/Public/KeyState.h
@@ -29,6 +29,12 @@ struct FKeyState
 	/** True if this key has been "consumed" by an InputComponent and should be ignored for further components during this update. */
 	uint8 bConsumed:1;
 
+	/** 
+	* True if this key was flushed this frame. This is used to flag to input processing that we may be receiving 
+	* a IE_Repeat event, but the IE_Pressed event accumulator may have been reset. 
+	*/
+	uint8 bWasJustFlushed:1;
+
 	/** Flag paired axes that have been sampled this tick. X = LSB, Z = MSB */
 	uint8 PairSampledAxes : 3;
 
@@ -51,6 +57,7 @@ struct FKeyState
 		, bDown(false)
 		, bDownPrevious(false)
 		, bConsumed(false)
+		, bWasJustFlushed(false)
 		, PairSampledAxes(0)
 		, SampleCountAccumulator(0)
 		, RawValueAccumulator(0.f, 0.f, 0.f)
