HLSLMaterialTranslator: Fix crash when UMaterialCollection referred to in an FEnvironmentDefines DDC entry has been redirected and an ObjectRedirector left behind.
#rnx

diff --git a/Engine/Source/Runtime/CoreUObject/Private/Serialization/ArchiveUObject.cpp b/Engine/Source/Runtime/CoreUObject/Private/Serialization/ArchiveUObject.cpp
--- a/Engine/Source/Runtime/CoreUObject/Private/Serialization/ArchiveUObject.cpp
+++ b/Engine/Source/Runtime/CoreUObject/Private/Serialization/ArchiveUObject.cpp
@@ -4,6 +4,7 @@
 #include "UObject/UObjectGlobals.h"
 #include "UObject/Object.h"
 #include "Serialization/SerializedPropertyScope.h"
+#include "UObject/ObjectRedirector.h"
 #include "UObject/UnrealType.h"
 #include "Serialization/ObjectAndNameAsStringProxyArchive.h"
 #include "Serialization/ArchiveReplaceObjectRef.h"
@@ -178,6 +179,32 @@ FArchive& FObjectAndNameAsStringProxyArchive::operator<<(UObject*& Obj)
 		{
 			Obj = LoadObject<UObject>(nullptr, *LoadedString);
 		}
+
+		if (bResolveRedirectors && Obj && Obj->IsA(UObjectRedirector::StaticClass()))
+		{
+			int32 Count = 0;
+			constexpr int32 MaxCountBeforeUseSet = 5;
+			while (Obj && Obj->IsA(UObjectRedirector::StaticClass()) && ++Count <= MaxCountBeforeUseSet)
+			{
+				Obj = static_cast<UObjectRedirector*>(Obj)->DestinationObject;
+			}
+			if (Obj && Obj->IsA(UObjectRedirector::StaticClass()))
+			{
+				TSet<UObject*> Seen;
+				Seen.Add(Obj);
+				while (Obj && Obj->IsA(UObjectRedirector::StaticClass()))
+				{
+					Obj = static_cast<UObjectRedirector*>(Obj)->DestinationObject;
+					bool bExists;
+					Seen.Add(Obj, &bExists);
+					if (bExists)
+					{
+						// Cycle, return null
+						Obj = nullptr;
+					}
+				}
+			}
+		}
 	}
 	else if (Obj)
 	{
diff --git a/Engine/Source/Runtime/CoreUObject/Public/Serialization/ObjectAndNameAsStringProxyArchive.h b/Engine/Source/Runtime/CoreUObject/Public/Serialization/ObjectAndNameAsStringProxyArchive.h
--- a/Engine/Source/Runtime/CoreUObject/Public/Serialization/ObjectAndNameAsStringProxyArchive.h
+++ b/Engine/Source/Runtime/CoreUObject/Public/Serialization/ObjectAndNameAsStringProxyArchive.h
@@ -33,6 +33,15 @@ struct FObjectAndNameAsStringProxyArchive : public FNameAsStringProxyArchive
 
 	/** If we fail to find an object during loading, try and load it. */
 	bool bLoadIfFindFails;
+	/**
+	 * If bResolveRedirectors is true, when loading, in operator<< functions that return a resolved object, 
+	 * (UObject*, FWeakObjectPtr, FObjectPtr if resolved), if a UObject is a UObjectRedirector, the UObjectRedirector
+	 * will be followed and the output Obj will receive the target of the redirector. A chain of redirectors will
+	 * return the object at the end of the chain; null will be returned if the chain has a cycle or ends in a null
+	 * target.
+	 * If false, any UObjectRedirectors will be left unresolved and returned in the output Obj.
+	 */
+	bool bResolveRedirectors = false;
 
 	COREUOBJECT_API virtual FArchive& operator<<(UObject*& Obj) override;
 	COREUOBJECT_API virtual FArchive& operator<<(FWeakObjectPtr& Obj) override;
diff --git a/Engine/Source/Runtime/Engine/Private/Materials/HLSLMaterialTranslator.cpp b/Engine/Source/Runtime/Engine/Private/Materials/HLSLMaterialTranslator.cpp
--- a/Engine/Source/Runtime/Engine/Private/Materials/HLSLMaterialTranslator.cpp
+++ b/Engine/Source/Runtime/Engine/Private/Materials/HLSLMaterialTranslator.cpp
@@ -44,6 +44,7 @@
 #include "Stats/StatsMisc.h"
 #include "Stats/StatsTrace.h"
 #include "SubstrateDefinitions.h"
+#include "UObject/ObjectRedirector.h"
 #include "VT/RuntimeVirtualTexture.h"
 #include <memory>
 #include <tuple>
@@ -410,7 +411,32 @@ struct FHLSLMaterialTranslator::FEnvironmentDefines
 		Ar << bSubstrateComplexSpecialPath;
 		Ar << bTextureSampleDebug;
 		Ar << SubstrateDefines;
-		Ar << ParameterCollections;
+
+		TArray<UObject*> ParameterCollectionObjects;
+		if (!Ar.IsLoading())
+		{
+			ParameterCollectionObjects.Reserve(ParameterCollections.Num());
+			for (const TObjectPtr<UMaterialParameterCollection>& Object : ParameterCollections)
+			{
+				ParameterCollectionObjects.Add(Object);
+			}
+		}
+		Ar << ParameterCollectionObjects;
+		if (Ar.IsLoading())
+		{
+			ParameterCollections.Reset(ParameterCollectionObjects.Num());
+			for (UObject* Object : ParameterCollectionObjects)
+			{
+				UMaterialParameterCollection* Collection = Cast<UMaterialParameterCollection>(Object);
+				if (!Collection && Object)
+				{
+					UE_LOG(LogMaterial, Error,
+						TEXT("FEnvironmentDefines.ParameterCollections had object %s which is not a UMaterialParameterCollection; setting it to null."),
+						*Object->GetFullName());
+				}
+				ParameterCollections.Add(Collection);
+			}
+		}
 	}
 };
 
@@ -1215,6 +1241,7 @@ bool FHLSLMaterialTranslator::Translate()
 		// Serialize the environment defines from the buffer retrieved from the DDC.
 		FMemoryReaderView EnvironmentDefinesBufferReader{ TArrayView<uint8>{ (uint8*)EnvironmentDefinesBuffer.GetData(), (int)EnvironmentDefinesBuffer.GetSize() } };
 		FObjectAndNameAsStringProxyArchive EnvironmentDefinesBufferReaderProxy{ EnvironmentDefinesBufferReader, true };
+		EnvironmentDefinesBufferReaderProxy.bResolveRedirectors = true;
 		EnvironmentDefines->Serialize(EnvironmentDefinesBufferReaderProxy);
 
 		MaterialCompilationOutput = DDCMaterialCompilationOutput;
