- Support Initial Runtime State for client-only/server-only Data Layers
- This implicitely allows support of client-only/server-only External Data Layers

diff --git a/Engine/Source/Editor/DataLayerEditor/Private/DataLayer/DataLayerInstanceCustomization.cpp b/Engine/Source/Editor/DataLayerEditor/Private/DataLayer/DataLayerInstanceCustomization.cpp
--- a/Engine/Source/Editor/DataLayerEditor/Private/DataLayer/DataLayerInstanceCustomization.cpp
+++ b/Engine/Source/Editor/DataLayerEditor/Private/DataLayer/DataLayerInstanceCustomization.cpp
@@ -30,20 +30,20 @@ void FDataLayerInstanceDetails::CustomizeDetails(IDetailLayoutBuilder& DetailBui
 	DetailBuilder.GetObjectsBeingCustomized(ObjectsBeingCustomized);
 	
 	uint32 CustomizedDataLayerInstanceWithAssetCount = 0;
-	bool bHasRuntime = false;
+	bool bHasInitialRuntimeState = false;
 	for (const TWeakObjectPtr<UObject>& SelectedObject : ObjectsBeingCustomized)
 	{
 		UDataLayerInstance* DataLayerInstance = Cast<UDataLayerInstance>(SelectedObject.Get());
-		if (DataLayerInstance && DataLayerInstance->IsRuntime() && !DataLayerInstance->IsClientOnly() && !DataLayerInstance->IsServerOnly())
+		if (DataLayerInstance && DataLayerInstance->IsRuntime())
 		{
-			bHasRuntime = true;
+			bHasInitialRuntimeState = true;
 		}
 		if (Cast<UDataLayerInstanceWithAsset>(SelectedObject.Get()))
 		{
 			++CustomizedDataLayerInstanceWithAssetCount;
 		}
 	}
-	if (!bHasRuntime)
+	if (!bHasInitialRuntimeState)
 	{
 		DetailBuilder.HideProperty(GET_MEMBER_NAME_CHECKED(UDataLayerInstance, InitialRuntimeState));
 	}
diff --git a/Engine/Source/Runtime/Engine/Private/WorldPartition/DataLayer/DataLayerInstance.cpp b/Engine/Source/Runtime/Engine/Private/WorldPartition/DataLayer/DataLayerInstance.cpp
--- a/Engine/Source/Runtime/Engine/Private/WorldPartition/DataLayer/DataLayerInstance.cpp
+++ b/Engine/Source/Runtime/Engine/Private/WorldPartition/DataLayer/DataLayerInstance.cpp
@@ -237,7 +237,7 @@ bool UDataLayerInstance::CanEditChange(const FProperty* InProperty) const
 
 	if (InProperty->GetFName() == GET_MEMBER_NAME_CHECKED(UDataLayerInstance, InitialRuntimeState))
 	{
-		if (!IsRuntime() || IsClientOnly() || IsServerOnly())
+		if (!IsRuntime())
 		{
 			return false;
 		}
diff --git a/Engine/Source/Runtime/Engine/Private/WorldPartition/DataLayer/WorldDataLayers.cpp b/Engine/Source/Runtime/Engine/Private/WorldPartition/DataLayer/WorldDataLayers.cpp
--- a/Engine/Source/Runtime/Engine/Private/WorldPartition/DataLayer/WorldDataLayers.cpp
+++ b/Engine/Source/Runtime/Engine/Private/WorldPartition/DataLayer/WorldDataLayers.cpp
@@ -128,18 +128,20 @@ void AWorldDataLayers::InitializeDataLayerRuntimeStates()
 		{
 			const bool bDataLayerClientOnly = DataLayerInstance->IsClientOnly();
 			const bool bDataLayerServerOnly = DataLayerInstance->IsServerOnly();
+			const bool bIsLocalDataLayer = bDataLayerClientOnly || bDataLayerServerOnly;
+			TSet<FName>& TargetLoadedDataLayerNames = bIsLocalDataLayer ? LocalLoadedDataLayerNames : LoadedDataLayerNames;
+			TSet<FName>& TargetActiveDataLayerNames = bIsLocalDataLayer ? LocalActiveDataLayerNames : ActiveDataLayerNames;
 
-			if (DataLayerInstance->IsRuntime() && !bDataLayerClientOnly && !bDataLayerServerOnly)
+			if (DataLayerInstance->IsRuntime())
 			{
-				if (DataLayerInstance->GetInitialRuntimeState() == EDataLayerRuntimeState::Activated)
+				if (DataLayerInstance->GetInitialRuntimeState() == EDataLayerRuntimeState::Loaded)
 				{
-					ActiveDataLayerNames.Add(DataLayerInstance->GetDataLayerFName());
+					TargetLoadedDataLayerNames.Add(DataLayerInstance->GetDataLayerFName());
 				}
-				else if (DataLayerInstance->GetInitialRuntimeState() == EDataLayerRuntimeState::Loaded)
+				else if (DataLayerInstance->GetInitialRuntimeState() == EDataLayerRuntimeState::Activated)
 				{
-					LoadedDataLayerNames.Add(DataLayerInstance->GetDataLayerFName());
+					TargetActiveDataLayerNames.Add(DataLayerInstance->GetDataLayerFName());
 				}
-
 				RuntimeDataLayerInstances.Add(DataLayerInstance);
 			}
 			return true;
@@ -192,6 +194,10 @@ void AWorldDataLayers::SetDataLayerRuntimeState(const UDataLayerInstance* InData
 	const ENetMode NetMode = GetNetMode();
 	const bool bDataLayerClientOnly = InDataLayerInstance->IsClientOnly();
 	const bool bDataLayerServerOnly = InDataLayerInstance->IsServerOnly();
+	const bool bIsLocalDataLayer = bDataLayerClientOnly || bDataLayerServerOnly;
+	TSet<FName>& TargetLoadedDataLayerNames = bIsLocalDataLayer ? LocalLoadedDataLayerNames : LoadedDataLayerNames;
+	TSet<FName>& TargetActiveDataLayerNames = bIsLocalDataLayer ? LocalActiveDataLayerNames : ActiveDataLayerNames;
+
 	const EDataLayerRuntimeState CurrentState = GetDataLayerRuntimeStateByName(InDataLayerInstance->GetDataLayerFName());
 
 	if (bDataLayerClientOnly)
@@ -245,36 +251,20 @@ void AWorldDataLayers::SetDataLayerRuntimeState(const UDataLayerInstance* InData
 			PRAGMA_ENABLE_DEPRECATION_WARNINGS
 		}
 
-		LoadedDataLayerNames.Remove(InDataLayerInstance->GetDataLayerFName());
-		ActiveDataLayerNames.Remove(InDataLayerInstance->GetDataLayerFName());
-		LocalLoadedDataLayerNames.Remove(InDataLayerInstance->GetDataLayerFName());
-		LocalActiveDataLayerNames.Remove(InDataLayerInstance->GetDataLayerFName());
+		TargetLoadedDataLayerNames.Remove(InDataLayerInstance->GetDataLayerFName());
+		TargetActiveDataLayerNames.Remove(InDataLayerInstance->GetDataLayerFName());
 
 		if (InState == EDataLayerRuntimeState::Loaded)
 		{
-			if (bDataLayerClientOnly || bDataLayerServerOnly)
-			{
-				LocalLoadedDataLayerNames.Add(InDataLayerInstance->GetDataLayerFName());
-			}
-			else
-			{
-				LoadedDataLayerNames.Add(InDataLayerInstance->GetDataLayerFName());
-			}
+			TargetLoadedDataLayerNames.Add(InDataLayerInstance->GetDataLayerFName());
 		}
 		else if (InState == EDataLayerRuntimeState::Activated)
 		{
-			if (bDataLayerClientOnly || bDataLayerServerOnly)
-			{
-				LocalActiveDataLayerNames.Add(InDataLayerInstance->GetDataLayerFName());
-			}
-			else
-			{
-				ActiveDataLayerNames.Add(InDataLayerInstance->GetDataLayerFName());
-			}
+			TargetActiveDataLayerNames.Add(InDataLayerInstance->GetDataLayerFName());
 		}
 
 		// Update replicated properties
-		if (NetMode == NM_DedicatedServer || NetMode == NM_ListenServer)
+		if (!bIsLocalDataLayer && (NetMode == NM_DedicatedServer || NetMode == NM_ListenServer))
 		{
 			FlushNetDormancy();
 			AWORLDDATALAYERS_UPDATE_REPLICATED_DATALAYERS(RepActiveDataLayerNames, ActiveDataLayerNames.Array());
@@ -419,6 +409,7 @@ void AWorldDataLayers::ResolveEffectiveRuntimeState(const UDataLayerInstance* In
 	const ENetMode NetMode = GetNetMode();
 	const bool bDataLayerClientOnly = InDataLayerInstance->IsClientOnly();
 	const bool bDataLayerServerOnly = InDataLayerInstance->IsServerOnly();
+	const bool bIsLocalDataLayer = (bDataLayerClientOnly || bDataLayerServerOnly);
 	const FName DataLayerName = InDataLayerInstance->GetDataLayerFName();
 	EDataLayerRuntimeState CurrentEffectiveRuntimeState = GetDataLayerEffectiveRuntimeStateByName(DataLayerName);
 	EDataLayerRuntimeState NewEffectiveRuntimeState = GetDataLayerRuntimeStateByName(DataLayerName);
@@ -436,36 +427,23 @@ void AWorldDataLayers::ResolveEffectiveRuntimeState(const UDataLayerInstance* In
 
 	if (CurrentEffectiveRuntimeState != NewEffectiveRuntimeState)
 	{
-		EffectiveLoadedDataLayerNames.Remove(DataLayerName);
-		EffectiveActiveDataLayerNames.Remove(DataLayerName);
-		LocalEffectiveLoadedDataLayerNames.Remove(DataLayerName);
-		LocalEffectiveActiveDataLayerNames.Remove(DataLayerName);
+		TSet<FName>& TargetEffectiveLoadedDataLayerNames = bIsLocalDataLayer ? LocalEffectiveLoadedDataLayerNames : EffectiveLoadedDataLayerNames;
+		TSet<FName>& TargetEffectiveActiveDataLayerNames = bIsLocalDataLayer ? LocalEffectiveActiveDataLayerNames : EffectiveActiveDataLayerNames;
+
+		TargetEffectiveLoadedDataLayerNames.Remove(DataLayerName);
+		TargetEffectiveActiveDataLayerNames.Remove(DataLayerName);
 
 		if (NewEffectiveRuntimeState == EDataLayerRuntimeState::Loaded)
 		{
-			if (bDataLayerClientOnly || bDataLayerServerOnly)
-			{
-				LocalEffectiveLoadedDataLayerNames.Add(DataLayerName);
-			}
-			else
-			{
-				EffectiveLoadedDataLayerNames.Add(DataLayerName);
-			}
+			TargetEffectiveLoadedDataLayerNames.Add(DataLayerName);
 		}
 		else if (NewEffectiveRuntimeState == EDataLayerRuntimeState::Activated)
 		{
-			if (bDataLayerClientOnly || bDataLayerServerOnly)
-			{
-				LocalEffectiveActiveDataLayerNames.Add(DataLayerName);
-			}
-			else
-			{
-				EffectiveActiveDataLayerNames.Add(DataLayerName);
-			}
+			TargetEffectiveActiveDataLayerNames.Add(DataLayerName);
 		}
 
 		// Update Replicated Properties
-		if (NetMode == NM_DedicatedServer || NetMode == NM_ListenServer)
+		if (!bIsLocalDataLayer && (NetMode == NM_DedicatedServer || NetMode == NM_ListenServer))
 		{
 			FlushNetDormancy();
 			AWORLDDATALAYERS_UPDATE_REPLICATED_DATALAYERS(RepEffectiveActiveDataLayerNames, EffectiveActiveDataLayerNames.Array());
diff --git a/Engine/Source/Runtime/Engine/Public/WorldPartition/DataLayer/DataLayerInstance.h b/Engine/Source/Runtime/Engine/Public/WorldPartition/DataLayer/DataLayerInstance.h
--- a/Engine/Source/Runtime/Engine/Public/WorldPartition/DataLayer/DataLayerInstance.h
+++ b/Engine/Source/Runtime/Engine/Public/WorldPartition/DataLayer/DataLayerInstance.h
@@ -170,7 +170,7 @@ public:
 	virtual FColor GetDebugColor() const { return FColor::Black; }
 
 	UFUNCTION(Category = "Data Layer|Runtime", BlueprintCallable)
-	EDataLayerRuntimeState GetInitialRuntimeState() const { return IsRuntime() && !IsClientOnly() && !IsServerOnly() ? InitialRuntimeState : EDataLayerRuntimeState::Unloaded; }
+	EDataLayerRuntimeState GetInitialRuntimeState() const { return IsRuntime() ? InitialRuntimeState : EDataLayerRuntimeState::Unloaded; }
 
 	UFUNCTION(Category = "Data Layer", BlueprintCallable)
 	virtual FString GetDataLayerShortName() const { return TEXT("Invalid Data Layer"); }
