Holdout: Enforce BeforeDOF translucency pass for alpha holdout blend mode.

This is done primarily to prevent confusion by forcing a non-separate translucency pass. It also improves consistency between regular renders and scene captures (which disable separate translucency), for which the feature was initially added, and brings the behavior closer to primitive alpha holdout.

diff --git a/Engine/Source/Runtime/Engine/Private/Materials/HLSLMaterialTranslator.cpp b/Engine/Source/Runtime/Engine/Private/Materials/HLSLMaterialTranslator.cpp
--- a/Engine/Source/Runtime/Engine/Private/Materials/HLSLMaterialTranslator.cpp
+++ b/Engine/Source/Runtime/Engine/Private/Materials/HLSLMaterialTranslator.cpp
@@ -1820,9 +1820,18 @@ void FHLSLMaterialTranslator::DoTranslate()
 		}
 	}
 
-	if (IsAlphaHoldoutBlendMode(BlendMode) && !MaterialShadingModels.IsUnlit())
+	if (IsAlphaHoldoutBlendMode(BlendMode))
 	{
-		Errorf(TEXT("Alpha Holdout blend mode must use unlit shading model."));
+		if (!MaterialShadingModels.IsUnlit())
+		{
+			Errorf(TEXT("Alpha Holdout blend mode must use unlit shading model."));
+		}
+
+		if (Material->IsTranslucencyAfterDOFEnabled() || Material->IsTranslucencyAfterMotionBlurEnabled())
+		{
+			// We must write into the alpha channel of the SceneColor now for later translucents rendered after DOF to be able to be affected by the alpha.
+			Errorf(TEXT("Alpha Holdout blend mode must use BeforeDOF under Translucency / Advanced / Translucency Pass.")); 
+		}
 	}
 
 	if (Domain == MD_Volume && BlendMode != BLEND_Additive)
