DirectionalLight optimization: Add a setter for AtmosphereSunDiskColorScale  on the proxy so we don't need to fully recreate the renderstate each time it changes. This avoids continual 20ms spikes on the render thread
#tests Tested on XB1

diff --git a/Engine/Source/Runtime/Engine/Private/Components/DirectionalLightComponent.cpp b/Engine/Source/Runtime/Engine/Private/Components/DirectionalLightComponent.cpp
--- a/Engine/Source/Runtime/Engine/Private/Components/DirectionalLightComponent.cpp
+++ b/Engine/Source/Runtime/Engine/Private/Components/DirectionalLightComponent.cpp
@@ -384,6 +384,17 @@ public:
 			});
 	}
 
+	void UpdateAtmosphereSunDiskColorScale_GameThread(const UDirectionalLightComponent* Component)
+	{
+		FLinearColor NewAtmosphereSunDiskColorScale = Component->AtmosphereSunDiskColorScale;
+		FDirectionalLightSceneProxy* Proxy = this;
+		ENQUEUE_RENDER_COMMAND(FUpdateLightShaftOverrideDirectionCommand)(
+			[Proxy, NewAtmosphereSunDiskColorScale](FRHICommandList& RHICmdList)
+			{
+				Proxy->AtmosphereSunDiskColorScale = NewAtmosphereSunDiskColorScale;
+			});
+	}
+
 	/** Accesses parameters needed for rendering the light. */
 	virtual void GetLightShaderParameters(FLightRenderParameters& LightParameters, uint32 Flags=0) const override
 	{
@@ -1375,7 +1386,11 @@ void UDirectionalLightComponent::SetAtmosphereSunDiskColorScale(FLinearColor New
 		&& AtmosphereSunDiskColorScale != NewValue)
 	{
 		AtmosphereSunDiskColorScale = NewValue;
-		MarkRenderStateDirty();
+		if (SceneProxy)
+		{
+			FDirectionalLightSceneProxy* DirectionalLightSceneProxy = (FDirectionalLightSceneProxy*)SceneProxy;
+			DirectionalLightSceneProxy->UpdateAtmosphereSunDiskColorScale_GameThread(this);
+		}
 	}
 }
 
diff --git a/Engine/Source/Runtime/Engine/Public/LightSceneProxy.h b/Engine/Source/Runtime/Engine/Public/LightSceneProxy.h
--- a/Engine/Source/Runtime/Engine/Public/LightSceneProxy.h
+++ b/Engine/Source/Runtime/Engine/Public/LightSceneProxy.h
@@ -451,7 +451,7 @@ protected:
 	/** The index of the atmospheric light. Multiple lights can be considered when computing the sky/atmospheric scattering. */
 	const uint8 AtmosphereSunLightIndex;
 
-	const FLinearColor AtmosphereSunDiskColorScale;
+	FLinearColor AtmosphereSunDiskColorScale;
 
 	/** The light type (ELightComponentType) */
 	const uint8 LightType;
