Adding cook tags for sound wave format selection

diff --git a/Engine/Source/Runtime/Engine/Private/AudioDerivedData.cpp b/Engine/Source/Runtime/Engine/Private/AudioDerivedData.cpp
--- a/Engine/Source/Runtime/Engine/Private/AudioDerivedData.cpp
+++ b/Engine/Source/Runtime/Engine/Private/AudioDerivedData.cpp
@@ -30,6 +30,10 @@
 #include "Serialization/MemoryWriter.h"
 #endif
 
+#if WITH_EDITOR
+#include "UObject/ArchiveCookContext.h"
+#endif
+
 DEFINE_LOG_CATEGORY_STATIC(LogAudioDerivedData, Log, All);
 
 static int32 AllowAsyncCompression = 1;
@@ -1504,6 +1508,12 @@ void FStreamedAudioPlatformData::Serialize(FArchive& Ar, USoundWave* Owner)
 	Ar << NumChunks;
 	Ar << AudioFormat;
 
+	if (Ar.IsCooking() && Ar.IsSaving() && Ar.GetCookContext() && Ar.GetCookContext()->GetCookTagList())
+	{
+		FCookTagList* CookTags = Ar.GetCookContext()->GetCookTagList();
+		CookTags->Add(Owner, "StreamingFormat", LexToString(AudioFormat));
+	}
+
 	if (Ar.IsLoading())
 	{
 		check(!AudioFormat.IsNone());
diff --git a/Engine/Source/Runtime/Engine/Private/SoundWave.cpp b/Engine/Source/Runtime/Engine/Private/SoundWave.cpp
--- a/Engine/Source/Runtime/Engine/Private/SoundWave.cpp
+++ b/Engine/Source/Runtime/Engine/Private/SoundWave.cpp
@@ -38,6 +38,10 @@
 #include "uewav.h"
 #endif
 
+#if WITH_EDITOR
+#include "UObject/ArchiveCookContext.h"
+#endif
+
 static int32 SoundWaveDefaultLoadingBehaviorCVar = static_cast<int32>(ESoundWaveLoadingBehavior::LoadOnDemand);
 FAutoConsoleVariableRef CVarSoundWaveDefaultLoadingBehavior(
 	TEXT("au.streamcache.SoundWaveDefaultLoadingBehavior"),
@@ -218,6 +222,28 @@ void FSoundWaveData::OverrideRuntimeFormat(const FName& InRuntimeFormat)
 	RuntimeFormat = InRuntimeFormat;
 }
 
+static FName GetBaseFormatFromSuffixedFormat(const FName& InFormat)
+{
+	// If the format contains a delimiter, it's using a platform cook overrides with a suffix.
+	// Return up to the first delimiter e.g. "ADPCM_some_suffix" -> "ADPCM"
+	
+	// It's not great that we store the entire suffixed format in FName as this heavily pollutes
+	// the name table and never gets freed.
+	// e.g. suffixed name: BINKA_SCVER_5028_R4DV_false_SR0_48000.000000_SR1_32000.000000_SR2_24000.000000_SR3_12000.000000_SR4_8000.000000_QMOD_1.000000_CQLT_-1_ASTH_0.000000_INLC_false_LCK1_0.000000_CSZE_65536_LCF_false_ZCS_256_MCSO_-1_END
+	TStringBuilder<FName::StringBufferSize> Name;
+	Name << InFormat;
+
+	const TCHAR* DelimAt = FCString::Strchr(Name.ToString(), TEXT('_'));
+	if (DelimAt)
+	{
+		TStringView BaseFormat(Name.ToString(), DelimAt - Name.ToString());
+
+		// We know this FName already exists in the name table because it started out as one before the suffixes.
+		return FName(BaseFormat);
+	}
+	return InFormat;
+}
+
 FName FSoundWaveData::FindRuntimeFormat(const USoundWave& InWave) const
 {		
 #if WITH_EDITOR	
@@ -248,21 +274,8 @@ FName FSoundWaveData::FindRuntimeFormat(const USoundWave& InWave) const
 
 		// Only one format is supported currently.
 		if (ensureMsgf(ContainedFormats.Num() == 1, TEXT("ContainedFormats::Num()=%d"), ContainedFormats.Num()))
-		{			
-			FName CompressedFormat = ContainedFormats[0];
-			FString CompressedFormatString = CompressedFormat.GetPlainNameString();
-
-			// If the format contains a delimiter, it's using a platform cook overrides with a suffix.
-			// Return up to the first delimiter e.g. "ADPCM_some_suffix" -> "ADPCM"
-			int32 DelimIndex = 0;
-			if (CompressedFormatString.FindChar('_', DelimIndex))
-			{
-				return *CompressedFormatString.Left(DelimIndex);
-			}
-			else
-			{
-				return CompressedFormat;
-			}
+		{
+			return GetBaseFormatFromSuffixedFormat(ContainedFormats[0]);
 		}
 		
 		// Fail.
@@ -1217,6 +1230,26 @@ void USoundWave::Serialize( FArchive& Ar )
 
 	SerializeCuePoints(Ar, bCooked && Ar.IsLoading());
 
+	// Save cook tags
+	if (Ar.IsCooking() && Ar.IsSaving() && Ar.GetCookContext() && Ar.GetCookContext()->GetCookTagList())
+	{
+		FCookTagList* CookTags = Ar.GetCookContext()->GetCookTagList();
+		TArray<FName> CompressedFormats;
+		SoundWaveDataPtr->CompressedFormatData.GetContainedFormats(CompressedFormats);
+		CookTags->Add(this, "FormatContainerCount", LexToString(CompressedFormats.Num()));
+		TStringBuilder<256> FormatsListStr;
+		for (int32 FormatIndex = 0; FormatIndex < CompressedFormats.Num(); FormatIndex++)
+		{
+			FName BaseFormat = GetBaseFormatFromSuffixedFormat(CompressedFormats[FormatIndex]);
+			FormatsListStr << BaseFormat;
+			if (FormatIndex != CompressedFormats.Num() - 1)
+			{
+				FormatsListStr << TEXT(',');
+			}
+		}
+		CookTags->Add(this, "FormatContainerFormats", FormatsListStr.ToString());
+	}
+
 	if (bCooked)
 	{
 #if WITH_EDITOR
