Fix for non editor build break

diff --git a/Engine/Source/Runtime/Engine/Private/AudioDerivedData.cpp b/Engine/Source/Runtime/Engine/Private/AudioDerivedData.cpp
--- a/Engine/Source/Runtime/Engine/Private/AudioDerivedData.cpp
+++ b/Engine/Source/Runtime/Engine/Private/AudioDerivedData.cpp
@@ -1508,11 +1508,13 @@ void FStreamedAudioPlatformData::Serialize(FArchive& Ar, USoundWave* Owner)
 	Ar << NumChunks;
 	Ar << AudioFormat;
 
+#if WITH_EDITOR
 	if (Ar.IsCooking() && Ar.IsSaving() && Ar.GetCookContext() && Ar.GetCookContext()->GetCookTagList())
 	{
 		FCookTagList* CookTags = Ar.GetCookContext()->GetCookTagList();
 		CookTags->Add(Owner, "StreamingFormat", LexToString(AudioFormat));
 	}
+#endif
 
 	if (Ar.IsLoading())
 	{
diff --git a/Engine/Source/Runtime/Engine/Private/SoundWave.cpp b/Engine/Source/Runtime/Engine/Private/SoundWave.cpp
--- a/Engine/Source/Runtime/Engine/Private/SoundWave.cpp
+++ b/Engine/Source/Runtime/Engine/Private/SoundWave.cpp
@@ -1231,6 +1231,7 @@ void USoundWave::Serialize( FArchive& Ar )
 	SerializeCuePoints(Ar, bCooked && Ar.IsLoading());
 
 	// Save cook tags
+#if WITH_EDITOR
 	if (Ar.IsCooking() && Ar.IsSaving() && Ar.GetCookContext() && Ar.GetCookContext()->GetCookTagList())
 	{
 		FCookTagList* CookTags = Ar.GetCookContext()->GetCookTagList();
@@ -1249,6 +1250,7 @@ void USoundWave::Serialize( FArchive& Ar )
 		}
 		CookTags->Add(this, "FormatContainerFormats", FormatsListStr.ToString());
 	}
+#endif
 
 	if (bCooked)
 	{
