Added operator== to FRHIRayTracingGeometry to simplify debug checks.

diff --git a/Engine/Source/Runtime/Engine/Private/StaticMesh.cpp b/Engine/Source/Runtime/Engine/Private/StaticMesh.cpp
--- a/Engine/Source/Runtime/Engine/Private/StaticMesh.cpp
+++ b/Engine/Source/Runtime/Engine/Private/StaticMesh.cpp
@@ -1438,7 +1438,7 @@ void FStaticMeshLODResources::InitResources(UStaticMesh* Parent, int32 LODIndex)
 }
 
 #if RHI_RAYTRACING
-void FStaticMeshLODResources::SetupRayTracingGeometryInitializer(FRayTracingGeometryInitializer& Initializer, const FName& DebugName, const FName& OwnerName) const
+void FStaticMeshLODResources::SetupRayTracingGeometryInitializer(FRayTracingGeometryInitializer& Initializer, const FDebugName& DebugName, const FName& OwnerName) const
 {
 	Initializer.DebugName = DebugName;
 	Initializer.OwnerName = OwnerName;
diff --git a/Engine/Source/Runtime/Engine/Private/Streaming/StaticMeshUpdate.cpp b/Engine/Source/Runtime/Engine/Private/Streaming/StaticMeshUpdate.cpp
--- a/Engine/Source/Runtime/Engine/Private/Streaming/StaticMeshUpdate.cpp
+++ b/Engine/Source/Runtime/Engine/Private/Streaming/StaticMeshUpdate.cpp
@@ -138,46 +138,13 @@ static void CheckRayTracingGeometryInitializer(
 	ERayTracingGeometryInitializerType ExpectedInitializerType,
 	const FRayTracingGeometryInitializer& Initializer)
 {
-	const FName OwnerName = UStaticMesh::GetLODPathName(Mesh, LODIdx);
-
 	FRayTracingGeometryInitializer TmpInitializer;
-	LODResource.SetupRayTracingGeometryInitializer(TmpInitializer, Mesh->GetFName(), OwnerName);
+	LODResource.SetupRayTracingGeometryInitializer(TmpInitializer, Initializer.DebugName, Initializer.OwnerName);
+	// simply copied names since they're not always available during streaming
 
 	TmpInitializer.Type = ExpectedInitializerType;
 
-	// Can't compare TmpInitializer == Initializer directly due to some members not having equality operators
-
-	check(TmpInitializer.IndexBuffer == Initializer.IndexBuffer);
-	check(TmpInitializer.IndexBufferOffset == Initializer.IndexBufferOffset);
-	check(TmpInitializer.GeometryType == Initializer.GeometryType);
-	check(TmpInitializer.TotalPrimitiveCount == Initializer.TotalPrimitiveCount);
-
-	// Can't compare Segments directly due to some members not having equality operators
-	check(TmpInitializer.Segments.Num() == Initializer.Segments.Num());
-
-	for (int32 SegmentIndex = 0; SegmentIndex < TmpInitializer.Segments.Num(); ++SegmentIndex)
-	{
-		//check(TmpInitializer.Segments[SegmentIndex] == Initializer.Segments[SegmentIndex]);
-		check(TmpInitializer.Segments[SegmentIndex].VertexBuffer == Initializer.Segments[SegmentIndex].VertexBuffer);
-		check(TmpInitializer.Segments[SegmentIndex].VertexBufferElementType == Initializer.Segments[SegmentIndex].VertexBufferElementType);
-		check(TmpInitializer.Segments[SegmentIndex].VertexBufferOffset == Initializer.Segments[SegmentIndex].VertexBufferOffset);
-		check(TmpInitializer.Segments[SegmentIndex].VertexBufferStride == Initializer.Segments[SegmentIndex].VertexBufferStride);
-		check(TmpInitializer.Segments[SegmentIndex].MaxVertices == Initializer.Segments[SegmentIndex].MaxVertices);
-		check(TmpInitializer.Segments[SegmentIndex].FirstPrimitive == Initializer.Segments[SegmentIndex].FirstPrimitive);
-		check(TmpInitializer.Segments[SegmentIndex].NumPrimitives == Initializer.Segments[SegmentIndex].NumPrimitives);
-		check(TmpInitializer.Segments[SegmentIndex].bForceOpaque == Initializer.Segments[SegmentIndex].bForceOpaque);
-		check(TmpInitializer.Segments[SegmentIndex].bAllowDuplicateAnyHitShaderInvocation == Initializer.Segments[SegmentIndex].bAllowDuplicateAnyHitShaderInvocation);
-		check(TmpInitializer.Segments[SegmentIndex].bEnabled == Initializer.Segments[SegmentIndex].bEnabled);
-	}
-
-	check(TmpInitializer.OfflineData == Initializer.OfflineData);
-	check(TmpInitializer.SourceGeometry == Initializer.SourceGeometry);
-	check(TmpInitializer.bFastBuild == Initializer.bFastBuild);
-	check(TmpInitializer.bAllowUpdate == Initializer.bAllowUpdate);
-	check(TmpInitializer.bAllowCompaction == Initializer.bAllowCompaction);
-	check(TmpInitializer.Type == Initializer.Type);
-	// Can't compare DebugName directly due to FDebugName not having equality operator
-	check(TmpInitializer.OwnerName == Initializer.OwnerName);
+	check(TmpInitializer == Initializer);
 }
 #endif
 
diff --git a/Engine/Source/Runtime/Engine/Public/StaticMeshResources.h b/Engine/Source/Runtime/Engine/Public/StaticMeshResources.h
--- a/Engine/Source/Runtime/Engine/Public/StaticMeshResources.h
+++ b/Engine/Source/Runtime/Engine/Public/StaticMeshResources.h
@@ -522,7 +522,7 @@ public:
 	ENGINE_API void GetResourceSizeEx(FResourceSizeEx& CumulativeResourceSize) const;
 
 #if RHI_RAYTRACING
-	void SetupRayTracingGeometryInitializer(FRayTracingGeometryInitializer& Initializer, const FName& DebugName, const FName& OwnerName) const;
+	void SetupRayTracingGeometryInitializer(FRayTracingGeometryInitializer& Initializer, const FDebugName& DebugName, const FName& OwnerName) const;
 #endif // RHI_RAYTRACING
 
 	/** Get the estimated memory overhead of buffers marked as NeedsCPUAccess. */
diff --git a/Engine/Source/Runtime/RHI/Public/RHIResources.h b/Engine/Source/Runtime/RHI/Public/RHIResources.h
--- a/Engine/Source/Runtime/RHI/Public/RHIResources.h
+++ b/Engine/Source/Runtime/RHI/Public/RHIResources.h
@@ -3150,6 +3150,67 @@ public:
 	FName OwnerName;
 };
 
+#if DO_CHECK
+FORCEINLINE bool operator==(const FRayTracingGeometryInitializer& LHS, const FRayTracingGeometryInitializer& RHS)
+{
+	// Can't compare LHS == RHS directly due to some members not having equality operators
+
+	if (LHS.IndexBuffer != RHS.IndexBuffer
+		|| LHS.IndexBufferOffset != RHS.IndexBufferOffset
+		|| LHS.GeometryType != RHS.GeometryType
+		|| LHS.TotalPrimitiveCount != RHS.TotalPrimitiveCount)
+	{
+		return false;
+	}
+
+	// Can't compare Segments directly due to some members not having equality operators
+	if (LHS.Segments.Num() != RHS.Segments.Num())
+	{
+		return false;
+	}
+
+	for (int32 SegmentIndex = 0; SegmentIndex < LHS.Segments.Num(); ++SegmentIndex)
+	{
+		//if (LHS.Segments[SegmentIndex] != RHS.Segments[SegmentIndex])
+		//{
+		//	return false;
+		//}
+
+		if (LHS.Segments[SegmentIndex].VertexBuffer != RHS.Segments[SegmentIndex].VertexBuffer
+			|| LHS.Segments[SegmentIndex].VertexBufferElementType != RHS.Segments[SegmentIndex].VertexBufferElementType
+			|| LHS.Segments[SegmentIndex].VertexBufferOffset != RHS.Segments[SegmentIndex].VertexBufferOffset
+			|| LHS.Segments[SegmentIndex].VertexBufferStride != RHS.Segments[SegmentIndex].VertexBufferStride
+			|| LHS.Segments[SegmentIndex].MaxVertices != RHS.Segments[SegmentIndex].MaxVertices
+			|| LHS.Segments[SegmentIndex].FirstPrimitive != RHS.Segments[SegmentIndex].FirstPrimitive
+			|| LHS.Segments[SegmentIndex].NumPrimitives != RHS.Segments[SegmentIndex].NumPrimitives
+			|| LHS.Segments[SegmentIndex].bForceOpaque != RHS.Segments[SegmentIndex].bForceOpaque
+			|| LHS.Segments[SegmentIndex].bAllowDuplicateAnyHitShaderInvocation != RHS.Segments[SegmentIndex].bAllowDuplicateAnyHitShaderInvocation
+			|| LHS.Segments[SegmentIndex].bEnabled != RHS.Segments[SegmentIndex].bEnabled)
+		{
+			return false;
+		}
+	}
+
+	if (LHS.OfflineData != RHS.OfflineData
+		|| LHS.SourceGeometry != RHS.SourceGeometry
+		|| LHS.bFastBuild != RHS.bFastBuild
+		|| LHS.bAllowUpdate != RHS.bAllowUpdate
+		|| LHS.bAllowCompaction != RHS.bAllowCompaction
+		|| LHS.Type != RHS.Type)
+	{
+		return false;
+	}
+
+	// Can't compare DebugName directly due to FDebugName not having equality operator
+	if (LHS.OwnerName != RHS.OwnerName)
+	{
+		return false;
+	}
+
+	return true;
+}
+#endif
+
 enum ERayTracingSceneLifetime
 {
 	// Scene may only be used during the frame when it was created.
