Disable concurrent tick queuing by default as it is slower in most cases and also causes reliability issues, it can be re-enabled with tick.AllowConcurrentTickQueue.
Remove define that hard disabled this on windows and android because that made it difficult to verify it worked properly.

diff --git a/Engine/Source/Runtime/Engine/Private/TickTaskManager.cpp b/Engine/Source/Runtime/Engine/Private/TickTaskManager.cpp
--- a/Engine/Source/Runtime/Engine/Private/TickTaskManager.cpp
+++ b/Engine/Source/Runtime/Engine/Private/TickTaskManager.cpp
@@ -49,10 +49,11 @@ static TAutoConsoleVariable<int32> CVarAllowAsyncComponentTicks(
 	1,
 	TEXT("Used to control async component ticks."));
 
+// This was disabled by default in UE 5.5
 static TAutoConsoleVariable<int32> CVarAllowConcurrentQueue(
 	TEXT("tick.AllowConcurrentTickQueue"),
-	1,
-	TEXT("If true, queue ticks concurrently."));
+	0,
+	TEXT("If true, queue ticks concurrently using multiple threads at once. This may be faster on platforms with many cores but can change the order of ticking."));
 
 static TAutoConsoleVariable<int32> CVarAllowAsyncTickDispatch(
 	TEXT("tick.AllowAsyncTickDispatch"),
@@ -1492,13 +1493,12 @@ public:
 
 		int32 NumWorkerThread = 0;
 		bool bConcurrentQueue = false;
-#if !PLATFORM_WINDOWS && !PLATFORM_ANDROID
-		// some schedulers will hang for seconds trying to do this algorithm, threads starve even though other threads are calling sleep(0)
+
 		if (!FTickTaskSequencer::SingleThreadedMode())
 		{
+			// Concurrent tick may be faster in some situations but can change the order of ticking
 			bConcurrentQueue = !!CVarAllowConcurrentQueue.GetValueOnGameThread();
 		}
-#endif
 
 		if (!bConcurrentQueue)
 		{
