Make it safe to call GetAssetRegistryTags on CDO from ContentBrowserAssetDataSource.

diff --git a/Engine/Source/Runtime/Engine/Private/ActorFolder.cpp b/Engine/Source/Runtime/Engine/Private/ActorFolder.cpp
--- a/Engine/Source/Runtime/Engine/Private/ActorFolder.cpp
+++ b/Engine/Source/Runtime/Engine/Private/ActorFolder.cpp
@@ -85,7 +85,7 @@ void UActorFolder::GetAssetRegistryTags(FAssetRegistryTagsContext Context) const
 	Context.AddTag(FAssetRegistryTag(ActorFolder::NAME_FolderLabel, *FolderLabel, FAssetRegistryTag::TT_Hidden));
 	Context.AddTag(FAssetRegistryTag(ActorFolder::NAME_FolderInitiallyExpanded, bFolderInitiallyExpanded ? TEXT("1") : TEXT("0"), FAssetRegistryTag::TT_Hidden));
 	Context.AddTag(FAssetRegistryTag(ActorFolder::NAME_FolderIsDeleted, bIsDeleted ? TEXT("1") : TEXT("0"), FAssetRegistryTag::TT_Hidden));
-	Context.AddTag(FAssetRegistryTag(ActorFolder::NAME_OuterPackageName, *GetOuterULevel()->GetPackage()->GetName(), FAssetRegistryTag::TT_Hidden));
+	Context.AddTag(FAssetRegistryTag(ActorFolder::NAME_OuterPackageName, HasAnyFlags(RF_ClassDefaultObject) ? TEXT("") : *GetOuterULevel()->GetPackage()->GetName(), FAssetRegistryTag::TT_Hidden));
 }
 
 FActorFolderDesc UActorFolder::GetAssetRegistryInfoFromPackage(FName ActorFolderPackageName)
