Allow users to create PerPlatformProperties for types outside Engine.
- Moved PerPlatformProperties to CoreUObject
- Moved the serialization templates to an inl so other modules can include it, necessary if you want to PerPlatformize something in other modules
- Added stub in Engine for code that uses #include "PerPlatformProperties.h" to include the CoreUObject version. Will add deprecation warnings in a subsequent CL.

diff --git a/Engine/Source/Runtime/CoreUObject/Private/UObject/PerPlatformProperties.cpp b/Engine/Source/Runtime/CoreUObject/Private/UObject/PerPlatformProperties.cpp
new file mode 100644
--- /dev/null
+++ b/Engine/Source/Runtime/CoreUObject/Private/UObject/PerPlatformProperties.cpp
@@ -0,0 +1,49 @@
+// Copyright Epic Games, Inc. All Rights Reserved.
+
+#include "UObject/PerPlatformProperties.h"
+#include "Concepts/StaticStructProvider.h"
+#include "Misc/DelayedAutoRegister.h"
+
+#include UE_INLINE_GENERATED_CPP_BY_NAME(PerPlatformProperties)
+
+#include "UObject/PerPlatformPropertiesImpl.inl"
+
+IMPLEMENT_TYPE_LAYOUT(FFreezablePerPlatformFloat);
+IMPLEMENT_TYPE_LAYOUT(FFreezablePerPlatformInt);
+
+
+template COREUOBJECT_API FArchive& operator<<(FArchive&, TPerPlatformProperty<FPerPlatformInt, int32, NAME_IntProperty>&);
+template COREUOBJECT_API FArchive& operator<<(FArchive&, TPerPlatformProperty<FPerPlatformFloat, float, NAME_FloatProperty>&);
+template COREUOBJECT_API FArchive& operator<<(FArchive&, TPerPlatformProperty<FPerPlatformBool, bool, NAME_BoolProperty>&);
+template COREUOBJECT_API FArchive& operator<<(FArchive&, TPerPlatformProperty<FFreezablePerPlatformFloat, float, NAME_FloatProperty>&);
+template COREUOBJECT_API void operator<<(FStructuredArchive::FSlot Slot, TPerPlatformProperty<FPerPlatformInt, int32, NAME_IntProperty>&);
+template COREUOBJECT_API void operator<<(FStructuredArchive::FSlot Slot, TPerPlatformProperty<FPerPlatformFloat, float, NAME_FloatProperty>&);
+template COREUOBJECT_API void operator<<(FStructuredArchive::FSlot Slot, TPerPlatformProperty<FPerPlatformBool, bool, NAME_BoolProperty>&);
+template COREUOBJECT_API void operator<<(FStructuredArchive::FSlot Slot, TPerPlatformProperty<FFreezablePerPlatformFloat, float, NAME_FloatProperty>&);
+
+template COREUOBJECT_API FArchive& operator<<(FArchive&, TPerPlatformProperty<FPerPlatformFrameRate, FFrameRate, NAME_FrameRate>&);
+template COREUOBJECT_API void operator<<(FStructuredArchive::FSlot Slot, TPerPlatformProperty<FPerPlatformFrameRate, FFrameRate, NAME_FrameRate>&);
+
+FString FPerPlatformInt::ToString() const
+{
+	FString Result = FString::FromInt(Default);
+
+#if WITH_EDITORONLY_DATA
+	TArray<FName> SortedPlatforms;
+	PerPlatform.GetKeys(/*out*/ SortedPlatforms);
+	SortedPlatforms.Sort(FNameLexicalLess());
+
+	for (FName Platform : SortedPlatforms)
+	{
+		Result = FString::Printf(TEXT("%s, %s=%d"), *Result, *Platform.ToString(), PerPlatform.FindChecked(Platform));
+	}
+#endif
+
+	return Result;
+}
+
+FString FFreezablePerPlatformInt::ToString() const
+{
+	return FPerPlatformInt(*this).ToString();
+}
+
diff --git a/Engine/Source/Runtime/CoreUObject/Public/UObject/PerPlatformProperties.h b/Engine/Source/Runtime/CoreUObject/Public/UObject/PerPlatformProperties.h
new file mode 100644
--- /dev/null
+++ b/Engine/Source/Runtime/CoreUObject/Public/UObject/PerPlatformProperties.h
@@ -0,0 +1,472 @@
+// Copyright Epic Games, Inc. All Rights Reserved.
+
+/*=============================================================================
+PerPlatformProperties.h: Property types that can be overridden on a per-platform basis at cook time
+=============================================================================*/
+
+#pragma once
+
+#include "CoreMinimal.h"
+#include "UObject/Object.h"
+#include "Serialization/Archive.h"
+#include "Containers/Map.h"
+#include "Algo/Find.h"
+#include "Serialization/MemoryLayout.h"
+#include "Misc/DataDrivenPlatformInfoRegistry.h"
+#include "Misc/FrameRate.h"
+
+#include "PerPlatformProperties.generated.h"
+
+namespace PerPlatformProperty::Private
+{
+	template<typename MapType>
+	struct TGetKeyType
+	{
+		using Type = typename TTupleElement<0, typename MapType::ElementType>::Type;
+	};
+
+	template<typename NameType>
+	struct FNameFuncs;
+
+	template<>
+	struct FNameFuncs<FName>
+	{
+		static FName NameToKey(FName Name) { return Name; }
+
+		template<typename ValueType>
+		static void SerializePerPlatformMap(FArchive& Ar, TMap<FName, ValueType>& Map)
+		{
+			Ar << Map;
+		}
+
+		template<typename ValueType>
+		static void SerializePerPlatformMap(FArchive& UnderlyingArchive, FStructuredArchive::FRecord& Record, TMap<FName, ValueType>& Map)
+		{
+			Record << SA_VALUE(TEXT("PerPlatform"), Map);
+		}
+	};
+
+	template<>
+	struct FNameFuncs<FMemoryImageName>
+	{
+		static FMemoryImageName NameToKey(FName Name) { return FMemoryImageName(Name); }
+
+		template<typename ValueType>
+		static void SerializePerPlatformMap(FArchive& Ar, TMemoryImageMap<FMemoryImageName, ValueType>& Map)
+		{
+			if( Ar.IsLoading())
+			{
+				TMemoryImageMap<FName, ValueType> TempMap;
+				Ar << TempMap;
+				Map.Reset();
+				for( TPair<FName, ValueType>& Pair : TempMap)
+				{
+					Map.Add(FMemoryImageName(Pair.Key), Pair.Value);
+				}
+			}
+			else
+			{
+				TMemoryImageMap<FName, ValueType> TempMap;
+				for (TPair<FMemoryImageName, ValueType>& Pair : Map)
+				{
+					TempMap.Add(FName(Pair.Key), Pair.Value);
+				}
+				Ar << TempMap;
+			}
+		}
+
+		template<typename ValueType>
+		static void SerializePerPlatformMap(FArchive& UnderlyingArchive, FStructuredArchive::FRecord& Record, TMemoryImageMap<FMemoryImageName, ValueType>& Map)
+		{
+			if (UnderlyingArchive.IsLoading())
+			{
+				TMemoryImageMap<FName, ValueType> TempMap;
+
+				Record << SA_VALUE(TEXT("PerPlatform"), TempMap);
+				Map.Reset();
+				for (TPair<FName, ValueType>& Pair : TempMap)
+				{
+					Map.Add(FMemoryImageName(Pair.Key), Pair.Value);
+				}
+			}
+			else
+			{
+				TMemoryImageMap<FName, ValueType> TempMap;
+				for (TPair<FMemoryImageName, ValueType>& Pair : Map)
+				{
+					TempMap.Add(FName(Pair.Key), Pair.Value);
+				}
+				Record << SA_VALUE(TEXT("PerPlatform"), TempMap);
+			}
+		}
+	};
+
+	template<typename MapType>
+	using KeyFuncs = FNameFuncs<typename TGetKeyType<MapType>::Type>;
+}
+
+/** TPerPlatformProperty - template parent class for per-platform properties 
+ *  Implements Serialize function to replace value at cook time, and 
+ *  backwards-compatible loading code for properties converted from simple types.
+ */
+template<typename _StructType, typename _ValueType, EName _BasePropertyName>
+struct TPerPlatformProperty
+{
+	typedef _ValueType ValueType;
+	typedef _StructType StructType;
+
+#if WITH_EDITOR
+	/** Get the value for the given platform (using standard "ini" name, so Windows, not Win64 or WindowsClient), which can be used to lookup the group */
+	_ValueType GetValueForPlatform(FName PlatformName) const
+	{
+		const _StructType* This = StaticCast<const _StructType*>(this);
+
+		using MapType = decltype(This->PerPlatform);
+		using KeyFuncs = typename PerPlatformProperty::Private::KeyFuncs<MapType>;
+
+		const _ValueType* Ptr = This->PerPlatform.Find(KeyFuncs::NameToKey(PlatformName));
+
+		if (Ptr == nullptr)
+		{
+			const FDataDrivenPlatformInfo& Info = FDataDrivenPlatformInfoRegistry::GetPlatformInfo(PlatformName);
+			if (Info.PlatformGroupName != NAME_None)
+			{
+				Ptr = This->PerPlatform.Find(KeyFuncs::NameToKey(Info.PlatformGroupName));
+			}
+		}
+
+		return Ptr ? *Ptr : This->Default;
+	}
+#endif
+
+	_ValueType GetDefault() const
+	{
+		const _StructType* This = StaticCast<const _StructType*>(this);
+		return This->Default;
+	}
+	
+	_ValueType GetValue() const
+	{
+#if WITH_EDITORONLY_DATA && WITH_EDITOR
+		FName PlatformName;
+		// Lookup the override preview platform info, if any
+		// @todo this doesn't set PlatformName, just a group, but GetValueForPlatform() will technically work being given a Group name instead of a platform name, so we just use it
+		if (UObject::OnGetPreviewPlatform.IsBound() && UObject::OnGetPreviewPlatform.Execute(PlatformName))
+		{
+			return GetValueForPlatform(PlatformName);
+		}
+		else		
+#endif
+		{
+			const _StructType* This = StaticCast<const _StructType*>(this);
+			return This->Default;
+		}
+	}
+
+	/* Load old properties that have been converted to FPerPlatformX */
+	bool SerializeFromMismatchedTag(const FPropertyTag& Tag, FArchive& Ar)
+	{
+		if (Tag.Type == _BasePropertyName)
+		{
+			_StructType* This = StaticCast<_StructType*>(this);
+			_ValueType OldValue;
+			Ar << OldValue;
+			*This = _StructType(OldValue);
+			return true;
+		}
+		return false;
+	}
+
+	/* Serialization */
+	bool Serialize(FArchive& Ar)
+	{
+		Ar << *this;
+		return true;
+	}
+
+	/* Serialization */
+	bool Serialize(FStructuredArchive::FSlot Slot)
+	{
+		Slot << *this;
+		return true;
+	}
+};
+
+
+template<typename _StructType, typename _ValueType, EName _BasePropertyName>
+FArchive& operator<<(FArchive& Ar, TPerPlatformProperty<_StructType, _ValueType, _BasePropertyName>& P);
+template<typename _StructType, typename _ValueType, EName _BasePropertyName>
+void operator<<(FStructuredArchive::FSlot Slot, TPerPlatformProperty<_StructType, _ValueType, _BasePropertyName>& P);
+
+struct FFreezablePerPlatformInt;
+
+/** FPerPlatformInt - int32 property with per-platform overrides */
+USTRUCT(BlueprintType)
+struct FPerPlatformInt
+#if CPP
+:	public TPerPlatformProperty<FPerPlatformInt, int32, NAME_IntProperty>
+#endif
+{
+	GENERATED_USTRUCT_BODY()
+
+	UPROPERTY(BlueprintReadOnly, EditAnywhere, Category = PerPlatform)
+	int32 Default;
+
+#if WITH_EDITORONLY_DATA
+	UPROPERTY(EditAnywhere, Category = PerPlatform)
+	TMap<FName, int32> PerPlatform;
+#endif
+
+	FPerPlatformInt()
+	:	Default(0)
+	{
+	}
+
+	FPerPlatformInt(int32 InDefaultValue)
+	:	Default(InDefaultValue)
+	{
+	}
+
+	COREUOBJECT_API FString ToString() const;
+
+	COREUOBJECT_API FPerPlatformInt(const FFreezablePerPlatformInt& Other);
+};
+
+USTRUCT()
+struct FFreezablePerPlatformInt
+#if CPP
+	: public TPerPlatformProperty<FFreezablePerPlatformInt, int32, NAME_IntProperty>
+#endif
+{
+	DECLARE_TYPE_LAYOUT(FFreezablePerPlatformInt, NonVirtual);
+
+	GENERATED_USTRUCT_BODY()
+
+public:
+	using KeyType = FMemoryImageName;
+	using FPerPlatformMap = TMemoryImageMap<FMemoryImageName, int32>;
+
+	LAYOUT_FIELD(int32, Default);
+	LAYOUT_FIELD_EDITORONLY(FPerPlatformMap, PerPlatform);
+
+	FFreezablePerPlatformInt() : Default(0) {}
+	FFreezablePerPlatformInt(int32 InDefaultValue) : Default(InDefaultValue) {}
+	FFreezablePerPlatformInt(const FPerPlatformInt& Other)
+		: Default(Other.Default)
+	{
+#if WITH_EDITORONLY_DATA
+		for (const TPair<FName, int32>& Pair : Other.PerPlatform)
+		{
+			PerPlatform.Add(FMemoryImageName(Pair.Key), Pair.Value);
+		}
+#endif
+	}
+
+	COREUOBJECT_API FString ToString() const;
+};
+
+inline FPerPlatformInt::FPerPlatformInt(const FFreezablePerPlatformInt& Other)
+	: Default(Other.Default)
+{
+#if WITH_EDITORONLY_DATA
+	for (const TPair<FMemoryImageName, int32>& Pair : Other.PerPlatform)
+	{
+		PerPlatform.Add(FName(Pair.Key), Pair.Value);
+	}
+#endif
+}
+
+extern template COREUOBJECT_API FArchive& operator<<(FArchive&, TPerPlatformProperty<FPerPlatformInt, int32, NAME_IntProperty>&);
+extern template COREUOBJECT_API void operator<<(FStructuredArchive::FSlot Slot, TPerPlatformProperty<FPerPlatformInt, int32, NAME_IntProperty>&);
+
+template<>
+struct TStructOpsTypeTraits<FPerPlatformInt>
+	: public TStructOpsTypeTraitsBase2<FPerPlatformInt>
+{
+	enum
+	{
+		WithSerializeFromMismatchedTag = true,
+		WithSerializer = true
+	};
+	static constexpr EPropertyObjectReferenceType WithSerializerObjectReferences = EPropertyObjectReferenceType::None;
+};
+
+struct FFreezablePerPlatformFloat;
+
+/** FPerPlatformFloat - float property with per-platform overrides */
+USTRUCT(BlueprintType, meta = (CanFlattenStruct))
+struct FPerPlatformFloat
+#if CPP
+:	public TPerPlatformProperty<FPerPlatformFloat, float, NAME_FloatProperty>
+#endif
+{
+	GENERATED_USTRUCT_BODY()
+
+	UPROPERTY(BlueprintReadOnly, EditAnywhere, Category = PerPlatform)
+	float Default;
+
+#if WITH_EDITORONLY_DATA
+	UPROPERTY(EditAnywhere, Category = PerPlatform)
+	TMap<FName, float> PerPlatform;
+#endif
+
+	FPerPlatformFloat()
+	:	Default(0.f)
+	{
+	}
+
+	FPerPlatformFloat(float InDefaultValue)
+	:	Default(InDefaultValue)
+	{
+	}
+
+	COREUOBJECT_API FPerPlatformFloat(const FFreezablePerPlatformFloat& Other);
+};
+extern template COREUOBJECT_API FArchive& operator<<(FArchive&, TPerPlatformProperty<FPerPlatformFloat, float, NAME_FloatProperty>&);
+
+struct FFreezablePerPlatformFloat
+#if CPP
+	: public TPerPlatformProperty<FFreezablePerPlatformFloat, float, NAME_FloatProperty>
+#endif
+{
+	DECLARE_TYPE_LAYOUT(FFreezablePerPlatformFloat, NonVirtual);
+public:
+	using FPerPlatformMap = TMemoryImageMap<FMemoryImageName, float>;
+	
+	LAYOUT_FIELD(float, Default);
+	LAYOUT_FIELD_EDITORONLY(FPerPlatformMap, PerPlatform);
+
+	FFreezablePerPlatformFloat() : Default(0.0f) {}
+	FFreezablePerPlatformFloat(float InDefaultValue) : Default(InDefaultValue) {}
+	FFreezablePerPlatformFloat(const FPerPlatformFloat& Other)
+		: Default(Other.Default)
+	{
+#if WITH_EDITORONLY_DATA
+		for (const TPair<FName, float>& Pair : Other.PerPlatform)
+		{
+			PerPlatform.Add(FMemoryImageName(Pair.Key), Pair.Value);
+		}
+#endif
+}
+};
+
+inline FPerPlatformFloat::FPerPlatformFloat(const FFreezablePerPlatformFloat& Other)
+	: Default(Other.Default)
+{
+#if WITH_EDITORONLY_DATA
+	for (const TPair<FMemoryImageName, float>& Pair : Other.PerPlatform)
+	{
+		PerPlatform.Add(FName(Pair.Key), Pair.Value);
+	}
+#endif
+}
+
+template<>
+struct TStructOpsTypeTraits<FPerPlatformFloat>
+:	public TStructOpsTypeTraitsBase2<FPerPlatformFloat>
+{
+	enum
+	{
+		WithSerializeFromMismatchedTag = true,
+		WithSerializer = true
+	};
+	static constexpr EPropertyObjectReferenceType WithSerializerObjectReferences = EPropertyObjectReferenceType::None;
+};
+
+/** FPerPlatformBool - bool property with per-platform overrides */
+USTRUCT()
+struct FPerPlatformBool
+#if CPP
+:	public TPerPlatformProperty<FPerPlatformBool, bool, NAME_BoolProperty>
+#endif
+{
+	GENERATED_USTRUCT_BODY()
+
+	UPROPERTY(EditAnywhere, Category = PerPlatform)
+	bool Default;
+
+#if WITH_EDITORONLY_DATA
+	UPROPERTY(EditAnywhere, Category = PerPlatform)
+	TMap<FName, bool> PerPlatform;
+#endif
+
+	FPerPlatformBool()
+	:	Default(false)
+	{
+	}
+
+	FPerPlatformBool(bool InDefaultValue)
+	:	Default(InDefaultValue)
+	{
+	}
+};
+extern template COREUOBJECT_API FArchive& operator<<(FArchive&, TPerPlatformProperty<FPerPlatformBool, bool, NAME_BoolProperty>&);
+
+template<>
+struct TStructOpsTypeTraits<FPerPlatformBool>
+	: public TStructOpsTypeTraitsBase2<FPerPlatformBool>
+{
+	enum
+	{
+		WithSerializeFromMismatchedTag = true,
+		WithSerializer = true
+	};
+	static constexpr EPropertyObjectReferenceType WithSerializerObjectReferences = EPropertyObjectReferenceType::None;
+};
+
+/** FPerPlatformFrameRate - FFrameRate property with per-platform overrides */
+USTRUCT(BlueprintType)
+struct FPerPlatformFrameRate
+#if CPP
+:	public TPerPlatformProperty<FPerPlatformFrameRate, FFrameRate, NAME_FrameRate>
+#endif
+{
+	GENERATED_USTRUCT_BODY()
+
+	UPROPERTY(EditAnywhere, Category = PerPlatform)
+	FFrameRate Default;
+
+#if WITH_EDITORONLY_DATA
+	UPROPERTY(EditAnywhere, Category = PerPlatform)
+	TMap<FName, FFrameRate> PerPlatform;
+#endif
+
+	FPerPlatformFrameRate()
+	:	Default(30, 1)
+	{
+	}
+
+	FPerPlatformFrameRate(FFrameRate InDefaultValue)
+	:	Default(InDefaultValue)
+	{
+	}
+	
+	bool SerializeFromMismatchedTag(const FPropertyTag& Tag, FArchive& Ar)
+	{
+		if(const UStruct* FrameRateStruct = FindObject<UStruct>(FTopLevelAssetPath("/Script/CoreUObject.FrameRate")))
+		{
+			FFrameRate Value;
+			Ar << Value.Denominator;
+			Ar << Value.Numerator;
+			Default = Value;
+
+			return true;
+		}
+		
+		return false;
+	}
+};
+
+extern template COREUOBJECT_API FArchive& operator<<(FArchive&, TPerPlatformProperty<FPerPlatformFrameRate, FFrameRate, NAME_FrameRate>&);
+
+template<>
+struct TStructOpsTypeTraits<FPerPlatformFrameRate>
+	: public TStructOpsTypeTraitsBase2<FPerPlatformFrameRate>
+{
+	enum
+	{
+		WithSerializeFromMismatchedTag = false,
+		WithSerializer = true
+	};
+	static constexpr EPropertyObjectReferenceType WithSerializerObjectReferences = EPropertyObjectReferenceType::None;
+};
diff --git a/Engine/Source/Runtime/CoreUObject/Public/UObject/PerPlatformPropertiesImpl.inl b/Engine/Source/Runtime/CoreUObject/Public/UObject/PerPlatformPropertiesImpl.inl
new file mode 100644
--- /dev/null
+++ b/Engine/Source/Runtime/CoreUObject/Public/UObject/PerPlatformPropertiesImpl.inl
@@ -0,0 +1,84 @@
+// Copyright Epic Games, Inc. All Rights Reserved.
+
+/*=============================================================================
+ * PerPlatformPropertiesImpl.inl: Serializer implementation
+ *
+ * This file needs to be included by a cpp once for each module that declares a PerPlatformProperty
+ * And following the include, template instantiations of the serializaers for each property class, for example:
+ * 
+ * template SLATECORE_API FArchive& operator<<(FArchive&, TPerPlatformProperty<FPerPlatformMyEnumType, EMyEnumType, NAME_EnumProperty>&);
+ * template SLATECORE_API void operator<<(FStructuredArchive::FSlot Slot, TPerPlatformProperty<FPerPlatformMyEnumType, EMyEnumType, NAME_EnumProperty>&);
+ * =============================================================================*/
+
+#pragma once
+
+#if WITH_EDITOR
+#include "Interfaces/ITargetPlatform.h"
+#endif
+
+/** Serializer to cook out the most appropriate platform override */
+template<typename StructType, typename ValueType, EName BasePropertyName>
+FArchive& operator<<(FArchive& Ar, TPerPlatformProperty<StructType, ValueType, BasePropertyName>& Property)
+{
+	bool bCooked = false;
+#if WITH_EDITOR
+	if (Ar.IsCooking())
+	{
+		bCooked = true;
+		Ar << bCooked;
+		// Save out platform override if it exists and Default otherwise
+		ValueType Value = Property.GetValueForPlatform(*Ar.CookingTarget()->IniPlatformName());
+		Ar << Value;
+	}
+	else
+#endif
+	{
+		StructType* This = StaticCast<StructType*>(&Property);
+		Ar << bCooked;
+		Ar << This->Default;
+#if WITH_EDITORONLY_DATA
+		if (!bCooked)
+		{
+			using MapType = decltype(This->PerPlatform);
+			using KeyFuncs = typename PerPlatformProperty::Private::KeyFuncs<MapType>;
+			KeyFuncs::SerializePerPlatformMap(Ar, This->PerPlatform);
+		}
+#endif
+	}
+	return Ar;
+}
+
+/** Serializer to cook out the most appropriate platform override */
+template<typename StructType, typename ValueType, EName BasePropertyName>
+void operator<<(FStructuredArchive::FSlot Slot, TPerPlatformProperty<StructType, ValueType, BasePropertyName>& Property)
+{
+	FArchive& UnderlyingArchive = Slot.GetUnderlyingArchive();
+	FStructuredArchive::FRecord Record = Slot.EnterRecord();
+
+	bool bCooked = false;
+#if WITH_EDITOR
+	if (UnderlyingArchive.IsCooking())
+	{
+		bCooked = true;
+		Record << SA_VALUE(TEXT("bCooked"), bCooked);
+		// Save out platform override if it exists and Default otherwise
+		ValueType Value = Property.GetValueForPlatform(*UnderlyingArchive.CookingTarget()->IniPlatformName());
+		Record << SA_VALUE(TEXT("Value"), Value);
+	}
+	else
+#endif
+	{
+		StructType* This = StaticCast<StructType*>(&Property);
+		Record << SA_VALUE(TEXT("bCooked"), bCooked);
+		Record << SA_VALUE(TEXT("Value"), This->Default);
+#if WITH_EDITORONLY_DATA
+		if (!bCooked)
+		{
+			using MapType = decltype(This->PerPlatform);
+			using KeyFuncs = typename PerPlatformProperty::Private::KeyFuncs<MapType>;
+			KeyFuncs::SerializePerPlatformMap(UnderlyingArchive, Record, This->PerPlatform);
+		}
+#endif
+	}
+}
+
diff --git a/Engine/Source/Runtime/Engine/Private/PerPlatformProperties.cpp b/Engine/Source/Runtime/Engine/Private/PerPlatformProperties.cpp
deleted file mode 100644
diff --git a/Engine/Source/Runtime/Engine/Public/PerPlatformProperties.h b/Engine/Source/Runtime/Engine/Public/PerPlatformProperties.h
--- a/Engine/Source/Runtime/Engine/Public/PerPlatformProperties.h
+++ b/Engine/Source/Runtime/Engine/Public/PerPlatformProperties.h
@@ -1,476 +1,9 @@
 // Copyright Epic Games, Inc. All Rights Reserved.
 
 /*=============================================================================
-PerPlatformProperties.h: Property types that can be overridden on a per-platform basis at cook time
+Redirect PerPlatformProperties.h include to UObject
 =============================================================================*/
 
 #pragma once
 
-#include "Serialization/Archive.h"
-#include "RHIDefinitions.h"
-#include "Containers/Map.h"
-#include "Algo/Find.h"
-#include "Serialization/MemoryLayout.h"
-#include "Misc/DataDrivenPlatformInfoRegistry.h"
-#include "Misc/FrameRate.h"
-
-#include "PerPlatformProperties.generated.h"
-
-/** Helper function to avoid including Engine.h in this header */
-#if WITH_EDITORONLY_DATA && WITH_EDITOR
-ENGINE_API bool GEngine_GetPreviewPlatformName(FName& PlatformName);
-#endif
-
-namespace PerPlatformProperty::Private
-{
-	template<typename MapType>
-	struct TGetKeyType
-	{
-		using Type = typename TTupleElement<0, typename MapType::ElementType>::Type;
-	};
-
-	template<typename NameType>
-	struct FNameFuncs;
-
-	template<>
-	struct FNameFuncs<FName>
-	{
-		static FName NameToKey(FName Name) { return Name; }
-
-		template<typename ValueType>
-		static void SerializePerPlatformMap(FArchive& Ar, TMap<FName, ValueType>& Map)
-		{
-			Ar << Map;
-		}
-
-		template<typename ValueType>
-		static void SerializePerPlatformMap(FArchive& UnderlyingArchive, FStructuredArchive::FRecord& Record, TMap<FName, ValueType>& Map)
-		{
-			Record << SA_VALUE(TEXT("PerPlatform"), Map);
-		}
-	};
-
-	template<>
-	struct FNameFuncs<FMemoryImageName>
-	{
-		static FMemoryImageName NameToKey(FName Name) { return FMemoryImageName(Name); }
-
-		template<typename ValueType>
-		static void SerializePerPlatformMap(FArchive& Ar, TMemoryImageMap<FMemoryImageName, ValueType>& Map)
-		{
-			if( Ar.IsLoading())
-			{
-				TMemoryImageMap<FName, ValueType> TempMap;
-				Ar << TempMap;
-				Map.Reset();
-				for( TPair<FName, ValueType>& Pair : TempMap)
-				{
-					Map.Add(FMemoryImageName(Pair.Key), Pair.Value);
-				}
-			}
-			else
-			{
-				TMemoryImageMap<FName, ValueType> TempMap;
-				for (TPair<FMemoryImageName, ValueType>& Pair : Map)
-				{
-					TempMap.Add(FName(Pair.Key), Pair.Value);
-				}
-				Ar << TempMap;
-			}
-		}
-
-		template<typename ValueType>
-		static void SerializePerPlatformMap(FArchive& UnderlyingArchive, FStructuredArchive::FRecord& Record, TMemoryImageMap<FMemoryImageName, ValueType>& Map)
-		{
-			if (UnderlyingArchive.IsLoading())
-			{
-				TMemoryImageMap<FName, ValueType> TempMap;
-
-				Record << SA_VALUE(TEXT("PerPlatform"), TempMap);
-				Map.Reset();
-				for (TPair<FName, ValueType>& Pair : TempMap)
-				{
-					Map.Add(FMemoryImageName(Pair.Key), Pair.Value);
-				}
-			}
-			else
-			{
-				TMemoryImageMap<FName, ValueType> TempMap;
-				for (TPair<FMemoryImageName, ValueType>& Pair : Map)
-				{
-					TempMap.Add(FName(Pair.Key), Pair.Value);
-				}
-				Record << SA_VALUE(TEXT("PerPlatform"), TempMap);
-			}
-		}
-	};
-
-	template<typename MapType>
-	using KeyFuncs = FNameFuncs<typename TGetKeyType<MapType>::Type>;
-}
-
-/** TPerPlatformProperty - template parent class for per-platform properties 
- *  Implements Serialize function to replace value at cook time, and 
- *  backwards-compatible loading code for properties converted from simple types.
- */
-template<typename _StructType, typename _ValueType, EName _BasePropertyName>
-struct TPerPlatformProperty
-{
-	typedef _ValueType ValueType;
-	typedef _StructType StructType;
-
-#if WITH_EDITOR
-	/** Get the value for the given platform (using standard "ini" name, so Windows, not Win64 or WindowsClient), which can be used to lookup the group */
-	_ValueType GetValueForPlatform(FName PlatformName) const
-	{
-		const _StructType* This = StaticCast<const _StructType*>(this);
-
-		using MapType = decltype(This->PerPlatform);
-		using KeyFuncs = typename PerPlatformProperty::Private::KeyFuncs<MapType>;
-
-		const _ValueType* Ptr = This->PerPlatform.Find(KeyFuncs::NameToKey(PlatformName));
-
-		if (Ptr == nullptr)
-		{
-			const FDataDrivenPlatformInfo& Info = FDataDrivenPlatformInfoRegistry::GetPlatformInfo(PlatformName);
-			if (Info.PlatformGroupName != NAME_None)
-			{
-				Ptr = This->PerPlatform.Find(KeyFuncs::NameToKey(Info.PlatformGroupName));
-			}
-		}
-
-		return Ptr ? *Ptr : This->Default;
-	}
-#endif
-
-	_ValueType GetDefault() const
-	{
-		const _StructType* This = StaticCast<const _StructType*>(this);
-		return This->Default;
-	}
-	
-	_ValueType GetValue() const
-	{
-#if WITH_EDITORONLY_DATA && WITH_EDITOR
-		FName PlatformName;
-		// Lookup the override preview platform info, if any
-		// @todo this doesn't set PlatformName, just a group, but GetValueForPlatform() will technically work being given a Group name instead of a platform name, so we just use it
-		if (GEngine_GetPreviewPlatformName(PlatformName))
-		{
-			return GetValueForPlatform(PlatformName);
-		}
-		else		
-#endif
-		{
-			const _StructType* This = StaticCast<const _StructType*>(this);
-			return This->Default;
-		}
-	}
-
-	/* Load old properties that have been converted to FPerPlatformX */
-	bool SerializeFromMismatchedTag(const FPropertyTag& Tag, FArchive& Ar)
-	{
-		if (Tag.Type == _BasePropertyName)
-		{
-			_StructType* This = StaticCast<_StructType*>(this);
-			_ValueType OldValue;
-			Ar << OldValue;
-			*This = _StructType(OldValue);
-			return true;
-		}
-		return false;
-	}
-
-	/* Serialization */
-	bool Serialize(FArchive& Ar)
-	{
-		Ar << *this;
-		return true;
-	}
-
-	/* Serialization */
-	bool Serialize(FStructuredArchive::FSlot Slot)
-	{
-		Slot << *this;
-		return true;
-	}
-};
-
-template<typename _StructType, typename _ValueType, EName _BasePropertyName>
-ENGINE_API FArchive& operator<<(FArchive& Ar, TPerPlatformProperty<_StructType, _ValueType, _BasePropertyName>& P);
-template<typename _StructType, typename _ValueType, EName _BasePropertyName>
-ENGINE_API void operator<<(FStructuredArchive::FSlot Slot, TPerPlatformProperty<_StructType, _ValueType, _BasePropertyName>& P);
-
-struct FFreezablePerPlatformInt;
-
-/** FPerPlatformInt - int32 property with per-platform overrides */
-USTRUCT(BlueprintType)
-struct FPerPlatformInt
-#if CPP
-:	public TPerPlatformProperty<FPerPlatformInt, int32, NAME_IntProperty>
-#endif
-{
-	GENERATED_USTRUCT_BODY()
-
-	UPROPERTY(BlueprintReadOnly, EditAnywhere, Category = PerPlatform)
-	int32 Default;
-
-#if WITH_EDITORONLY_DATA
-	UPROPERTY(EditAnywhere, Category = PerPlatform)
-	TMap<FName, int32> PerPlatform;
-#endif
-
-	FPerPlatformInt()
-	:	Default(0)
-	{
-	}
-
-	FPerPlatformInt(int32 InDefaultValue)
-	:	Default(InDefaultValue)
-	{
-	}
-
-	ENGINE_API FString ToString() const;
-
-	ENGINE_API FPerPlatformInt(const FFreezablePerPlatformInt& Other);
-
-};
-
-USTRUCT()
-struct FFreezablePerPlatformInt
-#if CPP
-	: public TPerPlatformProperty<FFreezablePerPlatformInt, int32, NAME_IntProperty>
-#endif
-{
-	DECLARE_TYPE_LAYOUT(FFreezablePerPlatformInt, NonVirtual);
-
-	GENERATED_USTRUCT_BODY()
-
-public:
-	using KeyType = FMemoryImageName;
-	using FPerPlatformMap = TMemoryImageMap<FMemoryImageName, int32>;
-
-	LAYOUT_FIELD(int32, Default);
-	LAYOUT_FIELD_EDITORONLY(FPerPlatformMap, PerPlatform);
-
-	FFreezablePerPlatformInt() : Default(0) {}
-	FFreezablePerPlatformInt(int32 InDefaultValue) : Default(InDefaultValue) {}
-	FFreezablePerPlatformInt(const FPerPlatformInt& Other)
-		: Default(Other.Default)
-	{
-#if WITH_EDITORONLY_DATA
-		for (const TPair<FName, int32>& Pair : Other.PerPlatform)
-		{
-			PerPlatform.Add(FMemoryImageName(Pair.Key), Pair.Value);
-		}
-#endif
-	}
-
-	ENGINE_API FString ToString() const;
-};
-
-inline FPerPlatformInt::FPerPlatformInt(const FFreezablePerPlatformInt& Other)
-	: Default(Other.Default)
-{
-#if WITH_EDITORONLY_DATA
-	for (const TPair<FMemoryImageName, int32>& Pair : Other.PerPlatform)
-	{
-		PerPlatform.Add(FName(Pair.Key), Pair.Value);
-	}
-#endif
-}
-
-extern template ENGINE_API FArchive& operator<<(FArchive&, TPerPlatformProperty<FPerPlatformInt, int32, NAME_IntProperty>&);
-extern template ENGINE_API void operator<<(FStructuredArchive::FSlot Slot, TPerPlatformProperty<FPerPlatformInt, int32, NAME_IntProperty>&);
-
-template<>
-struct TStructOpsTypeTraits<FPerPlatformInt>
-	: public TStructOpsTypeTraitsBase2<FPerPlatformInt>
-{
-	enum
-	{
-		WithSerializeFromMismatchedTag = true,
-		WithSerializer = true
-	};
-	static constexpr EPropertyObjectReferenceType WithSerializerObjectReferences = EPropertyObjectReferenceType::None;
-};
-
-struct FFreezablePerPlatformFloat;
-
-/** FPerPlatformFloat - float property with per-platform overrides */
-USTRUCT(BlueprintType, meta = (CanFlattenStruct))
-struct FPerPlatformFloat
-#if CPP
-:	public TPerPlatformProperty<FPerPlatformFloat, float, NAME_FloatProperty>
-#endif
-{
-	GENERATED_USTRUCT_BODY()
-
-	UPROPERTY(BlueprintReadOnly, EditAnywhere, Category = PerPlatform)
-	float Default;
-
-#if WITH_EDITORONLY_DATA
-	UPROPERTY(EditAnywhere, Category = PerPlatform)
-	TMap<FName, float> PerPlatform;
-#endif
-
-	FPerPlatformFloat()
-	:	Default(0.f)
-	{
-	}
-
-	FPerPlatformFloat(float InDefaultValue)
-	:	Default(InDefaultValue)
-	{
-	}
-
-	ENGINE_API FPerPlatformFloat(const FFreezablePerPlatformFloat& Other);
-};
-extern template ENGINE_API FArchive& operator<<(FArchive&, TPerPlatformProperty<FPerPlatformFloat, float, NAME_FloatProperty>&);
-
-struct FFreezablePerPlatformFloat
-#if CPP
-	: public TPerPlatformProperty<FFreezablePerPlatformFloat, float, NAME_FloatProperty>
-#endif
-{
-	DECLARE_TYPE_LAYOUT(FFreezablePerPlatformFloat, NonVirtual);
-public:
-	using FPerPlatformMap = TMemoryImageMap<FMemoryImageName, float>;
-	
-	LAYOUT_FIELD(float, Default);
-	LAYOUT_FIELD_EDITORONLY(FPerPlatformMap, PerPlatform);
-
-	FFreezablePerPlatformFloat() : Default(0.0f) {}
-	FFreezablePerPlatformFloat(float InDefaultValue) : Default(InDefaultValue) {}
-	FFreezablePerPlatformFloat(const FPerPlatformFloat& Other)
-		: Default(Other.Default)
-	{
-#if WITH_EDITORONLY_DATA
-		for (const TPair<FName, float>& Pair : Other.PerPlatform)
-		{
-			PerPlatform.Add(FMemoryImageName(Pair.Key), Pair.Value);
-		}
-#endif
-}
-};
-
-inline FPerPlatformFloat::FPerPlatformFloat(const FFreezablePerPlatformFloat& Other)
-	: Default(Other.Default)
-{
-#if WITH_EDITORONLY_DATA
-	for (const TPair<FMemoryImageName, float>& Pair : Other.PerPlatform)
-	{
-		PerPlatform.Add(FName(Pair.Key), Pair.Value);
-	}
-#endif
-}
-
-template<>
-struct TStructOpsTypeTraits<FPerPlatformFloat>
-:	public TStructOpsTypeTraitsBase2<FPerPlatformFloat>
-{
-	enum
-	{
-		WithSerializeFromMismatchedTag = true,
-		WithSerializer = true
-	};
-	static constexpr EPropertyObjectReferenceType WithSerializerObjectReferences = EPropertyObjectReferenceType::None;
-};
-
-/** FPerPlatformBool - bool property with per-platform overrides */
-USTRUCT()
-struct FPerPlatformBool
-#if CPP
-:	public TPerPlatformProperty<FPerPlatformBool, bool, NAME_BoolProperty>
-#endif
-{
-	GENERATED_USTRUCT_BODY()
-
-	UPROPERTY(EditAnywhere, Category = PerPlatform)
-	bool Default;
-
-#if WITH_EDITORONLY_DATA
-	UPROPERTY(EditAnywhere, Category = PerPlatform)
-	TMap<FName, bool> PerPlatform;
-#endif
-
-	FPerPlatformBool()
-	:	Default(false)
-	{
-	}
-
-	FPerPlatformBool(bool InDefaultValue)
-	:	Default(InDefaultValue)
-	{
-	}
-};
-extern template ENGINE_API FArchive& operator<<(FArchive&, TPerPlatformProperty<FPerPlatformBool, bool, NAME_BoolProperty>&);
-
-template<>
-struct TStructOpsTypeTraits<FPerPlatformBool>
-	: public TStructOpsTypeTraitsBase2<FPerPlatformBool>
-{
-	enum
-	{
-		WithSerializeFromMismatchedTag = true,
-		WithSerializer = true
-	};
-	static constexpr EPropertyObjectReferenceType WithSerializerObjectReferences = EPropertyObjectReferenceType::None;
-};
-
-/** FPerPlatformFrameRate - FFrameRate property with per-platform overrides */
-USTRUCT(BlueprintType)
-struct FPerPlatformFrameRate
-#if CPP
-:	public TPerPlatformProperty<FPerPlatformFrameRate, FFrameRate, NAME_FrameRate>
-#endif
-{
-	GENERATED_USTRUCT_BODY()
-
-	UPROPERTY(EditAnywhere, Category = PerPlatform)
-	FFrameRate Default;
-
-#if WITH_EDITORONLY_DATA
-	UPROPERTY(EditAnywhere, Category = PerPlatform)
-	TMap<FName, FFrameRate> PerPlatform;
-#endif
-
-	FPerPlatformFrameRate()
-	:	Default(30, 1)
-	{
-	}
-
-	FPerPlatformFrameRate(FFrameRate InDefaultValue)
-	:	Default(InDefaultValue)
-	{
-	}
-	
-	bool SerializeFromMismatchedTag(const FPropertyTag& Tag, FArchive& Ar)
-	{
-		if(const UStruct* FrameRateStruct = FindObject<UStruct>(FTopLevelAssetPath("/Script/CoreUObject.FrameRate")))
-		{
-			FFrameRate Value;
-			Ar << Value.Denominator;
-			Ar << Value.Numerator;
-			Default = Value;
-
-			return true;
-		}
-		
-		return false;
-	}
-};
-
-extern template ENGINE_API FArchive& operator<<(FArchive&, TPerPlatformProperty<FPerPlatformFrameRate, FFrameRate, NAME_FrameRate>&);
-
-template<>
-struct TStructOpsTypeTraits<FPerPlatformFrameRate>
-	: public TStructOpsTypeTraitsBase2<FPerPlatformFrameRate>
-{
-	enum
-	{
-		WithSerializeFromMismatchedTag = false,
-		WithSerializer = true
-	};
-	static constexpr EPropertyObjectReferenceType WithSerializerObjectReferences = EPropertyObjectReferenceType::None;
-};
+#include "UObject/PerPlatformProperties.h"
