Added check() in SVT streaming manager to catch potential zero-sized DDC requests.

#rnx

diff --git a/Engine/Source/Runtime/Engine/Private/SparseVolumeTexture/SparseVolumeTextureStreamingManager.cpp b/Engine/Source/Runtime/Engine/Private/SparseVolumeTexture/SparseVolumeTextureStreamingManager.cpp
--- a/Engine/Source/Runtime/Engine/Private/SparseVolumeTexture/SparseVolumeTextureStreamingManager.cpp
+++ b/Engine/Source/Runtime/Engine/Private/SparseVolumeTexture/SparseVolumeTextureStreamingManager.cpp
@@ -1861,6 +1861,8 @@ UE::DerivedData::FCacheGetChunkRequest FStreamingManager::BuildDDCRequest(const
 	const uint32 ReadSizeInChunk = Resources.StreamingMetaData.TileDataOffsets[FirstTileIndex + NumTiles] - Resources.StreamingMetaData.TileDataOffsets[FirstTileIndex];
 	const uint32 ChunkTotalSize = Resources.StreamingMetaData.TileDataOffsets[LastTileIndexInChunkPlusOne] - Resources.StreamingMetaData.TileDataOffsets[FirstTileIndexInChunk];
 
+	checkf(ReadSizeInChunk > 0, TEXT("DDC chunk request sizes must be greater than zero."));
+
 	FCacheGetChunkRequest Request;
 	Request.Id = FValueId(FMemoryView(Resources.DDCChunkIds[ChunkIndex].GetData(), 12));
 	Request.Key.Bucket = FCacheBucket(TEXT("SparseVolumeTexture"));
