[NetPhysics] Implement ISimCallback on first PreResimStep
- virtual void FirstPreResimStep_Internal(int32 PhysicsStep) override
- Requires ESimCallbackOptions::Rewind

diff --git a/Engine/Source/Runtime/Engine/Private/PhysicsEngine/NetworkPhysicsComponent.cpp b/Engine/Source/Runtime/Engine/Private/PhysicsEngine/NetworkPhysicsComponent.cpp
--- a/Engine/Source/Runtime/Engine/Private/PhysicsEngine/NetworkPhysicsComponent.cpp
+++ b/Engine/Source/Runtime/Engine/Private/PhysicsEngine/NetworkPhysicsComponent.cpp
@@ -174,6 +174,22 @@ void FNetworkPhysicsCallback::ProcessInputs_Internal(int32 PhysicsStep, const TA
 	PostProcessInputsInternal.Broadcast(PhysicsStep);
 }
 
+void FNetworkPhysicsCallback::PreResimStep_Internal(int32 PhysicsStep, bool bFirst)
+{
+	if (bFirst)
+	{
+		for (Chaos::ISimCallbackObject* SimCallbackObject : RewindableCallbackObjects)
+		{
+			SimCallbackObject->FirstPreResimStep_Internal(PhysicsStep);
+		}
+	}
+}
+
+void FNetworkPhysicsCallback::PostResimStep_Internal(int32 PhysicsStep)
+{
+
+}
+
 int32 FNetworkPhysicsCallback::TriggerRewindIfNeeded_Internal(int32 LatestStepCompleted)
 {
 	int32 ResimFrame = INDEX_NONE;
diff --git a/Engine/Source/Runtime/Engine/Public/Physics/NetworkPhysicsComponent.h b/Engine/Source/Runtime/Engine/Public/Physics/NetworkPhysicsComponent.h
--- a/Engine/Source/Runtime/Engine/Public/Physics/NetworkPhysicsComponent.h
+++ b/Engine/Source/Runtime/Engine/Public/Physics/NetworkPhysicsComponent.h
@@ -319,6 +319,8 @@ struct FNetworkPhysicsCallback : public Chaos::IRewindCallback
 	virtual void ProcessInputs_External(int32 PhysicsStep, const TArray<Chaos::FSimCallbackInputAndObject>& SimCallbackInputs);
 	virtual void ProcessInputs_Internal(int32 PhysicsStep, const TArray<Chaos::FSimCallbackInputAndObject>& SimCallbackInputs) override;
 	virtual void ApplyCallbacks_Internal(int32 PhysicsStep, const TArray<Chaos::ISimCallbackObject*>& SimCallbackObjects) override;
+	virtual void PreResimStep_Internal(int32 PhysicsStep, bool bFirst) override;
+	virtual void PostResimStep_Internal(int32 PhysicsStep) override;
 	virtual int32 TriggerRewindIfNeeded_Internal(int32 LatestStepCompleted) override;
 	virtual void RegisterRewindableSimCallback_Internal(Chaos::ISimCallbackObject* SimCallbackObject) override
 	{
