[ActorDescContainer] Don't register delegates for non-editor worlds (was causing problems when using PIE travel).

diff --git a/Engine/Source/Runtime/Engine/Private/WorldPartition/ActorDescContainer.cpp b/Engine/Source/Runtime/Engine/Private/WorldPartition/ActorDescContainer.cpp
--- a/Engine/Source/Runtime/Engine/Private/WorldPartition/ActorDescContainer.cpp
+++ b/Engine/Source/Runtime/Engine/Private/WorldPartition/ActorDescContainer.cpp
@@ -44,6 +44,7 @@ UActorDescContainer::UActorDescContainer(const FObjectInitializer& ObjectInitial
 	: Super(ObjectInitializer)
 #if WITH_EDITOR
 	, bContainerInitialized(false)
+	, bRegisteredDelegates(false)
 #endif
 {}
 
@@ -186,7 +187,12 @@ void UActorDescContainer::Initialize(const FInitializeParams& InitParams)
 		OnActorDescContainerInitialized.Broadcast(this);
 	}
 
-	RegisterEditorDelegates();
+	bRegisteredDelegates = InitParams.bShouldRegisterEditorDeletages && ShouldRegisterDelegates();
+	
+	if (bRegisteredDelegates)
+	{
+		RegisterEditorDelegates();
+	}
 
 	bContainerInitialized = true;
 }
@@ -195,7 +201,11 @@ void UActorDescContainer::Uninitialize()
 {
 	if (bContainerInitialized)
 	{
-		UnregisterEditorDelegates();
+		if (bRegisteredDelegates)
+		{
+			UnregisterEditorDelegates();
+			bRegisteredDelegates = false;
+		}
 		bContainerInitialized = false;
 	}
 
@@ -473,30 +483,24 @@ bool UActorDescContainer::ShouldRegisterDelegates() const
 
 void UActorDescContainer::RegisterEditorDelegates()
 {
-	if (ShouldRegisterDelegates())
-	{
-		FCoreUObjectDelegates::OnObjectPreSave.AddUObject(this, &UActorDescContainer::OnObjectPreSave);
-		FEditorDelegates::OnPackageDeleted.AddUObject(this, &UActorDescContainer::OnPackageDeleted);
+	FCoreUObjectDelegates::OnObjectPreSave.AddUObject(this, &UActorDescContainer::OnObjectPreSave);
+	FEditorDelegates::OnPackageDeleted.AddUObject(this, &UActorDescContainer::OnPackageDeleted);
 
-		FWorldPartitionClassDescRegistry& ClassDescRegistry = FWorldPartitionClassDescRegistry::Get();
-		ClassDescRegistry.OnClassDescriptorUpdated().AddUObject(this, &UActorDescContainer::OnClassDescriptorUpdated);
+	FWorldPartitionClassDescRegistry& ClassDescRegistry = FWorldPartitionClassDescRegistry::Get();
+	ClassDescRegistry.OnClassDescriptorUpdated().AddUObject(this, &UActorDescContainer::OnClassDescriptorUpdated);
 
-		UDeletedObjectPlaceholder::OnObjectCreated.AddUObject(this, &UActorDescContainer::OnDeletedObjectPlaceholderCreated);
-	}
+	UDeletedObjectPlaceholder::OnObjectCreated.AddUObject(this, &UActorDescContainer::OnDeletedObjectPlaceholderCreated);
 }
 
 void UActorDescContainer::UnregisterEditorDelegates()
 {
-	if (ShouldRegisterDelegates())
-	{
-		FCoreUObjectDelegates::OnObjectPreSave.RemoveAll(this);
-		FEditorDelegates::OnPackageDeleted.RemoveAll(this);
+	FCoreUObjectDelegates::OnObjectPreSave.RemoveAll(this);
+	FEditorDelegates::OnPackageDeleted.RemoveAll(this);
 
-		FWorldPartitionClassDescRegistry& ClassDescRegistry = FWorldPartitionClassDescRegistry::Get();
-		ClassDescRegistry.OnClassDescriptorUpdated().RemoveAll(this);
+	FWorldPartitionClassDescRegistry& ClassDescRegistry = FWorldPartitionClassDescRegistry::Get();
+	ClassDescRegistry.OnClassDescriptorUpdated().RemoveAll(this);
 
-		UDeletedObjectPlaceholder::OnObjectCreated.RemoveAll(this);
-	}
+	UDeletedObjectPlaceholder::OnObjectCreated.RemoveAll(this);
 }
 
 void UActorDescContainer::OnActorDescAdded(FWorldPartitionActorDesc* NewActorDesc)
diff --git a/Engine/Source/Runtime/Engine/Private/WorldPartition/ActorDescContainerInstance.cpp b/Engine/Source/Runtime/Engine/Private/WorldPartition/ActorDescContainerInstance.cpp
--- a/Engine/Source/Runtime/Engine/Private/WorldPartition/ActorDescContainerInstance.cpp
+++ b/Engine/Source/Runtime/Engine/Private/WorldPartition/ActorDescContainerInstance.cpp
@@ -36,6 +36,7 @@ void UActorDescContainerInstance::RegisterContainer(const FInitializeParams& InP
 	};
 	ContainerInitParams.ContentBundleGuid = InParams.ContentBundleGuid;
 	ContainerInitParams.ExternalDataLayerAsset = InParams.ExternalDataLayerAsset;
+	ContainerInitParams.bShouldRegisterEditorDeletages = InParams.bShouldRegisterEditorDeletages;
 
 	SetContainer(UActorDescContainerSubsystem::GetChecked().RegisterContainer(ContainerInitParams));
 }
@@ -168,7 +169,12 @@ void UActorDescContainerInstance::Initialize(const FInitializeParams& InParams)
 	}
 
 	// Register Delegates
-	RegisterDelegates();
+	bRegisteredDelegates = InParams.bShouldRegisterEditorDeletages && ShouldRegisterDelegates();
+	
+	if (bRegisteredDelegates)
+	{
+		RegisterDelegates();
+	}
 
 	bIsInitialized = true;
 }
@@ -193,7 +199,11 @@ void UActorDescContainerInstance::Uninitialize()
 
 	bIsInitialized = false;
 
-	UnregisterDelegates();
+	if (bRegisteredDelegates)
+	{
+		UnregisterDelegates();
+		bRegisteredDelegates = false;
+	}
 
 	for (TUniquePtr<FWorldPartitionActorDescInstance>& ActorDescInstancePtr : ActorDescList)
 	{
@@ -228,53 +238,47 @@ bool UActorDescContainerInstance::ShouldRegisterDelegates() const
 
 void UActorDescContainerInstance::RegisterDelegates()
 {
-	if (ShouldRegisterDelegates())
-	{
-		check(Container);
+	check(Container);
 
-		// Only listen to Object replaced events on ContainerInstance that have a direct World Partition outer (Loaded Container Instances: Main World or Loaded Level Instances)
-		if (GetOuterWorldPartition())
-		{
-			FCoreUObjectDelegates::OnObjectsReplaced.AddUObject(this, &UActorDescContainerInstance::OnObjectsReplaced);
-		}
+	// Only listen to Object replaced events on ContainerInstance that have a direct World Partition outer (Loaded Container Instances: Main World or Loaded Level Instances)
+	if (GetOuterWorldPartition())
+	{
+		FCoreUObjectDelegates::OnObjectsReplaced.AddUObject(this, &UActorDescContainerInstance::OnObjectsReplaced);
+	}
 
-		// Only listen to this event if we have registered Child Container Instances
-		if (bCreateChildContainerHierarchy)
-		{
-			UActorDescContainerSubsystem::GetChecked().ContainerReplaced().AddUObject(this, &UActorDescContainerInstance::OnContainerReplaced);
-		}
+	// Only listen to this event if we have registered Child Container Instances
+	if (bCreateChildContainerHierarchy)
+	{
+		UActorDescContainerSubsystem::GetChecked().ContainerReplaced().AddUObject(this, &UActorDescContainerInstance::OnContainerReplaced);
+	}
 
-		// No need to register Added descs events for instanced worlds as they don't support it for now (Level Instances get reloaded after an edit)
-		if (!GetInstancingContext())
-		{
-			Container->OnActorDescAddedEvent.AddUObject(this, &UActorDescContainerInstance::OnActorDescAdded);
-		}
+	// No need to register Added descs events for instanced worlds as they don't support it for now (Level Instances get reloaded after an edit)
+	if (!GetInstancingContext())
+	{
+		Container->OnActorDescAddedEvent.AddUObject(this, &UActorDescContainerInstance::OnActorDescAdded);
+	}
 
-		// Important to hook the other events that will invalidate existing FWorldPartitionActorDescInstance's because those can be hashed and loaded
-		// even in Instanced worlds (Level Instances)
-		Container->OnActorDescRemovedEvent.AddUObject(this, &UActorDescContainerInstance::OnActorDescRemoved);
+	// Important to hook the other events that will invalidate existing FWorldPartitionActorDescInstance's because those can be hashed and loaded
+	// even in Instanced worlds (Level Instances)
+	Container->OnActorDescRemovedEvent.AddUObject(this, &UActorDescContainerInstance::OnActorDescRemoved);
 		
-		Container->OnActorDescUpdatingEvent.AddUObject(this, &UActorDescContainerInstance::OnActorDescUpdating);
-		Container->OnActorDescUpdatedEvent.AddUObject(this, &UActorDescContainerInstance::OnActorDescUpdated);
-	}
+	Container->OnActorDescUpdatingEvent.AddUObject(this, &UActorDescContainerInstance::OnActorDescUpdating);
+	Container->OnActorDescUpdatedEvent.AddUObject(this, &UActorDescContainerInstance::OnActorDescUpdated);
 }
 
 void UActorDescContainerInstance::UnregisterDelegates()
 {
-	if (ShouldRegisterDelegates())
-	{
-		FCoreUObjectDelegates::OnObjectsReplaced.RemoveAll(this);
+	FCoreUObjectDelegates::OnObjectsReplaced.RemoveAll(this);
 
-		check(Container);
-		Container->OnActorDescAddedEvent.RemoveAll(this);
-		Container->OnActorDescRemovedEvent.RemoveAll(this);
-		Container->OnActorDescUpdatingEvent.RemoveAll(this);
-		Container->OnActorDescUpdatedEvent.RemoveAll(this);
+	check(Container);
+	Container->OnActorDescAddedEvent.RemoveAll(this);
+	Container->OnActorDescRemovedEvent.RemoveAll(this);
+	Container->OnActorDescUpdatingEvent.RemoveAll(this);
+	Container->OnActorDescUpdatedEvent.RemoveAll(this);
 
-		if (UActorDescContainerSubsystem* ActorDescContainerSubsystem = UActorDescContainerSubsystem::Get())
-		{
-			ActorDescContainerSubsystem->ContainerReplaced().RemoveAll(this);
-		}
+	if (UActorDescContainerSubsystem* ActorDescContainerSubsystem = UActorDescContainerSubsystem::Get())
+	{
+		ActorDescContainerSubsystem->ContainerReplaced().RemoveAll(this);
 	}
 }
 
diff --git a/Engine/Source/Runtime/Engine/Public/WorldPartition/ActorDescContainer.h b/Engine/Source/Runtime/Engine/Public/WorldPartition/ActorDescContainer.h
--- a/Engine/Source/Runtime/Engine/Public/WorldPartition/ActorDescContainer.h
+++ b/Engine/Source/Runtime/Engine/Public/WorldPartition/ActorDescContainer.h
@@ -75,6 +75,9 @@ public:
 		/** The associated Content Bundle Guid */
 		FGuid ContentBundleGuid;
 
+		/** If the container should bind to editor events */
+		bool bShouldRegisterEditorDeletages = true;
+
 		/** The associated External Data Layer Asset */
 		const UExternalDataLayerAsset* ExternalDataLayerAsset = nullptr;
 
@@ -145,6 +148,7 @@ public:
 	virtual ENGINE_API const FWorldPartitionActorDesc* GetActorDescByName(FName ActorName) const;
 
 	bool bContainerInitialized;
+	bool bRegisteredDelegates;
 
 	FName ContainerPackageName;
 	FGuid ContentBundleGuid;
diff --git a/Engine/Source/Runtime/Engine/Public/WorldPartition/ActorDescContainerInstance.h b/Engine/Source/Runtime/Engine/Public/WorldPartition/ActorDescContainerInstance.h
--- a/Engine/Source/Runtime/Engine/Public/WorldPartition/ActorDescContainerInstance.h
+++ b/Engine/Source/Runtime/Engine/Public/WorldPartition/ActorDescContainerInstance.h
@@ -27,6 +27,7 @@ protected:
 	UActorDescContainerInstance()
 #if WITH_EDITORONLY_DATA
 		: bIsInitialized(false) 
+		, bRegisteredDelegates(false)
 		, bCreateChildContainerHierarchy(false)
 #endif
 	{}
@@ -41,8 +42,7 @@ public:
 		FInitializeParams(FName InContainerPackageName, bool bInCreateContainerInstanceHierarchy = false)
 			: ContainerPackageName(InContainerPackageName)
 			, bCreateContainerInstanceHierarchy(bInCreateContainerInstanceHierarchy)
-		{
-		}
+		{}
 				
 		FInitializeParams& SetParent(const UActorDescContainerInstance* InParentContainerInstance, const FGuid& InContainerActorGuid)
 		{
@@ -58,6 +58,12 @@ public:
 			return *this;
 		}
 
+		FInitializeParams& SetShouldRegisterEditorDeletages(bool bInShouldRegisterEditorDeletages)
+		{
+			bShouldRegisterEditorDeletages = bInShouldRegisterEditorDeletages;
+			return *this;
+		}
+
 		FName ContainerPackageName;
 
 		FGuid ContainerActorGuid;
@@ -73,6 +79,7 @@ public:
 		TOptional<FTransform> Transform;
 
 		bool bCreateContainerInstanceHierarchy;
+		bool bShouldRegisterEditorDeletages = true;
 
 		/* Custom filter function used to filter actors descriptors. */
 		TUniqueFunction<bool(const FWorldPartitionActorDesc*)> FilterActorDescFunc;
@@ -195,6 +202,7 @@ protected:
 	TMap<FGuid, TObjectPtr<UActorDescContainerInstance>>		ChildContainerInstances;
 
 	bool														bIsInitialized;
+	bool														bRegisteredDelegates;
 	bool														bCreateChildContainerHierarchy;
 #endif
 };
\ No newline at end of file
