Fixed case where dummy morph buffer was not big enough for a skeletal mesh rendering, causing GPU crashes on some platforms

diff --git a/Engine/Source/Runtime/Engine/Private/GPUSkinVertexFactory.cpp b/Engine/Source/Runtime/Engine/Private/GPUSkinVertexFactory.cpp
--- a/Engine/Source/Runtime/Engine/Private/GPUSkinVertexFactory.cpp
+++ b/Engine/Source/Runtime/Engine/Private/GPUSkinVertexFactory.cpp
@@ -122,6 +122,26 @@ static_assert(sizeof(ispc::FMatrix44f) == sizeof(FMatrix44f), "sizeof(ispc::FMat
 static_assert(sizeof(ispc::FMatrix3x4) == sizeof(FMatrix3x4), "sizeof(ispc::FMatrix3x4) != sizeof(FMatrix3x4)");
 #endif
 
+class FNullMorphVertexBuffer : public FVertexBuffer
+{
+public:
+	FNullMorphVertexBuffer() = default;
+	~FNullMorphVertexBuffer() = default;
+
+	virtual void InitRHI(FRHICommandListBase& RHICmdList) override
+	{
+		// Enough data for 64k vertices mesh
+		uint32 Size = sizeof(FMorphGPUSkinVertex) * 65535;
+		FRHIResourceCreateInfo CreateInfo(TEXT("FNullMorphVertexBuffer"));
+		VertexBufferRHI = RHICmdList.CreateBuffer(Size, BUF_Static | BUF_VertexBuffer | BUF_ShaderResource, 0, ERHIAccess::VertexOrIndexBuffer | ERHIAccess::SRVMask, CreateInfo);
+		void* LockedData = RHICmdList.LockBuffer(VertexBufferRHI, 0, Size, RLM_WriteOnly);
+		FMemory::Memzero(LockedData, Size);
+		RHICmdList.UnlockBuffer(VertexBufferRHI);
+	}
+};
+
+TGlobalResource<FNullMorphVertexBuffer, FRenderResource::EInitPhase::Pre> GNullMorphVertexBuffer;
+
 /*-----------------------------------------------------------------------------
  FSharedPoolPolicyData
  -----------------------------------------------------------------------------*/
@@ -618,7 +638,7 @@ void FGPUBaseSkinVertexFactory::InitRHI(FRHICommandListBase& RHICmdList)
 		}
 	}
 
-	MorphDeltaBufferSlot = FRHIStreamSourceSlot::Create(GNullVertexBuffer.VertexBufferRHI.GetReference());
+	MorphDeltaBufferSlot = FRHIStreamSourceSlot::Create(GNullMorphVertexBuffer.VertexBufferRHI.GetReference());
 }
 
 void FGPUBaseSkinVertexFactory::ReleaseRHI()
@@ -669,7 +689,7 @@ void FGPUBaseSkinVertexFactory::UpdateMorphState(FRHICommandListBase& RHICmdList
 	if (bUseMorphTarget)
 	{
 		const FMorphVertexBuffer* MorphVertexBuffer = GetMorphVertexBuffer(false);
-		RHICmdList.UpdateStreamSourceSlot(MorphDeltaBufferSlot, MorphVertexBuffer ? MorphVertexBuffer->VertexBufferRHI : GNullVertexBuffer.VertexBufferRHI);
+		RHICmdList.UpdateStreamSourceSlot(MorphDeltaBufferSlot, MorphVertexBuffer ? MorphVertexBuffer->VertexBufferRHI : GNullMorphVertexBuffer.VertexBufferRHI);
 	}
 }
 
