Texture.cpp use ensure instead of check on condition that can be hit when loading corrupt assets

diff --git a/Engine/Source/Runtime/Engine/Private/Texture.cpp b/Engine/Source/Runtime/Engine/Private/Texture.cpp
--- a/Engine/Source/Runtime/Engine/Private/Texture.cpp
+++ b/Engine/Source/Runtime/Engine/Private/Texture.cpp
@@ -236,7 +236,7 @@ const FTextureResource* UTexture::GetResource() const
 		return PrivateResourceRenderThread;
 	}
 
-	ensureMsgf(false, TEXT("Attempted to access a texture resource from an unkown thread."));
+	ensureMsgf(false, TEXT("Attempted to access a texture resource from an unknown thread."));
 	return nullptr;
 }
 
@@ -251,7 +251,7 @@ FTextureResource* UTexture::GetResource()
 		return PrivateResourceRenderThread;
 	}
 
-	ensureMsgf(false, TEXT("Attempted to access a texture resource from an unkown thread."));
+	ensureMsgf(false, TEXT("Attempted to access a texture resource from an unknown thread."));
 	return nullptr;
 }
 
@@ -3038,7 +3038,17 @@ FSharedBuffer FTextureSource::DoUEDeltaTransform(FSharedBuffer InBuffer,bool bFo
 {
 	int64 InBufferSize = InBuffer.GetSize();		
 	int64 ImageSize = CalcTotalSize();
-	check( InBufferSize == ImageSize );
+	
+	if ( InBufferSize != ImageSize )
+	{
+		// this can be hit on corrupt uassets
+		ensureMsgf( InBufferSize == ImageSize , 
+			TEXT("DoUEDeltaTransform InBufferSize = %lld ImageSize = %lld mismatch ; likely corrupt asset."),
+			InBufferSize, ImageSize
+			);
+
+		return FSharedBuffer();
+	}
 
 	if ( ImageSize == 0 )
 	{
