MarkRenderDynamicDataDirty -> MarkRenderStateDirty

diff --git a/Engine/Source/Runtime/Engine/Private/Components/SkinnedMeshComponent.cpp b/Engine/Source/Runtime/Engine/Private/Components/SkinnedMeshComponent.cpp
--- a/Engine/Source/Runtime/Engine/Private/Components/SkinnedMeshComponent.cpp
+++ b/Engine/Source/Runtime/Engine/Private/Components/SkinnedMeshComponent.cpp
@@ -5051,7 +5051,7 @@ void USkinnedMeshComponent::UpdateSkinWeightOverrideBuffer()
 {
 	// Force a mesh update to ensure bone buffers are up to date
 	bForceMeshObjectUpdate = true;
-	MarkRenderDynamicDataDirty();
+	MarkRenderStateDirty();
 
 	// Queue an update of the skin weight buffer used by the current Mesh Object
 	if (MeshObject)
