Added StaticLighting support for LevelInstances

MapBuildDataID in StaticMeshComponent & LightBaseComponent are now always recomputed from OriginalMapBuildDataID & ActorInstanceGuid, this allows for LevelInstance Actors to have unique/stable IDs for linking the Static Lighting data
Added helper FActorInstanceGuid::GetActorInstanceGuid that relies on serialized ActorGuids (that'll be released after OnRegister()) to obtain the ActorInstanceGuid
Added helper ULevelInstanceSubsystem::GetOwningLevel() to identify which level acts as the serialized container for Actors in LevelInstances
StaticLighting data for LevelInstance actors is stored in their owning level and rebinded together on load.

diff --git a/Engine/Source/Editor/UnrealEd/Private/StaticLightingSystem/StaticLightingSystem.cpp b/Engine/Source/Editor/UnrealEd/Private/StaticLightingSystem/StaticLightingSystem.cpp
--- a/Engine/Source/Editor/UnrealEd/Private/StaticLightingSystem/StaticLightingSystem.cpp
+++ b/Engine/Source/Editor/UnrealEd/Private/StaticLightingSystem/StaticLightingSystem.cpp
@@ -87,6 +87,7 @@ DEFINE_LOG_CATEGORY(LogStaticLightingSystem);
 #include "BuildSettings.h"
 #include "Misc/EngineBuildSettings.h"
 #include "TargetReceipt.h"
+#include "LevelInstance/LevelInstanceSubsystem.h"
 
 #define LOCTEXT_NAMESPACE "StaticLightingSystem"
 
@@ -1368,6 +1369,13 @@ void FStaticLightingSystem::ApplyNewLightingData(bool bLightingSuccessful)
 
 			const bool bBuildLightingForLevel = Options.ShouldBuildLightingForLevel( Level );
 
+			// in this specific case the Level MapBuildData is shared with another level and 
+			// this specific level doesn't own it's data so we need to skip it
+			if (!Level->IsMapBuildDataOwner())
+			{
+				continue;
+			}
+
 			UMapBuildDataRegistry* Registry = LightingContext.GetRegistryForLevel(Level);
 		
 			if (Level->IsPersistentLevel() && LightingContext.World->IsPartitionedWorld())
diff --git a/Engine/Source/Runtime/Core/Public/UObject/FortniteMainBranchObjectVersion.h b/Engine/Source/Runtime/Core/Public/UObject/FortniteMainBranchObjectVersion.h
--- a/Engine/Source/Runtime/Core/Public/UObject/FortniteMainBranchObjectVersion.h
+++ b/Engine/Source/Runtime/Core/Public/UObject/FortniteMainBranchObjectVersion.h
@@ -462,6 +462,9 @@ struct FFortniteMainBranchObjectVersion
 		// Fix to get full name of templated type ( Tarray > TArray<Float> for example )
 		DataflowTemplatedTypeFix,
 
+		// Changes for LevelInstance support in StaticLighting
+		LevelInstanceStaticLightingSupport,
+
 		// -----<new versions can be added above this line>-------------------------------------------------
 		VersionPlusOne,
 		LatestVersion = VersionPlusOne - 1
diff --git a/Engine/Source/Runtime/Engine/Classes/Components/LightComponent.h b/Engine/Source/Runtime/Engine/Classes/Components/LightComponent.h
--- a/Engine/Source/Runtime/Engine/Classes/Components/LightComponent.h
+++ b/Engine/Source/Runtime/Engine/Classes/Components/LightComponent.h
@@ -513,6 +513,7 @@ struct FPrecomputedLightInstanceData : public FSceneComponentInstanceData
 	FPrecomputedLightInstanceData(const ULightComponent* SourceComponent)
 		: FSceneComponentInstanceData(SourceComponent)
 		, Transform(SourceComponent->GetComponentTransform())
+		, OriginalLightGuid(SourceComponent->OriginalLightGuid)		
 		, LightGuid(SourceComponent->LightGuid)
 		, PreviewShadowMapChannel(SourceComponent->PreviewShadowMapChannel)
 	{}
@@ -532,6 +533,9 @@ struct FPrecomputedLightInstanceData : public FSceneComponentInstanceData
 	UPROPERTY()
 	FTransform Transform;
 
+	UPROPERTY()
+	FGuid OriginalLightGuid;
+
 	UPROPERTY()
 	FGuid LightGuid;
 
diff --git a/Engine/Source/Runtime/Engine/Classes/Components/LightComponentBase.h b/Engine/Source/Runtime/Engine/Classes/Components/LightComponentBase.h
--- a/Engine/Source/Runtime/Engine/Classes/Components/LightComponentBase.h
+++ b/Engine/Source/Runtime/Engine/Classes/Components/LightComponentBase.h
@@ -20,8 +20,11 @@ class ULightComponentBase : public USceneComponent
 	 * The GUID changes whenever the light position changes.
 	 */
 	UPROPERTY()
-	FGuid LightGuid;
+	FGuid OriginalLightGuid;
 
+	UPROPERTY()
+	FGuid LightGuid;
+	
 	/**  */
 	UPROPERTY()
 	float Brightness_DEPRECATED;
diff --git a/Engine/Source/Runtime/Engine/Classes/Components/SkyLightComponent.h b/Engine/Source/Runtime/Engine/Classes/Components/SkyLightComponent.h
--- a/Engine/Source/Runtime/Engine/Classes/Components/SkyLightComponent.h
+++ b/Engine/Source/Runtime/Engine/Classes/Components/SkyLightComponent.h
@@ -389,6 +389,9 @@ public:
 		CastChecked<USkyLightComponent>(Component)->ApplyComponentInstanceData(this);
 	}
 
+	UPROPERTY()
+	FGuid OriginalLightGuid;
+
 	UPROPERTY()
 	FGuid LightGuid;
 
diff --git a/Engine/Source/Runtime/Engine/Classes/Components/StaticMeshComponent.h b/Engine/Source/Runtime/Engine/Classes/Components/StaticMeshComponent.h
--- a/Engine/Source/Runtime/Engine/Classes/Components/StaticMeshComponent.h
+++ b/Engine/Source/Runtime/Engine/Classes/Components/StaticMeshComponent.h
@@ -788,6 +788,9 @@ public:
 
 	ENGINE_API const FMeshMapBuildData* GetMeshMapBuildData(const FStaticMeshComponentLODInfo& LODInfo, bool bCheckForResourceCluster = true) const;
 
+	void UpdateMapBuildDataId();
+	void UpdateStaticLightingData();
+
 	/** Called during scene proxy creation to get the Nanite resource data */
 	DECLARE_DELEGATE_RetVal(const Nanite::FResources*, FOnGetNaniteResources);
 	virtual FOnGetNaniteResources& OnGetNaniteResources() { return OnGetNaniteResourcesEvent; }
diff --git a/Engine/Source/Runtime/Engine/Classes/Engine/Level.h b/Engine/Source/Runtime/Engine/Classes/Engine/Level.h
--- a/Engine/Source/Runtime/Engine/Classes/Engine/Level.h
+++ b/Engine/Source/Runtime/Engine/Classes/Engine/Level.h
@@ -1113,6 +1113,8 @@ public:
 
 	ENGINE_API UMapBuildDataRegistry* GetOrCreateMapBuildData();
 
+	ENGINE_API bool IsMapBuildDataOwner() const;
+
 	/** Sets whether this level is a lighting scenario and handles propagating the change. */
 	ENGINE_API void SetLightingScenario(bool bNewIsLightingScenario);
 
diff --git a/Engine/Source/Runtime/Engine/Classes/Engine/MapBuildDataRegistry.h b/Engine/Source/Runtime/Engine/Classes/Engine/MapBuildDataRegistry.h
--- a/Engine/Source/Runtime/Engine/Classes/Engine/MapBuildDataRegistry.h
+++ b/Engine/Source/Runtime/Engine/Classes/Engine/MapBuildDataRegistry.h
@@ -23,6 +23,16 @@ class FPrecomputedVolumetricLightmapData;
 struct FAssetCompileData;
 struct FVolumetricLightMapGridDesc;
 
+DECLARE_LOG_CATEGORY_EXTERN(LogMapBuildDataRegistry, Log, All);
+
+#define UE_LOG_MAPBUILDATA_ENABLED 0
+
+#if UE_LOG_MAPBUILDATA_ENABLED
+#define UE_LOG_MAPBUILDDATA(...) UE_LOG(LogMapBuildDataRegistry, __VA_ARGS__)
+#else
+#define UE_LOG_MAPBUILDDATA(...) 
+#endif
+
 struct FPerInstanceLightmapData
 {
 	FVector2f LightmapUVBias;
diff --git a/Engine/Source/Runtime/Engine/Private/Actor.cpp b/Engine/Source/Runtime/Engine/Private/Actor.cpp
--- a/Engine/Source/Runtime/Engine/Private/Actor.cpp
+++ b/Engine/Source/Runtime/Engine/Private/Actor.cpp
@@ -73,6 +73,7 @@
 #include "WorldPartition/DataLayer/ExternalDataLayerInstance.h"
 #include "WorldPartition/DataLayer/DataLayerManager.h"
 #include "WorldPartition/ContentBundle/ContentBundlePaths.h"
+#include "WorldPartition/ActorInstanceGuids.h"
 
 DEFINE_LOG_CATEGORY(LogActor);
 
@@ -919,6 +920,11 @@ void AActor::Serialize(FArchive& Ar)
 		}
 	}
 #endif
+
+	if ((Ar.IsCooking() || Ar.IsLoadingFromCookedPackage()) && (Ar.CustomVer(FFortniteMainBranchObjectVersion::GUID) >= FFortniteMainBranchObjectVersion::LevelInstanceStaticLightingSupport))
+	{		
+		FActorInstanceGuid::Serialize(Ar, *this);
+	}
 }
 
 void AActor::PostLoad()
@@ -3619,6 +3625,8 @@ void AActor::PostRegisterAllComponents()
 	ensureMsgf(bHasRegisteredAllComponents == true, TEXT("bHasRegisteredAllComponents must be set to true prior to calling PostRegisterAllComponents()"));
 
 	FNavigationSystem::OnActorRegistered(*this);
+
+	FActorInstanceGuid::ReleaseActorInstanceGuid(*this);
 }
 
 /** Util to call OnComponentCreated on components */
diff --git a/Engine/Source/Runtime/Engine/Private/Components/LightComponent.cpp b/Engine/Source/Runtime/Engine/Private/Components/LightComponent.cpp
--- a/Engine/Source/Runtime/Engine/Private/Components/LightComponent.cpp
+++ b/Engine/Source/Runtime/Engine/Private/Components/LightComponent.cpp
@@ -33,6 +33,7 @@
 #include "UObject/SoftObjectPath.h"
 #include "UObject/UnrealType.h"
 #include "UObject/Package.h"
+#include "WorldPartition/ActorInstanceGuids.h"
 
 #include UE_INLINE_GENERATED_CPP_BY_NAME(LightComponent)
 
@@ -163,8 +164,17 @@ void ULightComponentBase::SetSamplesPerPixel(int NewValue)
 
 void ULightComponentBase::Serialize(FArchive& Ar)
 {
+	Ar.UsingCustomVersion(FFortniteMainBranchObjectVersion::GUID);
+
 	Super::Serialize(Ar);
 
+	if (Ar.CustomVer(FFortniteMainBranchObjectVersion::GUID) < FFortniteMainBranchObjectVersion::LevelInstanceStaticLightingSupport)
+	{
+		// Starting from there LightGuid is a mix of OrigignalLightGuid & ActorInstanceGuid so move the value where it belongs
+		OriginalLightGuid = LightGuid;
+		LightGuid.Invalidate();
+	}
+
 	if (Ar.UEVer() < VER_UE4_INVERSE_SQUARED_LIGHTS_DEFAULT)
 	{
 		Intensity = Brightness_DEPRECATED;
@@ -242,15 +252,19 @@ void ULightComponentBase::PostEditChangeProperty(FPropertyChangedEvent& Property
 void ULightComponentBase::ValidateLightGUIDs()
 {
 	// Validate light guids.
-	if (!LightGuid.IsValid())
+	if (!OriginalLightGuid.IsValid())
 	{
 		UpdateLightGUIDs();
 	}
+
+	
+	LightGuid = (OriginalLightGuid.IsValid() && GetOwner() )? FGuid::Combine( OriginalLightGuid , FActorInstanceGuid::GetActorInstanceGuid(*GetOwner())) : FGuid();
 }
 
 void ULightComponentBase::UpdateLightGUIDs()
 {
-	LightGuid = (HasStaticShadowing() ? FGuid::NewGuid() : FGuid());
+	OriginalLightGuid = (HasStaticShadowing() ? FGuid::NewGuid() : FGuid());
+	LightGuid = (OriginalLightGuid.IsValid() && GetOwner() )? FGuid::Combine( OriginalLightGuid , FActorInstanceGuid::GetActorInstanceGuid(*GetOwner())) : FGuid();
 }
 
 bool ULightComponentBase::HasStaticLighting() const
@@ -270,6 +284,7 @@ void ULightComponentBase::PostLoad()
 
 	if (!HasStaticShadowing())
 	{
+		OriginalLightGuid.Invalidate();
 		LightGuid.Invalidate();
 	}
 }
@@ -278,6 +293,8 @@ void ULightComponentBase::OnRegister()
 {
 	Super::OnRegister();
 
+	ValidateLightGUIDs();
+
 	if (SpriteComponent)
 	{
 		SpriteComponent->SpriteInfo.Category = TEXT("Lighting");
@@ -541,7 +558,7 @@ void ULightComponent::Serialize(FArchive& Ar)
 			LegacyData->ShadowMapChannel = ShadowMapChannel_DEPRECATED;
 
 			FLightComponentLegacyMapBuildData LegacyLightData;
-			LegacyLightData.Id = LightGuid;
+			LegacyLightData.Id = OriginalLightGuid;
 			LegacyLightData.Data = LegacyData;
 			GLightComponentsWithLegacyBuildData.AddAnnotation(this, MoveTemp(LegacyLightData));
 		}
@@ -830,10 +847,22 @@ bool ULightComponent::IsReadyForFinishDestroy()
 
 void ULightComponent::OnRegister()
 {
-	Super::OnRegister();
-
 	// Update GUIDs on attachment if they are not valid.
 	ValidateLightGUIDs();
+
+	if (OriginalLightGuid.IsValid() && !HasStaticLighting() && HasStaticShadowing())	
+	{		
+		// Since we support LevelInstances the channel must always be reassigned on load
+		if (UMapBuildDataRegistry* DataRegistry = UMapBuildDataRegistry::Get(this))
+		{
+			if (FLightComponentMapBuildData* LightBuildData = DataRegistry->GetLightBuildData(LightGuid))
+			{
+				PreviewShadowMapChannel = LightBuildData->ShadowMapChannel;
+			}
+		}
+	}
+
+	Super::OnRegister();
 }
 
 void ULightComponent::CreateRenderState_Concurrent(FRegisterComponentContext* Context)
@@ -1295,6 +1324,7 @@ void ULightComponent::InvalidateLightingCacheDetailed(bool bInvalidateBuildEnque
 	else
 	{
 		// Movable lights will have a GUID of 0
+		OriginalLightGuid.Invalidate();
 		LightGuid.Invalidate();
 	}
 }
@@ -1314,7 +1344,8 @@ void ULightComponent::ApplyComponentInstanceData(FPrecomputedLightInstanceData*
 		return;
 	}
 
-	LightGuid = (HasStaticShadowing() ? LightMapData->LightGuid : FGuid());
+	OriginalLightGuid = (HasStaticShadowing() ? LightMapData->OriginalLightGuid : FGuid());
+	LightGuid = (HasStaticShadowing() ? LightMapData->LightGuid : FGuid());	
 	PreviewShadowMapChannel = LightMapData->PreviewShadowMapChannel;
 
 	MarkRenderStateDirty();
diff --git a/Engine/Source/Runtime/Engine/Private/Components/SkyLightComponent.cpp b/Engine/Source/Runtime/Engine/Private/Components/SkyLightComponent.cpp
--- a/Engine/Source/Runtime/Engine/Private/Components/SkyLightComponent.cpp
+++ b/Engine/Source/Runtime/Engine/Private/Components/SkyLightComponent.cpp
@@ -692,6 +692,7 @@ TStructOnScope<FActorComponentInstanceData> USkyLightComponent::GetComponentInst
 {
 	TStructOnScope<FActorComponentInstanceData> InstanceData = MakeStructOnScope<FActorComponentInstanceData, FPrecomputedSkyLightInstanceData>(this);
 	FPrecomputedSkyLightInstanceData* SkyLightInstanceData = InstanceData.Cast<FPrecomputedSkyLightInstanceData>();
+	SkyLightInstanceData->OriginalLightGuid = OriginalLightGuid;
 	SkyLightInstanceData->LightGuid = LightGuid;
 	SkyLightInstanceData->ProcessedSkyTexture = ProcessedSkyTexture;
 
@@ -707,6 +708,7 @@ void USkyLightComponent::ApplyComponentInstanceData(FPrecomputedSkyLightInstance
 {
 	check(LightMapData);
 
+	OriginalLightGuid = (HasStaticShadowing() ? LightMapData->OriginalLightGuid : FGuid());
 	LightGuid = (HasStaticShadowing() ? LightMapData->LightGuid : FGuid());
 	ProcessedSkyTexture = LightMapData->ProcessedSkyTexture;
 	IrradianceEnvironmentMap = LightMapData->IrradianceEnvironmentMap;
diff --git a/Engine/Source/Runtime/Engine/Private/Components/StaticMeshComponent.cpp b/Engine/Source/Runtime/Engine/Private/Components/StaticMeshComponent.cpp
--- a/Engine/Source/Runtime/Engine/Private/Components/StaticMeshComponent.cpp
+++ b/Engine/Source/Runtime/Engine/Private/Components/StaticMeshComponent.cpp
@@ -48,6 +48,7 @@
 #include "Rendering/NaniteResources.h"
 #include "NaniteVertexFactory.h"
 #include "StaticMeshSceneProxyDesc.h"
+#include "WorldPartition/ActorInstanceGuids.h"
 
 #include UE_INLINE_GENERATED_CPP_BY_NAME(StaticMeshComponent)
 
@@ -64,6 +65,7 @@ FStaticMeshComponentInstanceData::FStaticMeshComponentInstanceData(const UStatic
 {
 	for (const FStaticMeshComponentLODInfo& LODDataEntry : SourceComponent->LODData)
 	{
+		CachedStaticLighting.Add(LODDataEntry.OriginalMapBuildDataId);
 		CachedStaticLighting.Add(LODDataEntry.MapBuildDataId);
 	}
 
@@ -143,7 +145,8 @@ bool FStaticMeshComponentInstanceData::ApplyVertexColorData(UStaticMeshComponent
 			FStaticMeshComponentLODInfo& LODInfo = StaticMeshComponent->LODData[LODIndex];
 			if(CachedStaticLighting.IsValidIndex(LODIndex))
 			{
-				LODInfo.MapBuildDataId = CachedStaticLighting[LODIndex];
+				LODInfo.OriginalMapBuildDataId = CachedStaticLighting[(LODIndex*2)];
+				LODInfo.MapBuildDataId = CachedStaticLighting[(LODIndex*2)+1];
 			}
 		}
 
@@ -373,7 +376,7 @@ void UStaticMeshComponent::Serialize(FArchive& Ar)
 			if (LODData[LODIndex].LegacyMapBuildData)
 			{
 				LODData[LODIndex].LegacyMapBuildData->IrrelevantLights = IrrelevantLights_DEPRECATED;
-				LegacyComponentData.Data.Emplace(LODData[LODIndex].MapBuildDataId, LODData[LODIndex].LegacyMapBuildData);
+				LegacyComponentData.Data.Emplace(LODData[LODIndex].OriginalMapBuildDataId, LODData[LODIndex].LegacyMapBuildData);
 				LODData[LODIndex].LegacyMapBuildData = nullptr;
 			}
 		}
@@ -476,6 +479,22 @@ void UStaticMeshComponent::PreSave(FObjectPreSaveContext ObjectSaveContext)
 	Super::PreSave(ObjectSaveContext);
 
 	CachePaintedDataIfNecessary();
+	
+	// Detect when to create MapBuildDataID
+	// To avoid having to resubmit all Actors on 1st Static Lighting computation we want this GUID to always be created and serialized as soon as possible. 
+	// But we only create it for the editor save context, on cook builds we avoid creating it since if it's missing it won't link to anything and it'll make the
+	// cook non-deterministic
+	if (!ObjectSaveContext.IsCooking())
+	{
+		UpdateStaticLightingData();
+
+		int32 LODIndex = 0;
+		for (FStaticMeshComponentLODInfo& LODInfo : LODData)
+		{
+			LODInfo.CreateMapBuildDataId(LODIndex);
+			LODIndex++;
+		}
+	}
 }
 #endif // WITH_EDITORONLY_DATA
 
@@ -756,6 +775,8 @@ void UStaticMeshComponent::OnRegister()
 
 	Super::OnRegister();
 
+	UpdateMapBuildDataId();
+
 	// World transform might have changes causing negative determinant which changes the culling mode
 	PrecachePSOs();
 }
@@ -2906,7 +2927,8 @@ void UStaticMeshComponent::ApplyComponentInstanceData(FStaticMeshComponentInstan
 
 			for (int32 i = 0; i < NumLODLightMaps; ++i)
 			{
-				LODData[i].MapBuildDataId = StaticMeshInstanceData->CachedStaticLighting[i];
+				LODData[i].OriginalMapBuildDataId = StaticMeshInstanceData->CachedStaticLighting[(i*2)];
+				LODData[i].MapBuildDataId = StaticMeshInstanceData->CachedStaticLighting[(i*2)+1];
 			}
 		}
 		else
@@ -3413,6 +3435,54 @@ bool UStaticMeshComponent::ComponentIsTouchingSelectionFrustum(const FConvexVolu
 #endif // #if WITH_EDITOR
 
 
+FGuid GetMapDataIdForLOD(const FGuid& LOD0Guid, uint32 LodIndex)
+{
+	FString GuidBaseString = LOD0Guid.ToString(EGuidFormats::Digits);
+	GuidBaseString += TEXT("LOD_") + FString::FromInt(LodIndex);
+
+	FSHA1 Sha;
+	Sha.Update((uint8*)*GuidBaseString, GuidBaseString.Len() * sizeof(TCHAR));
+	Sha.Final();
+	// Retrieve the hash and use it to construct a pseudo-GUID.
+	uint32 Hash[5];
+	Sha.GetHash((uint8*)Hash);
+	return FGuid(Hash[0] ^ Hash[4], Hash[1], Hash[2], Hash[3]);
+}
+
+void UStaticMeshComponent::UpdateStaticLightingData()
+{
+	static const auto AllowStaticLightingVar = IConsoleManager::Get().FindTConsoleVariableDataInt(TEXT("r.AllowStaticLighting"));
+	const bool bAllowStaticLighting = (!AllowStaticLightingVar || AllowStaticLightingVar->GetValueOnGameThread() != 0);
+
+	if (!bAllowStaticLighting)
+		return;
+
+	if (HasStaticLighting() && LODData.IsEmpty() && StaticMesh)
+	{		
+		SetLODDataCount(StaticMesh->GetNumLODs(), LODData.Num());
+	}
+}
+
+void UStaticMeshComponent::UpdateMapBuildDataId()
+{	
+	if (GetOwner())	
+	{
+		//@todo_ow: Detect cases where MapBuildDataId has been successfully set on cook and skip?
+		if (LODData.Num() && LODData[0].OriginalMapBuildDataId.IsValid())
+		{
+			FGuid ActorInstanceGuid = FActorInstanceGuid::GetActorInstanceGuid(*GetOwner());
+
+			// We already have a GUID, but in case of a LevelInstance Actor we must adjust it to be unique
+			LODData[0].MapBuildDataId = FGuid::Combine(LODData[0].OriginalMapBuildDataId, ActorInstanceGuid);
+			
+			for (int i = 1; i < LODData.Num(); i++)
+			{
+				LODData[i].OriginalMapBuildDataId = GetMapDataIdForLOD(LODData[0].OriginalMapBuildDataId, i);
+				LODData[i].MapBuildDataId = FGuid::Combine(LODData[i].OriginalMapBuildDataId, ActorInstanceGuid);
+			}
+		}
+	}
+}
 //////////////////////////////////////////////////////////////////////////
 // StaticMeshComponentLODInfo
 
@@ -3439,30 +3509,29 @@ FStaticMeshComponentLODInfo::FStaticMeshComponentLODInfo(UStaticMeshComponent* I
 
 bool FStaticMeshComponentLODInfo::CreateMapBuildDataId(int32 LodIndex)
 {
-	if (!MapBuildDataId.IsValid())
+	bool bReturnVal = false;
+
+	if (!OriginalMapBuildDataId.IsValid())
 	{
 		if (LodIndex == 0 || OwningComponent == nullptr)
 		{
-			MapBuildDataId = FGuid::NewGuid();
+			OriginalMapBuildDataId = FGuid::NewGuid();
 		}
 		else
 		{
-			FString GuidBaseString = OwningComponent->LODData[0].MapBuildDataId.ToString(EGuidFormats::Digits);
-			GuidBaseString += TEXT("LOD_") + FString::FromInt(LodIndex);
-
-			FSHA1 Sha;
-			Sha.Update((uint8*)*GuidBaseString, GuidBaseString.Len() * sizeof(TCHAR));
-			Sha.Final();
-			// Retrieve the hash and use it to construct a pseudo-GUID.
-			uint32 Hash[5];
-			Sha.GetHash((uint8*)Hash);
-			MapBuildDataId = FGuid(Hash[0] ^ Hash[4], Hash[1], Hash[2], Hash[3]);
+			OriginalMapBuildDataId = GetMapDataIdForLOD(OwningComponent->LODData[0].OriginalMapBuildDataId, LodIndex);
 		}
 
-		return true;
+		bReturnVal = true;
 	}
 
-	return false;
+	if (!MapBuildDataId.IsValid() && OwningComponent)
+	{
+		OwningComponent->UpdateMapBuildDataId();
+		bReturnVal = true;
+	}
+
+	return bReturnVal;
 }
 
 /** Destructor */
@@ -3660,14 +3729,20 @@ FArchive& operator<<(FArchive& Ar,FStaticMeshComponentLODInfo& I)
 	{
 		if (Ar.IsLoading() && Ar.CustomVer(FRenderingObjectVersion::GUID) < FRenderingObjectVersion::MapBuildDataSeparatePackage)
 		{
-			I.MapBuildDataId = FGuid::NewGuid();
+			I.OriginalMapBuildDataId = FGuid::NewGuid();
 			I.LegacyMapBuildData = new FMeshMapBuildData();
 			Ar << I.LegacyMapBuildData->LightMap;
 			Ar << I.LegacyMapBuildData->ShadowMap;
 		}
 		else
-		{
-			Ar << I.MapBuildDataId;
+		{	
+			//@todo_ow: Serialize only MapBuildDataId for cooked versions? Doesn't work in runtime instanced levels
+			// so we'd need to know if the ID we have is the original or has the ActorInstanceGuid mixed in
+			if (Ar.IsCooking() || Ar.IsLoadingFromCookedPackage())
+			{
+				Ar << I.MapBuildDataId;
+			}
+			Ar << I.OriginalMapBuildDataId;
 		}
 	}
 
diff --git a/Engine/Source/Runtime/Engine/Private/InstancedStaticMesh.cpp b/Engine/Source/Runtime/Engine/Private/InstancedStaticMesh.cpp
--- a/Engine/Source/Runtime/Engine/Private/InstancedStaticMesh.cpp
+++ b/Engine/Source/Runtime/Engine/Private/InstancedStaticMesh.cpp
@@ -2075,6 +2075,7 @@ TStructOnScope<FActorComponentInstanceData> UInstancedStaticMeshComponent::GetCo
 
 	for (const FStaticMeshComponentLODInfo& LODDataEntry : LODData)
 	{
+		StaticMeshInstanceData->CachedStaticLighting.MapBuildDataIds.Add(LODDataEntry.OriginalMapBuildDataId);
 		StaticMeshInstanceData->CachedStaticLighting.MapBuildDataIds.Add(LODDataEntry.MapBuildDataId);
 	}
 
@@ -2159,7 +2160,8 @@ void UInstancedStaticMeshComponent::ApplyComponentInstanceData(FInstancedStaticM
 
 		for (int32 i = 0; i < NumLODLightMaps; ++i)
 		{
-			LODData[i].MapBuildDataId = InstancedMeshData->CachedStaticLighting.MapBuildDataIds[i];
+			LODData[i].OriginalMapBuildDataId = InstancedMeshData->CachedStaticLighting.MapBuildDataIds[(i*2)];
+			LODData[i].MapBuildDataId = InstancedMeshData->CachedStaticLighting.MapBuildDataIds[(i*2)+1];
 		}
 	}
 
@@ -2869,7 +2871,7 @@ void UInstancedStaticMeshComponent::ApplyLightMapping(FStaticLightingTextureMapp
 		SetLODDataCount(ResolvedMesh->GetNumLODs(), ResolvedMesh->GetNumLODs());
 		FStaticMeshComponentLODInfo& LODInfo = LODData[0];
 
-		// Ensure this LODInfo has a valid MapBuildDataId
+		// Ensure this LODInfo has a valid OriginalMapBuildDataId
 		if (LODInfo.CreateMapBuildDataId(0))
 		{
 			MarkPackageDirty();
diff --git a/Engine/Source/Runtime/Engine/Private/Level.cpp b/Engine/Source/Runtime/Engine/Private/Level.cpp
--- a/Engine/Source/Runtime/Engine/Private/Level.cpp
+++ b/Engine/Source/Runtime/Engine/Private/Level.cpp
@@ -80,6 +80,7 @@ Level.cpp: Level-related functions
 #include "UObject/MetaData.h"
 #include "UObject/PropertyBagRepository.h"
 #include "WorldPartition/WorldPartitionRuntimeCell.h"
+#include "LevelInstance/LevelInstanceSubsystem.h"
 
 #include UE_INLINE_GENERATED_CPP_BY_NAME(Level)
 
@@ -3036,6 +3037,11 @@ void ULevel::OnLevelLoaded()
 		WorldPartitionRuntimeCell = Cast<const UWorldPartitionRuntimeCell>(LevelStreaming->GetWorldPartitionCell());
 	}
 
+	if (ULevel* OwningLevel = ULevelInstanceSubsystem::GetOwningLevel(this, true))
+	{
+		MapBuildData = OwningLevel->MapBuildData;
+	}
+
 	// 1. Cook commandlet does it's own UWorldPartition::Initialize call in FWorldPartitionCookPackageSplitter::GetGenerateList
 	// 2. Do not Initialize if World doesn't have a UWorldPartitionSubsystem
 	if (!IsRunningCookCommandlet() && OwningWorld->HasSubsystem<UWorldPartitionSubsystem>())
@@ -3283,6 +3289,12 @@ void ULevel::InitializeRenderingResources()
 		ULevel* ActiveLightingScenario = OwningWorld->GetActiveLightingScenario();
 		UMapBuildDataRegistry* EffectiveMapBuildData = MapBuildData;
 
+		if (!IsMapBuildDataOwner())
+		{
+			// If we're using the MapBuildData of our owning Level don't touch the rendering resources, owning level will take care of it
+			return;
+		}
+
 		if (ActiveLightingScenario && ActiveLightingScenario->MapBuildData)
 		{
 			EffectiveMapBuildData = ActiveLightingScenario->MapBuildData;
@@ -3314,6 +3326,12 @@ void ULevel::ReleaseRenderingResources()
 {
 	if (OwningWorld && FApp::CanEverRender())
 	{
+		if (!IsMapBuildDataOwner())
+		{
+			// If we're using the MapBuildData of our owning Level don't touch the rendering resources, owning level will take care of it
+			return;
+		}
+
 		if (VolumetricLightmapGridManager)
 		{
 			VolumetricLightmapGridManager->RemoveFromScene(OwningWorld->Scene);
@@ -3474,6 +3492,11 @@ UPackage* ULevel::CreateMapBuildDataPackage() const
 	return BuiltDataPackage;
 }
 
+bool ULevel::IsMapBuildDataOwner() const
+{
+	return ULevelInstanceSubsystem::GetOwningLevel(this, true) == this;
+}
+
 UMapBuildDataRegistry* ULevel::GetOrCreateMapBuildData()
 {
 	if (!MapBuildData 
@@ -3495,6 +3518,14 @@ UMapBuildDataRegistry* ULevel::GetOrCreateMapBuildData()
 		FName ShortPackageName = FPackageName::GetShortFName(BuiltDataPackage->GetFName());
 		// Top level UObjects have to have both RF_Standalone and RF_Public to be saved into packages
 		MapBuildData = NewObject<UMapBuildDataRegistry>(BuiltDataPackage, ShortPackageName, RF_Standalone | RF_Public);
+		// assign to all transient Levels inside the World
+		for (ULevel* Level : GetWorld()->GetLevels())
+		{
+			if (ULevelInstanceSubsystem::GetOwningLevel(Level, true) == this)
+			{
+				Level->MapBuildData = MapBuildData;
+			}
+		}
 		MarkPackageDirty();
 	}
 
diff --git a/Engine/Source/Runtime/Engine/Private/LevelInstance/LevelInstanceActorImpl.cpp b/Engine/Source/Runtime/Engine/Private/LevelInstance/LevelInstanceActorImpl.cpp
--- a/Engine/Source/Runtime/Engine/Private/LevelInstance/LevelInstanceActorImpl.cpp
+++ b/Engine/Source/Runtime/Engine/Private/LevelInstance/LevelInstanceActorImpl.cpp
@@ -20,6 +20,8 @@
 
 #endif
 
+#include "WorldPartition/ActorInstanceGuids.h"
+
 #define LOCTEXT_NAMESPACE "LevelInstanceActor"
 
 void FLevelInstanceActorImpl::RegisterLevelInstance()
@@ -84,6 +86,8 @@ bool FLevelInstanceActorImpl::IsLoadingEnabled() const
 
 void FLevelInstanceActorImpl::OnLevelInstanceLoaded()
 {
+	FActorInstanceGuid::SetLevelInstanceGuid(LevelInstance->GetLoadedLevel(), ULevelInstanceSubsystem::GetOwningLevel(LevelInstance->GetLoadedLevel()), LevelInstance->GetLevelInstanceGuid());
+
 #if WITH_EDITOR
 	AActor* Actor = CastChecked<AActor>(LevelInstance);
 	if (!Actor->GetWorld()->IsGameWorld())
diff --git a/Engine/Source/Runtime/Engine/Private/LevelInstance/LevelInstanceSubsystem.cpp b/Engine/Source/Runtime/Engine/Private/LevelInstance/LevelInstanceSubsystem.cpp
--- a/Engine/Source/Runtime/Engine/Private/LevelInstance/LevelInstanceSubsystem.cpp
+++ b/Engine/Source/Runtime/Engine/Private/LevelInstance/LevelInstanceSubsystem.cpp
@@ -745,6 +745,51 @@ ILevelInstanceInterface* ULevelInstanceSubsystem::GetOwningLevelInstance(const U
 	return nullptr;
 }
 
+ULevel* ULevelInstanceSubsystem::GetOwningLevel(const ULevel* Level, bool bFollowChainToNonLevelInstanceOwningLevel /*= false*/)
+{	
+	auto GetOwningLevelInternal = [] (const ULevel* ForLevel) -> ULevel*
+	{
+		if (ForLevel->GetWorld() != ForLevel->GetTypedOuter<UWorld>())
+		{	
+			if (ULevelStreaming* LevelStreaming = ULevelStreaming::FindStreamingLevel(ForLevel))
+			{
+				if (ULevelStreamingLevelInstance* LevelInstanceStreaming = Cast<ULevelStreamingLevelInstance>(LevelStreaming))
+				{
+					if (AActor* LevelActor = Cast<AActor>(LevelInstanceStreaming->GetLevelInstance()))
+					{
+						return LevelActor->GetLevel();
+					}
+				}
+				if (UWorldPartitionLevelStreamingDynamic* WPLevelStreaming = Cast<UWorldPartitionLevelStreamingDynamic>(LevelStreaming))
+				{
+					if (UObject* WorldPartition =  WPLevelStreaming->GetOuterWorldPartition().ResolveObject())
+					{						
+						return WorldPartition->GetTypedOuter<ULevel>();
+					}
+				}
+			}
+		}
+
+		return nullptr;
+	};
+
+	// locate the non instanced level responsible for getting us in the world
+	const ULevel* OwningLevel = bFollowChainToNonLevelInstanceOwningLevel ? Level : nullptr;
+	const ULevel* CurrentLevel = Level;
+
+	do
+	{
+    	CurrentLevel = GetOwningLevelInternal(CurrentLevel);
+		if (CurrentLevel)
+		{
+			OwningLevel = CurrentLevel;
+		}
+	} while(CurrentLevel && bFollowChainToNonLevelInstanceOwningLevel);
+	
+	return const_cast<ULevel*>(OwningLevel);
+}
+
+
 void ULevelInstanceSubsystem::ForEachActorInLevelInstance(const ILevelInstanceInterface* LevelInstance, TFunctionRef<bool(AActor* LevelActor)> Operation) const
 {
 	if (ULevel* LevelInstanceLevel = GetLevelInstanceLevel(LevelInstance))
diff --git a/Engine/Source/Runtime/Engine/Private/MapBuildData.cpp b/Engine/Source/Runtime/Engine/Private/MapBuildData.cpp
--- a/Engine/Source/Runtime/Engine/Private/MapBuildData.cpp
+++ b/Engine/Source/Runtime/Engine/Private/MapBuildData.cpp
@@ -40,6 +40,8 @@ MapBuildData.cpp
 DECLARE_MEMORY_STAT(TEXT("Stationary Light Static Shadowmap"),STAT_StationaryLightBuildData,STATGROUP_MapBuildData);
 DECLARE_MEMORY_STAT(TEXT("Reflection Captures"),STAT_ReflectionCaptureBuildData,STATGROUP_MapBuildData);
 
+DEFINE_LOG_CATEGORY(LogMapBuildDataRegistry);
+
 FArchive& operator<<(FArchive& Ar, FMeshMapBuildData& MeshMapBuildData)
 {
 	Ar << MeshMapBuildData.LightMap;
@@ -580,6 +582,20 @@ void UMapBuildDataRegistry::Serialize(FArchive& Ar)
 			}
 		}
 	}
+
+#if UE_LOG_MAPBUILDATA_ENABLED 
+
+	UE_LOG_MAPBUILDDATA(Log, TEXT("Loaded Registry %s"), *GetFullName());
+	for (auto& It : MeshBuildData)
+	{
+		UE_LOG_MAPBUILDDATA(Log, TEXT("    => Mesh GUID : %s"), *It.Key.ToString());
+	}
+
+	for (auto& It : LightBuildData)
+	{
+		UE_LOG_MAPBUILDDATA(Log, TEXT("    => Light GUID : %s"), *It.Key.ToString());
+	}
+#endif
 }
 
 void UMapBuildDataRegistry::PostLoad()
@@ -692,6 +708,8 @@ FMeshMapBuildData& UMapBuildDataRegistry::AllocateMeshBuildData(const FGuid& Mes
 	check(MeshId.IsValid());
 	check(!bSetupResourceClusters);
 
+	UE_LOG_MAPBUILDDATA(Log, TEXT("Allocating MeshBuildData in Registry %s for Guid: %s"), *GetFullName(), *MeshId.ToString());
+
 	if (bMarkDirty)
 	{
 		MarkPackageDirty();
@@ -711,6 +729,7 @@ const FMeshMapBuildData* UMapBuildDataRegistry::GetMeshBuildData(FGuid MeshId) c
 		return nullptr;
 	}
 
+	UE_LOG_MAPBUILDDATA(Log, TEXT("Finding MeshBuildData (%p) in Registry %s for Guid: %s"), FoundData, *GetFullName(), *MeshId.ToString());
 	return FoundData;
 }
 
@@ -723,6 +742,7 @@ FMeshMapBuildData* UMapBuildDataRegistry::GetMeshBuildData(FGuid MeshId)
 		return nullptr;
 	}
 
+	UE_LOG_MAPBUILDDATA(Log, TEXT("Finding MeshBuildData (%p) in Registry %s for Guid: %s"), FoundData, *GetFullName(), *MeshId.ToString());
 	return FoundData;
 }
 
@@ -836,16 +856,20 @@ FLightComponentMapBuildData& UMapBuildDataRegistry::FindOrAllocateLightBuildData
 		MarkPackageDirty();
 	}
 
+	UE_LOG_MAPBUILDDATA(Log, TEXT("Allocating LightBuildData in Registry %s for Guid: %s"), *GetFullName(), *LightId.ToString());
+
 	return LightBuildData.FindOrAdd(LightId);
 }
 
 const FLightComponentMapBuildData* UMapBuildDataRegistry::GetLightBuildData(FGuid LightId) const
 {
+	UE_LOG_MAPBUILDDATA(Log, TEXT("Finding LightBuildData (%p) in Registry %s for Guid: %s"), LightBuildData.Find(LightId), *GetFullName(), *LightId.ToString());
 	return LightBuildData.Find(LightId);
 }
 
 FLightComponentMapBuildData* UMapBuildDataRegistry::GetLightBuildData(FGuid LightId)
 {
+	UE_LOG_MAPBUILDDATA(Log, TEXT("Finding LightBuildData (%p) in Registry %s for Guid: %s"), LightBuildData.Find(LightId), *GetFullName(), *LightId.ToString());
 	return LightBuildData.Find(LightId);
 }
 
@@ -1226,7 +1250,9 @@ UMapBuildDataRegistry* UMapBuildDataRegistry::Get(const AActor* Actor)
 		return World->PersistentLevel->MapBuildData;		
 	}
 
-	return Get(OwnerLevel, World);
+	UMapBuildDataRegistry* Registry = Get(OwnerLevel, World);
+	UE_LOG_MAPBUILDDATA(Log, TEXT("Returning Registry %s for Actor %s, %s"), *Registry->GetFullName(), *Actor->GetActorNameOrLabel(), *Actor->GetFullName());
+	return Registry;
 }
 
 UMapBuildDataRegistry* UMapBuildDataRegistry::Get(ULevel* OwnerLevel, UWorld* World)
diff --git a/Engine/Source/Runtime/Engine/Private/StaticLightingBuildContext.cpp b/Engine/Source/Runtime/Engine/Private/StaticLightingBuildContext.cpp
--- a/Engine/Source/Runtime/Engine/Private/StaticLightingBuildContext.cpp
+++ b/Engine/Source/Runtime/Engine/Private/StaticLightingBuildContext.cpp
@@ -9,6 +9,7 @@
 #include "Engine/MapBuildDataRegistry.h"
 #include "WorldPartition/WorldPartition.h"
 #include "PrecomputedVolumetricLightmap.h"
+#include "LevelInstance/LevelInstanceSubsystem.h"
 
 #if WITH_EDITOR
 
@@ -261,19 +262,26 @@ UMapBuildDataRegistry* FStaticLightingBuildContext::GetOrCreateRegistryForActor(
 {
 	check(Actor);
 
-	ULevel* Level = Actor->GetLevel();	
+	// For Actors in LevelInstances we need to defer storage to the level that owns the LevelInstance
+	ULevel* Level = ULevelInstanceSubsystem::GetOwningLevel(Actor->GetLevel(), true);
 
-	return GetLightingStorageLevel(Level)->GetOrCreateMapBuildData();
+	UMapBuildDataRegistry* Registry = GetLightingStorageLevel(Level)->GetOrCreateMapBuildData();
+	
+	UE_LOG_MAPBUILDDATA(Log, TEXT("Creating/Returning Registry %s for Actor %s, %s"), *Registry->GetFullName(), *Actor->GetActorNameOrLabel(), *Actor->GetFullName());
+
+	return Registry;
 }
 
 UMapBuildDataRegistry* FStaticLightingBuildContext::GetRegistryForLevel(ULevel* Level) const
 {
-	return GetLightingStorageLevel(Level)->MapBuildData;
+	// For  we need to defer storage to the level that owns the LevelInstance
+	return GetLightingStorageLevel(ULevelInstanceSubsystem::GetOwningLevel(Level, true))->MapBuildData;
 }
 
 UMapBuildDataRegistry* FStaticLightingBuildContext::GetOrCreateRegistryForLevel(ULevel* Level) const
 {
-	return GetLightingStorageLevel(Level)->GetOrCreateMapBuildData();
+	// For  we need to defer storage to the level that owns the LevelInstance
+	return GetLightingStorageLevel(ULevelInstanceSubsystem::GetOwningLevel(Level, true))->GetOrCreateMapBuildData();
 }
 
-#endif
\ No newline at end of file
+#endif
diff --git a/Engine/Source/Runtime/Engine/Private/WorldPartition/ActorInstanceGuids.cpp b/Engine/Source/Runtime/Engine/Private/WorldPartition/ActorInstanceGuids.cpp
new file mode 100644
--- /dev/null
+++ b/Engine/Source/Runtime/Engine/Private/WorldPartition/ActorInstanceGuids.cpp
@@ -0,0 +1,187 @@
+// Copyright Epic Games, Inc. All Rights Reserved.
+
+#include "WorldPartition/ActorInstanceGuids.h"
+#include "UObject/UObjectAnnotation.h"
+#include "GameFramework/Actor.h"
+#include "Engine/Level.h"
+#include "LevelInstance/LevelInstanceSubsystem.h"
+
+struct FLevelInstanceGuid
+{
+	TWeakObjectPtr<ULevel> Level;
+	TWeakObjectPtr<ULevel> OwnerLevel;
+	FGuid LevelInstanceGuid;
+	FGuid ResolvedLevelInstanceGuid;
+	bool bIsDefault = true;
+
+	bool IsDefault()
+	{
+		return bIsDefault;
+	}
+};
+
+FUObjectAnnotationSparse<FActorInstanceGuid, true>	GActorGuids;
+FUObjectAnnotationSparse<FLevelInstanceGuid, true>	GLevelInstanceGuids;
+
+void FActorInstanceGuid::SetLevelInstanceGuid(ULevel* Level, ULevel* OwnerLevel, const FGuid& Guid, const FGuid& ResolvedGuid)
+{
+	// double registration is often an order issue (LevelInstance GUID tried to be accessed before FLevelInstanceActorImpl::OnLevelInstanceLoaded was called)
+	check(!GLevelInstanceGuids.GetAnnotationMap().Contains(Level));
+
+	FLevelInstanceGuid	LIGuid;
+
+	LIGuid.Level = Level;
+	LIGuid.OwnerLevel = OwnerLevel;
+	LIGuid.LevelInstanceGuid = Guid;
+	LIGuid.ResolvedLevelInstanceGuid = ResolvedGuid;
+	LIGuid.bIsDefault = false;
+	
+	GLevelInstanceGuids.AddAnnotation(Level, LIGuid);
+
+	// try to force a resolve for the guid
+	GetLevelInstanceGuid(Level);
+}
+
+FGuid FActorInstanceGuid::GetLevelInstanceGuid(ULevel* Level)
+{
+	if (!Level)
+	{
+		return FGuid();
+	}
+
+	FLevelInstanceGuid LevelGuids = GLevelInstanceGuids.GetAnnotation(Level);	
+	
+	if (LevelGuids.IsDefault())
+	{
+		// If we can find an OwningLevel this is most likely the WorldPartition streaming level of a WP LevelInstance inside a non-WP map
+		if (ULevel* OwningLevel = ULevelInstanceSubsystem::GetOwningLevel(Level))
+		{
+			LevelGuids = GLevelInstanceGuids.GetAnnotation(OwningLevel);			
+			if (!LevelGuids.IsDefault())
+			{
+				// Duplicate the annotation to this level to avoid having to look it up all the time
+				check(LevelGuids.ResolvedLevelInstanceGuid.IsValid());
+				GLevelInstanceGuids.AddAnnotation(Level, LevelGuids);
+			}
+		}
+	}
+	
+	// When we reach a non-instanced Level in the chain, register it so that we don't try to resolve it later-on
+	if (LevelGuids.IsDefault())
+	{
+		SetLevelInstanceGuid(Level, nullptr, FGuid());
+	}
+	
+	if (!LevelGuids.ResolvedLevelInstanceGuid.IsValid() && LevelGuids.LevelInstanceGuid.IsValid())
+	{		
+		FGuid OwnerLevelGuid;
+
+		if (LevelGuids.OwnerLevel.IsValid())
+		{
+			OwnerLevelGuid = GetLevelInstanceGuid(LevelGuids.OwnerLevel.Get());
+		}
+
+		LevelGuids.ResolvedLevelInstanceGuid = FGuid::Combine(OwnerLevelGuid, LevelGuids.LevelInstanceGuid);
+
+		// Update the Annotations after modifying it
+		GLevelInstanceGuids.AddAnnotation(Level, LevelGuids);
+	}
+
+	return LevelGuids.ResolvedLevelInstanceGuid;
+}
+
+#if WITH_EDITOR
+void FActorInstanceGuid::InitializeFrom(AActor& InActor)
+{
+	ActorGuid = InActor.GetActorGuid();
+	ActorInstanceGuid = InActor.GetActorInstanceGuid();	
+
+	if (ActorInstanceGuid == ActorGuid)
+	{
+		// not an instance/owner level unknown delay resolve until we have all the info available
+		ActorInstanceGuid.Invalidate();
+	}
+}
+#endif
+
+bool FActorInstanceGuid::IsDefault()
+{
+	return !ActorGuid.IsValid() && !ActorInstanceGuid.IsValid();
+}
+
+FActorInstanceGuid FActorInstanceGuid::GetActorGuids(AActor& InActor)
+{
+	FActorInstanceGuid Guids;
+
+#if WITH_EDITOR
+	Guids.InitializeFrom(InActor);
+#else
+	Guids = GActorGuids.GetAnnotation(&InActor);
+#endif
+		
+	//@todo_ow: ideally would be done during serialization in runtime, but in the async thread it's not safe to use GetLevelInstanceGuid
+	// ActorInstanceGuid is invalid, either the Actor is in the main map or nobody assigned the ActorInstanceGuid upon loading a LevelInstance
+	// So verify which case we're in
+	if (!Guids.ActorInstanceGuid.IsValid())
+	{
+		// Resolve ActorInstanceGuid by combining it with the LevelGUID 
+		FGuid LevelInstanceGuid = GetLevelInstanceGuid(InActor.GetLevel());
+
+		if (LevelInstanceGuid.IsValid())
+		{
+			Guids.ActorInstanceGuid = FGuid::Combine(LevelInstanceGuid, Guids.ActorGuid);
+		}
+		else
+		{
+			Guids.ActorInstanceGuid = Guids.ActorGuid;
+		}
+
+#if !WITH_EDITOR
+		GActorGuids.AddAnnotation(&InActor, Guids);
+#endif
+	}
+	
+	return Guids;
+}
+
+FGuid FActorInstanceGuid::GetActorInstanceGuid(AActor& InActor)
+{
+	FActorInstanceGuid Guids = GetActorGuids(InActor);
+
+	return Guids.ActorInstanceGuid;
+}
+
+void FActorInstanceGuid::ReleaseActorInstanceGuid(AActor& InActor)
+{
+	// Only loaded objects will have those annotations
+	if (InActor.HasAnyFlags(RF_WasLoaded))
+	{
+		GActorGuids.RemoveAnnotation(&InActor);
+	}
+}
+
+FArchive &operator <<(FArchive& Ar, FActorInstanceGuid& InActorInstanceGuids)
+{
+	Ar << InActorInstanceGuids.ActorGuid;
+	Ar << InActorInstanceGuids.ActorInstanceGuid;
+
+	return Ar;
+}
+
+void FActorInstanceGuid::Serialize(FArchive& Ar, AActor& InActor)
+{
+	FActorInstanceGuid Guids;
+
+#if WITH_EDITOR
+	Guids.InitializeFrom(InActor);
+#endif
+
+	Ar << Guids;
+	
+	if (Ar.IsLoadingFromCookedPackage())
+	{
+		GActorGuids.AddAnnotation(&InActor, MoveTemp(Guids));
+	}
+}
+
+
diff --git a/Engine/Source/Runtime/Engine/Public/LevelInstance/LevelInstanceSubsystem.h b/Engine/Source/Runtime/Engine/Public/LevelInstance/LevelInstanceSubsystem.h
--- a/Engine/Source/Runtime/Engine/Public/LevelInstance/LevelInstanceSubsystem.h
+++ b/Engine/Source/Runtime/Engine/Public/LevelInstance/LevelInstanceSubsystem.h
@@ -75,6 +75,7 @@ public:
 	ENGINE_API ILevelInstanceInterface* GetLevelInstance(const FLevelInstanceID& LevelInstanceID) const;
 	ENGINE_API ILevelInstanceInterface* GetOwningLevelInstance(const ULevel* Level) const;
 	ENGINE_API FLevelInstanceID RegisterLevelInstance(ILevelInstanceInterface* LevelInstance);
+	ENGINE_API static ULevel* GetOwningLevel(const ULevel* Level, bool bFollowChainToNonLevelInstanceOwningLevel = false);
 	ENGINE_API void UnregisterLevelInstance(ILevelInstanceInterface* LevelInstance);
 	ENGINE_API void RequestLoadLevelInstance(ILevelInstanceInterface* LevelInstance, bool bUpdate);
 	ENGINE_API void RequestUnloadLevelInstance(ILevelInstanceInterface* LevelInstance);
diff --git a/Engine/Source/Runtime/Engine/Public/StaticMeshComponentLODInfo.h b/Engine/Source/Runtime/Engine/Public/StaticMeshComponentLODInfo.h
--- a/Engine/Source/Runtime/Engine/Public/StaticMeshComponentLODInfo.h
+++ b/Engine/Source/Runtime/Engine/Public/StaticMeshComponentLODInfo.h
@@ -22,6 +22,9 @@ struct FStaticMeshComponentLODInfo
 	GENERATED_USTRUCT_BODY()
 
 	/** Uniquely identifies this LOD's built map data. */
+	FGuid OriginalMapBuildDataId;
+
+	/** Uniquely identifies this LOD's built map data, a combination of OriginalMapBuildDataID and the ActorInstanceID **/
 	FGuid MapBuildDataId;
 
 	/** Used during deserialization to temporarily store legacy lightmap data. */
@@ -61,7 +64,7 @@ struct FStaticMeshComponentLODInfo
 	void CleanUp();
 
 	/** 
-	 * Ensure this LODInfo has a valid MapBuildDataId GUID.
+	 * Ensure this LODInfo has a valid OriginalMapBuildDataId GUID.
 	 * @param LodIndex Index of the LOD this LODInfo represents.
 	 * @return true if a new GUID was created, false otherwise.
 	 */
diff --git a/Engine/Source/Runtime/Engine/Public/WorldPartition/ActorInstanceGuids.h b/Engine/Source/Runtime/Engine/Public/WorldPartition/ActorInstanceGuids.h
new file mode 100644
--- /dev/null
+++ b/Engine/Source/Runtime/Engine/Public/WorldPartition/ActorInstanceGuids.h
@@ -0,0 +1,30 @@
+// Copyright Epic Games, Inc. All Rights Reserved.
+
+#pragma once
+
+#include "Misc/Guid.h"
+
+class AActor;
+class FArchive;
+class ULevel;
+
+struct FActorInstanceGuid
+{
+	FGuid ActorGuid;
+	FGuid ActorInstanceGuid;
+
+	bool IsDefault();
+#if WITH_EDITOR
+	void InitializeFrom(AActor& InActor);
+#endif
+
+	ENGINE_API static void SetLevelInstanceGuid(ULevel* Level, ULevel* OwnerLevel, const FGuid& Guid, const FGuid& ResolvedGuid = FGuid());
+	ENGINE_API static FGuid GetLevelInstanceGuid(ULevel* Level);
+
+	ENGINE_API static FActorInstanceGuid GetActorGuids(AActor& InActor);
+	ENGINE_API static FGuid GetActorInstanceGuid(AActor& InActor);
+	ENGINE_API static void ReleaseActorInstanceGuid(AActor& InActor);
+	ENGINE_API static void Serialize(FArchive& Ar, AActor& InActor);
+};
+
+FArchive &operator <<(FArchive& Ar, FActorInstanceGuid& InActorInstanceGuids);
