- Added an advanced runtime setting for Data Layer Instances to specify whether a streaming cells containing the Data Layer(s) should override the default 'Block On Slow Streaming' of its Runtime Grid /  Runtime Partition. If the cell has more than one Data Layer with an override value, the blocking state wins over non-blocking state.

diff --git a/Engine/Source/Editor/DataLayerEditor/Private/DataLayer/DataLayerInstanceCustomization.cpp b/Engine/Source/Editor/DataLayerEditor/Private/DataLayer/DataLayerInstanceCustomization.cpp
--- a/Engine/Source/Editor/DataLayerEditor/Private/DataLayer/DataLayerInstanceCustomization.cpp
+++ b/Engine/Source/Editor/DataLayerEditor/Private/DataLayer/DataLayerInstanceCustomization.cpp
@@ -46,6 +46,7 @@ void FDataLayerInstanceDetails::CustomizeDetails(IDetailLayoutBuilder& DetailBui
 	if (!bHasInitialRuntimeState)
 	{
 		DetailBuilder.HideProperty(GET_MEMBER_NAME_CHECKED(UDataLayerInstance, InitialRuntimeState));
+		DetailBuilder.HideProperty(GET_MEMBER_NAME_CHECKED(UDataLayerInstance, OverrideBlockOnSlowStreaming));
 	}
 
 	TSharedRef<IPropertyHandle> DataLayerAssetProperty = DetailBuilder.GetProperty("DataLayerAsset", UDataLayerInstanceWithAsset::StaticClass());
diff --git a/Engine/Source/Runtime/Engine/Private/WorldPartition/DataLayer/DataLayerInstance.cpp b/Engine/Source/Runtime/Engine/Private/WorldPartition/DataLayer/DataLayerInstance.cpp
--- a/Engine/Source/Runtime/Engine/Private/WorldPartition/DataLayer/DataLayerInstance.cpp
+++ b/Engine/Source/Runtime/Engine/Private/WorldPartition/DataLayer/DataLayerInstance.cpp
@@ -28,6 +28,7 @@ UDataLayerInstance::UDataLayerInstance(const FObjectInitializer& ObjectInitializ
 	, bIsInitiallyLoadedInEditor(true)
 	, bIsLoadedInEditor(true)
 	, bIsLocked(false)
+	, OverrideBlockOnSlowStreaming(EOverrideBlockOnSlowStreaming::NoOverride)
 #endif
 	, InitialRuntimeState(EDataLayerRuntimeState::Unloaded)
 {}
diff --git a/Engine/Source/Runtime/Engine/Private/WorldPartition/RuntimeHashSet/WorldPartitionRuntimeHashSetStreamingGeneration.cpp b/Engine/Source/Runtime/Engine/Private/WorldPartition/RuntimeHashSet/WorldPartitionRuntimeHashSetStreamingGeneration.cpp
--- a/Engine/Source/Runtime/Engine/Private/WorldPartition/RuntimeHashSet/WorldPartitionRuntimeHashSetStreamingGeneration.cpp
+++ b/Engine/Source/Runtime/Engine/Private/WorldPartition/RuntimeHashSet/WorldPartitionRuntimeHashSetStreamingGeneration.cpp
@@ -153,8 +153,10 @@ bool UWorldPartitionRuntimeHashSet::GenerateStreaming(UWorldPartitionStreamingPo
 		RuntimeCell->SetDataLayers(CellDescInstance.DataLayerInstances);
 		RuntimeCell->SetContentBundleUID(CellDescInstance.ContentBundleID);
 		RuntimeCell->SetClientOnlyVisible(CellDescInstance.bClientOnlyVisible);
-		RuntimeCell->SetBlockOnSlowLoading(CellDescInstance.bBlockOnSlowStreaming);
-		RuntimeCell->SetIsHLOD(CellDescInstance.SourcePartition->HLODIndex != INDEX_NONE);
+		const bool bIsHLOD = CellDescInstance.SourcePartition->HLODIndex != INDEX_NONE;
+		const bool bBlockOnSlowStreaming = ResolveBlockOnSlowStreamingForCell(CellDescInstance.bBlockOnSlowStreaming, bIsHLOD, CellDescInstance.DataLayerInstances);
+		RuntimeCell->SetBlockOnSlowLoading(bBlockOnSlowStreaming);
+		RuntimeCell->SetIsHLOD(bIsHLOD);
 		RuntimeCell->SetGuid(CellUniqueId.Guid);
 		RuntimeCell->SetCellDebugColor(CellDescInstance.SourcePartition->DebugColor);
 
diff --git a/Engine/Source/Runtime/Engine/Private/WorldPartition/WorldPartitionRuntimeHash.cpp b/Engine/Source/Runtime/Engine/Private/WorldPartition/WorldPartitionRuntimeHash.cpp
--- a/Engine/Source/Runtime/Engine/Private/WorldPartition/WorldPartitionRuntimeHash.cpp
+++ b/Engine/Source/Runtime/Engine/Private/WorldPartition/WorldPartitionRuntimeHash.cpp
@@ -480,6 +480,30 @@ void UWorldPartitionRuntimeHash::ForceExternalActorLevelReference(bool bForceExt
 		}
 	}
 }
+
+bool UWorldPartitionRuntimeHash::ResolveBlockOnSlowStreamingForCell(bool bInOwnerBlockOnSlowStreaming, bool bInIsHLODCell, const TArray<const UDataLayerInstance*>& InCellDataLayerInstances) const
+{
+	if (bInIsHLODCell)
+	{
+		return false;
+	}
+
+	TOptional<bool> DataLayersOverrideBlockOnSlowStreaming;
+	for (const UDataLayerInstance* DataLayerInstance : InCellDataLayerInstances)
+	{
+		if (DataLayerInstance->GetOverrideBlockOnSlowStreaming() != EOverrideBlockOnSlowStreaming::NoOverride)
+		{
+			bool bIsBlocking = DataLayerInstance->GetOverrideBlockOnSlowStreaming() == EOverrideBlockOnSlowStreaming::Blocking;
+			DataLayersOverrideBlockOnSlowStreaming = bIsBlocking;
+			if (bIsBlocking)
+			{
+				break;
+			}
+		}
+	}
+	const bool bBlockOnSlowStreaming = DataLayersOverrideBlockOnSlowStreaming.IsSet() ? DataLayersOverrideBlockOnSlowStreaming.GetValue() : bInOwnerBlockOnSlowStreaming;
+	return bBlockOnSlowStreaming;
+}
 #endif
 
 bool UWorldPartitionRuntimeHash::IsCellRelevantFor(bool bClientOnlyVisible) const
diff --git a/Engine/Source/Runtime/Engine/Private/WorldPartition/WorldPartitionRuntimeSpatialHash.cpp b/Engine/Source/Runtime/Engine/Private/WorldPartition/WorldPartitionRuntimeSpatialHash.cpp
--- a/Engine/Source/Runtime/Engine/Private/WorldPartition/WorldPartitionRuntimeSpatialHash.cpp
+++ b/Engine/Source/Runtime/Engine/Private/WorldPartition/WorldPartitionRuntimeSpatialHash.cpp
@@ -1576,13 +1576,14 @@ bool UWorldPartitionRuntimeSpatialHash::CreateStreamingGrid(const FSpatialHashRu
 
 					const TArray<const UDataLayerInstance*>& DataLayers = GridCellDataChunk.GetDataLayers();
 					const bool bAreClientOnlyDataLayers = DataLayers.Num() && Algo::AllOf(DataLayers, [](const UDataLayerInstance* DataLayerInstance) { return DataLayerInstance->IsClientOnly(); });
-
+					const bool bIsHLOD = RuntimeGrid.HLODLayer ? true : false;
+					const bool bBlockOnSlowStreaming = ResolveBlockOnSlowStreamingForCell(CurrentStreamingGrid.bBlockOnSlowStreaming, bIsHLOD, DataLayers);
 					StreamingCell->SetIsAlwaysLoaded(bIsCellAlwaysLoaded);
 					StreamingCell->SetDataLayers(DataLayers);
 					StreamingCell->SetContentBundleUID(GridCellDataChunk.GetContentBundleID());
 					StreamingCell->SetClientOnlyVisible(CurrentStreamingGrid.bClientOnlyVisible || bAreClientOnlyDataLayers);
-					StreamingCell->SetBlockOnSlowLoading(CurrentStreamingGrid.bBlockOnSlowStreaming);
-					StreamingCell->SetIsHLOD(RuntimeGrid.HLODLayer ? true : false);
+					StreamingCell->SetBlockOnSlowLoading(bBlockOnSlowStreaming);
+					StreamingCell->SetIsHLOD(bIsHLOD);
 					StreamingCell->SetGuid(CellGuid);
 
 					FBox2D Bounds;
diff --git a/Engine/Source/Runtime/Engine/Public/WorldPartition/DataLayer/DataLayerInstance.h b/Engine/Source/Runtime/Engine/Public/WorldPartition/DataLayer/DataLayerInstance.h
--- a/Engine/Source/Runtime/Engine/Public/WorldPartition/DataLayer/DataLayerInstance.h
+++ b/Engine/Source/Runtime/Engine/Public/WorldPartition/DataLayer/DataLayerInstance.h
@@ -34,6 +34,17 @@ enum class EDataLayerRuntimeState : uint8
 	Activated
 };
 
+UENUM(BlueprintType)
+enum class EOverrideBlockOnSlowStreaming : uint8
+{
+	// Uses default runtime partition 'Block on Slow Streaming' setting
+	NoOverride,
+	// Blocks on slow streaming (Overrides runtime partition 'Block on Slow Streaming' setting)
+	Blocking,
+	// Doesn't block on slow streaming (Overrides runtime partition 'Block on Slow Streaming' setting)
+	NotBlocking
+};
+
 const inline TCHAR* GetDataLayerRuntimeStateName(EDataLayerRuntimeState State)
 {
 	switch (State)
@@ -92,6 +103,7 @@ public:
 	ENGINE_API bool IsLoadedInEditorChangedByUserOperation() const { return bIsLoadedInEditorChangedByUserOperation; }
 	ENGINE_API virtual bool IsReadOnly(FText* OutReason = nullptr) const;
 	virtual bool IsIncludedInActorFilterDefault() const { return false; }
+	ENGINE_API EOverrideBlockOnSlowStreaming GetOverrideBlockOnSlowStreaming() const { return OverrideBlockOnSlowStreaming; }
 
 	// Data Layer Instance features support
 	ENGINE_API bool CanBeChildOf(const UDataLayerInstance* InParent, FText* OutReason = nullptr) const;
@@ -242,6 +254,9 @@ protected:
 	/** Whether this data layer is locked, which means the user can't change actors assignation, remove or rename it */
 	UPROPERTY()
 	uint32 bIsLocked : 1;
+
+	UPROPERTY(Category = "Runtime|Advanced", EditAnywhere)
+	EOverrideBlockOnSlowStreaming OverrideBlockOnSlowStreaming;
 #endif
 
 	/** Initial runtime state of this data layer instance. Only supported if it's runtime and not client/server only. */
diff --git a/Engine/Source/Runtime/Engine/Public/WorldPartition/WorldPartitionRuntimeHash.h b/Engine/Source/Runtime/Engine/Public/WorldPartition/WorldPartitionRuntimeHash.h
--- a/Engine/Source/Runtime/Engine/Public/WorldPartition/WorldPartitionRuntimeHash.h
+++ b/Engine/Source/Runtime/Engine/Public/WorldPartition/WorldPartitionRuntimeHash.h
@@ -264,6 +264,7 @@ protected:
 protected:
 #if WITH_EDITOR
 	ENGINE_API void ForceExternalActorLevelReference(bool bForceExternalActorLevelReferenceForPIE);
+	ENGINE_API bool ResolveBlockOnSlowStreamingForCell(bool bInOwnerBlockOnSlowStreaming, bool bInIsHLODCell, const TArray<const UDataLayerInstance*>& InCellDataLayerInstances) const;
 #endif
 
 	TSet<TWeakObjectPtr<URuntimeHashExternalStreamingObjectBase>> InjectedExternalStreamingObjects;
