Added helper function FExternalDataLayerHelper::GetRelativeExternalActorPackagePath to convert an actor package using EDL from a full path to a relative path.
(will be used by a custom asset reference validator)

diff --git a/Engine/Source/Runtime/Engine/Private/WorldPartition/DataLayer/ExternalDataLayerHelper.cpp b/Engine/Source/Runtime/Engine/Private/WorldPartition/DataLayer/ExternalDataLayerHelper.cpp
--- a/Engine/Source/Runtime/Engine/Private/WorldPartition/DataLayer/ExternalDataLayerHelper.cpp
+++ b/Engine/Source/Runtime/Engine/Private/WorldPartition/DataLayer/ExternalDataLayerHelper.cpp
@@ -389,6 +389,28 @@ void FExternalDataLayerHelper::ForEachExternalDataLayerLevelPackagePath(const FS
 	}
 }
 
+FStringView FExternalDataLayerHelper::GetRelativeExternalActorPackagePath(FStringView InExternalDataLayerExternalActorPackagePath)
+{
+	uint32 ExternalActorIdx = UE::String::FindFirst(InExternalDataLayerExternalActorPackagePath, FPackagePath::GetExternalActorsFolderName(), ESearchCase::IgnoreCase);
+	if (ExternalActorIdx != INDEX_NONE)
+	{
+		FStringView RelativePath = InExternalDataLayerExternalActorPackagePath.RightChop(ExternalActorIdx + FCString::Strlen(FPackagePath::GetExternalActorsFolderName()));
+		if (RelativePath.Left(ExternalDataLayerFolder.Len()).Equals(GetExternalDataLayerFolder()))
+		{
+			if (!RelativePath.IsEmpty())
+			{
+				check(RelativePath.Left(GetExternalDataLayerFolder().Len()).Equals(GetExternalDataLayerFolder()));
+				RelativePath = RelativePath.RightChop(GetExternalDataLayerFolder().Len());
+				if (!RelativePath.IsEmpty())
+				{
+					return RelativePath.RightChop(RelativePath.Find(TEXT("/")));
+				}
+			}
+		}
+	}
+	return FStringView();
+}
+
 bool FExternalDataLayerHelper::IsExternalDataLayerPath(FStringView InExternalDataLayerPath, FExternalDataLayerUID* OutExternalDataLayerUID)
 {
 	int32 ExternalDataLayerFolderIdx = UE::String::FindFirst(InExternalDataLayerPath, GetExternalDataLayerFolder(), ESearchCase::IgnoreCase);
diff --git a/Engine/Source/Runtime/Engine/Public/WorldPartition/DataLayer/ExternalDataLayerHelper.h b/Engine/Source/Runtime/Engine/Public/WorldPartition/DataLayer/ExternalDataLayerHelper.h
--- a/Engine/Source/Runtime/Engine/Public/WorldPartition/DataLayer/ExternalDataLayerHelper.h
+++ b/Engine/Source/Runtime/Engine/Public/WorldPartition/DataLayer/ExternalDataLayerHelper.h
@@ -47,6 +47,12 @@ public:
 
 	/** Iterates through all possible External Data Layer Level Package Paths using Asset Registry. */
 	ENGINE_API static void ForEachExternalDataLayerLevelPackagePath(const FString& InLevelPackageName, TFunctionRef<void(const FString&)> Func);
+
+	/* Returns the external actor package relative path for an actor package of an actor using External Data Layers.
+	 * InExternalDataLayerExternalActorPackagePath format is : /{MountPoint}/{ExternalActorFolder}/{ExternalDataLayerFolder}/{EDL_UID}/{ExternalActorPackagePath}
+	 * return format is : /{ExternalActorPackagePath}, empty otherwise
+	 */
+	ENGINE_API static FStringView GetRelativeExternalActorPackagePath(FStringView InExternalDataLayerExternalActorPackagePath);
 #endif
 
 private:
