Fixed bound check  on CachedStaticLighting to reflect the fact we always add 2 IDs now

diff --git a/Engine/Source/Runtime/Engine/Private/Components/StaticMeshComponent.cpp b/Engine/Source/Runtime/Engine/Private/Components/StaticMeshComponent.cpp
--- a/Engine/Source/Runtime/Engine/Private/Components/StaticMeshComponent.cpp
+++ b/Engine/Source/Runtime/Engine/Private/Components/StaticMeshComponent.cpp
@@ -143,7 +143,7 @@ bool FStaticMeshComponentInstanceData::ApplyVertexColorData(UStaticMeshComponent
 		for(int32 LODIndex = 0; LODIndex < StaticMeshComponent->LODData.Num(); ++LODIndex)
 		{
 			FStaticMeshComponentLODInfo& LODInfo = StaticMeshComponent->LODData[LODIndex];
-			if(CachedStaticLighting.IsValidIndex(LODIndex))
+			if(CachedStaticLighting.IsValidIndex((LODIndex*2)+1))
 			{
 				LODInfo.OriginalMapBuildDataId = CachedStaticLighting[(LODIndex*2)];
 				LODInfo.MapBuildDataId = CachedStaticLighting[(LODIndex*2)+1];
