Fix NumLODLightmapData to take into account we have 2 guid in ApplyComponentInstanceData for StaticMeshComponent & InstanctedStaticMeshComponent

diff --git a/Engine/Source/Runtime/Engine/Private/Components/StaticMeshComponent.cpp b/Engine/Source/Runtime/Engine/Private/Components/StaticMeshComponent.cpp
--- a/Engine/Source/Runtime/Engine/Private/Components/StaticMeshComponent.cpp
+++ b/Engine/Source/Runtime/Engine/Private/Components/StaticMeshComponent.cpp
@@ -2916,7 +2916,7 @@ void UStaticMeshComponent::ApplyComponentInstanceData(FStaticMeshComponentInstan
 		return;
 	}
 
-	const int32 NumLODLightMaps = StaticMeshInstanceData->CachedStaticLighting.Num();
+	const int32 NumLODLightMaps = StaticMeshInstanceData->CachedStaticLighting.Num()/2;
 
 	if (HasStaticLighting() && NumLODLightMaps > 0)
 	{
diff --git a/Engine/Source/Runtime/Engine/Private/InstancedStaticMesh.cpp b/Engine/Source/Runtime/Engine/Private/InstancedStaticMesh.cpp
--- a/Engine/Source/Runtime/Engine/Private/InstancedStaticMesh.cpp
+++ b/Engine/Source/Runtime/Engine/Private/InstancedStaticMesh.cpp
@@ -2155,7 +2155,7 @@ void UInstancedStaticMeshComponent::ApplyComponentInstanceData(FInstancedStaticM
 	// Restore static lighting if appropriate
 	if (bMatch)
 	{
-		const int32 NumLODLightMaps = InstancedMeshData->CachedStaticLighting.MapBuildDataIds.Num();
+		const int32 NumLODLightMaps = InstancedMeshData->CachedStaticLighting.MapBuildDataIds.Num()/2;
 		SetLODDataCount(NumLODLightMaps, NumLODLightMaps);
 
 		for (int32 i = 0; i < NumLODLightMaps; ++i)
