Skip creation of OrignalMapBuildDataIDs on PreSave for Archetypes/CDOs, this is never useful
Skip creation of OrignalMapBuildDataID on PreSave for ULightComponent on cook, if you don't have it now, you don't need it

diff --git a/Engine/Source/Runtime/Engine/Private/Components/LightComponent.cpp b/Engine/Source/Runtime/Engine/Private/Components/LightComponent.cpp
--- a/Engine/Source/Runtime/Engine/Private/Components/LightComponent.cpp
+++ b/Engine/Source/Runtime/Engine/Private/Components/LightComponent.cpp
@@ -614,7 +614,11 @@ void ULightComponent::PreSave(const class ITargetPlatform* TargetPlatform)
 void ULightComponent::PreSave(FObjectPreSaveContext ObjectSaveContext)
 {
 	Super::PreSave(ObjectSaveContext);
-	ValidateLightGUIDs();
+
+	if (!ObjectSaveContext.IsCooking() && !IsTemplate())
+	{
+		ValidateLightGUIDs();
+	}
 }
 
 bool ULightComponent::CanEditChange(const FProperty* InProperty) const
diff --git a/Engine/Source/Runtime/Engine/Private/Components/StaticMeshComponent.cpp b/Engine/Source/Runtime/Engine/Private/Components/StaticMeshComponent.cpp
--- a/Engine/Source/Runtime/Engine/Private/Components/StaticMeshComponent.cpp
+++ b/Engine/Source/Runtime/Engine/Private/Components/StaticMeshComponent.cpp
@@ -484,7 +484,7 @@ void UStaticMeshComponent::PreSave(FObjectPreSaveContext ObjectSaveContext)
 	// To avoid having to resubmit all Actors on 1st Static Lighting computation we want this GUID to always be created and serialized as soon as possible. 
 	// But we only create it for the editor save context, on cook builds we avoid creating it since if it's missing it won't link to anything and it'll make the
 	// cook non-deterministic
-	if (!ObjectSaveContext.IsCooking())
+	if (!ObjectSaveContext.IsCooking() && !IsTemplate())
 	{
 		UpdateStaticLightingData();
 
