Prevent UObject access in HeterogeneousVolumeComponent's PrimitiveSceneProxy

diff --git a/Engine/Source/Runtime/Engine/Private/Components/HeterogeneousVolumeComponent.cpp b/Engine/Source/Runtime/Engine/Private/Components/HeterogeneousVolumeComponent.cpp
--- a/Engine/Source/Runtime/Engine/Private/Components/HeterogeneousVolumeComponent.cpp
+++ b/Engine/Source/Runtime/Engine/Private/Components/HeterogeneousVolumeComponent.cpp
@@ -51,10 +51,13 @@ public:
 	//~ End FPrimitiveSceneProxy Interface.
 
 private:
-	UMaterialInterface* MaterialInterface;
 	FLocalVertexFactory VertexFactory;
 	FStaticMeshVertexBuffers StaticMeshVertexBuffers;
 	FHeterogeneousVolumeData HeterogeneousVolumeData;
+
+	// Cache UObject values
+	FMaterialRelevance MaterialRelevance;
+	FMaterialRenderProxy* MaterialRenderProxy;
 };
 
 /*=============================================================================
@@ -63,7 +66,6 @@ private:
 
 FHeterogeneousVolumeSceneProxy::FHeterogeneousVolumeSceneProxy(UHeterogeneousVolumeComponent* InComponent)
 	: FPrimitiveSceneProxy(InComponent)
-	, MaterialInterface(InComponent->GetMaterial(0))
 	, VertexFactory(GetScene().GetFeatureLevel(), "FHeterogeneousVolumeSceneProxy")
 #if ACTOR_HAS_LABELS
 	, HeterogeneousVolumeData(this, InComponent->GetReadableName())
@@ -88,11 +90,19 @@ FHeterogeneousVolumeSceneProxy::FHeterogeneousVolumeSceneProxy(UHeterogeneousVol
 	VoxelSize.Z /= InComponent->VolumeResolution.Z;
 	HeterogeneousVolumeData.MinimumVoxelSize = FMath::Max(VoxelSize.GetMin(), 0.001);
 
+	UMaterialInterface* MaterialInterface = InComponent->GetMaterial(0);
 	if (InComponent->MaterialInstanceDynamic)
 	{
 		MaterialInterface = InComponent->MaterialInstanceDynamic;
 	}
 
+	MaterialRenderProxy = nullptr;
+	if (MaterialInterface)
+	{
+		MaterialRelevance = MaterialInterface->GetRelevance_Concurrent(GetScene().GetFeatureLevel());
+		MaterialRenderProxy = MaterialInterface->GetRenderProxy();
+	}
+
 	HeterogeneousVolumeData.StepFactor = InComponent->StepFactor;
 	HeterogeneousVolumeData.ShadowStepFactor = InComponent->ShadowStepFactor;
 	HeterogeneousVolumeData.ShadowBiasFactor = InComponent->ShadowBiasFactor;
@@ -154,12 +164,7 @@ FPrimitiveViewRelevance FHeterogeneousVolumeSceneProxy::GetViewRelevance(const F
 {
 	FPrimitiveViewRelevance Result;
 
-	if (MaterialInterface)
-	{
-		FMaterialRelevance MaterialRelevance = MaterialInterface->GetRelevance_Concurrent(View->GetFeatureLevel());
-		MaterialRelevance.SetPrimitiveViewRelevance(Result);
-	}
-
+	MaterialRelevance.SetPrimitiveViewRelevance(Result);
 	Result.bDrawRelevance = IsShown(View) && View->Family->EngineShowFlags.HeterogeneousVolumes;
 	Result.bOpaque = false;
 	Result.bStaticRelevance = false;
@@ -187,13 +192,13 @@ void FHeterogeneousVolumeSceneProxy::GetDynamicMeshElements(
 		{
 			if (VisibilityMap & (1 << ViewIndex))
 			{
-				if (MaterialInterface)
+				if (MaterialRenderProxy)
 				{
 					// Set up MeshBatch
 					FMeshBatch& Mesh = Collector.AllocateMesh();
 
 					Mesh.VertexFactory = &VertexFactory;
-					Mesh.MaterialRenderProxy = MaterialInterface->GetRenderProxy();
+					Mesh.MaterialRenderProxy = MaterialRenderProxy;
 					Mesh.LCI = NULL;
 					Mesh.ReverseCulling = IsLocalToWorldDeterminantNegative() ? true : false;
 					Mesh.CastShadow = CastsDynamicShadow();
