Fixed check condition, the the level is owned by a level with no GUID the resolved GUID is expected to be invalid

diff --git a/Engine/Source/Runtime/Engine/Private/WorldPartition/ActorInstanceGuids.cpp b/Engine/Source/Runtime/Engine/Private/WorldPartition/ActorInstanceGuids.cpp
--- a/Engine/Source/Runtime/Engine/Private/WorldPartition/ActorInstanceGuids.cpp
+++ b/Engine/Source/Runtime/Engine/Private/WorldPartition/ActorInstanceGuids.cpp
@@ -60,7 +60,7 @@ FGuid FActorInstanceGuid::GetLevelInstanceGuid(ULevel* Level)
 			if (!LevelGuids.IsDefault())
 			{
 				// Duplicate the annotation to this level to avoid having to look it up all the time
-				check(LevelGuids.ResolvedLevelInstanceGuid.IsValid());
+				check(LevelGuids.ResolvedLevelInstanceGuid.IsValid() || !LevelGuids.LevelInstanceGuid.IsValid());
 				GLevelInstanceGuids.AddAnnotation(Level, LevelGuids);
 			}
 		}
