Substrate - volumetric fog/cloud node now handle the legacy unlit (emissive only) shading model at conversion time and also when overriden on a material instance.

diff --git a/Engine/Shaders/Private/HeterogeneousVolumes/HeterogeneousVolumesLiveShadingUtils.ush b/Engine/Shaders/Private/HeterogeneousVolumes/HeterogeneousVolumesLiveShadingUtils.ush
--- a/Engine/Shaders/Private/HeterogeneousVolumes/HeterogeneousVolumesLiveShadingUtils.ush
+++ b/Engine/Shaders/Private/HeterogeneousVolumes/HeterogeneousVolumesLiveShadingUtils.ush
@@ -1,5 +1,8 @@
 // Copyright Epic Games, Inc. All Rights Reserved.
 
+// Change this to force recompilation of all shaders referencing this file.
+#pragma message("UESHADERMETADATA_VERSION 80584AA0-BFD9-4021-BE67-273077B48601")
+
 FPrimitiveSceneData GetPrimitiveData(
 	uint PrimitiveId,
 	float4x4 LocalToWorld,
diff --git a/Engine/Shaders/Private/VolumetricCloud.usf b/Engine/Shaders/Private/VolumetricCloud.usf
--- a/Engine/Shaders/Private/VolumetricCloud.usf
+++ b/Engine/Shaders/Private/VolumetricCloud.usf
@@ -7,7 +7,7 @@
 
 
 // Change this to force recompilation of all shaders referencing this file.
-#pragma message("UESHADERMETADATA_VERSION B0DC9137-AC62-453E-A2F9-A01047CC479C")
+#pragma message("UESHADERMETADATA_VERSION 58372ED6-B830-41C4-829E-A21186A68401")
 
 #include "Common.ush"
 #include "ColorMap.ush"
diff --git a/Engine/Shaders/Private/VolumetricFogVoxelization.usf b/Engine/Shaders/Private/VolumetricFogVoxelization.usf
--- a/Engine/Shaders/Private/VolumetricFogVoxelization.usf
+++ b/Engine/Shaders/Private/VolumetricFogVoxelization.usf
@@ -4,6 +4,9 @@
 	VolumetricFogVoxelization.usf
 =============================================================================*/
 
+// Change this to force recompilation of all shaders referencing this file.
+#pragma message("UESHADERMETADATA_VERSION 43567765-38CF-4EC3-B288-49B6F14EF801")
+
 // So vertex factories can interpolate dependencies of VertexFactoryGetTranslatedPrimitiveVolumeBounds
 #define PASS_NEEDS_PRIMITIVE_VOLUME_BOUNDS 1
 
diff --git a/Engine/Source/Runtime/Engine/Private/Materials/HLSLMaterialTranslator.cpp b/Engine/Source/Runtime/Engine/Private/Materials/HLSLMaterialTranslator.cpp
--- a/Engine/Source/Runtime/Engine/Private/Materials/HLSLMaterialTranslator.cpp
+++ b/Engine/Source/Runtime/Engine/Private/Materials/HLSLMaterialTranslator.cpp
@@ -13483,8 +13483,36 @@ int32 FHLSLMaterialTranslator::SubstrateConversionFromLegacy(
 	);
 }
 
-int32 FHLSLMaterialTranslator::SubstrateVolumetricFogCloudBSDF(int32 Albedo, int32 Extinction, int32 EmissiveColor, int32 AmbientOcclusion)
+int32 FHLSLMaterialTranslator::SubstrateVolumetricFogCloudBSDF(int32 Albedo, int32 Extinction, int32 EmissiveColor, int32 AmbientOcclusion, bool bEmissiveOnly)
 {
+	// EmissiveOnly isrequired due to the legacy material shading model set by the user, which could be Unlit or Lit for volumetric domain.
+	// It was set on the material detail panel or overriden from a mateiral instance. Note: material attributes could not be used with the volumetric domain.
+
+	UMaterial* BaseMaterial = Material->GetMaterialInterface()->GetBaseMaterial();
+	// Material is probably an instanced material, so we check that this is the case before potentially overriding the ShadingModel from the material instance.
+	UMaterialInterface* BaseMaterialInterface = static_cast<UMaterialInterface*>(BaseMaterial);
+	UMaterialInterface* MaterialInterface = Material->GetMaterialInterface();
+	if (BaseMaterial && (BaseMaterialInterface != MaterialInterface))
+	{
+		FMaterialShadingModelField BaseMaterialShadingModels = BaseMaterial->GetShadingModels();
+		FMaterialShadingModelField MaterialShadingModels = Material->GetShadingModels();
+
+		// If the potentially instanced material does not have the same shading model as the instance material, this means it would have overridden the shading model.
+		// From the UI, only a single shading model is selectable, so we simply apply the one coming form the material instance.
+		if (MaterialShadingModels.IsValid() && MaterialShadingModels.CountShadingModels() == 1 && MaterialShadingModels != BaseMaterialShadingModels)
+		{
+			// Use the mode required by the overriden shading model.
+			bEmissiveOnly = MaterialShadingModels.GetFirstShadingModel() == MSM_Unlit ? true : false;
+		}
+	}
+
+	// Override to EmissiveOnly if that is the final shading model.
+	if (bEmissiveOnly)
+	{
+		Albedo = Extinction = Constant3(0.0f, 0.0f, 0.0f);
+		AmbientOcclusion = Constant(1.0f);
+	}
+
 	return AddCodeChunk(
 		MCT_Substrate, TEXT("GetSubstrateVolumeFogCloudBSDF(%s, %s, %s, %s)"),
 		*SubstrateGetCastParameterCode(Albedo,				MCT_Float3),
diff --git a/Engine/Source/Runtime/Engine/Private/Materials/HLSLMaterialTranslator.h b/Engine/Source/Runtime/Engine/Private/Materials/HLSLMaterialTranslator.h
--- a/Engine/Source/Runtime/Engine/Private/Materials/HLSLMaterialTranslator.h
+++ b/Engine/Source/Runtime/Engine/Private/Materials/HLSLMaterialTranslator.h
@@ -1288,7 +1288,7 @@ protected:
 		int32 ClearCoat_Normal, int32 ClearCoat_Tangent, const FString& ClearCoat_SharedLocalBasisIndexMacro, 
 		int32 CustomTangent_Tangent,
 		FSubstrateOperator* PromoteToOperator) override;
-	virtual int32 SubstrateVolumetricFogCloudBSDF(int32 Albedo, int32 Extinction, int32 EmissiveColor, int32 AmbientOcclusion) override;
+	virtual int32 SubstrateVolumetricFogCloudBSDF(int32 Albedo, int32 Extinction, int32 EmissiveColor, int32 AmbientOcclusion, bool bEmissiveOnly) override;
 	virtual int32 SubstrateUnlitBSDF(int32 EmissiveColor, int32 TransmittanceColor, int32 Normal, FSubstrateOperator* PromoteToOperator) override;
 	virtual int32 SubstrateHairBSDF(int32 BaseColor, int32 Scatter, int32 Specular, int32 Roughness, int32 Backlit, int32 EmissiveColor, int32 Tangent, const FString& SharedLocalBasisIndexMacro, FSubstrateOperator* PromoteToOperator) override;
 	virtual int32 SubstrateEyeBSDF(int32 DiffuseColor, int32 Roughness, int32 IrisMask, int32 IrisDistance, int32 EmissiveColor, int32 CorneaNormal, int32 IrisNormal, int32 IrisPlaneNormal, int32 SSSProfileId, const FString& SharedLocalBasisIndexMacro, FSubstrateOperator* PromoteToOperator) override;
diff --git a/Engine/Source/Runtime/Engine/Private/Materials/Material.cpp b/Engine/Source/Runtime/Engine/Private/Materials/Material.cpp
--- a/Engine/Source/Runtime/Engine/Private/Materials/Material.cpp
+++ b/Engine/Source/Runtime/Engine/Private/Materials/Material.cpp
@@ -3604,6 +3604,8 @@ void UMaterial::ConvertMaterialToSubstrateMaterial()
 			ColorMatInputConnectionTo(EditorOnly->EmissiveColor, VolBSDF, 2, MP_EmissiveColor);	
 			ScalarMatInputConnectionTo(EditorOnly->AmbientOcclusion, VolBSDF, 3, MP_AmbientOcclusion);
 
+			VolBSDF->bEmissiveOnly = ShadingModel == MSM_Unlit;
+
 			// SUBSTRATE_TODO remove the VolumetricAdvancedOutput node and add the input onto FogCloudBSDF even if only used by the cloud renderer?
 			EditorOnly->FrontMaterial.Connect(0, VolBSDF);
 			bInvalidateShader = true;
@@ -4539,7 +4541,8 @@ bool UMaterial::CanEditChange(const FProperty* InProperty) const
 	
 		if (PropertyName == GET_MEMBER_NAME_STRING_CHECKED(UMaterial, ShadingModel))
 		{
-			return !bSubstrateEnabled && MaterialDomain == MD_Surface;
+			// Volume can also use the unlit mode for emissive only volumetric effect (also overridable on material instance)
+			return !bSubstrateEnabled && (MaterialDomain == MD_Surface || MaterialDomain == MD_Volume);
 		}
 
 		if (PropertyName == GET_MEMBER_NAME_STRING_CHECKED(UMaterial, bIsThinSurface))
@@ -5084,11 +5087,11 @@ void UMaterial::RebuildShadingModelField()
 			}
 			else if (MaterialDomain == MD_LightFunction)
 			{
-				SubstrateMaterialInfo.AddShadingModel(SSM_VolumetricFogCloud);
+				SubstrateMaterialInfo.AddShadingModel(SSM_LightFunction);
 			}
 			else if (MaterialDomain == MD_Volume)
 			{
-				SubstrateMaterialInfo.AddShadingModel(SSM_LightFunction);
+				SubstrateMaterialInfo.AddShadingModel(SSM_VolumetricFogCloud);
 			}
 			else if (MaterialDomain == MD_PostProcess)
 			{
@@ -5173,7 +5176,14 @@ void UMaterial::RebuildShadingModelField()
 			else if (SubstrateMaterialInfo.HasOnlyShadingModel(SSM_VolumetricFogCloud))
 			{
 				MaterialDomain = EMaterialDomain::MD_Volume;
-				ShadingModel = MSM_DefaultLit;
+				// Volumetric shading model do support Unlit as an EmissiveOnly mode, on top of the lit mode.
+				TArray<UMaterialExpressionSubstrateVolumetricFogCloudBSDF*> ShadingModelExpressions;
+				GetAllExpressionsInMaterialAndFunctionsOfType(ShadingModelExpressions);
+				for (UMaterialExpressionSubstrateVolumetricFogCloudBSDF* MatExpr : ShadingModelExpressions)
+				{
+					ShadingModel = MatExpr->bEmissiveOnly ? MSM_Unlit : MSM_DefaultLit;
+					break; // We can only keep a single shading model due to an assert in UMaterialInstance::UpdateOverridableBaseProperties which want a single shading model for non material attribute workflow.
+				}
 				BlendMode = EBlendMode::BLEND_Additive;
 			}
 			else if (SubstrateMaterialInfo.HasOnlyShadingModel(SSM_Hair))
diff --git a/Engine/Source/Runtime/Engine/Private/Materials/MaterialExpressions.cpp b/Engine/Source/Runtime/Engine/Private/Materials/MaterialExpressions.cpp
--- a/Engine/Source/Runtime/Engine/Private/Materials/MaterialExpressions.cpp
+++ b/Engine/Source/Runtime/Engine/Private/Materials/MaterialExpressions.cpp
@@ -25643,14 +25643,20 @@ int32 UMaterialExpressionSubstrateVolumetricFogCloudBSDF::Compile(class FMateria
 		CompileWithDefaultFloat3(Compiler, Albedo, 0.0f, 0.0f, 0.0f),
 		CompileWithDefaultFloat3(Compiler, Extinction, 0.0f, 0.0f, 0.0f),
 		CompileWithDefaultFloat3(Compiler, EmissiveColor, 0.0f, 0.0f, 0.0f),
-		CompileWithDefaultFloat1(Compiler, AmbientOcclusion, 1.0f));
+		CompileWithDefaultFloat1(Compiler, AmbientOcclusion, 1.0f),
+		bEmissiveOnly);
 
 	return OutputCodeChunk;
 }
 
 void UMaterialExpressionSubstrateVolumetricFogCloudBSDF::GetCaption(TArray<FString>& OutCaptions) const
 {
-	OutCaptions.Add(TEXT("Substrate Volumetric-Fog-Cloud BSDF"));
+	FString Caption = TEXT("Substrate Volumetric-Fog-Cloud BSDF");
+	if (bEmissiveOnly)
+	{
+		Caption += TEXT("(Emissive Only)");
+	}
+	OutCaptions.Add(Caption);
 }
 
 uint32 UMaterialExpressionSubstrateVolumetricFogCloudBSDF::GetOutputType(int32 OutputIndex)
diff --git a/Engine/Source/Runtime/Engine/Public/MaterialCompiler.h b/Engine/Source/Runtime/Engine/Public/MaterialCompiler.h
--- a/Engine/Source/Runtime/Engine/Public/MaterialCompiler.h
+++ b/Engine/Source/Runtime/Engine/Public/MaterialCompiler.h
@@ -680,7 +680,7 @@ public:
 		int32 ClearCoat_Normal, int32 ClearCoat_Tangent, const FString& ClearCoat_SharedLocalBasisIndexMacro,
 		int32 CustomTangent_Tangent,
 		FSubstrateOperator* PromoteToOperator) = 0;
-	virtual int32 SubstrateVolumetricFogCloudBSDF(int32 Albedo, int32 Extinction, int32 EmissiveColor, int32 AmbientOcclusion) = 0;
+	virtual int32 SubstrateVolumetricFogCloudBSDF(int32 Albedo, int32 Extinction, int32 EmissiveColor, int32 AmbientOcclusion, bool bEmissiveOnly) = 0;
 	virtual int32 SubstrateUnlitBSDF(int32 EmissiveColor, int32 TransmittanceColor, int32 Normal, FSubstrateOperator* PromoteToOperator) = 0;
 	virtual int32 SubstrateHairBSDF(int32 BaseColor, int32 Scatter, int32 Specular, int32 Roughness, int32 Backlit, int32 EmissiveColor, int32 Tangent, const FString& SharedLocalBasisIndexMacro, FSubstrateOperator* PromoteToOperator) = 0;
 	virtual int32 SubstrateEyeBSDF(int32 DiffuseAlbedo, int32 Roughness, int32 IrisMask, int32 IrisDistance, int32 IrisNormal, int32 IrisPlaneNormal, int32 SSSProfileId, int32 EmissiveColor, int32 CorneaNormal, const FString& SharedLocalBasisIndexMacro, FSubstrateOperator* PromoteToOperator) = 0;
@@ -1387,9 +1387,9 @@ public:
 			PromoteToOperator);
 	}
 
-	virtual int32 SubstrateVolumetricFogCloudBSDF(int32 Albedo, int32 Extinction, int32 EmissiveColor, int32 AmbientOcclusion) override
+	virtual int32 SubstrateVolumetricFogCloudBSDF(int32 Albedo, int32 Extinction, int32 EmissiveColor, int32 AmbientOcclusion, bool bEmissiveOnly) override
 	{
-		return Compiler->SubstrateVolumetricFogCloudBSDF(Albedo, Extinction, EmissiveColor, AmbientOcclusion);
+		return Compiler->SubstrateVolumetricFogCloudBSDF(Albedo, Extinction, EmissiveColor, AmbientOcclusion, bEmissiveOnly);
 	}
 
 	virtual int32 SubstrateUnlitBSDF(int32 EmissiveColor, int32 TransmittanceColor, int32 Normal, FSubstrateOperator* PromoteToOperator) override
diff --git a/Engine/Source/Runtime/Engine/Public/Materials/MaterialExpressionSubstrate.h b/Engine/Source/Runtime/Engine/Public/Materials/MaterialExpressionSubstrate.h
--- a/Engine/Source/Runtime/Engine/Public/Materials/MaterialExpressionSubstrate.h
+++ b/Engine/Source/Runtime/Engine/Public/Materials/MaterialExpressionSubstrate.h
@@ -445,6 +445,10 @@ class UMaterialExpressionSubstrateVolumetricFogCloudBSDF : public UMaterialExpre
 	UPROPERTY()
 	FExpressionInput AmbientOcclusion;
 
+	/** Enabling this corresponds to selecting the Unlit shading model from the legacy material system. In this case, only the Emissive input will be considered. */
+	UPROPERTY(EditAnywhere, Category = Shading, meta = (DisplayName = "Unlit (EmissiveOnly)"))
+	uint32 bEmissiveOnly : 1;
+
 	//~ Begin UMaterialExpression Interface
 #if WITH_EDITOR
 	virtual int32 Compile(class FMaterialCompiler* Compiler, int32 OutputIndex) override;
