UE-207416 - Derived struct data is not copied to base struct

CL 22605284 made Generic_SetStructurePropertyByName a bit too strict by only allowing structs that matched exactly for copying properties. Prior to the change, this function was incredibly lax by allowing completely unrelated structures to be copied. The appropriate middleground allows derived structs. While this may lead to object slicing in some instances, this is permissable.

Additionally, removed the public-facing Generic_SetStructurePropertyByName function since it was deprecated.

diff --git a/Engine/Source/Runtime/Engine/Classes/Kismet/KismetSystemLibrary.h b/Engine/Source/Runtime/Engine/Classes/Kismet/KismetSystemLibrary.h
--- a/Engine/Source/Runtime/Engine/Classes/Kismet/KismetSystemLibrary.h
+++ b/Engine/Source/Runtime/Engine/Classes/Kismet/KismetSystemLibrary.h
@@ -996,9 +996,6 @@ class UKismetSystemLibrary : public UBlueprintFunctionLibrary
 	UFUNCTION(BlueprintCallable, CustomThunk, meta = (BlueprintInternalUseOnly = "true", CustomStructureParam = "Value", AutoCreateRefTerm = "Value"))
 	static ENGINE_API void SetStructurePropertyByName(UObject* Object, FName PropertyName, const FGenericStruct& Value);
 
-	UE_DEPRECATED(5.2, "Function has been deprecated.")
-	static ENGINE_API void Generic_SetStructurePropertyByName(UObject* OwnerObject, FName StructPropertyName, const void* SrcStructAddr);
-
 	/** Based on UKismetArrayLibrary::execSetArrayPropertyByName */
 	DECLARE_FUNCTION(execSetStructurePropertyByName);
 
diff --git a/Engine/Source/Runtime/Engine/Private/KismetSystemLibrary.cpp b/Engine/Source/Runtime/Engine/Private/KismetSystemLibrary.cpp
--- a/Engine/Source/Runtime/Engine/Private/KismetSystemLibrary.cpp
+++ b/Engine/Source/Runtime/Engine/Private/KismetSystemLibrary.cpp
@@ -63,26 +63,37 @@ namespace UE::Blueprint::Private
 		bBlamePrintString,
 		TEXT("When true, prints the Blueprint Asset and Function that generated calls to Print String. Useful for tracking down screen message spam."));
 
-	void Generic_SetStructurePropertyByName(UObject* OwnerObject, FName StructPropertyName, FStructProperty* SrcStructProperty, const void* SrcStructAddr)
+	void Generic_SetStructurePropertyByName(UObject* OwnerObject, FName StructPropertyName, const FStructProperty* SrcStructProperty, const void* SrcStructAddr)
 	{
 		if (OwnerObject != nullptr)
 		{
-			FStructProperty* DestStructProperty = FindFProperty<FStructProperty>(OwnerObject->GetClass(), StructPropertyName);
+			const FStructProperty* DestStructProperty = FindFProperty<FStructProperty>(OwnerObject->GetClass(), StructPropertyName);
 
 			// SrcStructAddr and SrcStructProperty can be null in certain scenarios.
 			// For example, retrieving an element reference from an array of user structs in BP can result in a null source.
 			// We'll report a BP exception in that case, but we also need to soft-fail here to prevent an assert in CopyValuesInternal.
 
-			bool bCanSetStructureProperty =
+			const bool bHasValidParameters =
 				(DestStructProperty != nullptr) &&
 				(SrcStructProperty != nullptr) &&
-				(SrcStructAddr != nullptr) &&
-				SrcStructProperty->SameType(DestStructProperty);
+				(SrcStructAddr != nullptr)
+			;
 
-			if (bCanSetStructureProperty)
+			if (bHasValidParameters)
 			{
-				void* DestStructAddr = DestStructProperty->ContainerPtrToValuePtr<void>(OwnerObject);
-				DestStructProperty->CopyValuesInternal(DestStructAddr, SrcStructAddr, 1);
+				const UStruct* SourceStruct = SrcStructProperty->Struct;
+				const UStruct* DestStruct = DestStructProperty->Struct;
+
+				if (ensure(SourceStruct && DestStruct))
+				{
+					// For derived structs, this can lead to object slicing in some contexts, but that's expected behavior.
+
+					if (SourceStruct->IsChildOf(DestStruct))
+					{
+						void* DestStructAddr = DestStructProperty->ContainerPtrToValuePtr<void>(OwnerObject);
+						DestStructProperty->CopyValuesInternal(DestStructAddr, SrcStructAddr, 1);
+					}
+				}
 			}
 		}
 	}
@@ -1558,19 +1569,6 @@ void UKismetSystemLibrary::SetCollisionProfileNameProperty(UObject* Object, FNam
 	check(0);
 }
 
-void UKismetSystemLibrary::Generic_SetStructurePropertyByName(UObject* OwnerObject, FName StructPropertyName, const void* SrcStructAddr)
-{
-	if (OwnerObject != nullptr)
-	{
-		FStructProperty* StructProp = FindFProperty<FStructProperty>(OwnerObject->GetClass(), StructPropertyName);
-		if (StructProp != nullptr)
-		{
-			void* Dest = StructProp->ContainerPtrToValuePtr<void>(OwnerObject);
-			StructProp->CopyValuesInternal(Dest, SrcStructAddr, 1);
-		}
-	}
-}
-
 void UKismetSystemLibrary::GetActorListFromComponentList(const TArray<UPrimitiveComponent*>& ComponentList, UClass* ActorClassFilter, TArray<class AActor*>& OutActorList)
 {
 	OutActorList.Empty();
@@ -1591,8 +1589,6 @@ void UKismetSystemLibrary::GetActorListFromComponentList(const TArray<UPrimitive
 	}
 }
 
-
-
 bool UKismetSystemLibrary::SphereOverlapActors(const UObject* WorldContextObject, const FVector SpherePos, float SphereRadius, const TArray<TEnumAsByte<EObjectTypeQuery> > & ObjectTypes, UClass* ActorClassFilter, const TArray<AActor*>& ActorsToIgnore, TArray<AActor*>& OutActors)
 {
 	OutActors.Empty();
