Fixes crash when batch retargeting old animation with UAnimSequence::AnimationTrackNames

diff --git a/Engine/Source/Runtime/Engine/Classes/Animation/AnimSequence.h b/Engine/Source/Runtime/Engine/Classes/Animation/AnimSequence.h
--- a/Engine/Source/Runtime/Engine/Classes/Animation/AnimSequence.h
+++ b/Engine/Source/Runtime/Engine/Classes/Animation/AnimSequence.h
@@ -264,7 +264,7 @@ protected:
 	 * This is name of RawAnimationData tracks for editoronly - if we lose skeleton, we'll need relink them
 	 */
 	UE_DEPRECATED(5.0, "Animation track names has been deprecated see FBoneAnimationTrack::Name")
-	UPROPERTY(VisibleAnywhere, Category="Animation")
+	UPROPERTY(VisibleAnywhere, DuplicateTransient, Category="Animation")
 	TArray<FName> AnimationTrackNames;
 
 	/**
