Fix editor crash when raytracing is enabled and enter kanji characters into TextRenderActor.

#tests Attached projects to the jira

diff --git a/Engine/Source/Runtime/Engine/Private/Components/TextRenderComponent.cpp b/Engine/Source/Runtime/Engine/Private/Components/TextRenderComponent.cpp
--- a/Engine/Source/Runtime/Engine/Private/Components/TextRenderComponent.cpp
+++ b/Engine/Source/Runtime/Engine/Private/Components/TextRenderComponent.cpp
@@ -923,7 +923,7 @@ void FTextRenderSceneProxy::GetDynamicRayTracingInstances(FRayTracingMaterialGat
 			Mesh.bCanApplyViewModeOverrides = !bAlwaysRenderAsText;
 			Mesh.LODIndex = 0;
 			Mesh.bUseWireframeSelectionColoring = IsSelected() ? 1 : 0;
-			Mesh.SegmentIndex = 0;
+			Mesh.SegmentIndex = BatchIndex;
 			Mesh.MeshIdInPrimitive = 0;
 		}
 		RayTracingInstance.MaterialsView = MakeArrayView(CachedRayTracingMaterials);
@@ -1139,18 +1139,22 @@ void FTextRenderSceneProxy::UpdateRayTracingGeometry_RenderingThread(FRHICommand
 	Initializer.bFastBuild = true;
 	Initializer.bAllowUpdate = false;
 
-	FRayTracingGeometrySegment Segment;
-	Segment.VertexBuffer = VertexBuffers.PositionVertexBuffer.VertexBufferRHI;
-	Segment.VertexBufferElementType = VET_Float3;
-	Segment.VertexBufferStride = VertexBuffers.PositionVertexBuffer.GetStride();
-	Segment.VertexBufferOffset = 0;
-	Segment.MaxVertices = VertexBuffers.PositionVertexBuffer.GetNumVertices();
-	Segment.FirstPrimitive = 0;
-	Segment.NumPrimitives = IndexBuffer.Indices.Num() / 3;
-	Segment.bEnabled = true;
-	Segment.bForceOpaque = false;
-
-	Initializer.Segments.Add(Segment);
+	for (int BatchIndex = 0; BatchIndex < TextBatches.Num(); ++BatchIndex)
+	{
+		const FTextBatch& TextBatch = TextBatches[BatchIndex];
+		FRayTracingGeometrySegment Segment;
+		Segment.VertexBuffer = VertexBuffers.PositionVertexBuffer.VertexBufferRHI;
+		Segment.VertexBufferElementType = VET_Float3;
+		Segment.VertexBufferStride = VertexBuffers.PositionVertexBuffer.GetStride();
+		Segment.VertexBufferOffset = 0;
+		Segment.MaxVertices = VertexBuffers.PositionVertexBuffer.GetNumVertices();
+		Segment.FirstPrimitive = TextBatch.IndexBufferOffset / 3;
+		Segment.NumPrimitives =  TextBatch.IndexBufferCount / 3;
+		Segment.bEnabled = true;
+		Segment.bForceOpaque = false;
+
+		Initializer.Segments.Add(Segment);
+	}
 
 	RayTracingGeometry.SetInitializer(Initializer);
 	RayTracingGeometry.InitResource(RHICmdList);
