Skeletal Mesh: Deprecate and eviscerate {Set,Get}HasBeenSimplified. This information is stored per LOD and makes no sense as a global flag.

diff --git a/Engine/Plugins/Experimental/Mutable/Source/CustomizableObject/Private/MuCO/CustomizableObjectInstance.cpp b/Engine/Plugins/Experimental/Mutable/Source/CustomizableObject/Private/MuCO/CustomizableObjectInstance.cpp
--- a/Engine/Plugins/Experimental/Mutable/Source/CustomizableObject/Private/MuCO/CustomizableObjectInstance.cpp
+++ b/Engine/Plugins/Experimental/Mutable/Source/CustomizableObject/Private/MuCO/CustomizableObjectInstance.cpp
@@ -3438,7 +3438,6 @@ void UCustomizableInstancePrivate::InitSkeletalMeshData(const TSharedRef<FUpdate
 	SkeletalMesh->SetMinLod(FMath::Max(CustomizableObject.LODSettings.MinLOD.GetDefault(), (int32)FirstLODAvailable));
 	SkeletalMesh->SetQualityLevelMinLod(CustomizableObject.LODSettings.MinQualityLevelLOD);
 
-	SkeletalMesh->SetHasBeenSimplified(false);
 	SkeletalMesh->SetHasVertexColors(false);
 
 	// Set the default Physics Assets
diff --git a/Engine/Plugins/Experimental/Mutable/Source/CustomizableObjectEditor/Private/MuCOE/MutableMeshPreviewUtils.cpp b/Engine/Plugins/Experimental/Mutable/Source/CustomizableObjectEditor/Private/MuCOE/MutableMeshPreviewUtils.cpp
--- a/Engine/Plugins/Experimental/Mutable/Source/CustomizableObjectEditor/Private/MuCOE/MutableMeshPreviewUtils.cpp
+++ b/Engine/Plugins/Experimental/Mutable/Source/CustomizableObjectEditor/Private/MuCOE/MutableMeshPreviewUtils.cpp
@@ -368,7 +368,6 @@ namespace MutableMeshPreviewUtils
 		 */
 		bool BuildSkeletalMeshRenderData(const mu::MeshPtrConst InMutableMesh, USkeletalMesh* OutSkeletalMesh)
 		{
-			OutSkeletalMesh->SetHasBeenSimplified(false);
 			OutSkeletalMesh->SetHasVertexColors(false);
 
 			// Find how many bones the bonemap could have
diff --git a/Engine/Plugins/Experimental/SkeletalReduction/Source/Private/SkeletalMeshReductionPlugin.cpp b/Engine/Plugins/Experimental/SkeletalReduction/Source/Private/SkeletalMeshReductionPlugin.cpp
--- a/Engine/Plugins/Experimental/SkeletalReduction/Source/Private/SkeletalMeshReductionPlugin.cpp
+++ b/Engine/Plugins/Experimental/SkeletalReduction/Source/Private/SkeletalMeshReductionPlugin.cpp
@@ -2324,7 +2324,6 @@ void FQuadricSkeletalMeshReduction::ReduceSkeletalMesh(USkeletalMesh& SkeletalMe
 		}
 		// Flag this LOD as having been simplified.
 		ReducedLODInfoPtr->bHasBeenSimplified = true;
-		SkeletalMesh.SetHasBeenSimplified(true);
 
 		//Restore the user sections data to what it was. It must be done if we want to avoid changing the DDC key. I.E. UserSectionData is part of the key
 		//DDC key cannot be change during the build
@@ -2381,7 +2380,6 @@ void FQuadricSkeletalMeshReduction::ReduceSkeletalMesh(USkeletalMesh& SkeletalMe
 
 		NewModel->RequiredBones.Empty();
 		SkeletalMesh.GetLODInfo(LODIndex)->bHasBeenSimplified = true;
-		SkeletalMesh.SetHasBeenSimplified(true);
 	}
 	
 	if (!bLODModelAdded)
diff --git a/Engine/Source/Runtime/Engine/Classes/Engine/SkeletalMesh.h b/Engine/Source/Runtime/Engine/Classes/Engine/SkeletalMesh.h
--- a/Engine/Source/Runtime/Engine/Classes/Engine/SkeletalMesh.h
+++ b/Engine/Source/Runtime/Engine/Classes/Engine/SkeletalMesh.h
@@ -22,6 +22,7 @@
 #include "SkeletalMeshSampling.h"
 #include "SkeletalMeshSourceModel.h"
 #include "SkinnedAsset.h"
+#include "SkinnedAssetCommon.h"
 #include "Templates/SubclassOf.h"
 #include "UObject/Object.h"
 #include "UObject/ObjectMacros.h"
@@ -106,7 +107,7 @@ enum class ESkeletalMeshAsyncProperties : uint64
 	NegativeBoundsExtension = 1llu << 32,
 	PositiveBoundsExtension = 1llu << 33,
 	ExtendedBounds = 1llu << 34,
-	HasBeenSimplified = 1llu << 35,
+	// Spare = 1llu << 35,
 	EnablePerPolyCollision = 1llu << 36,
 	BodySetup = 1llu << 37,
 	MorphTargetIndexMap = 1llu << 38,
@@ -1278,31 +1279,18 @@ public:
 	UPROPERTY()
 	uint8 bUseHighPrecisionTangentBasis_DEPRECATED : 1;
 
-	/** true if this mesh has ever been simplified with Simplygon. */
-	UE_DEPRECATED(4.27, "Please do not access this member directly; use USkeletalMesh::GetHasBeenSimplified() or USkeletalMesh::SetHasBeenSimplified().")
-	UPROPERTY()
-	uint8 bHasBeenSimplified:1;
+	UE_DEPRECATED(5.5, "This functionality has been removed. Use FSkeletalMeshLODInfo::bHasBeenSimplified instead.")
 	static FName GetHasBeenSimplifiedMemberName()
 	{
-		PRAGMA_DISABLE_DEPRECATION_WARNINGS
-		return GET_MEMBER_NAME_CHECKED(USkeletalMesh, bHasBeenSimplified);
-		PRAGMA_ENABLE_DEPRECATION_WARNINGS
+		return NAME_None;
 	}
 
-	bool GetHasBeenSimplified() const
-	{
-		WaitUntilAsyncPropertyReleased(ESkeletalMeshAsyncProperties::HasBeenSimplified, ESkinnedAssetAsyncPropertyLockType::ReadOnly);
-		PRAGMA_DISABLE_DEPRECATION_WARNINGS
-		return bHasBeenSimplified;
-		PRAGMA_ENABLE_DEPRECATION_WARNINGS
-	}
+	UE_DEPRECATED(5.5, "This functionality has been removed. Use FSkeletalMeshLODInfo::bHasBeenSimplified instead.")
+	bool GetHasBeenSimplified() const;
 
-	void SetHasBeenSimplified(bool bInHasBeenSimplified)
+	UE_DEPRECATED(5.5, "This functionality has been removed. Use FSkeletalMeshLODInfo::bHasBeenSimplified instead.")
+	static void SetHasBeenSimplified(bool bInHasBeenSimplified)
 	{
-		WaitUntilAsyncPropertyReleased(ESkeletalMeshAsyncProperties::HasBeenSimplified);
-		PRAGMA_DISABLE_DEPRECATION_WARNINGS
-		bHasBeenSimplified = bInHasBeenSimplified;
-		PRAGMA_ENABLE_DEPRECATION_WARNINGS
 	}
 
 	/** Whether or not the mesh has vertex colors */
diff --git a/Engine/Source/Runtime/Engine/Private/SkeletalMesh.cpp b/Engine/Source/Runtime/Engine/Private/SkeletalMesh.cpp
--- a/Engine/Source/Runtime/Engine/Private/SkeletalMesh.cpp
+++ b/Engine/Source/Runtime/Engine/Private/SkeletalMesh.cpp
@@ -5741,6 +5741,18 @@ TSoftObjectPtr<UObject> USkeletalMesh::GetDefaultAnimatingRig() const
 	PRAGMA_ENABLE_DEPRECATION_WARNINGS
 }
 
+bool USkeletalMesh::GetHasBeenSimplified() const
+{
+	for (int32 LODIndex = 0, LODCount = GetLODNum(); LODIndex < LODCount; LODIndex++)
+	{
+		if (GetLODInfo(LODIndex)->bHasBeenSimplified)
+		{
+			return true;
+		}
+	}
+	return false;
+}
+
 void USkeletalMesh::GetMappableNodeData(TArray<FName>& OutNames, TArray<FNodeItem>& OutNodeItems) const
 {
 	
