Moved DeadBlending AnimNode out of Experimental.

diff --git a/Engine/Source/Editor/AnimGraph/Public/AnimGraphNode_DeadBlending.h b/Engine/Source/Editor/AnimGraph/Public/AnimGraphNode_DeadBlending.h
--- a/Engine/Source/Editor/AnimGraph/Public/AnimGraphNode_DeadBlending.h
+++ b/Engine/Source/Editor/AnimGraph/Public/AnimGraphNode_DeadBlending.h
@@ -9,7 +9,7 @@
 #include "AnimGraphNode_DeadBlending.generated.h"
 
 
-UCLASS(Experimental, MinimalAPI)
+UCLASS(MinimalAPI)
 class UAnimGraphNode_DeadBlending : public UAnimGraphNode_Base
 {
 	GENERATED_BODY()
diff --git a/Engine/Source/Runtime/Engine/Classes/Animation/AnimNode_DeadBlending.h b/Engine/Source/Runtime/Engine/Classes/Animation/AnimNode_DeadBlending.h
--- a/Engine/Source/Runtime/Engine/Classes/Animation/AnimNode_DeadBlending.h
+++ b/Engine/Source/Runtime/Engine/Classes/Animation/AnimNode_DeadBlending.h
@@ -27,7 +27,7 @@
  * `ExtrapolationHalfLifeMax` parameters, which specify the approximate average, min, and max decay periods.
  * More specifically they specify the "half-life" - or how it takes for the velocities to be decayed by half.
  */
-USTRUCT(Experimental, BlueprintInternalUseOnly)
+USTRUCT(BlueprintInternalUseOnly)
 struct FAnimNode_DeadBlending : public FAnimNode_Base, public IBoneReferenceSkeletonProvider
 {
 	GENERATED_BODY()
