InvalidateExport component templates that are deemed unnecessary, this prevents the linker from recreating them when they are imported
[RN] Fixed underlying cause of "Content is missing from cook" error when cooking deep blueprint hierarchies with components

diff --git a/Engine/Source/Runtime/Engine/Private/InheritableComponentHandler.cpp b/Engine/Source/Runtime/Engine/Private/InheritableComponentHandler.cpp
--- a/Engine/Source/Runtime/Engine/Private/InheritableComponentHandler.cpp
+++ b/Engine/Source/Runtime/Engine/Private/InheritableComponentHandler.cpp
@@ -259,6 +259,13 @@ void UInheritableComponentHandler::ValidateTemplates()
 						Record.ComponentTemplate->ClearFlags(RF_Standalone);
 						Record.ComponentTemplate->RemoveFromRoot();
 						Record.ComponentTemplate->MarkAsGarbage();
+						// Rename won't invalidate the linker's export, and linker lifetime extends long beyond an actual loadpackage 
+						// invocation. Consequently, if the template object is garbage collected (as we hope it will be) it 
+						// could tragically be recreated by FLinkerLoad unless we invalidate the export. Zen loader has some 
+						// logic to avoid recreating the object, but it is buggy and we want to avoid object recreation when 
+						// not using Zen, anyway. We would not need to invalidate the export if: 1. Rename invalidated the 
+						// export or 2. FLinkerLoad's lifetime were reigned in.
+						FLinkerLoad::InvalidateExport(Record.ComponentTemplate);
 #endif // WITH_EDITOR
 						UnnecessaryComponents.AddUnique(Record.ComponentTemplate);
 					}
