[Backout] - CL33463496 - CIS Cook Warning
#rnx
Original CL Desc
-----------------------------------------------------------------
Cherrypick of 33033713. In addition to the below description, this CL also bumps MaterialTranslationDDCVersion in FFortniteReleaseBranchCustomObjectVersion to ensure that the material translation cache is invalidated and picks up on the changed code gen.

Fix for invalid ddx/ddy calls in generated material code when using the DDX and DDY nodes in a material that can be executed in a compute shader context other than Nanite shading (such as HeterogeneousVolumes or Volumetric Clouds).
Fixed by using existing DDX() and DDY() wrapper functions which are already set up to return 0 in all shader types except pixel shaders and nanite compute shading compute shaders.

#rnx

diff --git a/Engine/Shaders/Private/MaterialTemplate.ush b/Engine/Shaders/Private/MaterialTemplate.ush
--- a/Engine/Shaders/Private/MaterialTemplate.ush
+++ b/Engine/Shaders/Private/MaterialTemplate.ush
@@ -42,7 +42,7 @@
 
 // Update this GUID to force all material shaders to recompile
 // Merge conflicts on this line should be resolved by generating a new GUID
-#pragma message("UESHADERMETADATA_VERSION 5AE81BAA-6F06-4443-AEC8-98E87D189D36")
+#pragma message("UESHADERMETADATA_VERSION 838540DA-BB92-4973-80CB-D9E993EBA40D")
 
 #if USES_SPEEDTREE
 	#include "/Engine/Private/SpeedTreeCommon.ush"
diff --git a/Engine/Source/Runtime/Core/Private/UObject/FortniteReleaseBranchCustomObjectVersion.cpp b/Engine/Source/Runtime/Core/Private/UObject/FortniteReleaseBranchCustomObjectVersion.cpp
--- a/Engine/Source/Runtime/Core/Private/UObject/FortniteReleaseBranchCustomObjectVersion.cpp
+++ b/Engine/Source/Runtime/Core/Private/UObject/FortniteReleaseBranchCustomObjectVersion.cpp
@@ -8,6 +8,6 @@ TMap<FGuid, FGuid> FFortniteReleaseBranchCustomObjectVersion::GetSystemGuids()
 
 	SystemGuids.Add(DevGuids.Niagara_LatestScriptCompileVersion, FGuid("5F589BAA37394A51BF718ECB10D73C19"));
 	SystemGuids.Add(DevGuids.SkeletalMeshDerivedDataVersion, FGuid("5705956EC7134274A5A1BD3FC74DB3A9"));
-	SystemGuids.Add(DevGuids.MaterialTranslationDDCVersion, FGuid("CC1B3A336E144A0B9E78014E07775988"));
+	SystemGuids.Add(DevGuids.MaterialTranslationDDCVersion, FGuid("555395300FC34E35A49DB477FC8ADD85"));
 	return SystemGuids;
 }
diff --git a/Engine/Source/Runtime/Engine/Private/Materials/HLSLMaterialTranslator.cpp b/Engine/Source/Runtime/Engine/Private/Materials/HLSLMaterialTranslator.cpp
--- a/Engine/Source/Runtime/Engine/Private/Materials/HLSLMaterialTranslator.cpp
+++ b/Engine/Source/Runtime/Engine/Private/Materials/HLSLMaterialTranslator.cpp
@@ -11127,7 +11127,7 @@ int32 FHLSLMaterialTranslator::DisplacementMagnitude()
 	return Constant(Material->GetDisplacementScaling().Magnitude);
 }
 
-int32 FHLSLMaterialTranslator::Derivative(int32 A, EDervativeComponent Component)
+int32 FHLSLMaterialTranslator::Derivative(int32 A, const TCHAR* Component)
 {
 	if (A == INDEX_NONE)
 	{
@@ -11149,29 +11149,10 @@ int32 FHLSLMaterialTranslator::Derivative(int32 A, EDervativeComponent Component
 
 	const FDerivInfo ADerivInfo = GetDerivInfo(A);
 	const EMaterialValueType ResultType = MakeNonLWCType(ADerivInfo.Type);
-	const bool bIsLWCType = IsLWCType(ADerivInfo.Type);
-	const TCHAR* FunctionName = bIsLWCType ? TEXT("WSDd") : TEXT("DD");
-	const TCHAR* FunctionNamePostfix = bIsLWCType ? TEXT("Demote") : TEXT("");
-	
-	// For non-LWC types, we use the DDX() and DDY() functions defined in Common.ush, which correctly deal with shader types that do not support hardware derivatives.
-	const TCHAR* ComponentLowerCase = TEXT("");
-	const TCHAR* ComponentUpperCase = TEXT("");
-	switch (Component)
-	{
-	case EDervativeComponent::X:
-		ComponentLowerCase = TEXT("x");
-		ComponentUpperCase = TEXT("X");
-		break;
-	case EDervativeComponent::Y:
-		ComponentLowerCase = TEXT("y");
-		ComponentUpperCase = TEXT("Y");
-		break;
-	default:
-		checkNoEntry();
-		break;
-	}
+	const TCHAR* FunctionName = IsLWCType(ADerivInfo.Type) ? TEXT("WSDd") : TEXT("dd");
+	const TCHAR* FunctionNamePostfix = IsLWCType(ADerivInfo.Type) ? TEXT("Demote") : TEXT("");
     
-	if (bIsLWCType)
+	if (IsLWCType(ADerivInfo.Type))
 	{
 		AddLWCFuncUsage(ELWCFunctionKind::Other);
 	}
@@ -11188,16 +11169,16 @@ int32 FHLSLMaterialTranslator::Derivative(int32 A, EDervativeComponent Component
     	case MCT_Float4: CastType = TEXT("(float4)"); break;
     }
             
-	const FString FiniteCode = FString::Printf(TEXT("%s%s%s(%s%s)"), FunctionName, bIsLWCType ? ComponentLowerCase : ComponentUpperCase, FunctionNamePostfix, CastType, *GetParameterCode(A));
+	const FString FiniteCode = FString::Printf(TEXT("%s%s%s(%s%s)"), FunctionName, Component, FunctionNamePostfix, CastType, *GetParameterCode(A));
 
 	if (IsAnalyticDerivEnabled() && ADerivInfo.DerivativeStatus == EDerivativeStatus::Valid)
 	{
 		FString ADeriv = GetParameterCodeDeriv(A, CompiledPDV_Analytic);
 
 		FString AnalyticCode = DerivativeAutogen.ConstructDeriv(
-			FString::Printf(TEXT("%s.Dd%s"), *ADeriv, ComponentLowerCase),
-			FString::Printf(TEXT("DD%s(%s.Ddx)"), ComponentUpperCase, *ADeriv),
-			FString::Printf(TEXT("DD%s(%s.Ddy)"), ComponentUpperCase, *ADeriv),
+			FString::Printf(TEXT("%s.Dd%s"), *ADeriv, Component),
+			FString::Printf(TEXT("dd%s(%s.Ddx)"), Component, *ADeriv),
+			FString::Printf(TEXT("dd%s(%s.Ddy)"), Component, *ADeriv),
 			MakeNonLWCType(ADerivInfo.DerivType)
 		);
 
@@ -11211,12 +11192,12 @@ int32 FHLSLMaterialTranslator::Derivative(int32 A, EDervativeComponent Component
 
 int32 FHLSLMaterialTranslator::DDX( int32 A )
 {
-	return Derivative(A, EDervativeComponent::X);
+	return Derivative(A, TEXT("x"));
 }
 
 int32 FHLSLMaterialTranslator::DDY( int32 A )
 {
-	return Derivative(A, EDervativeComponent::Y);
+	return Derivative(A, TEXT("y"));
 }
 
 int32 FHLSLMaterialTranslator::AntialiasedTextureMask(int32 Tex, int32 UV, float Threshold, uint8 Channel)
diff --git a/Engine/Source/Runtime/Engine/Private/Materials/HLSLMaterialTranslator.h b/Engine/Source/Runtime/Engine/Private/Materials/HLSLMaterialTranslator.h
--- a/Engine/Source/Runtime/Engine/Private/Materials/HLSLMaterialTranslator.h
+++ b/Engine/Source/Runtime/Engine/Private/Materials/HLSLMaterialTranslator.h
@@ -1168,12 +1168,7 @@ protected:
 	virtual int32 DDX(int32 A) override;
 	virtual int32 DDY(int32 A) override;
 
-	enum class EDervativeComponent
-	{
-		X, 
-		Y,
-	};
-	int32 Derivative(int32 A, EDervativeComponent Component);
+	int32 Derivative(int32 A, const TCHAR* Component);
 
 	virtual int32 AntialiasedTextureMask(int32 Tex, int32 UV, float Threshold, uint8 Channel) override;
 	virtual int32 DepthOfFieldFunction(int32 Depth, int32 FunctionValueIndex) override;
