Fix handling of non power of 2 textures when building virtual texture.

VT building was not using PowerOfTwoMode build setting correctly to handle non-pow2 texture stretching/padding to pow2. This resulted in wrong mip bias calculation leading to a crash.

Code adds extra checks inside InitializeFromBuildSettings() to prevent VT build in various cases when odd VT texture sizes are used:
1) explicit resize to smaller size than VT tile size. Each block must be with at least VT tile size.
2) explicit resize without specifying both target width/height for non-pow2 texture sizes - this is to avoid situation when user enables explicit resize in UI, but has not yet entered width/height (which means texture will be used with non-pow2 size - not supported by VT).

diff --git a/Engine/Source/Runtime/Engine/Private/TextureDerivedDataTask.cpp b/Engine/Source/Runtime/Engine/Private/TextureDerivedDataTask.cpp
--- a/Engine/Source/Runtime/Engine/Private/TextureDerivedDataTask.cpp
+++ b/Engine/Source/Runtime/Engine/Private/TextureDerivedDataTask.cpp
@@ -190,7 +190,7 @@ void FTextureSourceData::InitAsPlaceholder()
 	bValid = true;
 }
 
-void FTextureSourceData::Init(UTexture& InTexture, TextureMipGenSettings InMipGenSettings, bool bInCubeMap, bool bInTextureArray, bool bInVolumeTexture, bool bAllowAsyncLoading)
+void FTextureSourceData::Init(UTexture& InTexture, TextureMipGenSettings InMipGenSettings, bool bInCubeMap, bool bInTextureArray, bool bInVolumeTexture, ETexturePowerOfTwoSetting::Type InPow2Setting, int32 InResizeDuringBuildX, int32 InResizeDuringBuildY, bool bAllowAsyncLoading)
 {
 	check( bValid == false ); // we set to true at the end, acts as our return value
 	
@@ -271,9 +271,16 @@ void FTextureSourceData::Init(UTexture& InTexture, TextureMipGenSettings InMipGe
 
 	for (FTextureSourceBlockData& Block : Blocks)
 	{
+		int32 AdjustedSizeX, AdjustedSizeY, AdjustedSizeZ;
+		UE::TextureBuildUtilities::GetPowerOfTwoTargetTextureSize(
+			Block.SizeX, Block.SizeY, 1,
+			false,
+			InPow2Setting, InResizeDuringBuildX, InResizeDuringBuildY,
+			AdjustedSizeX, AdjustedSizeY, AdjustedSizeZ);
+
 		// for the common case of NumBlocks == 1, BlockSizeX == Block.SizeX, MipBiasX/Y will both be zero
-		const int32 MipBiasX = FMath::CeilLogTwo(BlockSizeX / Block.SizeX);
-		const int32 MipBiasY = FMath::CeilLogTwo(BlockSizeY / Block.SizeY);
+		const int32 MipBiasX = FMath::CeilLogTwo(BlockSizeX / AdjustedSizeX);
+		const int32 MipBiasY = FMath::CeilLogTwo(BlockSizeY / AdjustedSizeY);
 		if (MipBiasX != MipBiasY)
 		{
 			// @todo Oodle: this is failing even if "pad to pow2 square" is set, can we allow it through in that case?
@@ -1656,10 +1663,19 @@ FTextureCacheDerivedDataWorker::FTextureCacheDerivedDataWorker(
 
 	// All of these settings are fixed across build settings and are derived directly from the texture.
 	// So we can just use layer 0 of whatever we have.
-	TextureData.Init(Texture, (TextureMipGenSettings)BuildSettingsPerLayerFetchOrBuild[0].MipGenSettings, BuildSettingsPerLayerFetchOrBuild[0].bCubemap, BuildSettingsPerLayerFetchOrBuild[0].bTextureArray, BuildSettingsPerLayerFetchOrBuild[0].bVolume, bAllowAsyncLoading);
+	const FTextureBuildSettings& BuildSettings = BuildSettingsPerLayerFetchOrBuild[0];
+	TextureData.Init(Texture,
+		(TextureMipGenSettings)BuildSettings.MipGenSettings,
+		BuildSettings.bCubemap,
+		BuildSettings.bTextureArray,
+		BuildSettings.bVolume,
+		(ETexturePowerOfTwoSetting::Type)BuildSettings.PowerOfTwoMode,
+		BuildSettings.ResizeDuringBuildX,
+		BuildSettings.ResizeDuringBuildY,
+		bAllowAsyncLoading);
 
 	bool bNeedsCompositeData = Texture.GetCompositeTexture() && Texture.CompositeTextureMode != CTM_Disabled && Texture.GetCompositeTexture()->Source.IsValid();
-	if (BuildSettingsPerLayerFetchOrBuild[0].bCPUAccessible)
+	if (BuildSettings.bCPUAccessible)
 	{
 		// CPU accessible textures don't run image processing and thus don't need the composite data.
 		bNeedsCompositeData = false;
@@ -1677,7 +1693,15 @@ FTextureCacheDerivedDataWorker::FTextureCacheDerivedDataWorker(
 
 		if ( bMatchingBlocks )
 		{
-			CompositeTextureData.Init(*Texture.GetCompositeTexture(), (TextureMipGenSettings)BuildSettingsPerLayerFetchOrBuild[0].MipGenSettings, BuildSettingsPerLayerFetchOrBuild[0].bCubemap, BuildSettingsPerLayerFetchOrBuild[0].bTextureArray, BuildSettingsPerLayerFetchOrBuild[0].bVolume, bAllowAsyncLoading);
+			CompositeTextureData.Init(*Texture.GetCompositeTexture(),
+				(TextureMipGenSettings)BuildSettings.MipGenSettings,
+				BuildSettings.bCubemap,
+				BuildSettings.bTextureArray,
+				BuildSettings.bVolume,
+				(ETexturePowerOfTwoSetting::Type)BuildSettings.PowerOfTwoMode,
+				BuildSettings.ResizeDuringBuildX,
+				BuildSettings.ResizeDuringBuildY,
+				bAllowAsyncLoading);
 		}
 	}
 }
diff --git a/Engine/Source/Runtime/Engine/Private/TextureDerivedDataTask.h b/Engine/Source/Runtime/Engine/Private/TextureDerivedDataTask.h
--- a/Engine/Source/Runtime/Engine/Private/TextureDerivedDataTask.h
+++ b/Engine/Source/Runtime/Engine/Private/TextureDerivedDataTask.h
@@ -92,7 +92,7 @@ struct FTextureSourceData
 	// Clear the current source data and make it into a placeholder texture.
 	void InitAsPlaceholder();
 
-	void Init(UTexture& InTexture, TextureMipGenSettings InMipGenSettings, bool bInCubeMap, bool bInTextureArray, bool bInVolumeTexture, bool bAllowAsyncLoading);
+	void Init(UTexture& InTexture, TextureMipGenSettings InMipGenSettings, bool bInCubeMap, bool bInTextureArray, bool bInVolumeTexture, ETexturePowerOfTwoSetting::Type InPow2Setting, int32 InResizeDuringBuildX, int32 InResizeDuringBuildY, bool bAllowAsyncLoading);
 	bool IsValid() const { return bValid; }
 
 	// ImageWrapperModule is not used
diff --git a/Engine/Source/Runtime/Engine/Private/VT/VirtualTextureDataBuilder.cpp b/Engine/Source/Runtime/Engine/Private/VT/VirtualTextureDataBuilder.cpp
--- a/Engine/Source/Runtime/Engine/Private/VT/VirtualTextureDataBuilder.cpp
+++ b/Engine/Source/Runtime/Engine/Private/VT/VirtualTextureDataBuilder.cpp
@@ -254,41 +254,44 @@ bool FVirtualTextureBuilderDerivedInfo::InitializeFromBuildSettings(const FTextu
 	const FTextureBuildSettings& BuildSettingsLayer0 = InSettingsPerLayer[0];
 	const int32 TileSize = BuildSettingsLayer0.VirtualTextureTileSize;
 
-	BlockSizeX = InSourceData.BlockSizeX;
-	BlockSizeY = InSourceData.BlockSizeY;
-
-	// BlockSize is potentially adjusted by rounding to power of 2
-	switch (BuildSettingsLayer0.PowerOfTwoMode)
+	if (BuildSettingsLayer0.PowerOfTwoMode == ETexturePowerOfTwoSetting::ResizeToSpecificResolution)
 	{
-	case ETexturePowerOfTwoSetting::None:
-		break;
-	case ETexturePowerOfTwoSetting::PadToPowerOfTwo:
-	case ETexturePowerOfTwoSetting::StretchToPowerOfTwo:
-		BlockSizeX = FMath::RoundUpToPowerOfTwo(BlockSizeX);
-		BlockSizeY = FMath::RoundUpToPowerOfTwo(BlockSizeY);
-		break;
-	case ETexturePowerOfTwoSetting::PadToSquarePowerOfTwo:
-	case ETexturePowerOfTwoSetting::StretchToSquarePowerOfTwo:
-		BlockSizeX = FMath::RoundUpToPowerOfTwo(BlockSizeX);
-		BlockSizeY = FMath::RoundUpToPowerOfTwo(BlockSizeY);
-		BlockSizeX = FMath::Max(BlockSizeX, BlockSizeY);
-		BlockSizeY = BlockSizeX;
-		break;
-	case ETexturePowerOfTwoSetting::ResizeToSpecificResolution:
-		if (BuildSettingsLayer0.ResizeDuringBuildX)
+		// do not allow to set target width or height smaller than VT tile size
+		if ((BuildSettingsLayer0.ResizeDuringBuildX && BuildSettingsLayer0.ResizeDuringBuildX < TileSize) ||
+			(BuildSettingsLayer0.ResizeDuringBuildY && BuildSettingsLayer0.ResizeDuringBuildY < TileSize))
 		{
-			BlockSizeX = BuildSettingsLayer0.ResizeDuringBuildX;
+			// will need to adjust miptail block calculations for this to work
+
+			UE_LOG(LogVirtualTexturing, Warning, TEXT("InitializeFromBuildSettings failed : Explicit resize to smaller than tile size (%d) not supported (%d x %d) [%s]"),
+				TileSize, BuildSettingsLayer0.ResizeDuringBuildX, BuildSettingsLayer0.ResizeDuringBuildY, *InSourceData.TextureFullName);
+
+			return false;
 		}
-		if (BuildSettingsLayer0.ResizeDuringBuildY)
+
+		for (const auto& SourceBlock : InSourceData.Blocks)
 		{
-			BlockSizeY = BuildSettingsLayer0.ResizeDuringBuildY;
+			// if any of block sizes is not power of two any of them is smaller than VT tile size
+			if (!FMath::IsPowerOfTwo(SourceBlock.SizeX) || !FMath::IsPowerOfTwo(SourceBlock.SizeY) || SourceBlock.SizeX < TileSize || SourceBlock.SizeY < TileSize)
+			{
+				// then both target resize width & height must be set
+				if (BuildSettingsLayer0.ResizeDuringBuildX == 0 || BuildSettingsLayer0.ResizeDuringBuildY == 0)
+				{
+					UE_LOG(LogVirtualTexturing, Warning, TEXT("InitializeFromBuildSettings failed : Both resized width and height (%d x %d) must be set if any block is smaller than tile size (%d) [%s]"),
+						BuildSettingsLayer0.ResizeDuringBuildX, BuildSettingsLayer0.ResizeDuringBuildY, TileSize, *InSourceData.TextureFullName);
+					return false;
+				}
+			}
 		}
-		break;
-	default:
-		checkNoEntry();
-		break;
 	}
 
+	int32 BlockSizeZ; // not needed here
+	UE::TextureBuildUtilities::GetPowerOfTwoTargetTextureSize(
+		InSourceData.BlockSizeX, InSourceData.BlockSizeY, 1,
+		false,
+		(ETexturePowerOfTwoSetting::Type)BuildSettingsLayer0.PowerOfTwoMode,
+		BuildSettingsLayer0.ResizeDuringBuildX, BuildSettingsLayer0.ResizeDuringBuildY,
+		BlockSizeX, BlockSizeY, BlockSizeZ);
+
 	check(InSettingsPerLayer[0].MaxTextureResolution >= (uint32)TileSize);
 
 	// Clamp BlockSizeX and BlockSizeY to MaxTextureResolution, but don't change aspect ratio
@@ -944,10 +947,19 @@ void FVirtualTextureDataBuilder::BuildLayerBlocks(FSlowTask& BuildTask, uint32 L
 			TBSettings.MipGenSettings = TMGS_SimpleAverage;
 		}
 
+		// in case image is smaller than tile size, we need PowerOfTwoMode adjusted size to correctly calculate MipBias & LocalBlockSizeScale in a while loop below
+		int32 AdjustedSizeX, AdjustedSizeY, AdjustedSizeZ;
+		UE::TextureBuildUtilities::GetPowerOfTwoTargetTextureSize(
+			SourceMips[0].SizeX, SourceMips[0].SizeY, 1,
+			false,
+			(ETexturePowerOfTwoSetting::Type)SettingsPerLayer[0].PowerOfTwoMode,
+			SettingsPerLayer[0].ResizeDuringBuildX, SettingsPerLayer[0].ResizeDuringBuildY,
+			AdjustedSizeX, AdjustedSizeY, AdjustedSizeZ);
+
 		// For multi-block images, we may have scaled the max block size to be tile-sized, but individual blocks may still be smaller than 1 tile
 		// These need to be scaled up as well (scaling up individual blocks has the effect of reducing the block's mip-bias)
 		int32 LocalBlockSizeScale = DerivedInfo.BlockSizeScale;
-		while (SourceMips[0].SizeX * LocalBlockSizeScale < TileSize || SourceMips[0].SizeY * LocalBlockSizeScale < TileSize)
+		while (AdjustedSizeX * LocalBlockSizeScale < TileSize || AdjustedSizeY * LocalBlockSizeScale < TileSize)
 		{
 			check(BlockData.MipBias > 0u);
 			--BlockData.MipBias;
@@ -1114,6 +1126,7 @@ void FVirtualTextureDataBuilder::BuildLayerBlocks(FSlowTask& BuildTask, uint32 L
 		TBSettings.BaseTextureFormatName = LayerData.FormatName; // VTs never have platform prefix
 		TBSettings.bSRGB = BuildSettingsForLayer.bSRGB;
 		TBSettings.bUseLegacyGamma = BuildSettingsForLayer.bUseLegacyGamma;
+		TBSettings.PowerOfTwoMode = ETexturePowerOfTwoSetting::None; // no resizing - that's for source blocks only, miptail block size is already set up to be a power of two
 
 		// Make sure the output of the texture builder is in the same gamma space as we expect it.
 		check(TBSettings.GetDestGammaSpace() == BuildSettingsForLayer.GetDestGammaSpace());
