If the shader compile worker crashes, the error arrays were getting out of sync

The function BuildErrorStringAndReport assumes the UniqueErrors and UniqueErrorPrefixes arrays are of the same size, so make sure they get updated together.

diff --git a/Engine/Source/Runtime/Engine/Private/ShaderCompiler/ShaderCompiler.cpp b/Engine/Source/Runtime/Engine/Private/ShaderCompiler/ShaderCompiler.cpp
--- a/Engine/Source/Runtime/Engine/Private/ShaderCompiler/ShaderCompiler.cpp
+++ b/Engine/Source/Runtime/Engine/Private/ShaderCompiler/ShaderCompiler.cpp
@@ -2812,6 +2812,8 @@ static void BuildErrorStringAndReport(const FShaderDiagnosticInfo& DiagInfo, FSt
 {
 	bool bReportedDebugInfo = false;
 
+	check(DiagInfo.UniqueErrors.Num() == DiagInfo.UniqueErrorPrefixes.Num());
+
 	for (int32 ErrorIndex = 0; ErrorIndex < DiagInfo.UniqueErrors.Num(); ErrorIndex++)
 	{
 		FString UniqueErrorString = DiagInfo.UniqueErrorPrefixes[ErrorIndex] + DiagInfo.UniqueErrors[ErrorIndex] + TEXT("\n");
diff --git a/Engine/Source/Runtime/RenderCore/Private/ShaderDiagnostics.cpp b/Engine/Source/Runtime/RenderCore/Private/ShaderDiagnostics.cpp
--- a/Engine/Source/Runtime/RenderCore/Private/ShaderDiagnostics.cpp
+++ b/Engine/Source/Runtime/RenderCore/Private/ShaderDiagnostics.cpp
@@ -186,6 +186,7 @@ void FShaderDiagnosticInfo::AddAndProcessErrorsForJob(FShaderCommonCompileJob& J
 			{
 				// Job hard crashed
 				UniqueErrors.Add(FString::Printf(TEXT("Internal Error!\n\t%s"), *GetSingleJobCompilationDump(&Job)));
+				UniqueErrorPrefixes.Add(TEXT(""));
 			}
 
 			// If we filter all error messages because they are interpreted as warnings, we have to assume all error messages are in fact errors and not warnings.
