Minimal Skinned/Skeletal Mesh Component threaded tick

Ensure all relevant tick work is skipped when ticking on worker threads (just perform LOD calculations)

diff --git a/Engine/Source/Runtime/Engine/Classes/Components/SkinnedMeshComponent.h b/Engine/Source/Runtime/Engine/Classes/Components/SkinnedMeshComponent.h
--- a/Engine/Source/Runtime/Engine/Classes/Components/SkinnedMeshComponent.h
+++ b/Engine/Source/Runtime/Engine/Classes/Components/SkinnedMeshComponent.h
@@ -40,6 +40,11 @@ namespace Nanite
 	struct FResources;
 }
 
+namespace UE::Anim
+{
+	struct FSkinnedMeshComponentExtensions;
+}
+
 DECLARE_DELEGATE_OneParam(FOnAnimUpdateRateParamsCreated, FAnimUpdateRateParameters*)
 DECLARE_MULTICAST_DELEGATE_ThreeParams(FOnTickPose, USkinnedMeshComponent* /*SkinnedMeshComponent*/, float /*DeltaTime*/, bool /*bNeedsValidRootMotion*/)
 
@@ -1988,6 +1993,7 @@ public:
 
 	friend class FRenderStateRecreator;
 	friend class FSkeletalMeshStreamOut;
+	friend struct UE::Anim::FSkinnedMeshComponentExtensions;
 
 	//
 	// Animation required bones
diff --git a/Engine/Source/Runtime/Engine/Private/Components/SkeletalMeshComponent.cpp b/Engine/Source/Runtime/Engine/Private/Components/SkeletalMeshComponent.cpp
--- a/Engine/Source/Runtime/Engine/Private/Components/SkeletalMeshComponent.cpp
+++ b/Engine/Source/Runtime/Engine/Private/Components/SkeletalMeshComponent.cpp
@@ -1696,62 +1696,68 @@ void USkeletalMeshComponent::TickComponent(float DeltaTime, enum ELevelTick Tick
 	}
 #endif
 
-	if (ClothingSimulation)
+	if(!PrimaryComponentTick.bRunOnAnyThread)
 	{
-		ClothingSimulation->UpdateWorldForces(this);
-	}
+		if (ClothingSimulation)
+		{
+			ClothingSimulation->UpdateWorldForces(this);
+		}
 
-	UpdateEndPhysicsTickRegisteredState();
-	UpdateClothTickRegisteredState();
+		UpdateEndPhysicsTickRegisteredState();
+		UpdateClothTickRegisteredState();
 
-	// If we are suspended, we will not simulate clothing, but as clothing is simulated in local space
-	// relative to a root bone we need to extract simulation positions as this bone could be animated.
-	if((!CVarEnableClothPhysics.GetValueOnGameThread() || bClothingSimulationSuspended) && ClothingSimulation)
-	{
-		CSV_SCOPED_TIMING_STAT(Animation, Cloth);
+		// If we are suspended, we will not simulate clothing, but as clothing is simulated in local space
+		// relative to a root bone we need to extract simulation positions as this bone could be animated.
+		if((!CVarEnableClothPhysics.GetValueOnGameThread() || bClothingSimulationSuspended) && ClothingSimulation)
+		{
+			CSV_SCOPED_TIMING_STAT(Animation, Cloth);
 
-		// First update the simulation context, since the simulation isn't ticking
-		// and it is still required to get the correct simulation data and bounds.
-		constexpr bool bIsInitialization = false;
-		ClothingSimulation->FillContext(this, DeltaTime, ClothingSimulationContext, bIsInitialization);
+			// First update the simulation context, since the simulation isn't ticking
+			// and it is still required to get the correct simulation data and bounds.
+			constexpr bool bIsInitialization = false;
+			ClothingSimulation->FillContext(this, DeltaTime, ClothingSimulationContext, bIsInitialization);
 
-		ClothingSimulation->GetSimulationData(CurrentSimulationData, this, Cast<USkeletalMeshComponent>(LeaderPoseComponent.Get()));
+			ClothingSimulation->GetSimulationData(CurrentSimulationData, this, Cast<USkeletalMeshComponent>(LeaderPoseComponent.Get()));
+		}
 	}
 
 	Super::TickComponent(DeltaTime, TickType, ThisTickFunction);
-	
-	PendingRadialForces.Reset();
 
-	// Update bOldForceRefPose
-	bOldForceRefPose = bForceRefpose;
-
-	/** Update the end group and tick priority */
-	const bool bDoLateEnd = CVarAnimationDelaysEndGroup.GetValueOnGameThread() > 0;
-	const bool bRequiresPhysics = EndPhysicsTickFunction.IsTickFunctionRegistered();
-	const ETickingGroup EndTickGroup = bDoLateEnd && !bRequiresPhysics ? TG_PostPhysics : TG_PrePhysics;
-	if (ThisTickFunction)
+	if(!PrimaryComponentTick.bRunOnAnyThread)
 	{
-		ThisTickFunction->EndTickGroup = EndTickGroup;
+		PendingRadialForces.Reset();
 
-		const bool bDoHiPri = CVarHiPriSkinnedMeshesTicks.GetValueOnGameThread() > 0;
-		check(PrimaryComponentTick.bHighPriority == bDoHiPri)
-	}
+		// Update bOldForceRefPose
+		bOldForceRefPose = bForceRefpose;
 
-	// If we are waiting for ParallelEval to complete or if we require Physics, 
-	// then FinalizeBoneTransform will be called and Anim events will be dispatched there. 
-	// We prefer doing it there so these events are triggered once we have a new updated pose.
-	// Note that it's possible that FinalizeBoneTransform has already been called here if not using ParallelUpdate.
-	// or it's possible that it hasn't been called at all if we're skipping Evaluate due to not being visible.
-	// ConditionallyDispatchQueuedAnimEvents will catch that and only Dispatch events if not already done.
-	if (!IsRunningParallelEvaluation() && !bRequiresPhysics)
-	{
-		/////////////////////////////////////////////////////////////////////////////
-		// Notify / Event Handling!
-		// This can do anything to our component (including destroy it) 
-		// Any code added after this point needs to take that into account
-		/////////////////////////////////////////////////////////////////////////////
+		/** Update the end group and tick priority */
+		const bool bDoLateEnd = CVarAnimationDelaysEndGroup.GetValueOnGameThread() > 0;
+		const bool bRequiresPhysics = EndPhysicsTickFunction.IsTickFunctionRegistered();
+		const ETickingGroup EndTickGroup = bDoLateEnd && !bRequiresPhysics ? TG_PostPhysics : TG_PrePhysics;
+		if (ThisTickFunction)
+		{
+			ThisTickFunction->EndTickGroup = EndTickGroup;
 
-		ConditionallyDispatchQueuedAnimEvents();
+			const bool bDoHiPri = CVarHiPriSkinnedMeshesTicks.GetValueOnGameThread() > 0;
+			check(PrimaryComponentTick.bHighPriority == bDoHiPri)
+		}
+
+		// If we are waiting for ParallelEval to complete or if we require Physics, 
+		// then FinalizeBoneTransform will be called and Anim events will be dispatched there. 
+		// We prefer doing it there so these events are triggered once we have a new updated pose.
+		// Note that it's possible that FinalizeBoneTransform has already been called here if not using ParallelUpdate.
+		// or it's possible that it hasn't been called at all if we're skipping Evaluate due to not being visible.
+		// ConditionallyDispatchQueuedAnimEvents will catch that and only Dispatch events if not already done.
+		if (!IsRunningParallelEvaluation() && !bRequiresPhysics)
+		{
+			/////////////////////////////////////////////////////////////////////////////
+			// Notify / Event Handling!
+			// This can do anything to our component (including destroy it) 
+			// Any code added after this point needs to take that into account
+			/////////////////////////////////////////////////////////////////////////////
+
+			ConditionallyDispatchQueuedAnimEvents();
+		}
 	}
 }
 
diff --git a/Engine/Source/Runtime/Engine/Private/Components/SkinnedMeshComponent.cpp b/Engine/Source/Runtime/Engine/Private/Components/SkinnedMeshComponent.cpp
--- a/Engine/Source/Runtime/Engine/Private/Components/SkinnedMeshComponent.cpp
+++ b/Engine/Source/Runtime/Engine/Private/Components/SkinnedMeshComponent.cpp
@@ -1462,8 +1462,12 @@ void USkinnedMeshComponent::TickComponent(float DeltaTime, enum ELevelTick TickT
 	SCOPED_NAMED_EVENT(USkinnedMeshComponent_TickComponent, FColor::Yellow);
 	SCOPE_CYCLE_COUNTER(STAT_SkinnedMeshCompTick);
 
-	// Tick ActorComponent first.
-	Super::TickComponent(DeltaTime, TickType, ThisTickFunction);
+	// We cant run the blueprint events (ReceiveTick, latent actions) on worker threads, so skip them if we are running as such
+	if(!PrimaryComponentTick.bRunOnAnyThread)
+	{
+		// Tick ActorComponent first.
+		Super::TickComponent(DeltaTime, TickType, ThisTickFunction);
+	}
 
 	// See if this mesh was rendered recently. This has to happen first because other data will rely on this
 	bRecentlyRendered = (GetLastRenderTime() > GetWorld()->TimeSeconds - 1.0f);
@@ -1472,41 +1476,45 @@ void USkinnedMeshComponent::TickComponent(float DeltaTime, enum ELevelTick TickT
 	// This must be done BEFORE animation Update and Evaluate (TickPose and RefreshBoneTransforms respectively)
 	const bool bLODHasChanged = UpdateLODStatus();
 
-	// Tick Pose first
-	if (ShouldTickPose())
+	// Skip the rest of the work if we are running on worker threads as none of it is safe
+	if(!PrimaryComponentTick.bRunOnAnyThread)
 	{
-		TickPose(DeltaTime, false);
-	}
+		// Tick Pose first
+		if (ShouldTickPose())
+		{
+			TickPose(DeltaTime, false);
+		}
 
-	// If we have been recently rendered, and bForceRefPose has been on for at least a frame, or the LOD changed, update bone matrices.
-	if( ShouldUpdateTransform(bLODHasChanged) )
-	{
-		// Do not update bones if we are taking bone transforms from another SkelMeshComp
-		if( LeaderPoseComponent.IsValid() )
+		// If we have been recently rendered, and bForceRefPose has been on for at least a frame, or the LOD changed, update bone matrices.
+		if( ShouldUpdateTransform(bLODHasChanged) )
 		{
-			UpdateFollowerComponent();
+			// Do not update bones if we are taking bone transforms from another SkelMeshComp
+			if( LeaderPoseComponent.IsValid() )
+			{
+				UpdateFollowerComponent();
+			}
+			else 
+			{
+				RefreshBoneTransforms(ThisTickFunction);
+			}
 		}
-		else 
+		else if(VisibilityBasedAnimTickOption == EVisibilityBasedAnimTickOption::AlwaysTickPose)
 		{
-			RefreshBoneTransforms(ThisTickFunction);
+			// We are not refreshing bone transforms, but we do want to tick pose. We may need to kick off a parallel task
+			DispatchParallelTickPose(ThisTickFunction);
 		}
-	}
-	else if(VisibilityBasedAnimTickOption == EVisibilityBasedAnimTickOption::AlwaysTickPose)
-	{
-		// We are not refreshing bone transforms, but we do want to tick pose. We may need to kick off a parallel task
-		DispatchParallelTickPose(ThisTickFunction);
-	}
 #if WITH_EDITOR
-	else 
-	{
-		// only do this for level viewport actors
-		UWorld* World = GetWorld();
-		if (World && World->WorldType == EWorldType::Editor)
+		else 
 		{
-			RefreshMorphTargets();
+			// only do this for level viewport actors
+			UWorld* World = GetWorld();
+			if (World && World->WorldType == EWorldType::Editor)
+			{
+				RefreshMorphTargets();
+			}
 		}
-	}
 #endif // WITH_EDITOR
+	}
 }
 
 UObject const* USkinnedMeshComponent::AdditionalStatObject() const
