Fix missing transitions on the skin cache tangent accumulation buffers.

FGPUSkinCache::MakeBufferTransitions was skipping transitions if a resource was already in the target state, but that's incorrect for UAV -> UAV. These transitions are necessary in order to serialize dispatches which write to the same UAV. Also, the tangent accumulation buffers can appear multiple times in the array passed to MakeBufferTransitions, so we need to make sure we issue a single transition per resource (which requires sorting the array).

#rnx

diff --git a/Engine/Source/Runtime/Engine/Private/GPUSkinCache.cpp b/Engine/Source/Runtime/Engine/Private/GPUSkinCache.cpp
--- a/Engine/Source/Runtime/Engine/Private/GPUSkinCache.cpp
+++ b/Engine/Source/Runtime/Engine/Private/GPUSkinCache.cpp
@@ -25,6 +25,7 @@ GPUSkinCache.cpp: Performs skinning on a compute shader into a buffer to avoid v
 #include "Stats/StatsTrace.h"
 #include "UObject/UObjectIterator.h"
 #include "ComponentRecreateRenderStateContext.h"
+#include "Algo/Sort.h"
 
 DEFINE_STAT(STAT_GPUSkinCache_TotalNumChunks);
 DEFINE_STAT(STAT_GPUSkinCache_TotalNumVertices);
@@ -1259,11 +1260,23 @@ void FGPUSkinCache::MakeBufferTransitions(FRHICommandList& RHICmdList, TArray<FS
 {
 	if (Buffers.Num() > 0)
 	{
+		// The tangent accumulation buffers are shared between sections so they can end up in the list more than once. We
+		// need to make sure we don't issue multiple transitions for the same resource in a single call, the RHIs can't deal with that.
+		Algo::Sort(Buffers);
+
 		TArray<FRHITransitionInfo, SceneRenderingAllocator> UAVs;
 		UAVs.Reserve(Buffers.Num());
+		
+		FSkinCacheRWBuffer* LastBuffer = nullptr;
 		for (FSkinCacheRWBuffer* Buffer : Buffers)
 		{
-			if (Buffer->AccessState != ToState)
+			if (Buffer == LastBuffer)
+			{
+				continue;
+			}
+
+			LastBuffer = Buffer;
+			if ( EnumHasAnyFlags(ToState, ERHIAccess::UAVMask) || Buffer->AccessState != ToState)
 			{
 				UAVs.Add(Buffer->UpdateAccessState(ToState));
 			}
