Add padding to ShowFlags to fix a bug when building them for shipping

diff --git a/Engine/Source/Runtime/Engine/Public/ShowFlags.h b/Engine/Source/Runtime/Engine/Public/ShowFlags.h
--- a/Engine/Source/Runtime/Engine/Public/ShowFlags.h
+++ b/Engine/Source/Runtime/Engine/Public/ShowFlags.h
@@ -58,14 +58,18 @@ struct FEngineShowFlags
 	// A show flag is either an uint32:1 or static const bool (if optimized out according to UE_BUILD_OPTIMIZED_SHOWFLAGS)
 
 #if PLATFORM_USE_SHOWFLAGS_ALWAYS_BITFIELD
-	#define SHOWFLAG_ALWAYS_ACCESSIBLE(a,...) uint32 a : 1; void Set##a(bool bVal){ a = bVal?1:0;}
+	#define SHOWFLAG_ALWAYS_ACCESSIBLE(a,...) uint32 a : 1; void Set##a(bool bVal){ a = bVal ? 1:0; }
 #else
 	// broken bit field compilers will render the background black
-	#define SHOWFLAG_ALWAYS_ACCESSIBLE(a,...) bool a; void Set##a(bool bVal){ a = bVal;}
+	#define SHOWFLAG_ALWAYS_ACCESSIBLE(a,...) bool a; void Set##a(bool bVal){ a = bVal; }
 #endif
 
 	#if UE_BUILD_OPTIMIZED_SHOWFLAGS 
-		#define SHOWFLAG_FIXED_IN_SHIPPING(v,a,...) static const bool a = v; void Set##a(bool bVal){}
+		#if PLATFORM_USE_SHOWFLAGS_ALWAYS_BITFIELD
+			#define SHOWFLAG_FIXED_IN_SHIPPING(v,a,...)  uint32 _padding_##a : 1; static const bool a = v; void Set##a(bool bVal){}
+		#else
+			#define SHOWFLAG_FIXED_IN_SHIPPING(v,a,...)  bool _padding_##a; static const bool a = v; void Set##a(bool bVal){} 
+		#endif 
 	#else
 		#define SHOWFLAG_FIXED_IN_SHIPPING(v,a,b,c) SHOWFLAG_ALWAYS_ACCESSIBLE(a,b,c)
 	#endif
