Modified FRayTracingGeometry::Evict() to call RequestUpdateCachedRenderState(...)
- less error prone than having to manually call it after evicting geometries.

diff --git a/Engine/Source/Runtime/Engine/Private/Rendering/RayTracingGeometryManager.cpp b/Engine/Source/Runtime/Engine/Private/Rendering/RayTracingGeometryManager.cpp
--- a/Engine/Source/Runtime/Engine/Private/Rendering/RayTracingGeometryManager.cpp
+++ b/Engine/Source/Runtime/Engine/Private/Rendering/RayTracingGeometryManager.cpp
@@ -316,11 +316,6 @@ void FRayTracingGeometryManager::Tick(FRHICommandList& RHICmdList)
 			if (Geometry->GetRHI() != nullptr)
 			{
 				Geometry->Evict();
-
-				if (Geometry->GroupHandle != INDEX_NONE)
-				{
-					RequestUpdateCachedRenderState(Geometry->GroupHandle);
-				}
 			}
 		}
 	}
diff --git a/Engine/Source/Runtime/Engine/Private/StaticMesh.cpp b/Engine/Source/Runtime/Engine/Private/StaticMesh.cpp
--- a/Engine/Source/Runtime/Engine/Private/StaticMesh.cpp
+++ b/Engine/Source/Runtime/Engine/Private/StaticMesh.cpp
@@ -2047,7 +2047,7 @@ void UStaticMesh::RequestUpdateCachedRenderState() const
 	if (IsRayTracingAllowed())
 	{
 		// TODO: this should only be necessary when a BLAS build was not requested (ie: non-compressed offline BLAS)
-		((FRayTracingGeometryManager*)GRayTracingGeometryManager)->RequestUpdateCachedRenderState(GetRenderData()->RayTracingGeometryGroupHandle);
+		GRayTracingGeometryManager->RequestUpdateCachedRenderState(GetRenderData()->RayTracingGeometryGroupHandle);
 	}
 #endif
 
diff --git a/Engine/Source/Runtime/Engine/Public/Rendering/RayTracingGeometryManager.h b/Engine/Source/Runtime/Engine/Public/Rendering/RayTracingGeometryManager.h
--- a/Engine/Source/Runtime/Engine/Public/Rendering/RayTracingGeometryManager.h
+++ b/Engine/Source/Runtime/Engine/Public/Rendering/RayTracingGeometryManager.h
@@ -36,7 +36,7 @@ public:
 	ENGINE_API void RegisterProxyWithCachedRayTracingState(FPrimitiveSceneProxy* Proxy, RayTracing::GeometryGroupHandle InRayTracingGeometryGroupHandle);
 	ENGINE_API void UnregisterProxyWithCachedRayTracingState(FPrimitiveSceneProxy* Proxy, RayTracing::GeometryGroupHandle InRayTracingGeometryGroupHandle);
 
-	void RequestUpdateCachedRenderState(RayTracing::GeometryGroupHandle InRayTracingGeometryGroupHandle);
+	ENGINE_API virtual void RequestUpdateCachedRenderState(RayTracing::GeometryGroupHandle InRayTracingGeometryGroupHandle) override;
 
 	ENGINE_API void AddReferencedGeometry(const FRayTracingGeometry* Geometry);
 	ENGINE_API void AddReferencedGeometryGroups(const TSet<RayTracing::GeometryGroupHandle>& GeometryGroups);
diff --git a/Engine/Source/Runtime/RenderCore/Private/RayTracingGeometry.cpp b/Engine/Source/Runtime/RenderCore/Private/RayTracingGeometry.cpp
--- a/Engine/Source/Runtime/RenderCore/Private/RayTracingGeometry.cpp
+++ b/Engine/Source/Runtime/RenderCore/Private/RayTracingGeometry.cpp
@@ -130,6 +130,11 @@ void FRayTracingGeometry::Evict()
 	RemoveBuildRequest();
 	RayTracingGeometryRHI.SafeRelease();
 	EnumAddFlags(GeometryState, EGeometryStateFlags::Evicted);
+		
+	if (GroupHandle != INDEX_NONE)
+	{
+		GRayTracingGeometryManager->RequestUpdateCachedRenderState(GroupHandle);
+	}
 }
 
 void FRayTracingGeometry::CreateRayTracingGeometry(FRHICommandListBase& RHICmdList, ERTAccelerationStructureBuildPriority InBuildPriority)
diff --git a/Engine/Source/Runtime/RenderCore/Public/RayTracingGeometryManagerInterface.h b/Engine/Source/Runtime/RenderCore/Public/RayTracingGeometryManagerInterface.h
--- a/Engine/Source/Runtime/RenderCore/Public/RayTracingGeometryManagerInterface.h
+++ b/Engine/Source/Runtime/RenderCore/Public/RayTracingGeometryManagerInterface.h
@@ -51,6 +51,8 @@ public:
 	RENDERCORE_API virtual RayTracing::GeometryGroupHandle RegisterRayTracingGeometryGroup(uint32 NumLODs) = 0;
 	RENDERCORE_API virtual void ReleaseRayTracingGeometryGroup(RayTracing::GeometryGroupHandle Handle) = 0;
 
+	RENDERCORE_API virtual void RequestUpdateCachedRenderState(RayTracing::GeometryGroupHandle InRayTracingGeometryGroupHandle) = 0;
+
 	RENDERCORE_API virtual void PreRender() = 0;
 	RENDERCORE_API virtual void Tick(FRHICommandList& RHICmdList) = 0;
 };
