Skeletal Mesh: Only recover LODModel into MeshDescription after all fixups are done.

#rnx

diff --git a/Engine/Source/Runtime/Engine/Classes/Engine/SkeletalMesh.h b/Engine/Source/Runtime/Engine/Classes/Engine/SkeletalMesh.h
--- a/Engine/Source/Runtime/Engine/Classes/Engine/SkeletalMesh.h
+++ b/Engine/Source/Runtime/Engine/Classes/Engine/SkeletalMesh.h
@@ -2383,13 +2383,13 @@ private:
 
 	//When loading a legacy asset (saved before the skeletalmesh build refactor), we need to create the user sections data.
 	//This function should be call only in the PostLoad
-	ENGINE_API void CreateUserSectionsDataForLegacyAssets();
+	void CreateUserSectionsDataForLegacyAssets();
 
 
 	/*
 	 * This function will enforce the user section data is coherent with the sections.
 	 */
-	ENGINE_API void PostLoadValidateUserSectionData();
+	void PostLoadValidateUserSectionData();
 
 	/*
 	 * This function will ensure we have valid tangent in all LODs. If we found an invalid tangent axis, we will try to set it with the cross product of the two other axis.
@@ -2397,8 +2397,14 @@ private:
 	 * It will validate tangents only for assets that do not have source build data. (This means assets imported before the build refactor that was done in UE 4.24.)
 	 * @note - If it finds a bad normal, it will LOG a warning to let the user know they have to re-import their mesh.
 	 */
-	ENGINE_API void PostLoadVerifyAndFixBadTangent();
+	void PostLoadVerifyAndFixBadTangent();
 
+
+	/** After all loading is done and fixups completed, attempt to recover any LOD models into source model slots that
+	 *  are still empty.
+	 */
+	void PostLoadRecoverConvertLODModelsToMeshDescription();
+		
 public:
 	/*
 	 * This function will enforce valid material index in the sections and the LODMaterialMap of all LOD.
diff --git a/Engine/Source/Runtime/Engine/Private/SkeletalMesh.cpp b/Engine/Source/Runtime/Engine/Private/SkeletalMesh.cpp
--- a/Engine/Source/Runtime/Engine/Private/SkeletalMesh.cpp
+++ b/Engine/Source/Runtime/Engine/Private/SkeletalMesh.cpp
@@ -3100,6 +3100,91 @@ void USkeletalMesh::PostLoadVerifyAndFixBadTangent()
 	}
 }
 
+void USkeletalMesh::PostLoadRecoverConvertLODModelsToMeshDescription()
+{
+	// If we didn't get any meshes from the bulk data, then try to recover them from the LODModel listings.
+	for (int32 LODIndex = 0; LODIndex < GetImportedModel()->LODModels.Num(); ++LODIndex)
+	{
+		if (HasMeshDescription(LODIndex))
+		{
+			continue;
+		}
+
+		// If the mesh was not pulled out of the reduction data, we need to reset the LOD settings
+		// so that the mesh doesn't get reduced again if it gets regenerated.
+		FSkeletalMeshLODInfo* MeshLODInfo = GetLODInfo(LODIndex); 
+		const bool bReductionActive = IsReductionActive(LODIndex);
+		const bool bInlineReduction = (MeshLODInfo->ReductionSettings.BaseLOD == LODIndex);
+		if (!bReductionActive || bInlineReduction)
+		{
+			const FSkeletalMeshLODModel& LODModel = GetImportedModel()->LODModels[LODIndex];
+			FMeshDescription MeshDescription;
+			LODModel.GetMeshDescription(this, LODIndex, MeshDescription);
+			CreateMeshDescription(LODIndex, MoveTemp(MeshDescription));
+			CommitMeshDescription(LODIndex);
+			
+			// Ensure normals aren't automatically computed when we rebuild.
+			FSkeletalMeshBuildSettings& BuildSettings = MeshLODInfo->BuildSettings;
+			BuildSettings.bRecomputeNormals = false;
+
+			// Reset the reduction settings so that we don't re-reduce the mesh and possibly lose morph targets
+			// in the process.
+			FSkeletalMeshOptimizationSettings& ReductionSettings = MeshLODInfo->ReductionSettings;
+		
+			//Remove the reduction settings
+			ReductionSettings.NumOfTrianglesPercentage = 1.0f;
+			ReductionSettings.NumOfVertPercentage = 1.0f;
+			ReductionSettings.MaxNumOfTrianglesPercentage = MAX_uint32;
+			ReductionSettings.MaxNumOfVertsPercentage = MAX_uint32;
+			ReductionSettings.TerminationCriterion = SMTC_NumOfTriangles;
+			MeshLODInfo->bHasBeenSimplified = false;
+		}
+		else if (MeshLODInfo->LODMaterialMap.IsEmpty())
+		{
+			// Generated LODs (not inline) do not need imported data. We do need a material map though,
+			// because in many cases the map was not created when a section material got overridden, so reconstruct one if it isn't available.
+			const FSkeletalMeshLODModel& BaseLODModel = GetImportedModel()->LODModels[MeshLODInfo->ReductionSettings.BaseLOD];
+			const FSkeletalMeshLODModel& LODModel = GetImportedModel()->LODModels[LODIndex];
+			TArray<int32> MaterialMap;
+			MaterialMap.Init(INDEX_NONE, LODModel.Sections.Num());
+
+			if (BaseLODModel.Sections.Num() == LODModel.Sections.Num())
+			{
+				for (int32 SectionIndex = 0; SectionIndex < LODModel.Sections.Num(); SectionIndex++)
+				{
+					const int32 MaterialIndex = LODModel.Sections[SectionIndex].MaterialIndex;
+					if (BaseLODModel.Sections[SectionIndex].MaterialIndex != MaterialIndex)
+					{
+						MaterialMap[SectionIndex] = MaterialIndex;
+					}
+				}
+			}
+			else
+			{
+				for (int32 SectionIndex = 0; SectionIndex < LODModel.Sections.Num(); SectionIndex++)
+				{
+					const int32 BaseSectionIndex = LODModel.Sections[SectionIndex].OriginalDataSectionIndex;
+					
+					if (BaseLODModel.Sections.IsValidIndex(BaseSectionIndex))
+					{
+						const int32 MaterialIndex = LODModel.Sections[SectionIndex].MaterialIndex;
+
+						if (BaseLODModel.Sections[BaseSectionIndex].MaterialIndex != MaterialIndex)
+						{
+							MaterialMap[SectionIndex] = MaterialIndex;
+						}
+					}
+				}
+			}
+			if (Algo::AnyOf(MaterialMap, [](int32 Item) { return Item != INDEX_NONE; }))
+			{
+				MeshLODInfo->LODMaterialMap = MoveTemp(MaterialMap);
+			}
+		}
+	}
+}
+
+
 #endif // WITH_EDITOR
 
 bool USkeletalMesh::IsPostLoadThreadSafe() const
@@ -3314,87 +3399,6 @@ void USkeletalMesh::BeginPostLoadInternal(FSkinnedAssetPostLoadContext& Context)
 			}
 		}
 		
-		// If we didn't get any meshes from the bulk data, then try to recover them from the LODModel listings.
-		for (int32 LODIndex = 0; LODIndex < GetImportedModel()->LODModels.Num(); ++LODIndex)
-		{
-			if (HasMeshDescription(LODIndex))
-			{
-				continue;
-			}
-
-			// If the mesh was not pulled out of the reduction data, we need to reset the LOD settings
-			// so that the mesh doesn't get reduced again if it gets regenerated.
-			FSkeletalMeshLODInfo* MeshLODInfo = GetLODInfo(LODIndex); 
-			const bool bReductionActive = IsReductionActive(LODIndex);
-			const bool bInlineReduction = (MeshLODInfo->ReductionSettings.BaseLOD == LODIndex);
-			if (!bReductionActive || bInlineReduction)
-			{
-				const FSkeletalMeshLODModel& LODModel = GetImportedModel()->LODModels[LODIndex];
-				FMeshDescription MeshDescription;
-				LODModel.GetMeshDescription(this, LODIndex, MeshDescription);
-				CreateMeshDescription(LODIndex, MoveTemp(MeshDescription));
-				CommitMeshDescription(LODIndex);
-				
-				// Ensure normals aren't automatically computed when we rebuild.
-				FSkeletalMeshBuildSettings& BuildSettings = MeshLODInfo->BuildSettings;
-				BuildSettings.bRecomputeNormals = false;
-
-				// Reset the reduction settings so that we don't re-reduce the mesh and possibly lose morph targets
-				// in the process.
-				FSkeletalMeshOptimizationSettings& ReductionSettings = MeshLODInfo->ReductionSettings;
-			
-				//Remove the reduction settings
-				ReductionSettings.NumOfTrianglesPercentage = 1.0f;
-				ReductionSettings.NumOfVertPercentage = 1.0f;
-				ReductionSettings.MaxNumOfTrianglesPercentage = MAX_uint32;
-				ReductionSettings.MaxNumOfVertsPercentage = MAX_uint32;
-				ReductionSettings.TerminationCriterion = SMTC_NumOfTriangles;
-				MeshLODInfo->bHasBeenSimplified = false;
-			}
-			else if (MeshLODInfo->LODMaterialMap.IsEmpty())
-			{
-				// Generated LODs (not inline) do not need imported data. We do need a material map though,
-				// because in many cases the map was not created when a section material got overridden, so reconstruct one if it isn't available.
-				const FSkeletalMeshLODModel& BaseLODModel = GetImportedModel()->LODModels[MeshLODInfo->ReductionSettings.BaseLOD];
-				const FSkeletalMeshLODModel& LODModel = GetImportedModel()->LODModels[LODIndex];
-				TArray<int32> MaterialMap;
-				MaterialMap.Init(INDEX_NONE, LODModel.Sections.Num());
-
-				if (BaseLODModel.Sections.Num() == LODModel.Sections.Num())
-				{
-					for (int32 SectionIndex = 0; SectionIndex < LODModel.Sections.Num(); SectionIndex++)
-					{
-						const int32 MaterialIndex = LODModel.Sections[SectionIndex].MaterialIndex;
-						if (BaseLODModel.Sections[SectionIndex].MaterialIndex != MaterialIndex)
-						{
-							MaterialMap[SectionIndex] = MaterialIndex;
-						}
-					}
-				}
-				else
-				{
-					for (int32 SectionIndex = 0; SectionIndex < LODModel.Sections.Num(); SectionIndex++)
-					{
-						const int32 BaseSectionIndex = LODModel.Sections[SectionIndex].OriginalDataSectionIndex;
-						
-						if (BaseLODModel.Sections.IsValidIndex(BaseSectionIndex))
-						{
-							const int32 MaterialIndex = LODModel.Sections[SectionIndex].MaterialIndex;
-
-							if (BaseLODModel.Sections[BaseSectionIndex].MaterialIndex != MaterialIndex)
-							{
-								MaterialMap[SectionIndex] = MaterialIndex;
-							}
-						}
-					}
-				}
-				if (Algo::AnyOf(MaterialMap, [](int32 Item) { return Item != INDEX_NONE; }))
-				{
-					MeshLODInfo->LODMaterialMap = MoveTemp(MaterialMap);
-				}
-			}
-		}
-
 		if (GetLinkerCustomVersion(FEditorObjectVersion::GUID) < FEditorObjectVersion::SkeletalMeshBuildRefactor)
 		{
 			CreateUserSectionsDataForLegacyAssets();
@@ -3448,6 +3452,8 @@ void USkeletalMesh::ExecutePostLoadInternal(FSkinnedAssetPostLoadContext& Contex
 			MeshUtilities->FixupMaterialSlotNames(this);
 		}
 
+		PostLoadRecoverConvertLODModelsToMeshDescription();		
+
 		if (GetResourceForRendering() == nullptr)
 		{
 			CacheDerivedData(&Context);
