Improvements to FRayTracingGeometryManager:
- New "Ray Tracing Geometry" stat group
- Moved existing stats "Pending Builds" and "Pending Build Primitives" to the new group (previously in "Scene Rendering")
- Added new stats to track number of registered ray tracing geometries and groups.

diff --git a/Engine/Source/Runtime/Engine/Private/Rendering/RayTracingGeometryManager.cpp b/Engine/Source/Runtime/Engine/Private/Rendering/RayTracingGeometryManager.cpp
--- a/Engine/Source/Runtime/Engine/Private/Rendering/RayTracingGeometryManager.cpp
+++ b/Engine/Source/Runtime/Engine/Private/Rendering/RayTracingGeometryManager.cpp
@@ -75,8 +75,13 @@ static FAutoConsoleVariableRef CVarRayTracingPendingBuildPriorityBoostPerFrame(
 	ECVF_RenderThreadSafe
 );
 
-DECLARE_DWORD_ACCUMULATOR_STAT(TEXT("Ray tracing pending builds"), STAT_RayTracingPendingBuilds, STATGROUP_SceneRendering);
-DECLARE_DWORD_ACCUMULATOR_STAT(TEXT("Ray tracing pending build primitives"), STAT_RayTracingPendingBuildPrimitives, STATGROUP_SceneRendering);
+DECLARE_STATS_GROUP(TEXT("Ray Tracing Geometry"), STATGROUP_RayTracingGeometry, STATCAT_Advanced);
+
+DECLARE_DWORD_ACCUMULATOR_STAT(TEXT("Geometry Count"), STAT_RayTracingGeometryCount, STATGROUP_RayTracingGeometry);
+DECLARE_DWORD_ACCUMULATOR_STAT(TEXT("Geometry Group Count"), STAT_RayTracingGeometryGroupCount, STATGROUP_RayTracingGeometry);
+
+DECLARE_DWORD_ACCUMULATOR_STAT(TEXT("Pending Builds"), STAT_RayTracingPendingBuilds, STATGROUP_RayTracingGeometry);
+DECLARE_DWORD_ACCUMULATOR_STAT(TEXT("Pending Build Primitives"), STAT_RayTracingPendingBuildPrimitives, STATGROUP_RayTracingGeometry);
 
 FRayTracingGeometryManager::~FRayTracingGeometryManager()
 {
@@ -138,6 +143,9 @@ RayTracing::GeometryGroupHandle FRayTracingGeometryManager::RegisterRayTracingGe
 	Group.Geometries.AddDefaulted(NumLODs);
 
 	RayTracing::GeometryGroupHandle Handle = RegisteredGroups.Add(MoveTemp(Group));
+
+	INC_DWORD_STAT(STAT_RayTracingGeometryGroupCount);
+
 	return Handle;
 }
 
@@ -157,6 +165,8 @@ void FRayTracingGeometryManager::ReleaseRayTracingGeometryGroup(RayTracing::Geom
 		// set flag on group so that it is released once the last primitive is unregistered
 		RegisteredGroups[Handle].bPendingRelease = true;
 	}
+
+	DEC_DWORD_STAT(STAT_RayTracingGeometryGroupCount);
 }
 
 FRayTracingGeometryManager::RayTracingGeometryHandle FRayTracingGeometryManager::RegisterRayTracingGeometry(FRayTracingGeometry* InGeometry)
@@ -179,6 +189,8 @@ FRayTracingGeometryManager::RayTracingGeometryHandle FRayTracingGeometryManager:
 
 			Group.Geometries[InGeometry->LODIndex] = InGeometry;
 		}
+		
+		INC_DWORD_STAT(STAT_RayTracingGeometryCount);
 
 		return Handle;
 	}
@@ -193,6 +205,8 @@ void FRayTracingGeometryManager::ReleaseRayTracingGeometryHandle(RayTracingGeome
 		FScopeLock ScopeLock(&RequestCS);
 		RegisteredGeometries.RemoveAt(Handle);
 		ReferencedGeometryHandles.Remove(Handle);
+
+		DEC_DWORD_STAT(STAT_RayTracingGeometryCount);
 	}	
 }
 
