Convert StaticMesh build checkf assertion on invalid mesh description to error log output to gracefully fail during cooks.

diff --git a/Engine/Source/Runtime/Engine/Private/StaticMesh.cpp b/Engine/Source/Runtime/Engine/Private/StaticMesh.cpp
--- a/Engine/Source/Runtime/Engine/Private/StaticMesh.cpp
+++ b/Engine/Source/Runtime/Engine/Private/StaticMesh.cpp
@@ -3113,7 +3113,11 @@ void FStaticMeshRenderData::Cache(const ITargetPlatform* TargetPlatform, UStatic
 			FStaticMeshStatusMessageContext StatusContext( FText::Format( NSLOCTEXT("Engine", "BuildingStaticMeshStatus", "Building static mesh {StaticMeshName} (Required Memory Estimate: {EstimatedMemory} MB)..."), Args ) );
 
 			checkf(!Owner->HasAnyFlags(RF_NeedLoad), TEXT("StaticMesh %s being PostLoaded before having been serialized - this suggests an async loading problem."), *GetPathNameSafe(Owner));
-			checkf(Owner->IsMeshDescriptionValid(0), TEXT("Bad MeshDescription on %s"), *GetPathNameSafe(Owner));
+			if (!Owner->IsMeshDescriptionValid(0))
+			{
+				UE_LOG(LogStaticMesh, Error, TEXT("Bad MeshDescription on %s"), *GetPathNameSafe(Owner));
+				return;
+			}
 
 			if (Owner->bDoFastBuild)
 			{
@@ -3122,8 +3126,14 @@ void FStaticMeshRenderData::Cache(const ITargetPlatform* TargetPlatform, UStatic
 				AllocateLODResources(NumSourceModels);
 				for (int32 LodIndex = 0; LodIndex < NumSourceModels; LodIndex++)
 				{
-					checkf(Owner->IsMeshDescriptionValid(LodIndex), TEXT("Bad MeshDescription at lod index %d on %s"), LodIndex, *GetPathNameSafe(Owner));
-					Owner->BuildFromMeshDescription(*Owner->GetMeshDescription(LodIndex), LODResources[LodIndex]);
+					if (Owner->IsMeshDescriptionValid(LodIndex))
+					{
+						Owner->BuildFromMeshDescription(*Owner->GetMeshDescription(LodIndex), LODResources[LodIndex]);
+					}
+					else
+					{
+						UE_LOG(LogStaticMesh, Error, TEXT("Bad MeshDescription at lod index %d on %s"), LodIndex, *GetPathNameSafe(Owner));
+					}
 				}
 			}
 			else
