Added LevelStreaming.AllowIncrementalRemovalWhilePendingVisibility to control whether incremental removal of a streaming level can be done while there's a pending visible streaming level being processed (enabled by default). This flag will allow to process pending removals even in the case where there's a lot of streaming level to be added to the world and regardless of the priority of streaming levels.

diff --git a/Engine/Source/Runtime/Engine/Classes/Engine/LevelStreaming.h b/Engine/Source/Runtime/Engine/Classes/Engine/LevelStreaming.h
--- a/Engine/Source/Runtime/Engine/Classes/Engine/LevelStreaming.h
+++ b/Engine/Source/Runtime/Engine/Classes/Engine/LevelStreaming.h
@@ -642,6 +642,9 @@ public:
 	/** If true level streaming can reuse an unloaded level that wasn't GC'd yet. */
 	static ENGINE_API bool ShouldReuseUnloadedButStillAroundLevels(const ULevel* InLevel);
 
+	/** Whether incremental removal of a streaming level can be done while there's a pending visible streaming level being processed. */
+	static ENGINE_API bool AllowIncrementalRemovalWhilePendingVisibility();
+
 	/** 
 	 * Traverses all streaming level objects in the persistent world and in all inner worlds and calls appropriate delegate for streaming objects that refer specified level 
 	 *
diff --git a/Engine/Source/Runtime/Engine/Private/LevelStreaming.cpp b/Engine/Source/Runtime/Engine/Private/LevelStreaming.cpp
--- a/Engine/Source/Runtime/Engine/Private/LevelStreaming.cpp
+++ b/Engine/Source/Runtime/Engine/Private/LevelStreaming.cpp
@@ -88,6 +88,14 @@ namespace LevelStreamingCVars
 		TEXT("Whether level streaming will reuse the unloaded levels that aren't GC'd yet.\n")
 		TEXT("0: Disable, 1: Enable"),
 		ECVF_ReadOnly);
+
+	static bool bAllowIncrementalRemovalWhilePendingVisibility = true;
+	FAutoConsoleVariableRef CVarAllowIncrementalRemovalWhilePendingVisibility(
+		TEXT("LevelStreaming.AllowIncrementalRemovalWhilePendingVisibility"),
+		bAllowIncrementalRemovalWhilePendingVisibility,
+		TEXT("Whether incremental removal of a streaming level can be done while there's a pending visible streaming level being processed.\n")
+		TEXT("0: Disable, 1: Enable"),
+		ECVF_Default);
 }
 
 bool ULevelStreaming::DefaultAllowClientUseMakingInvisibleTransactionRequests()
@@ -129,6 +137,11 @@ bool ULevelStreaming::ShouldServerUseMakingVisibleTransactionRequest()
 	return LevelStreamingCVars::bShouldServerUseMakingVisibleTransactionRequest;
 }
 
+bool ULevelStreaming::AllowIncrementalRemovalWhilePendingVisibility()
+{
+	return LevelStreamingCVars::bAllowIncrementalRemovalWhilePendingVisibility;
+}
+
 bool ULevelStreaming::ShouldReuseUnloadedButStillAroundLevels(const ULevel* InLevel)
 {
 #if WITH_EDITOR
diff --git a/Engine/Source/Runtime/Engine/Private/World.cpp b/Engine/Source/Runtime/Engine/Private/World.cpp
--- a/Engine/Source/Runtime/Engine/Private/World.cpp
+++ b/Engine/Source/Runtime/Engine/Private/World.cpp
@@ -3394,7 +3394,7 @@ void UWorld::AddToWorld( ULevel* Level, const FTransform& LevelTransform, bool b
 
 			bExecuteNextStep = Level->IsFinishedRouteActorInitialization() && (!bConsiderTimeLimit || !IsTimeLimitExceeded( TEXT("routing Initialize on actors"), StartTime, Level, TimeLimit ));
 		}
-
+		
 		// Sort the actor list; can't do this on save as the relevant properties for sorting might have been changed by code
 		if( bExecuteNextStep && !Level->bAlreadySortedActorList )
 		{
@@ -3567,9 +3567,12 @@ void UWorld::RemoveFromWorld( ULevel* Level, bool bAllowIncrementalRemoval, FNet
 		}
 	};
 
-	// To be removed from the world a world must be visible and not pending being made visible (this may be redundent, but for safety)
 	// If the level may be removed incrementally then there must also be no level pending visibility
-	if ( ((CurrentLevelPendingVisibility == nullptr) || (!bAllowIncrementalRemoval && (CurrentLevelPendingVisibility != Level))) && (Level->bIsVisible || Level->bIsBeingRemoved) )
+	// except if LevelStreaming.AllowIncrementalRemovalWhilePendingVisibility is true.
+	const bool bIsCandidateForRemoval = ULevelStreaming::AllowIncrementalRemovalWhilePendingVisibility() ? true : !bAllowIncrementalRemoval;
+	
+	// To be removed from the world a world must be visible and not pending being made visible (this may be redundent, but for safety)
+	if ( ((CurrentLevelPendingVisibility == nullptr) || (bIsCandidateForRemoval && (CurrentLevelPendingVisibility != Level))) && (Level->bIsVisible || Level->bIsBeingRemoved) )
 	{
 #if PERF_TRACK_DETAILED_ASYNC_STATS
 		// Keep track of timing.
@@ -3633,7 +3636,7 @@ void UWorld::RemoveFromWorld( ULevel* Level, bool bAllowIncrementalRemoval, FNet
 			BeginRemoval();
 		}
 
-		if ( bFinishRemovingLevel )
+		if (bFinishRemovingLevel)
 		{
 			for (int32 ActorIdx = 0; ActorIdx < Level->Actors.Num(); ActorIdx++)
 			{
@@ -3666,8 +3669,8 @@ void UWorld::RemoveFromWorld( ULevel* Level, bool bAllowIncrementalRemoval, FNet
 			Level->ReleaseRenderingResources();
 
 			// Remove from the world's level array and destroy actor components.
-			IStreamingManager::Get().RemoveLevel( Level );
-		
+			IStreamingManager::Get().RemoveLevel(Level);
+
 			Level->ClearLevelComponents();
 
 			if (bIsGameWorld && !Level->bClientOnlyVisible)
@@ -3703,7 +3706,7 @@ void UWorld::RemoveFromWorld( ULevel* Level, bool bAllowIncrementalRemoval, FNet
 				Levels.Remove(Level);
 				Level->OwningWorld = nullptr;
 			}
-				
+
 			// let the universe know we have removed a level
 			FWorldDelegates::LevelRemovedFromWorld.Broadcast(Level, this);
 			BroadcastLevelsChanged();
