Include time spent finishing the incremental removal of a streamed level in the tracked time limit s.UnregisterComponentsTimeLimit.

diff --git a/Engine/Source/Runtime/Engine/Private/World.cpp b/Engine/Source/Runtime/Engine/Private/World.cpp
--- a/Engine/Source/Runtime/Engine/Private/World.cpp
+++ b/Engine/Source/Runtime/Engine/Private/World.cpp
@@ -3642,6 +3642,9 @@ void UWorld::RemoveFromWorld( ULevel* Level, bool bAllowIncrementalRemoval, FNet
 
 		if (bFinishRemovingLevel)
 		{
+			TRACE_CPUPROFILER_EVENT_SCOPE(RemoveFromWorld_FinishRemovingLevel);
+			double StartTime = FPlatformTime::Seconds();
+
 			for (int32 ActorIdx = 0; ActorIdx < Level->Actors.Num(); ActorIdx++)
 			{
 				if (AActor* Actor = Level->Actors[ActorIdx])
@@ -3722,6 +3725,13 @@ void UWorld::RemoveFromWorld( ULevel* Level, bool bAllowIncrementalRemoval, FNet
 			}
 
 			Level->bIsBeingRemoved = false;
+
+			// Keep track of time spent completing the incremental removal of the level
+			if (bAllowIncrementalRemoval)
+			{
+				double DeltaTime = (FPlatformTime::Seconds() - StartTime) * 1000;
+				GRemoveFromWorldUnregisterComponentTimeCumul += DeltaTime;
+			}
 		} // if ( bFinishRemovingLevel )
 
 		Level->bIsDisassociatingLevel = false;
