[Character] Option to base character on attachment root.

diff --git a/Engine/Source/Runtime/Engine/Classes/GameFramework/CharacterMovementComponent.h b/Engine/Source/Runtime/Engine/Classes/GameFramework/CharacterMovementComponent.h
--- a/Engine/Source/Runtime/Engine/Classes/GameFramework/CharacterMovementComponent.h
+++ b/Engine/Source/Runtime/Engine/Classes/GameFramework/CharacterMovementComponent.h
@@ -1383,6 +1383,10 @@ public:
 	UPROPERTY(Category = "Character Movement (General Settings)", EditAnywhere, BlueprintReadWrite)
 	bool bBasedMovementIgnorePhysicsBase = false;
 
+	/** Property to set if characters should stay based on objects attachment root instead of the traced object */
+	UPROPERTY(Category = "Character Movement (General Settings)", EditAnywhere, BlueprintReadWrite)
+	bool bBaseOnAttachmentRoot = false;
+
 	/** Property to set if characters should stay based on objects while jumping */
 	UPROPERTY(Category = "Character Movement: Jumping / Falling", EditAnywhere, BlueprintReadWrite)
 	bool bStayBasedInAir = false;
diff --git a/Engine/Source/Runtime/Engine/Private/Components/CharacterMovementComponent.cpp b/Engine/Source/Runtime/Engine/Private/Components/CharacterMovementComponent.cpp
--- a/Engine/Source/Runtime/Engine/Private/Components/CharacterMovementComponent.cpp
+++ b/Engine/Source/Runtime/Engine/Private/Components/CharacterMovementComponent.cpp
@@ -2291,6 +2291,14 @@ void UCharacterMovementComponent::SetBase( UPrimitiveComponent* NewBase, FName B
 
 	if (CharacterOwner && !bIsNavWalkingOnServer)
 	{
+		if (bBaseOnAttachmentRoot && NewBase != nullptr && BoneName == NAME_None)
+		{
+			// The attachment root might not be a PrimitiveComponent so check first
+			if (UPrimitiveComponent* NewRootBase = Cast<UPrimitiveComponent>(NewBase->GetAttachmentRoot()))
+			{
+				NewBase = NewRootBase;
+			}
+		}
 		CharacterOwner->SetBase(NewBase, NewBase ? BoneName : NAME_None, bNotifyActor);
 	}
 }
