Substrate - Fixed compilation problems wit material attribute.

diff --git a/Engine/Source/Runtime/Engine/Private/Materials/MaterialExpressions.cpp b/Engine/Source/Runtime/Engine/Private/Materials/MaterialExpressions.cpp
--- a/Engine/Source/Runtime/Engine/Private/Materials/MaterialExpressions.cpp
+++ b/Engine/Source/Runtime/Engine/Private/Materials/MaterialExpressions.cpp
@@ -27840,7 +27840,8 @@ int32 UMaterialExpressionSubstrateConvertMaterialAttributes::Compile(class FMate
 	}
 
 	int32 ShadingModelCodeChunk = MaterialAttributes.CompileWithDefault(Compiler, FMaterialAttributeDefinitionMap::GetID(MP_ShadingModel));
-	const bool bHasShadingModelExpression = IsMaterialAttributeInputConnected(Cached, MP_ShadingModel);
+	const bool bHasShadingModelExpression = IsMaterialAttributeInputConnected(Cached, MP_ShadingModel) 
+											|| ShadingModelOverride == MSM_FromMaterialExpression; // In this case, rely on the default compilation to return DefaultLit.
 	if (!bHasShadingModelExpression)
 	{
 		ShadingModelCodeChunk = Compiler->Constant(float(ShadingModelOverride));
@@ -27970,7 +27971,7 @@ void UMaterialExpressionSubstrateConvertMaterialAttributes::GatherSubstrateMater
 	if (IsMaterialAttributeInputConnected(Cached, MP_CustomData1)) 		{ SubstrateMaterialInfo.AddPropertyConnected(MP_CustomData1); }
 	if (IsMaterialAttributeInputConnected(Cached, MP_Opacity)) 			{ SubstrateMaterialInfo.AddPropertyConnected(MP_Opacity); }
 
-	if (IsMaterialAttributeInputConnected(Cached, MP_ShadingModel))
+	if (IsMaterialAttributeInputConnected(Cached, MP_ShadingModel) || ShadingModelOverride == MSM_FromMaterialExpression)
 	{
 		SubstrateMaterialInfo.AddPropertyConnected(MP_ShadingModel);
 
@@ -28028,7 +28029,7 @@ FSubstrateOperator* UMaterialExpressionSubstrateConvertMaterialAttributes::Subst
 	FMaterialShadingModelField ShadingModels = Compiler->GetMaterialShadingModels();
 
 	// Logic about shading models and complexity should match UMaterialExpressionSubstrateConvertMaterialAttributes::Compile.
-	const bool bHasShadingModelFromExpression = IsMaterialAttributeInputConnected(Cached, MP_ShadingModel); // We keep HasShadingModelFromExpression in case all shading models cannot be safely recovered from material functions.
+	const bool bHasShadingModelFromExpression = IsMaterialAttributeInputConnected(Cached, MP_ShadingModel) || ShadingModelOverride == MSM_FromMaterialExpression; // We keep HasShadingModelFromExpression in case all shading models cannot be safely recovered from material functions.
 	if ((ShadingModels.CountShadingModels() > 1) || bHasShadingModelFromExpression) 
 	{
 		return AddDefaultWorstCase(true, true);
