[PCG] Add "Create Spline Mesh" element
* Consumes a UPCGPolyLineData and produces a USplineMeshComponent for each spline segment
* Add (Soft)SplineMeshComponentDescriptor which mirrors the (Soft)ISMComponentDescriptor to easily expose UStaticMesh properties and initialize USplineMeshComponents without lots of boilerplate
* Add UPCGManagedSplineMeshComponent to allow for reuse of USplineMeshComponents
* Clean up old code in "Create Spline" element

diff --git a/Engine/Plugins/PCG/Source/PCG/Private/Elements/PCGCreateSpline.cpp b/Engine/Plugins/PCG/Source/PCG/Private/Elements/PCGCreateSpline.cpp
--- a/Engine/Plugins/PCG/Source/PCG/Private/Elements/PCGCreateSpline.cpp
+++ b/Engine/Plugins/PCG/Source/PCG/Private/Elements/PCGCreateSpline.cpp
@@ -21,15 +21,6 @@
 
 #define LOCTEXT_NAMESPACE "PCGCreateSpline"
 
-UPCGCreateSplineSettings::UPCGCreateSplineSettings(const FObjectInitializer& ObjectInitializer)
-	: UPCGSettings(ObjectInitializer)
-{
-	if (PCGHelpers::IsNewObjectAndNotDefault(this))
-	{
-		AttachOptions = EPCGAttachOptions::InFolder;
-	}
-}
-
 #if WITH_EDITOR
 FText UPCGCreateSplineSettings::GetNodeTooltipText() const
 {
@@ -50,12 +41,6 @@ FPCGElementPtr UPCGCreateSplineSettings::CreateElement() const
 	return MakeShared<FPCGCreateSplineElement>();
 }
 
-bool FPCGCreateSplineElement::CanExecuteOnlyOnMainThread(FPCGContext* Context) const
-{
-	const UPCGCreateSplineSettings* Settings = Context->GetInputSettings<UPCGCreateSplineSettings>();
-	return Settings && Settings->Mode == EPCGCreateSplineMode::CreateNewActor;
-}
-
 bool FPCGCreateSplineElement::IsCacheable(const UPCGSettings* InSettings) const
 {
 	const UPCGCreateSplineSettings* Settings = Cast<const UPCGCreateSplineSettings>(InSettings);
@@ -74,15 +59,6 @@ bool FPCGCreateSplineElement::ExecuteInternal(FPCGContext* Context) const
 
 	EPCGCreateSplineMode Mode = Settings->Mode;
 
-	if (Mode == EPCGCreateSplineMode::CreateNewActor)
-	{
-		const bool bHasAuthority = !Context->SourceComponent.IsValid() || (Context->SourceComponent->GetOwner() && Context->SourceComponent->GetOwner()->HasAuthority());
-		if (!bHasAuthority)
-		{
-			Mode = EPCGCreateSplineMode::CreateDataOnly;
-		}
-	}
-
 	for (const FPCGTaggedData& Input : Inputs)
 	{
 		const UPCGSpatialData* SpatialData = Cast<UPCGSpatialData>(Input.Data);
@@ -142,36 +118,6 @@ bool FPCGCreateSplineElement::ExecuteInternal(FPCGContext* Context) const
 		UPCGSplineData* SplineData = NewObject<UPCGSplineData>();
 		AActor* SplineActor = TargetActor;
 
-		if (Settings->Mode == EPCGCreateSplineMode::CreateDataOnly)
-		{
-			// Nothing
-		}
-		else if(Settings->Mode == EPCGCreateSplineMode::CreateNewActor)
-		{
-			// TODO: check ownership of target actor like in the spawn actor?
-			// TODO: allow template?
-			FActorSpawnParameters ActorSpawnParams;
-			SplineActor = UPCGActorHelpers::SpawnDefaultActor(TargetActor->GetWorld(), TargetActor->GetLevel(), AActor::StaticClass(), TargetActor->GetTransform(), ActorSpawnParams);
-
-			if (!SplineActor)
-			{
-				PCGE_LOG(Error, GraphAndLog, LOCTEXT("FailedToCreateActor", "Failed to create actor to hold the spline"));
-				continue;
-			}
-
-			UPCGManagedActors* ManagedActors = NewObject<UPCGManagedActors>(Context->SourceComponent.Get());
-
-			FTransform Transform = TargetActor->GetTransform();
-
-			SplineActor->Tags = TargetActor->Tags;
-			SplineActor->Tags.AddUnique(PCGHelpers::DefaultPCGActorTag);
-			PCGHelpers::AttachToParent(SplineActor, TargetActor, Settings->AttachOptions);
-			SplineData->TargetActor = SplineActor;
-
-			ManagedActors->GeneratedActors.Add(SplineActor);
-			Context->SourceComponent->AddToManagedResources(ManagedActors);
-		}
-
 		const TArray<FPCGPoint>& Points = PointData->GetPoints();
 		TArray<FSplinePoint> SplinePoints;
 		SplinePoints.Reserve(Points.Num());
diff --git a/Engine/Plugins/PCG/Source/PCG/Private/Elements/PCGCreateSplineMesh.cpp b/Engine/Plugins/PCG/Source/PCG/Private/Elements/PCGCreateSplineMesh.cpp
new file mode 100644
--- /dev/null
+++ b/Engine/Plugins/PCG/Source/PCG/Private/Elements/PCGCreateSplineMesh.cpp
@@ -0,0 +1,234 @@
+// Copyright Epic Games, Inc. All Rights Reserved.
+
+#include "Elements/PCGCreateSplineMesh.h"
+
+#include "PCGComponent.h"
+#include "PCGContext.h"
+#include "PCGManagedResource.h"
+#include "PCGPin.h"
+#include "PCGSettings.h"
+#include "Data/PCGLandscapeSplineData.h"
+#include "Data/PCGSplineData.h"
+#include "Helpers/PCGActorHelpers.h"
+#include "Helpers/PCGHelpers.h"
+
+#include "Components/SplineMeshComponent.h"
+#include "Engine/StaticMesh.h"
+#include "Engine/World.h"
+#include "Materials/MaterialInterface.h"
+#include "VT/RuntimeVirtualTexture.h"
+
+#include "DrawDebugHelpers.h"
+#include "PCGSubsystem.h"
+
+#include UE_INLINE_GENERATED_CPP_BY_NAME(PCGCreateSplineMesh)
+
+#define LOCTEXT_NAMESPACE "PCGCreateSplineMeshElement"
+
+#if WITH_EDITOR
+FText UPCGCreateSplineMeshSettings::GetNodeTooltipText() const
+{
+	return LOCTEXT("CreateSplineMeshTooltip", "Create a USplineMeshComponent for each segment along a given spline.");
+}
+#endif
+
+TArray<FPCGPinProperties> UPCGCreateSplineMeshSettings::InputPinProperties() const
+{
+	TArray<FPCGPinProperties> PinProperties;
+	PinProperties.Emplace(PCGPinConstants::DefaultInputLabel, EPCGDataType::PolyLine);
+
+	return PinProperties;
+}
+
+FPCGElementPtr UPCGCreateSplineMeshSettings::CreateElement() const
+{
+	return MakeShared<FPCGCreateSplineMeshElement>();
+}
+
+FPCGContext* FPCGCreateSplineMeshElement::CreateContext()
+{
+	return new FPCGCreateSplineMeshContext();
+}
+
+bool FPCGCreateSplineMeshElement::PrepareDataInternal(FPCGContext* InContext) const
+{
+	TRACE_CPUPROFILER_EVENT_SCOPE(FPCGCreateSplineMeshElement::PrepareDataInternal);
+
+	FPCGCreateSplineMeshContext* Context = static_cast<FPCGCreateSplineMeshContext*>(InContext);
+	check(Context);
+
+	const UPCGCreateSplineMeshSettings* Settings = Context->GetInputSettings<UPCGCreateSplineMeshSettings>();
+	check(Settings);
+
+	const FSoftSplineMeshComponentDescriptor& Descriptor = Settings->SplineMeshDescriptor;
+
+	if (Descriptor.StaticMesh.IsNull())
+	{
+		return true;
+	}
+
+	if (!Context->WasLoadRequested())
+	{
+		TArray<FSoftObjectPath> ObjectsToLoad = { Descriptor.StaticMesh.ToSoftObjectPath() };
+
+		if (!Descriptor.OverlayMaterial.IsNull())
+		{
+			ObjectsToLoad.Emplace(Descriptor.OverlayMaterial.ToSoftObjectPath());
+		}
+
+		for (int MaterialIndex = 0; MaterialIndex < Descriptor.OverrideMaterials.Num(); ++MaterialIndex)
+		{
+			if (!Descriptor.OverrideMaterials[MaterialIndex].IsNull())
+			{
+				ObjectsToLoad.Emplace(Descriptor.OverrideMaterials[MaterialIndex].ToSoftObjectPath());
+			}
+		}
+
+		for (int RVTIndex = 0; RVTIndex < Descriptor.RuntimeVirtualTextures.Num(); ++RVTIndex)
+		{
+			if (!Descriptor.RuntimeVirtualTextures[RVTIndex].IsNull())
+			{
+				ObjectsToLoad.Emplace(Descriptor.RuntimeVirtualTextures[RVTIndex].ToSoftObjectPath());
+			}
+		}
+
+		return Context->RequestResourceLoad(Context, std::move(ObjectsToLoad), !Settings->bSynchronousLoad);
+	}
+
+	return true;
+}
+
+bool FPCGCreateSplineMeshElement::ExecuteInternal(FPCGContext* Context) const
+{
+	TRACE_CPUPROFILER_EVENT_SCOPE(FPCGCreateSplineMeshElement::Execute);
+
+	const UPCGCreateSplineMeshSettings* Settings = Context->GetInputSettings<UPCGCreateSplineMeshSettings>();
+	check(Settings);
+
+	FSplineMeshComponentDescriptor Descriptor(Settings->SplineMeshDescriptor);
+
+	// TODO: Optionally read static mesh, materials, etc. from attributes.
+	UStaticMesh* StaticMesh = Descriptor.StaticMesh.Get();
+
+	if (!StaticMesh)
+	{
+		return true;
+	}
+
+	const FVector MeshExtents = StaticMesh->GetBounds().BoxExtent;
+
+	FPCGSplineMeshComponentBuilderParameters SMCBuilderParams;
+	SMCBuilderParams.Descriptor = Descriptor;
+	SMCBuilderParams.SplineMeshParams = Settings->SplineMeshParams;
+
+	TArray<FPCGTaggedData> Inputs = Context->InputData.GetInputs();
+	TArray<FPCGTaggedData>& Outputs = Context->OutputData.TaggedData;
+
+	for (const FPCGTaggedData& Input : Inputs)
+	{
+		const UPCGPolyLineData* SplineData = Cast<UPCGPolyLineData>(Input.Data);
+		const UPCGLandscapeSplineData* LandscapeSplineData = Cast<UPCGLandscapeSplineData>(SplineData);
+
+		if (!SplineData)
+		{
+			PCGE_LOG(Error, GraphAndLog, LOCTEXT("InvalidInputData", "Invalid input data, must be of type PolyLine."));
+			continue;
+		}
+
+		AActor* TargetActor = Settings->TargetActor.Get() ? Settings->TargetActor.Get() : Context->GetTargetActor(nullptr);
+
+		if (!TargetActor)
+		{
+			PCGE_LOG(Error, GraphAndLog, LOCTEXT("InvalidTargetActor", "Invalid target actor."));
+			continue;
+		}
+
+		const int NumSegments = SplineData->GetNumSegments();
+		const bool bIsClosed = SplineData->IsClosed();
+
+		for (int Index = 0; Index < NumSegments; ++Index)
+		{
+			const int NextIndex = (bIsClosed && Index == NumSegments - 1) ? 0 : (Index + 1);
+
+			FBox Bounds, NextBounds;
+			FTransform Transform = SplineData->GetTransformAtDistance(Index, /*Distance=*/0.0, /*bWorldSpace=*/true, &Bounds);
+			FTransform NextTransform = SplineData->GetTransformAtDistance(Index, SplineData->GetSegmentLength(Index), /*bWorldSpace=*/true, &NextBounds);
+
+			if (LandscapeSplineData && !SMCBuilderParams.SplineMeshParams.bScaleMeshToLandscapeSplineFullWidth)
+			{
+				Bounds = FBox::BuildAABB(FVector::ZeroVector, FVector::OneVector);
+				NextBounds = FBox::BuildAABB(FVector::ZeroVector, FVector::OneVector);
+			}
+
+			FVector LeaveTangent, ArriveTangent, Dummy;
+			SplineData->GetTangentsAtSegmentStart(Index, Dummy, LeaveTangent);
+			SplineData->GetTangentsAtSegmentStart(NextIndex, ArriveTangent, Dummy);
+
+			// Set Position
+			SMCBuilderParams.SplineMeshParams.StartPosition = Transform.GetLocation();
+			SMCBuilderParams.SplineMeshParams.StartTangent = LeaveTangent;
+			SMCBuilderParams.SplineMeshParams.EndPosition = NextTransform.GetLocation();
+			SMCBuilderParams.SplineMeshParams.EndTangent = ArriveTangent;
+
+			// Set Roll (Rotation)
+			const FQuat StartRotation = Transform.GetRotation();
+			const FQuat EndRotation = NextTransform.GetRotation();
+
+			SMCBuilderParams.SplineMeshParams.StartRollDegrees = StartRotation.Rotator().Roll;
+			SMCBuilderParams.SplineMeshParams.EndRollDegrees = EndRotation.Rotator().Roll;
+
+			// Set Scale
+			FVector StartScale = Transform.GetScale3D() * Bounds.GetExtent();
+			FVector EndScale = NextTransform.GetScale3D() * NextBounds.GetExtent();
+
+			if (SMCBuilderParams.SplineMeshParams.bScaleMeshToBounds)
+			{
+				FVector::FReal ScaleFactorY = 1.0f;
+				FVector::FReal ScaleFactorZ = 1.0f;
+
+				// We only scale in two dimensions since we are extruding along one of the axes. Scale on the two axes we are not extruding along.
+				if (SMCBuilderParams.SplineMeshParams.ForwardAxis == EPCGSplineMeshForwardAxis::X)
+				{
+					ScaleFactorY = MeshExtents.Y;
+					ScaleFactorZ = MeshExtents.Z;
+				}
+				else if (SMCBuilderParams.SplineMeshParams.ForwardAxis == EPCGSplineMeshForwardAxis::Y)
+				{
+					ScaleFactorY = MeshExtents.X;
+					ScaleFactorZ = MeshExtents.Z;
+				}
+				else if (SMCBuilderParams.SplineMeshParams.ForwardAxis == EPCGSplineMeshForwardAxis::Z)
+				{
+					ScaleFactorY = MeshExtents.X;
+					ScaleFactorZ = MeshExtents.Y;
+				}
+
+				StartScale.Y /= ScaleFactorY;
+				StartScale.Z /= ScaleFactorZ;
+				EndScale.Y /= ScaleFactorY;
+				EndScale.Z /= ScaleFactorZ;
+			}
+
+			SMCBuilderParams.SplineMeshParams.StartScale = FVector2D(StartScale.Y, StartScale.Z);
+			SMCBuilderParams.SplineMeshParams.EndScale = FVector2D(EndScale.Y, EndScale.Z);
+
+			USplineMeshComponent* SplineMeshComponent = UPCGActorHelpers::GetOrCreateSplineMeshComponent(TargetActor, Context->SourceComponent.Get(), Settings->UID, SMCBuilderParams);
+
+			// TODO: Write out the geometry to a dynamic mesh type.
+			//SplineMeshComponent->BodySetup->TriMeshGeometries
+		}
+
+		// Execute PostProcess Functions
+		if (TargetActor)
+		{
+			for (UFunction* Function : PCGHelpers::FindUserFunctions(TargetActor->GetClass(), Settings->PostProcessFunctionNames, { UPCGFunctionPrototypes::GetPrototypeWithNoParams() }, Context))
+			{
+				TargetActor->ProcessEvent(Function, nullptr);
+			}
+		}
+	}
+
+	return true;
+}
+
+#undef LOCTEXT_NAMESPACE
diff --git a/Engine/Plugins/PCG/Source/PCG/Private/Elements/PCGSplineMeshParams.cpp b/Engine/Plugins/PCG/Source/PCG/Private/Elements/PCGSplineMeshParams.cpp
new file mode 100644
--- /dev/null
+++ b/Engine/Plugins/PCG/Source/PCG/Private/Elements/PCGSplineMeshParams.cpp
@@ -0,0 +1,65 @@
+// Copyright Epic Games, Inc. All Rights Reserved.
+
+#include "Elements/PCGSplineMeshParams.h"
+
+#include UE_INLINE_GENERATED_CPP_BY_NAME(PCGSplineMeshParams)
+
+uint32 GetTypeHash(const FPCGSplineMeshParams& Key)
+{
+	if (!Key.Hash.IsSet())
+	{
+		Key.Hash = Key.ComputeHash();
+	}
+
+	return Key.Hash.GetValue();
+}
+
+bool FPCGSplineMeshParams::operator==(const FPCGSplineMeshParams& Other) const
+{
+	return ForwardAxis == Other.ForwardAxis
+		&& bScaleMeshToBounds == Other.bScaleMeshToBounds
+		&& bScaleMeshToLandscapeSplineFullWidth == Other.bScaleMeshToLandscapeSplineFullWidth
+		&& bSmoothInterpRollScale == Other.bSmoothInterpRollScale
+		&& FMath::IsNearlyEqual(NaniteClusterBoundsScale, Other.NaniteClusterBoundsScale)
+		&& FMath::IsNearlyEqual(SplineBoundaryMin, Other.SplineBoundaryMin)
+		&& FMath::IsNearlyEqual(SplineBoundaryMax, Other.SplineBoundaryMax)
+		&& FMath::IsNearlyEqual(StartRollDegrees, Other.StartRollDegrees)
+		&& FMath::IsNearlyEqual(EndRollDegrees, Other.EndRollDegrees)
+		&& (StartPosition - Other.StartPosition).IsNearlyZero()
+		&& (StartTangent - Other.StartTangent).IsNearlyZero()
+		&& (EndPosition - Other.EndPosition).IsNearlyZero()
+		&& (EndTangent - Other.EndTangent).IsNearlyZero()
+		&& (StartScale - Other.StartScale).IsNearlyZero()
+		&& (EndScale - Other.EndScale).IsNearlyZero();
+}
+
+uint32 FPCGSplineMeshParams::ComputeHash() const
+{
+	auto HashCombineAll = [](auto... Values)
+	{
+		uint32 HashVal = 0;
+		([&HashVal, &Values]()
+		{
+			HashVal = HashCombine(HashVal, GetTypeHash(Values));
+		}(), ...);
+		return HashVal;
+	};
+	
+	return HashCombineAll(
+		ForwardAxis,
+		bScaleMeshToBounds,
+		bScaleMeshToLandscapeSplineFullWidth,
+		SplineUpDir,
+		NaniteClusterBoundsScale,
+		SplineBoundaryMin,
+		SplineBoundaryMax,
+		bSmoothInterpRollScale,
+		StartPosition,
+		StartTangent,
+		StartRollDegrees,
+		StartScale,
+		EndPosition,
+		EndTangent,
+		EndRollDegrees,
+		EndScale);
+}
diff --git a/Engine/Plugins/PCG/Source/PCG/Private/Elements/PCGTypedGetter.cpp b/Engine/Plugins/PCG/Source/PCG/Private/Elements/PCGTypedGetter.cpp
--- a/Engine/Plugins/PCG/Source/PCG/Private/Elements/PCGTypedGetter.cpp
+++ b/Engine/Plugins/PCG/Source/PCG/Private/Elements/PCGTypedGetter.cpp
@@ -142,7 +142,7 @@ UPCGGetSplineSettings::UPCGGetSplineSettings()
 TArray<FPCGPinProperties> UPCGGetSplineSettings::OutputPinProperties() const
 {
 	TArray<FPCGPinProperties> PinProperties;
-	PinProperties.Emplace(PCGPinConstants::DefaultOutputLabel, EPCGDataType::Spline);
+	PinProperties.Emplace(PCGPinConstants::DefaultOutputLabel, EPCGDataType::PolyLine);
 
 	return PinProperties;
 }
diff --git a/Engine/Plugins/PCG/Source/PCG/Private/Helpers/PCGActorHelpers.cpp b/Engine/Plugins/PCG/Source/PCG/Private/Helpers/PCGActorHelpers.cpp
--- a/Engine/Plugins/PCG/Source/PCG/Private/Helpers/PCGActorHelpers.cpp
+++ b/Engine/Plugins/PCG/Source/PCG/Private/Helpers/PCGActorHelpers.cpp
@@ -5,10 +5,12 @@
 #include "PCGComponent.h"
 #include "PCGManagedResource.h"
 #include "PCGModule.h"
+#include "Elements/PCGSplineMeshParams.h"
 #include "Helpers/PCGHelpers.h"
 
 #include "EngineUtils.h"
 #include "Components/HierarchicalInstancedStaticMeshComponent.h"
+#include "Components/SplineMeshComponent.h"
 #include "Engine/StaticMesh.h"
 #include "Materials/MaterialInterface.h"
 
@@ -149,6 +151,116 @@ UPCGManagedISMComponent* UPCGActorHelpers::GetOrCreateManagedISMC(AActor* InTarg
 	return Resource;
 }
 
+USplineMeshComponent* UPCGActorHelpers::GetOrCreateSplineMeshComponent(AActor* InTargetActor, UPCGComponent* InSourceComponent, uint64 SettingsUID, const FPCGSplineMeshComponentBuilderParameters& InParams)
+{
+	UPCGManagedSplineMeshComponent* ManagedComponent = GetOrCreateManagedSplineMeshComponent(InTargetActor, InSourceComponent, SettingsUID, InParams);
+	if (ManagedComponent)
+	{
+		return ManagedComponent->GetComponent();
+	}
+	else
+	{
+		return nullptr;
+	}
+}
+
+UPCGManagedSplineMeshComponent* UPCGActorHelpers::GetOrCreateManagedSplineMeshComponent(AActor* InTargetActor, UPCGComponent* InSourceComponent, uint64 SettingsUID, const FPCGSplineMeshComponentBuilderParameters& InParams)
+{
+	check(InTargetActor && InSourceComponent);
+
+	const UStaticMesh* StaticMesh = InParams.Descriptor.StaticMesh;
+
+	if (!StaticMesh)
+	{
+		return nullptr;
+	}
+
+	TRACE_CPUPROFILER_EVENT_SCOPE(UPCGActorHelpers::GetOrCreateManagedSplineMeshComponent);
+
+	{
+		TRACE_CPUPROFILER_EVENT_SCOPE(UPCGActorHelpers::GetOrCreateManagedSplineMeshComponent::FindMatchingManagedSplineMeshComponent);
+
+		UPCGManagedSplineMeshComponent* MatchingResource = nullptr;
+		InSourceComponent->ForEachManagedResource([&MatchingResource, &InParams, &InTargetActor, SettingsUID](UPCGManagedResource* InResource)
+		{
+			// Early out if already found a match
+			if (MatchingResource)
+			{
+				return;
+			}
+
+			if (UPCGManagedSplineMeshComponent* Resource = Cast<UPCGManagedSplineMeshComponent>(InResource))
+			{
+				if (Resource->GetSettingsUID() != SettingsUID || !Resource->CanBeUsed())
+				{
+					return;
+				}
+
+				if (USplineMeshComponent* SplineMeshComponent = Resource->GetComponent())
+				{
+					if (IsValid(SplineMeshComponent)
+						&& SplineMeshComponent->GetOwner() == InTargetActor
+						&& Resource->GetDescriptor() == InParams.Descriptor
+						&& Resource->GetSplineMeshParams() == InParams.SplineMeshParams)
+					{
+						MatchingResource = Resource;
+					}
+				}
+			}
+		});
+
+		if (MatchingResource)
+		{
+			TRACE_CPUPROFILER_EVENT_SCOPE(UPCGActorHelpers::GetOrCreateManagedSplineMeshComponent::MarkAsUsed);
+			MatchingResource->MarkAsUsed();
+
+			return MatchingResource;
+		}
+	}
+
+	// No matching component found, let's create a new one.
+	InTargetActor->Modify(!InSourceComponent->IsInPreviewMode());
+
+	FString ComponentName = TEXT("PCGSplineMeshComponent_") + StaticMesh->GetName();
+	const EObjectFlags ObjectFlags = (InSourceComponent->IsInPreviewMode() ? RF_Transient : RF_NoFlags);
+	USplineMeshComponent* SplineMeshComponent = NewObject<USplineMeshComponent>(InTargetActor, MakeUniqueObjectName(InTargetActor, USplineMeshComponent::StaticClass(), FName(ComponentName)), ObjectFlags);
+
+	// Init Component
+	{
+		InParams.Descriptor.InitComponent(SplineMeshComponent);
+
+		const FPCGSplineMeshParams& SplineMeshParams = InParams.SplineMeshParams;
+		SplineMeshComponent->SetStartAndEnd(SplineMeshParams.StartPosition, SplineMeshParams.StartTangent, SplineMeshParams.EndPosition, SplineMeshParams.EndTangent);
+		SplineMeshComponent->SetStartRollDegrees(SplineMeshParams.StartRollDegrees);
+		SplineMeshComponent->SetEndRollDegrees(SplineMeshParams.EndRollDegrees);
+		SplineMeshComponent->SetStartScale(SplineMeshParams.StartScale);
+		SplineMeshComponent->SetEndScale(SplineMeshParams.EndScale);
+		SplineMeshComponent->SetForwardAxis((ESplineMeshAxis::Type)SplineMeshParams.ForwardAxis);
+		SplineMeshComponent->SetSplineUpDir(SplineMeshParams.SplineUpDir);
+		SplineMeshComponent->SplineParams.NaniteClusterBoundsScale = SplineMeshParams.NaniteClusterBoundsScale;
+		SplineMeshComponent->SplineBoundaryMin = SplineMeshParams.SplineBoundaryMin;
+		SplineMeshComponent->SplineBoundaryMax = SplineMeshParams.SplineBoundaryMax;
+		SplineMeshComponent->bSmoothInterpRollScale = SplineMeshParams.bSmoothInterpRollScale;
+	}
+
+	SplineMeshComponent->RegisterComponent();
+	InTargetActor->AddInstanceComponent(SplineMeshComponent);
+
+	SplineMeshComponent->AttachToComponent(InTargetActor->GetRootComponent(), FAttachmentTransformRules(EAttachmentRule::KeepWorld, EAttachmentRule::KeepWorld, EAttachmentRule::KeepWorld, false));
+	SplineMeshComponent->ComponentTags.Add(InSourceComponent->GetFName());
+	SplineMeshComponent->ComponentTags.Add(PCGHelpers::DefaultPCGTag);
+
+	// Create managed resource on source component
+	UPCGManagedSplineMeshComponent* Resource = NewObject<UPCGManagedSplineMeshComponent>(InSourceComponent);
+	Resource->SetComponent(SplineMeshComponent);
+	Resource->SetDescriptor(InParams.Descriptor);
+	Resource->SetSplineMeshParams(InParams.SplineMeshParams);
+	Resource->SetSettingsUID(SettingsUID);
+	InSourceComponent->AddToManagedResources(Resource);
+
+	return Resource;
+}
+
 bool UPCGActorHelpers::DeleteActors(UWorld* World, const TArray<TSoftObjectPtr<AActor>>& ActorsToDelete)
 {
 	if (!World || ActorsToDelete.Num() == 0)
diff --git a/Engine/Plugins/PCG/Source/PCG/Private/PCGManagedResource.cpp b/Engine/Plugins/PCG/Source/PCG/Private/PCGManagedResource.cpp
--- a/Engine/Plugins/PCG/Source/PCG/Private/PCGManagedResource.cpp
+++ b/Engine/Plugins/PCG/Source/PCG/Private/PCGManagedResource.cpp
@@ -7,6 +7,7 @@
 #include "Helpers/PCGHelpers.h"
 
 #include "Components/InstancedStaticMeshComponent.h"
+#include "Components/SplineMeshComponent.h"
 #include "Engine/Level.h"
 #include "UObject/Package.h"
 #include "Utils/PCGGeneratedResourcesLogging.h"
@@ -664,4 +665,35 @@ void UPCGManagedISMComponent::SetComponent(UInstancedStaticMeshComponent* InComp
 {
 	GeneratedComponent = InComponent;
 	CachedRawComponentPtr = InComponent;
-}
\ No newline at end of file
+}
+
+void UPCGManagedSplineMeshComponent::ForgetComponent()
+{
+	Super::ForgetComponent();
+	CachedRawComponentPtr = nullptr;
+}
+
+USplineMeshComponent* UPCGManagedSplineMeshComponent::GetComponent() const
+{
+	if (!CachedRawComponentPtr)
+	{
+		USplineMeshComponent* GeneratedComponentPtr = Cast<USplineMeshComponent>(GeneratedComponent.Get());
+
+		// Implementation note:
+		// There is no surefire way to make sure that we can use the raw pointer UNLESS it is from the same owner
+		if (GeneratedComponentPtr && Cast<UPCGComponent>(GetOuter()) && GeneratedComponentPtr->GetOwner() == Cast<UPCGComponent>(GetOuter())->GetOwner())
+		{
+			CachedRawComponentPtr = GeneratedComponentPtr;
+		}
+
+		return GeneratedComponentPtr;
+	}
+
+	return CachedRawComponentPtr;
+}
+
+void UPCGManagedSplineMeshComponent::SetComponent(USplineMeshComponent* InComponent)
+{
+	GeneratedComponent = InComponent;
+	CachedRawComponentPtr = InComponent;
+}
diff --git a/Engine/Plugins/PCG/Source/PCG/Public/Elements/PCGCreateSpline.h b/Engine/Plugins/PCG/Source/PCG/Public/Elements/PCGCreateSpline.h
--- a/Engine/Plugins/PCG/Source/PCG/Public/Elements/PCGCreateSpline.h
+++ b/Engine/Plugins/PCG/Source/PCG/Public/Elements/PCGCreateSpline.h
@@ -15,7 +15,6 @@ enum class EPCGCreateSplineMode : uint8
 {
 	CreateDataOnly,
 	CreateComponent,
-	CreateNewActor UMETA(Hidden)
 };
 
 /** PCG node that creates a spline presentation from the input points data, with optional tangents */
@@ -25,8 +24,6 @@ class UPCGCreateSplineSettings : public UPCGSettings
 	GENERATED_BODY()
 
 public:
-	UPCGCreateSplineSettings(const FObjectInitializer& ObjectInitializer);
-
 	//~Begin UPCGSettings interface
 #if WITH_EDITOR
 	virtual FName GetDefaultNodeName() const override { return FName(TEXT("CreateSpline")); }
@@ -67,9 +64,6 @@ public:
 	UPROPERTY(meta = (PCG_Overridable))
 	TSoftObjectPtr<AActor> TargetActor;
 
-	UPROPERTY(BlueprintReadWrite, EditAnywhere, Category = Settings, meta = (EditCondition = "Mode==EPCGCreateSplineMode::CreateNewActor", EditConditionHides))
-	EPCGAttachOptions AttachOptions = EPCGAttachOptions::Attached; // Note that this is no longer the default value for new nodes, it is now EPCGAttachOptions::InFolder
-
 	/** Specify a list of functions to be called on the target actor after spline creation. Functions need to be parameter-less and with "CallInEditor" flag enabled. */
 	UPROPERTY(BlueprintReadWrite, EditAnywhere, Category = Settings)
 	TArray<FName> PostProcessFunctionNames;
@@ -78,7 +72,6 @@ public:
 class FPCGCreateSplineElement : public IPCGElement
 {
 protected:
-	virtual bool CanExecuteOnlyOnMainThread(FPCGContext* Context) const override;
 	virtual bool IsCacheable(const UPCGSettings* InSettings) const override;
 	virtual bool ExecuteInternal(FPCGContext* Context) const override;
 };
\ No newline at end of file
diff --git a/Engine/Plugins/PCG/Source/PCG/Public/Elements/PCGCreateSplineMesh.h b/Engine/Plugins/PCG/Source/PCG/Public/Elements/PCGCreateSplineMesh.h
new file mode 100644
--- /dev/null
+++ b/Engine/Plugins/PCG/Source/PCG/Public/Elements/PCGCreateSplineMesh.h
@@ -0,0 +1,67 @@
+// Copyright Epic Games, Inc. All Rights Reserved.
+
+#pragma once
+
+#include "PCGCommon.h"
+#include "PCGContext.h"
+#include "PCGElement.h"
+#include "PCGSettings.h"
+#include "PCGSplineMeshParams.h"
+#include "Async/PCGAsyncLoadingContext.h"
+
+#include "Engine/SplineMeshComponentDescriptor.h"
+
+#include "PCGCreateSplineMesh.generated.h"
+
+struct FPCGContext;
+
+/** Create a USplineMeshComponent for each segment along a given spline. */
+UCLASS(MinimalAPI, BlueprintType, ClassGroup = (Procedural))
+class UPCGCreateSplineMeshSettings : public UPCGSettings
+{
+	GENERATED_BODY()
+
+public:
+	//~Begin UPCGSettings interface
+#if WITH_EDITOR
+	virtual FName GetDefaultNodeName() const override { return FName(TEXT("CreateSplineMesh")); }
+	virtual FText GetDefaultNodeTitle() const override { return NSLOCTEXT("PCGCreateSplineMeshElement", "NodeTitle", "Create Spline Mesh"); }
+	virtual FText GetNodeTooltipText() const override;
+	virtual EPCGSettingsType GetType() const override { return EPCGSettingsType::Spatial; }
+#endif
+
+protected:
+	virtual TArray<FPCGPinProperties> InputPinProperties() const override;
+	virtual TArray<FPCGPinProperties> OutputPinProperties() const override { return TArray<FPCGPinProperties>(); }
+	virtual FPCGElementPtr CreateElement() const override;
+	//~End UPCGSettings interface
+
+public:
+	UPROPERTY(EditAnywhere, Category = Settings)
+	FSoftSplineMeshComponentDescriptor SplineMeshDescriptor;
+
+	UPROPERTY(EditAnywhere, Category = Settings, meta = (PCG_Overridable))
+	FPCGSplineMeshParams SplineMeshParams;
+
+	UPROPERTY(meta = (PCG_Overridable))
+	TSoftObjectPtr<AActor> TargetActor;
+
+	/** Specify a list of functions to be called on the target actor after spline mesh creation. Functions need to be parameter-less and with "CallInEditor" flag enabled. */
+	UPROPERTY(BlueprintReadWrite, EditAnywhere, Category = Settings)
+	TArray<FName> PostProcessFunctionNames;
+
+	/** Force meshes/materials to load synchronously. */
+	UPROPERTY(BlueprintReadWrite, EditAnywhere, Category = "Settings|Debug")
+	bool bSynchronousLoad = false;
+};
+
+struct FPCGCreateSplineMeshContext : public FPCGContext, public IPCGAsyncLoadingContext {};
+
+class FPCGCreateSplineMeshElement : public IPCGElement
+{
+protected:
+	virtual bool IsCacheable(const UPCGSettings* InSettings) const override { return false; }
+	virtual FPCGContext* CreateContext() override;
+	virtual bool PrepareDataInternal(FPCGContext* InContext) const override;
+	virtual bool ExecuteInternal(FPCGContext* Context) const override;
+};
diff --git a/Engine/Plugins/PCG/Source/PCG/Public/Elements/PCGSplineMeshParams.h b/Engine/Plugins/PCG/Source/PCG/Public/Elements/PCGSplineMeshParams.h
new file mode 100644
--- /dev/null
+++ b/Engine/Plugins/PCG/Source/PCG/Public/Elements/PCGSplineMeshParams.h
@@ -0,0 +1,88 @@
+// Copyright Epic Games, Inc. All Rights Reserved.
+
+#pragma once
+
+#include "PCGSplineMeshParams.generated.h"
+
+UENUM()
+enum class EPCGSplineMeshForwardAxis : uint32
+{
+	X,
+	Y,
+	Z
+};
+
+USTRUCT(BlueprintType)
+struct FPCGSplineMeshParams
+{
+	GENERATED_BODY()
+
+	/** Chooses the forward axis for the spline mesh orientation. */
+	UPROPERTY(BlueprintReadWrite, EditAnywhere, Category = "Spline Mesh Params", meta = (PCG_Overridable))
+	EPCGSplineMeshForwardAxis ForwardAxis = EPCGSplineMeshForwardAxis::X;
+
+	/** Scale mesh to the spline control point bounds. Especially useful on Landscape Splines, where the bounds come from the width. */
+	UPROPERTY(BlueprintReadWrite, EditAnywhere, Category = "Spline Mesh Params", meta = (PCG_Overridable))
+	bool bScaleMeshToBounds = false;
+
+	/** Scale the mesh to the full width of the Landscape Spline (including Falloff). Only applies to Landscape Splines. */
+	UPROPERTY(BlueprintReadWrite, EditAnywhere, Category = "Spline Mesh Params", meta = (PCG_Overridable, EditCondition = "bScaleMeshToBounds"))
+	bool bScaleMeshToLandscapeSplineFullWidth = false;
+
+	/** Axis (in component space) that is used to determine X axis for coordinates along spline. */
+	UPROPERTY(BlueprintReadWrite, EditAnywhere, Category = "Spline Mesh Params", meta = (PCG_Overridable))
+	FVector SplineUpDir = FVector(0.0, 0.0, 1.0);
+
+	/**
+	 * How much to scale the calculated culling bounds of Nanite clusters after deformation.
+	 * NOTE: This should only be set greater than 1.0 if it fixes visible issues with clusters being
+	 * incorrectly culled.
+	 */
+	UPROPERTY(BlueprintReadWrite, EditAnywhere, Category = "Spline Mesh Params", meta = (ClampMin = 1.0, PCG_Overridable))
+	float NaniteClusterBoundsScale = 1.0f;
+
+	/** Minimum coordinate along the spline forward axis which corresponds to start of spline. If set to 0.0, will use bounding box to determine bounds. */
+	UPROPERTY(BlueprintReadWrite, EditAnywhere, Category = "Spline Mesh Params", meta = (PCG_Overridable))
+	float SplineBoundaryMin = 0.0f;
+
+	/** Maximum coordinate along the spline forward axis which corresponds to end of spline. If set to 0.0, will use bounding box to determine bounds. */
+	UPROPERTY(BlueprintReadWrite, EditAnywhere, Category = "Spline Mesh Params", meta = (PCG_Overridable))
+	float SplineBoundaryMax = 0.0f;
+
+	/** If true, will use smooth interpolation (ease in/out) for Scale, Roll, and Offset along this section of spline. If false, uses linear. */
+	UPROPERTY(BlueprintReadWrite, EditAnywhere, Category = "Spline Mesh Params", meta = (PCG_Overridable))
+	bool bSmoothInterpRollScale = true;
+
+	// Below properties are not exposed for edit because they are computed internally.
+
+	UPROPERTY()
+	FVector StartPosition = FVector::ZeroVector;
+
+	UPROPERTY()
+	FVector StartTangent = FVector::ZeroVector;
+
+	UPROPERTY()
+	float StartRollDegrees = 0.0f;
+
+	UPROPERTY()
+	FVector2D StartScale = FVector2D::ZeroVector;
+
+	UPROPERTY()
+	FVector EndPosition = FVector::ZeroVector;
+
+	UPROPERTY()
+	FVector EndTangent = FVector::ZeroVector;
+
+	UPROPERTY()
+	float EndRollDegrees = 0.0f;
+
+	UPROPERTY()
+	FVector2D EndScale = FVector2D::ZeroVector;
+
+	friend uint32 GetTypeHash(const FPCGSplineMeshParams& Key);
+	bool operator==(const FPCGSplineMeshParams& Other) const;
+
+private:
+	uint32 ComputeHash() const;
+	mutable TOptional<uint32> Hash;
+};
diff --git a/Engine/Plugins/PCG/Source/PCG/Public/Helpers/PCGActorHelpers.h b/Engine/Plugins/PCG/Source/PCG/Public/Helpers/PCGActorHelpers.h
--- a/Engine/Plugins/PCG/Source/PCG/Public/Helpers/PCGActorHelpers.h
+++ b/Engine/Plugins/PCG/Source/PCG/Public/Helpers/PCGActorHelpers.h
@@ -2,7 +2,10 @@
 
 #pragma once
 
+#include "Elements/PCGSplineMeshParams.h"
+
 #include "Engine/EngineTypes.h"
+#include "Engine/SplineMeshComponentDescriptor.h"
 #include "Engine/World.h"
 #include "ISMPartition/ISMComponentDescriptor.h"
 #include "Kismet/BlueprintFunctionLibrary.h"
@@ -15,8 +18,10 @@ class UActorComponent;
 class UInstancedStaticMeshComponent;
 class ULevel;
 class UMaterialInterface;
+class USplineMeshComponent;
 class UPCGComponent;
 class UPCGManagedISMComponent;
+class UPCGManagedSplineMeshComponent;
 class UStaticMesh;
 class UWorld;
 
@@ -33,6 +38,19 @@ struct FPCGISMCBuilderParameters
 	inline bool operator==(const FPCGISMCBuilderParameters& Other) const { return Descriptor == Other.Descriptor && NumCustomDataFloats == Other.NumCustomDataFloats; }
 };
 
+struct FPCGSplineMeshComponentBuilderParameters
+{
+	FSplineMeshComponentDescriptor Descriptor;
+	FPCGSplineMeshParams SplineMeshParams;
+
+	friend inline uint32 GetTypeHash(const FPCGSplineMeshComponentBuilderParameters& Key)
+	{
+		return HashCombine(GetTypeHash(Key.Descriptor), GetTypeHash(Key.SplineMeshParams));
+	}
+
+	inline bool operator==(const FPCGSplineMeshComponentBuilderParameters& Other) const { return Descriptor == Other.Descriptor && SplineMeshParams == Other.SplineMeshParams; }
+};
+
 UCLASS(BlueprintType)
 class PCG_API UPCGActorHelpers : public UBlueprintFunctionLibrary
 {
@@ -41,6 +59,8 @@ class PCG_API UPCGActorHelpers : public UBlueprintFunctionLibrary
 public:
 	static UInstancedStaticMeshComponent* GetOrCreateISMC(AActor* InTargetActor, UPCGComponent* SourceComponent, uint64 SettingsUID, const FPCGISMCBuilderParameters& Params);
 	static UPCGManagedISMComponent* GetOrCreateManagedISMC(AActor* InTargetActor, UPCGComponent* SourceComponent, uint64 SettingsUID, const FPCGISMCBuilderParameters& Params);
+	static USplineMeshComponent* GetOrCreateSplineMeshComponent(AActor* InTargetActor, UPCGComponent* SourceComponent, uint64 SettingsUID, const FPCGSplineMeshComponentBuilderParameters& Params);
+	static UPCGManagedSplineMeshComponent* GetOrCreateManagedSplineMeshComponent(AActor* InTargetActor, UPCGComponent* SourceComponent, uint64 SettingsUID, const FPCGSplineMeshComponentBuilderParameters& Params);
 	static bool DeleteActors(UWorld* World, const TArray<TSoftObjectPtr<AActor>>& ActorsToDelete);
 
 	template <typename T, typename = typename std::enable_if_t<std::is_base_of_v<AActor, T>>>
diff --git a/Engine/Plugins/PCG/Source/PCG/Public/PCGManagedResource.h b/Engine/Plugins/PCG/Source/PCG/Public/PCGManagedResource.h
--- a/Engine/Plugins/PCG/Source/PCG/Public/PCGManagedResource.h
+++ b/Engine/Plugins/PCG/Source/PCG/Public/PCGManagedResource.h
@@ -4,13 +4,16 @@
 
 #include "PCGCommon.h"
 #include "PCGCrc.h"
+#include "Elements/PCGSplineMeshParams.h"
 
+#include "Engine/SplineMeshComponentDescriptor.h"
 #include "ISMPartition/ISMComponentDescriptor.h"
 
 #include "PCGManagedResource.generated.h"
 
 class UActorComponent;
 class UInstancedStaticMeshComponent;
+class USplineMeshComponent;
 
 /** 
 * This class is used to hold resources and their mechanism to delete them on demand.
@@ -195,6 +198,44 @@ protected:
 	mutable UInstancedStaticMeshComponent* CachedRawComponentPtr = nullptr;
 };
 
+UCLASS(BlueprintType)
+class PCG_API UPCGManagedSplineMeshComponent : public UPCGManagedComponent
+{
+	GENERATED_BODY()
+
+public:
+	//~Begin UPCGManagedComponents interface
+	virtual void ResetComponent() override { /* Does nothing, but implementation is required to support reuse. */ }
+	virtual bool SupportsComponentReset() const override { return true; }
+	virtual void ForgetComponent() override;
+	//~End UPCGManagedComponents interface
+
+	USplineMeshComponent* GetComponent() const;
+	void SetComponent(USplineMeshComponent* InComponent);
+
+	void SetDescriptor(const FSplineMeshComponentDescriptor& InDescriptor) { Descriptor = InDescriptor; }
+	const FSplineMeshComponentDescriptor& GetDescriptor() const { return Descriptor; }
+
+	void SetSplineMeshParams(const FPCGSplineMeshParams& InSplineMeshParams) { SplineMeshParams = InSplineMeshParams; }
+	const FPCGSplineMeshParams& GetSplineMeshParams() const { return SplineMeshParams; }
+
+	uint64 GetSettingsUID() const { return SettingsUID; }
+	void SetSettingsUID(uint64 InSettingsUID) { SettingsUID = InSettingsUID; }
+
+protected:
+	UPROPERTY()
+	FSplineMeshComponentDescriptor Descriptor;
+
+	UPROPERTY()
+	FPCGSplineMeshParams SplineMeshParams;
+
+	UPROPERTY(Transient)
+	uint64 SettingsUID = -1; // purposefully a value that will never happen in data
+
+	// Cached raw pointer to USplineMeshComponent
+	mutable USplineMeshComponent* CachedRawComponentPtr = nullptr;
+};
+
 #if UE_ENABLE_INCLUDE_ORDER_DEPRECATED_IN_5_2
 #include "CoreMinimal.h"
 #endif
diff --git a/Engine/Source/Runtime/Engine/Classes/Engine/SplineMeshComponentDescriptor.h b/Engine/Source/Runtime/Engine/Classes/Engine/SplineMeshComponentDescriptor.h
new file mode 100644
--- /dev/null
+++ b/Engine/Source/Runtime/Engine/Classes/Engine/SplineMeshComponentDescriptor.h
@@ -0,0 +1,271 @@
+// Copyright Epic Games, Inc. All Rights Reserved.
+
+#pragma once
+
+#include "CoreMinimal.h"
+#include "Components/PrimitiveComponent.h"
+#include "Engine/EngineTypes.h"
+#include "PhysicsEngine/BodyInstance.h"
+
+#include "SplineMeshComponentDescriptor.generated.h"
+
+// Inspired in large part from ISMComponentDescriptor.h
+
+class USplineMeshComponent;
+class UStaticMesh;
+enum class ERuntimeVirtualTextureMainPassType : uint8;
+enum class ERendererStencilMask : uint8;
+
+USTRUCT()
+struct FSplineMeshComponentDescriptorBase
+{
+	GENERATED_BODY()
+
+	ENGINE_API FSplineMeshComponentDescriptorBase();
+	explicit FSplineMeshComponentDescriptorBase(ENoInit) {}
+	virtual ~FSplineMeshComponentDescriptorBase() {}
+
+	ENGINE_API USplineMeshComponent* CreateComponent(UObject* Outer, FName Name = NAME_None, EObjectFlags ObjectFlags = EObjectFlags::RF_NoFlags) const;
+
+	ENGINE_API virtual void InitFrom(const UStaticMeshComponent* Component, bool bInitBodyInstance = true);
+	ENGINE_API virtual uint32 ComputeHash() const;
+	ENGINE_API virtual void InitComponent(USplineMeshComponent* SplineMeshComponent) const;
+
+	ENGINE_API void PostLoadFixup(UObject* Loader);
+
+	ENGINE_API bool operator!=(const FSplineMeshComponentDescriptorBase& Other) const;
+	ENGINE_API bool operator==(const FSplineMeshComponentDescriptorBase& Other) const;
+
+	friend inline uint32 GetTypeHash(const FSplineMeshComponentDescriptorBase& Key)
+	{
+		return Key.GetTypeHash();
+	}
+
+	uint32 GetTypeHash() const
+	{
+		if (Hash == 0)
+		{
+			ComputeHash();
+		}
+		return Hash;
+	}
+
+public:
+	UPROPERTY()
+	mutable uint32 Hash = 0;
+
+	UPROPERTY(EditAnywhere, Category = "Component Settings")
+	TSubclassOf<USplineMeshComponent> ComponentClass;
+	
+	UPROPERTY(EditAnywhere, Category = "Component Settings")
+	TEnumAsByte<EComponentMobility::Type> Mobility;
+		
+	UPROPERTY(EditAnywhere, Category = "Component Settings")
+	ERuntimeVirtualTextureMainPassType VirtualTextureRenderPassType;
+	
+	UPROPERTY(EditAnywhere, Category = "Component Settings")
+	ELightmapType LightmapType;
+
+	UPROPERTY(EditAnywhere, Category = "Component Settings")
+	FLightingChannels LightingChannels;
+
+	UPROPERTY(EditAnywhere, Category = "Component Settings")
+	int32 RayTracingGroupId;
+
+	UPROPERTY(EditAnywhere, Category = "Component Settings")
+	ERayTracingGroupCullingPriority RayTracingGroupCullingPriority;
+
+	UPROPERTY(EditAnywhere, Category = "Component Settings")
+	TEnumAsByte<EHasCustomNavigableGeometry::Type> bHasCustomNavigableGeometry;
+
+	UPROPERTY(EditAnywhere, Category = "Component Settings")
+	ERendererStencilMask CustomDepthStencilWriteMask;
+
+	UPROPERTY(EditAnywhere, Category = "Component Settings")
+	FBodyInstance BodyInstance;
+		
+	UPROPERTY(EditAnywhere, Category = "Component Settings")
+	int32 VirtualTextureCullMips;
+
+	UPROPERTY(EditAnywhere, Category = "Component Settings")
+	int32 TranslucencySortPriority;
+
+	UPROPERTY(EditAnywhere, Category = "Component Settings")
+	int32 OverriddenLightMapRes;
+
+	UPROPERTY(EditAnywhere, Category = "Component Settings")
+	int32 CustomDepthStencilValue;
+
+#if WITH_EDITORONLY_DATA
+	UPROPERTY(EditAnywhere, Category = "Component Settings")
+	EHLODBatchingPolicy HLODBatchingPolicy;
+#endif
+	
+	UPROPERTY(EditAnywhere, Category = "Component Settings")
+	uint8 bCastShadow : 1;
+
+	UPROPERTY(EditAnywhere, Category = "Component Settings")
+	uint8 bEmissiveLightSource : 1;
+		
+	UPROPERTY(EditAnywhere, Category = "Component Settings")
+	uint8 bCastDynamicShadow : 1;
+
+	UPROPERTY(EditAnywhere, Category = "Component Settings")
+	uint8 bCastStaticShadow : 1;
+
+	UPROPERTY(EditAnywhere, Category = "Component Settings")
+	uint8 bCastContactShadow : 1;
+
+	UPROPERTY(EditAnywhere, Category = "Component Settings")
+	uint8 bCastShadowAsTwoSided : 1;
+
+	UPROPERTY(EditAnywhere, Category = "Component Settings")
+	uint8 bCastHiddenShadow : 1;
+
+	UPROPERTY(EditAnywhere, Category = "Component Settings")
+	uint8 bAffectDynamicIndirectLighting : 1;
+
+	UPROPERTY(EditAnywhere, Category = "Component Settings")
+	uint8 bAffectDynamicIndirectLightingWhileHidden : 1;
+
+	UPROPERTY(EditAnywhere, Category = "Component Settings")
+	uint8 bAffectDistanceFieldLighting : 1;
+
+	UPROPERTY(EditAnywhere, Category = "Component Settings")
+	uint8 bReceivesDecals : 1;
+
+	UPROPERTY(EditAnywhere, Category = "Component Settings")
+	uint8 bOverrideLightMapRes : 1;
+
+	UPROPERTY(EditAnywhere, Category = "Component Settings")
+	uint8 bUseAsOccluder : 1;
+
+	UPROPERTY(EditAnywhere, Category = "Component Settings")
+	uint8 bEnableDensityScaling : 1;
+
+	UPROPERTY(EditAnywhere, Category = "Component Settings")
+	uint8 bEnableDiscardOnLoad : 1;
+
+	UPROPERTY(EditAnywhere, Category = "Component Settings")
+	uint8 bRenderCustomDepth : 1;
+
+	UPROPERTY(EditAnywhere, Category = "Component Settings")
+	uint8 bVisibleInRayTracing : 1;
+
+	UPROPERTY(EditAnywhere, Category = "Component Settings")
+	uint8 bHiddenInGame : 1;
+
+	UPROPERTY(EditAnywhere, Category = "Component Settings")
+	uint8 bIsEditorOnly : 1;
+
+	UPROPERTY(EditAnywhere, Category = "Component Settings")
+	uint8 bVisible : 1;
+	
+	UPROPERTY(EditAnywhere, Category = "Component Settings")
+	uint8 bEvaluateWorldPositionOffset : 1;
+
+	UPROPERTY(EditAnywhere, Category = "Component Settings")
+	uint8 bReverseCulling : 1;
+
+#if WITH_EDITORONLY_DATA
+	UPROPERTY(EditAnywhere, Category = "Component Settings")
+	uint8 bIncludeInHLOD : 1;
+
+	UPROPERTY(EditAnywhere, Category = "Component Settings")
+	uint8 bConsiderForActorPlacementWhenHidden : 1;
+#endif
+
+	UPROPERTY(EditAnywhere, Category = "Component Settings", meta = (DisplayAfter = "bGenerateOverlapEvents"))
+	uint8 bUseDefaultCollision : 1;
+
+	UPROPERTY(EditAnywhere, Category = "Component Settings", meta = (DisplayAfter = "CustomDepthStencilWriteMask"))
+	uint8 bGenerateOverlapEvents : 1;
+
+	UPROPERTY(EditAnywhere, Category = "Component Settings")
+	uint8 bOverrideNavigationExport : 1;
+
+	UPROPERTY(EditAnywhere, Category = "Component Settings")
+	uint8 bForceNavigationObstacle : 1;
+
+	UPROPERTY(EditAnywhere, Category = "Component Settings")
+	uint8 bFillCollisionUnderneathForNavmesh : 1;
+
+	UPROPERTY(EditAnywhere, Category = "Component Settings")
+	int32 WorldPositionOffsetDisableDistance;
+
+	UPROPERTY(EditAnywhere, Category = "Component Settings")
+	EShadowCacheInvalidationBehavior ShadowCacheInvalidationBehavior;
+
+	UPROPERTY(EditAnywhere, Category = "Component Settings")
+	TEnumAsByte<enum EDetailMode> DetailMode;
+};
+
+USTRUCT()
+struct FSplineMeshComponentDescriptor : public FSplineMeshComponentDescriptorBase
+{
+	GENERATED_BODY()
+
+	ENGINE_API FSplineMeshComponentDescriptor();
+	ENGINE_API explicit FSplineMeshComponentDescriptor(const FSoftSplineMeshComponentDescriptor& Other);
+	static ENGINE_API FSplineMeshComponentDescriptor CreateFrom(const TSubclassOf<UStaticMeshComponent>& ComponentClass);
+
+	ENGINE_API virtual void InitFrom(const UStaticMeshComponent* Component, bool bInitBodyInstance = true) override;
+	ENGINE_API virtual uint32 ComputeHash() const;
+	ENGINE_API virtual void InitComponent(USplineMeshComponent* SplineMeshComponent) const override;
+		
+	ENGINE_API bool operator!=(const FSplineMeshComponentDescriptor& Other) const;
+	ENGINE_API bool operator==(const FSplineMeshComponentDescriptor& Other) const;
+
+	friend inline bool operator<(const FSplineMeshComponentDescriptor& Lhs, const FSplineMeshComponentDescriptor& Rhs)
+	{
+		return Lhs.Hash < Rhs.Hash;
+	}
+
+public:
+	UPROPERTY(EditAnywhere, Category = "Component Settings", meta = (DisplayAfter = "ComponentClass"))
+	TObjectPtr<UStaticMesh> StaticMesh = nullptr;
+
+	UPROPERTY(EditAnywhere, Category = "Component Settings", meta = (DisplayAfter = "ComponentClass"))
+	TArray<TObjectPtr<UMaterialInterface>> OverrideMaterials;
+
+	UPROPERTY(EditAnywhere, Category = "Component Settings", meta = (DisplayAfter = "ComponentClass"))
+	TObjectPtr<UMaterialInterface> OverlayMaterial;
+
+	UPROPERTY(EditAnywhere, Category = "Component Settings", meta = (DisplayAfter = "ComponentClass"))
+	TArray<TObjectPtr<URuntimeVirtualTexture>> RuntimeVirtualTextures;
+};
+
+USTRUCT()
+struct FSoftSplineMeshComponentDescriptor : public FSplineMeshComponentDescriptorBase
+{
+	GENERATED_BODY()
+
+	ENGINE_API FSoftSplineMeshComponentDescriptor();
+	ENGINE_API explicit FSoftSplineMeshComponentDescriptor(const FSplineMeshComponentDescriptor& Other);
+	static ENGINE_API FSoftSplineMeshComponentDescriptor CreateFrom(const TSubclassOf<UStaticMeshComponent>& ComponentClass);
+
+	ENGINE_API virtual void InitFrom(const UStaticMeshComponent* Component, bool bInitBodyInstance = true) override;
+	ENGINE_API virtual uint32 ComputeHash() const;
+	ENGINE_API virtual void InitComponent(USplineMeshComponent* SplineMeshComponent) const override;
+
+	ENGINE_API bool operator!=(const FSoftSplineMeshComponentDescriptor& Other) const;
+	ENGINE_API bool operator==(const FSoftSplineMeshComponentDescriptor& Other) const;
+
+	friend inline bool operator<(const FSoftSplineMeshComponentDescriptor& Lhs, const FSoftSplineMeshComponentDescriptor& Rhs)
+	{
+		return Lhs.Hash < Rhs.Hash;
+	}
+
+public:
+	UPROPERTY(EditAnywhere, Category = "Component Settings", meta = (DisplayAfter = "ComponentClass"))
+	TSoftObjectPtr<UStaticMesh> StaticMesh = nullptr;
+
+	UPROPERTY(EditAnywhere, Category = "Component Settings", meta = (DisplayAfter = "ComponentClass"))
+	TArray<TSoftObjectPtr<UMaterialInterface>> OverrideMaterials;
+
+	UPROPERTY(EditAnywhere, Category = "Component Settings", meta = (DisplayAfter = "ComponentClass"))
+	TSoftObjectPtr<UMaterialInterface> OverlayMaterial;
+
+	UPROPERTY(EditAnywhere, Category = "Component Settings", meta = (DisplayAfter = "ComponentClass"))
+	TArray<TSoftObjectPtr<URuntimeVirtualTexture>> RuntimeVirtualTextures;
+};
diff --git a/Engine/Source/Runtime/Engine/Private/SplineMeshComponentDescriptor.cpp b/Engine/Source/Runtime/Engine/Private/SplineMeshComponentDescriptor.cpp
new file mode 100644
--- /dev/null
+++ b/Engine/Source/Runtime/Engine/Private/SplineMeshComponentDescriptor.cpp
@@ -0,0 +1,445 @@
+// Copyright Epic Games, Inc. All Rights Reserved.
+
+#include "Engine/SplineMeshComponentDescriptor.h"
+
+#include "Algo/Transform.h"
+#include "Components/SplineMeshComponent.h"
+#include "Components/StaticMeshComponent.h"
+#include "Concepts/StaticStructProvider.h"
+#include "Engine/StaticMesh.h"
+#include "Materials/MaterialInstanceDynamic.h"
+#include "Materials/MaterialInterface.h"
+#include "Serialization/ArchiveCrc32.h"
+#include "VT/RuntimeVirtualTexture.h"
+
+#include UE_INLINE_GENERATED_CPP_BY_NAME(SplineMeshComponentDescriptor)
+
+FSplineMeshComponentDescriptorBase::FSplineMeshComponentDescriptorBase()
+{
+	// Note: should not really be used - prefer using FSplineMeshComponentDescriptor & FSoftSplineMeshDescriptor instead
+	InitFrom(USplineMeshComponent::StaticClass()->GetDefaultObject<USplineMeshComponent>());
+}
+
+FSplineMeshComponentDescriptor::FSplineMeshComponentDescriptor()
+	: FSplineMeshComponentDescriptorBase(NoInit)
+{
+	// Make sure we have proper defaults
+	InitFrom(USplineMeshComponent::StaticClass()->GetDefaultObject<USplineMeshComponent>());
+}
+
+FSplineMeshComponentDescriptor::FSplineMeshComponentDescriptor(const FSoftSplineMeshComponentDescriptor& Other)
+	: FSplineMeshComponentDescriptorBase(Other)
+{
+	StaticMesh = Other.StaticMesh.LoadSynchronous();
+	Algo::Transform(Other.OverrideMaterials, OverrideMaterials, [](TSoftObjectPtr<UMaterialInterface> Material) { return Material.LoadSynchronous(); });
+	OverlayMaterial = Other.OverlayMaterial.LoadSynchronous();
+	Algo::Transform(Other.RuntimeVirtualTextures, RuntimeVirtualTextures, [](TSoftObjectPtr<URuntimeVirtualTexture> RVT) { return RVT.LoadSynchronous(); });
+	Hash = Other.Hash;
+}
+
+FSoftSplineMeshComponentDescriptor::FSoftSplineMeshComponentDescriptor()
+	: FSplineMeshComponentDescriptorBase(NoInit)
+{
+	// Make sure we have proper defaults
+	InitFrom(USplineMeshComponent::StaticClass()->GetDefaultObject<USplineMeshComponent>());
+}
+
+FSoftSplineMeshComponentDescriptor::FSoftSplineMeshComponentDescriptor(const FSplineMeshComponentDescriptor& Other)
+	: FSplineMeshComponentDescriptorBase(Other)
+{
+	StaticMesh = Other.StaticMesh;
+	Algo::Transform(Other.OverrideMaterials, OverrideMaterials, [](TObjectPtr<UMaterialInterface> Material) { return Material; });
+	OverlayMaterial = Other.OverlayMaterial;
+	Algo::Transform(Other.RuntimeVirtualTextures, RuntimeVirtualTextures, [](TObjectPtr<URuntimeVirtualTexture> RVT) { return RVT; });
+	Hash = Other.Hash;
+}
+
+FSplineMeshComponentDescriptor FSplineMeshComponentDescriptor::CreateFrom(const TSubclassOf<UStaticMeshComponent>& From)
+{
+	FSplineMeshComponentDescriptor ComponentDescriptor;
+
+	ComponentDescriptor.InitFrom(From->GetDefaultObject<UStaticMeshComponent>());
+	ComponentDescriptor.ComputeHash();
+
+	return ComponentDescriptor;
+}
+
+void FSplineMeshComponentDescriptorBase::InitFrom(const UStaticMeshComponent* Template, bool bInitBodyInstance)
+{
+	check(Template);
+	bEnableDiscardOnLoad = false;
+
+	// Disregard the template class if it does not stem from a spline mesh component
+	if (Template->IsA(USplineMeshComponent::StaticClass()))
+	{
+		ComponentClass = Template->GetClass();
+	}
+
+	Mobility = Template->Mobility;
+	VirtualTextureRenderPassType = Template->VirtualTextureRenderPassType;
+	LightmapType = Template->GetLightmapType();
+	LightingChannels = Template->LightingChannels;
+	RayTracingGroupId = Template->RayTracingGroupId;
+	RayTracingGroupCullingPriority = Template->RayTracingGroupCullingPriority;
+	bHasCustomNavigableGeometry = Template->bHasCustomNavigableGeometry;
+	CustomDepthStencilWriteMask = Template->CustomDepthStencilWriteMask;
+	VirtualTextureCullMips = Template->VirtualTextureCullMips;
+	TranslucencySortPriority = Template->TranslucencySortPriority;
+	OverriddenLightMapRes = Template->OverriddenLightMapRes;
+	CustomDepthStencilValue = Template->CustomDepthStencilValue;
+	bCastShadow = Template->CastShadow;
+	bEmissiveLightSource = Template->bEmissiveLightSource;
+	bCastStaticShadow = Template->bCastStaticShadow;
+	bCastDynamicShadow = Template->bCastDynamicShadow;
+	bCastContactShadow = Template->bCastContactShadow;
+	bCastShadowAsTwoSided = Template->bCastShadowAsTwoSided;
+	bCastHiddenShadow = Template->bCastHiddenShadow;
+	bAffectDynamicIndirectLighting = Template->bAffectDynamicIndirectLighting;
+	bAffectDynamicIndirectLightingWhileHidden = Template->bAffectIndirectLightingWhileHidden;
+	bAffectDistanceFieldLighting = Template->bAffectDistanceFieldLighting;
+	bReceivesDecals = Template->bReceivesDecals;
+	bOverrideLightMapRes = Template->bOverrideLightMapRes;
+	bUseAsOccluder = Template->bUseAsOccluder;
+	bRenderCustomDepth = Template->bRenderCustomDepth;
+	bHiddenInGame = Template->bHiddenInGame;
+	bIsEditorOnly = Template->bIsEditorOnly;
+	bVisible = Template->GetVisibleFlag();
+	bVisibleInRayTracing = Template->bVisibleInRayTracing;
+	bEvaluateWorldPositionOffset = Template->bEvaluateWorldPositionOffset;
+	WorldPositionOffsetDisableDistance = Template->WorldPositionOffsetDisableDistance;
+	ShadowCacheInvalidationBehavior = Template->ShadowCacheInvalidationBehavior;
+	DetailMode = Template->DetailMode;
+	// Determine if this must render with reversed culling based on both scale and the component property
+	const bool bIsLocalToWorldDeterminantNegative = Template->GetRenderMatrix().Determinant() < 0;
+	bReverseCulling = Template->bReverseCulling != bIsLocalToWorldDeterminantNegative;
+	bUseDefaultCollision = Template->bUseDefaultCollision;
+	bGenerateOverlapEvents = Template->GetGenerateOverlapEvents();
+	bOverrideNavigationExport = Template->bOverrideNavigationExport;
+	bForceNavigationObstacle = Template->bForceNavigationObstacle;
+	bFillCollisionUnderneathForNavmesh = Template->bFillCollisionUnderneathForNavmesh;
+
+#if WITH_EDITORONLY_DATA
+	HLODBatchingPolicy = Template->HLODBatchingPolicy;
+	bIncludeInHLOD = Template->bEnableAutoLODGeneration;
+	bConsiderForActorPlacementWhenHidden = Template->bConsiderForActorPlacementWhenHidden;
+#endif
+
+	if (bInitBodyInstance)
+	{
+		BodyInstance.CopyBodyInstancePropertiesFrom(&Template->BodyInstance);
+	}
+}
+
+void FSplineMeshComponentDescriptor::InitFrom(const UStaticMeshComponent* Template, bool bInitBodyInstance)
+{
+	StaticMesh = Template->GetStaticMesh();
+	OverrideMaterials = Template->OverrideMaterials;
+	OverlayMaterial = Template->OverlayMaterial;
+	RuntimeVirtualTextures = Template->RuntimeVirtualTextures;
+
+	Super::InitFrom(Template, bInitBodyInstance);
+}
+
+void FSoftSplineMeshComponentDescriptor::InitFrom(const UStaticMeshComponent* Template, bool bInitBodyInstance)
+{
+	StaticMesh = Template->GetStaticMesh();
+	Algo::Transform(Template->OverrideMaterials, OverrideMaterials, [](TObjectPtr<UMaterialInterface> Material) { return Material; });
+	OverlayMaterial = Template->OverlayMaterial;
+	Algo::Transform(Template->RuntimeVirtualTextures, RuntimeVirtualTextures, [](TObjectPtr<URuntimeVirtualTexture> RVT) { return RVT; });
+
+	Super::InitFrom(Template, bInitBodyInstance);
+}
+
+void FSplineMeshComponentDescriptorBase::PostLoadFixup(UObject* Loader)
+{
+	check(Loader);
+
+	// Necessary to update the collision Response Container from the array
+	BodyInstance.FixupData(Loader);
+}
+
+bool FSplineMeshComponentDescriptorBase::operator!=(const FSplineMeshComponentDescriptorBase& Other) const
+{
+	return !(*this == Other);
+}
+
+bool FSplineMeshComponentDescriptor::operator!=(const FSplineMeshComponentDescriptor& Other) const
+{
+	return !(*this == Other);
+}
+
+bool FSoftSplineMeshComponentDescriptor::operator!=(const FSoftSplineMeshComponentDescriptor& Other) const
+{
+	return !(*this == Other);
+}
+
+bool FSplineMeshComponentDescriptorBase::operator==(const FSplineMeshComponentDescriptorBase& Other) const
+{
+	return ComponentClass == Other.ComponentClass &&
+	Mobility == Other.Mobility &&
+	VirtualTextureRenderPassType == Other.VirtualTextureRenderPassType &&
+	LightmapType == Other.LightmapType &&
+	GetLightingChannelMaskForStruct(LightingChannels) == GetLightingChannelMaskForStruct(Other.LightingChannels) &&
+	RayTracingGroupId == Other.RayTracingGroupId &&
+	RayTracingGroupCullingPriority == Other.RayTracingGroupCullingPriority &&
+	bHasCustomNavigableGeometry == Other.bHasCustomNavigableGeometry &&
+	CustomDepthStencilWriteMask == Other.CustomDepthStencilWriteMask &&
+	VirtualTextureCullMips == Other.VirtualTextureCullMips &&
+	TranslucencySortPriority == Other.TranslucencySortPriority &&
+	OverriddenLightMapRes == Other.OverriddenLightMapRes &&
+	CustomDepthStencilValue == Other.CustomDepthStencilValue &&
+	bCastShadow == Other.bCastShadow &&
+	bEmissiveLightSource == Other.bEmissiveLightSource &&
+	bCastStaticShadow == Other.bCastStaticShadow &&
+	bCastDynamicShadow == Other.bCastDynamicShadow &&
+	bCastContactShadow == Other.bCastContactShadow &&
+	bCastShadowAsTwoSided == Other.bCastShadowAsTwoSided &&
+	bCastHiddenShadow == Other.bCastHiddenShadow &&
+	bAffectDynamicIndirectLighting == Other.bAffectDynamicIndirectLighting &&
+	bAffectDynamicIndirectLightingWhileHidden == Other.bAffectDynamicIndirectLightingWhileHidden &&
+	bAffectDistanceFieldLighting == Other.bAffectDistanceFieldLighting &&
+	bReceivesDecals == Other.bReceivesDecals &&
+	bOverrideLightMapRes == Other.bOverrideLightMapRes &&
+	bUseAsOccluder == Other.bUseAsOccluder &&
+	bRenderCustomDepth == Other.bRenderCustomDepth &&
+	bEnableDiscardOnLoad == Other.bEnableDiscardOnLoad &&
+	bHiddenInGame == Other.bHiddenInGame &&
+	bIsEditorOnly == Other.bIsEditorOnly &&
+	bVisible == Other.bVisible &&
+	bVisibleInRayTracing == Other.bVisibleInRayTracing &&
+	bEvaluateWorldPositionOffset == Other.bEvaluateWorldPositionOffset &&
+	bReverseCulling == Other.bReverseCulling &&
+	bUseDefaultCollision == Other.bUseDefaultCollision &&
+	bGenerateOverlapEvents == Other.bGenerateOverlapEvents &&
+	bOverrideNavigationExport == Other.bOverrideNavigationExport &&
+	bForceNavigationObstacle == Other.bForceNavigationObstacle &&
+	bFillCollisionUnderneathForNavmesh == Other.bFillCollisionUnderneathForNavmesh &&
+	WorldPositionOffsetDisableDistance == Other.WorldPositionOffsetDisableDistance &&
+	ShadowCacheInvalidationBehavior == Other.ShadowCacheInvalidationBehavior &&
+	DetailMode == Other.DetailMode &&
+#if WITH_EDITORONLY_DATA
+	HLODBatchingPolicy == Other.HLODBatchingPolicy &&
+	bIncludeInHLOD == Other.bIncludeInHLOD &&
+	bConsiderForActorPlacementWhenHidden == Other.bConsiderForActorPlacementWhenHidden &&
+#endif // WITH_EDITORONLY_DATA
+	BodyInstance.GetCollisionEnabled() == Other.BodyInstance.GetCollisionEnabled() && 
+	BodyInstance.GetCollisionResponse() == Other.BodyInstance.GetCollisionResponse() &&
+	BodyInstance.DoesUseCollisionProfile() == Other.BodyInstance.DoesUseCollisionProfile() &&
+	(!BodyInstance.DoesUseCollisionProfile() || (BodyInstance.GetCollisionProfileName() == Other.BodyInstance.GetCollisionProfileName()));
+}
+
+bool FSplineMeshComponentDescriptor::operator==(const FSplineMeshComponentDescriptor& Other) const
+{
+	return (Hash == 0 || Other.Hash == 0 || Hash == Other.Hash) && // Check hash first, other checks are in case of Hash collision
+		StaticMesh == Other.StaticMesh &&
+		OverrideMaterials == Other.OverrideMaterials &&
+		OverlayMaterial == Other.OverlayMaterial &&
+		RuntimeVirtualTextures == Other.RuntimeVirtualTextures &&
+		Super::operator==(Other);
+}
+
+bool FSoftSplineMeshComponentDescriptor::operator==(const FSoftSplineMeshComponentDescriptor& Other) const
+{
+	return (Hash == 0 || Other.Hash == 0 || Hash == Other.Hash) && // Check hash first, other checks are in case of Hash collision
+		StaticMesh == Other.StaticMesh &&
+		OverrideMaterials == Other.OverrideMaterials &&
+		OverlayMaterial == Other.OverlayMaterial &&
+		RuntimeVirtualTextures == Other.RuntimeVirtualTextures &&
+		Super::operator==(Other);
+}
+
+uint32 FSplineMeshComponentDescriptorBase::ComputeHash() const
+{
+	FArchiveCrc32 CrcArchive;
+
+	Hash = 0; // we don't want the hash to impact the calculation
+	CrcArchive << *this;
+	Hash = CrcArchive.GetCrc();
+
+	return Hash;
+}
+
+uint32 FSplineMeshComponentDescriptor::ComputeHash() const
+{
+	Super::ComputeHash();
+
+	FSplineMeshComponentDescriptor& MutableSelf = *const_cast<FSplineMeshComponentDescriptor*>(this);
+	FArchiveCrc32 CrcArchive(Hash);
+	CrcArchive << MutableSelf.StaticMesh;
+	CrcArchive << MutableSelf.OverrideMaterials;
+	CrcArchive << MutableSelf.OverlayMaterial;
+	CrcArchive << MutableSelf.RuntimeVirtualTextures;
+	Hash = CrcArchive.GetCrc();
+
+	return Hash;
+}
+
+uint32 FSoftSplineMeshComponentDescriptor::ComputeHash() const
+{
+	Super::ComputeHash();
+
+	FSoftSplineMeshComponentDescriptor& MutableSelf = *const_cast<FSoftSplineMeshComponentDescriptor*>(this);
+	FArchiveCrc32 CrcArchive(Hash);
+	CrcArchive << MutableSelf.StaticMesh;
+	CrcArchive << MutableSelf.OverrideMaterials;
+	CrcArchive << MutableSelf.OverlayMaterial;
+	CrcArchive << MutableSelf.RuntimeVirtualTextures;
+	Hash = CrcArchive.GetCrc();
+
+	return Hash;
+}
+
+USplineMeshComponent* FSplineMeshComponentDescriptorBase::CreateComponent(UObject* Outer, FName Name, EObjectFlags ObjectFlags) const
+{
+	USplineMeshComponent* SplineMeshComponent = NewObject<USplineMeshComponent>(Outer, ComponentClass, Name, ObjectFlags);
+	
+	InitComponent(SplineMeshComponent);
+
+	return SplineMeshComponent;
+}
+
+void FSplineMeshComponentDescriptorBase::InitComponent(USplineMeshComponent* SplineMeshComponent) const
+{
+	SplineMeshComponent->Mobility = Mobility;
+	SplineMeshComponent->VirtualTextureRenderPassType = VirtualTextureRenderPassType;
+	SplineMeshComponent->SetLightmapType(LightmapType);
+	SplineMeshComponent->LightingChannels = LightingChannels;
+	SplineMeshComponent->RayTracingGroupId = RayTracingGroupId;
+	SplineMeshComponent->RayTracingGroupCullingPriority = RayTracingGroupCullingPriority;
+	SplineMeshComponent->bHasCustomNavigableGeometry = bHasCustomNavigableGeometry;
+	SplineMeshComponent->CustomDepthStencilWriteMask = CustomDepthStencilWriteMask;
+	SplineMeshComponent->BodyInstance.CopyBodyInstancePropertiesFrom(&BodyInstance);
+	SplineMeshComponent->VirtualTextureCullMips = VirtualTextureCullMips;
+	SplineMeshComponent->TranslucencySortPriority = TranslucencySortPriority;
+	SplineMeshComponent->OverriddenLightMapRes = OverriddenLightMapRes;
+	SplineMeshComponent->CustomDepthStencilValue = CustomDepthStencilValue;
+	SplineMeshComponent->CastShadow = bCastShadow;
+	SplineMeshComponent->bEmissiveLightSource = bEmissiveLightSource;
+	SplineMeshComponent->bCastStaticShadow = bCastStaticShadow;
+	SplineMeshComponent->bCastDynamicShadow = bCastDynamicShadow;
+	SplineMeshComponent->bCastContactShadow = bCastContactShadow;
+	SplineMeshComponent->bCastShadowAsTwoSided = bCastShadowAsTwoSided;
+	SplineMeshComponent->bCastHiddenShadow = bCastHiddenShadow;
+	SplineMeshComponent->bAffectDynamicIndirectLighting = bAffectDynamicIndirectLighting;
+	SplineMeshComponent->bAffectIndirectLightingWhileHidden = bAffectDynamicIndirectLightingWhileHidden;
+	SplineMeshComponent->bAffectDistanceFieldLighting = bAffectDistanceFieldLighting;
+	SplineMeshComponent->bReceivesDecals = bReceivesDecals;
+	SplineMeshComponent->bOverrideLightMapRes = bOverrideLightMapRes;
+	SplineMeshComponent->bUseAsOccluder = bUseAsOccluder;
+	SplineMeshComponent->bRenderCustomDepth = bRenderCustomDepth;
+	SplineMeshComponent->bHiddenInGame = bHiddenInGame;
+	SplineMeshComponent->bIsEditorOnly = bIsEditorOnly;
+	SplineMeshComponent->SetVisibleFlag(bVisible);
+	SplineMeshComponent->bVisibleInRayTracing = bVisibleInRayTracing;
+	SplineMeshComponent->bEvaluateWorldPositionOffset = bEvaluateWorldPositionOffset;
+	SplineMeshComponent->bReverseCulling = bReverseCulling;
+	SplineMeshComponent->bUseDefaultCollision = bUseDefaultCollision;
+	SplineMeshComponent->SetGenerateOverlapEvents(bGenerateOverlapEvents);
+	SplineMeshComponent->bOverrideNavigationExport = bOverrideNavigationExport;
+	SplineMeshComponent->bForceNavigationObstacle = bForceNavigationObstacle;
+	SplineMeshComponent->bFillCollisionUnderneathForNavmesh = bFillCollisionUnderneathForNavmesh;
+	SplineMeshComponent->WorldPositionOffsetDisableDistance = WorldPositionOffsetDisableDistance;
+	SplineMeshComponent->ShadowCacheInvalidationBehavior = ShadowCacheInvalidationBehavior;
+	SplineMeshComponent->DetailMode = DetailMode;
+	
+#if WITH_EDITORONLY_DATA
+	SplineMeshComponent->HLODBatchingPolicy = HLODBatchingPolicy;
+	SplineMeshComponent->bEnableAutoLODGeneration = bIncludeInHLOD;
+	SplineMeshComponent->bConsiderForActorPlacementWhenHidden = bConsiderForActorPlacementWhenHidden;
+#endif // WITH_EDITORONLY_DATA
+}
+
+void FSplineMeshComponentDescriptor::InitComponent(USplineMeshComponent* SplineMeshComponent) const
+{
+	SplineMeshComponent->SetStaticMesh(StaticMesh);
+
+	auto GetMaterial = [SplineMeshComponent](UMaterialInterface* MaterialInterface)
+	{
+		if (MaterialInterface && !MaterialInterface->IsAsset())
+		{
+			// If the material is equivalent to its parent, just take a reference to its parent rather than making another redundant object 
+			if (UMaterialInstance* Instance = Cast<UMaterialInstance>(MaterialInterface); Instance && Instance->IsRedundant())
+			{
+				MaterialInterface = Instance->Parent;
+			}
+			else
+			{
+				// As override materials are normally outered to their owner component, we need to duplicate them here to make sure we don't create
+				// references to actors in other levels (for packed level instances or HLOD actors).
+				MaterialInterface = DuplicateObject<UMaterialInterface>(MaterialInterface, SplineMeshComponent);
+
+				// If the MID we just duplicated has a nanite override that's also not an asset, duplicate that too
+				UMaterialInstanceDynamic* OverrideMID = Cast<UMaterialInstanceDynamic>(MaterialInterface);
+				UMaterialInterface* NaniteOverride = OverrideMID ? OverrideMID->GetNaniteOverride() : nullptr;
+				if (NaniteOverride && !NaniteOverride->IsAsset())
+				{
+					OverrideMID->SetNaniteOverride(DuplicateObject<UMaterialInterface>(NaniteOverride, SplineMeshComponent));
+				}
+			}
+		}
+
+		return MaterialInterface;
+	};
+
+	SplineMeshComponent->OverrideMaterials.Empty(OverrideMaterials.Num());
+	for (UMaterialInterface* OverrideMaterial : OverrideMaterials)
+	{
+		SplineMeshComponent->OverrideMaterials.Add(GetMaterial(OverrideMaterial));
+	}
+	SplineMeshComponent->OverlayMaterial = GetMaterial(OverlayMaterial);
+	SplineMeshComponent->RuntimeVirtualTextures = RuntimeVirtualTextures;
+
+	Super::InitComponent(SplineMeshComponent);
+}
+
+void FSoftSplineMeshComponentDescriptor::InitComponent(USplineMeshComponent* SplineMeshComponent) const
+{
+	SplineMeshComponent->SetStaticMesh(StaticMesh.LoadSynchronous());
+
+	auto GetMaterial = [SplineMeshComponent](const TSoftObjectPtr<UMaterialInterface>& MaterialInterfacePtr)
+	{
+		UMaterialInterface* MaterialInterface = MaterialInterfacePtr.LoadSynchronous();
+		if (MaterialInterface && !MaterialInterface->IsAsset())
+		{
+			// If the material is equivalent to its parent, just take a reference to its parent rather than making another redundant object 
+			if (UMaterialInstance* Instance = Cast<UMaterialInstance>(MaterialInterface); Instance && Instance->IsRedundant())
+			{
+				MaterialInterface = Instance->Parent;
+			}
+			else
+			{
+				// As override materials are normally outered to their owner component, we need to duplicate them here to make sure we don't create
+				// references to actors in other levels (for packed level instances or HLOD actors).
+				MaterialInterface = DuplicateObject<UMaterialInterface>(MaterialInterface, SplineMeshComponent);
+
+				// If the MID we just duplicated has a nanite override that's also not an asset, duplicate that too
+				UMaterialInstanceDynamic* OverrideMID = Cast<UMaterialInstanceDynamic>(MaterialInterface);
+				UMaterialInterface* NaniteOverride = OverrideMID ? OverrideMID->GetNaniteOverride() : nullptr; 
+				if (NaniteOverride && !NaniteOverride->IsAsset())
+				{
+					OverrideMID->SetNaniteOverride(DuplicateObject<UMaterialInterface>(NaniteOverride, SplineMeshComponent));
+				}
+			}
+		}
+		
+		return MaterialInterface;
+	};
+
+	SplineMeshComponent->OverrideMaterials.Empty(OverrideMaterials.Num());
+	for (const TSoftObjectPtr<UMaterialInterface>& OverrideMaterialPtr : OverrideMaterials)
+	{
+		SplineMeshComponent->OverrideMaterials.Add(GetMaterial(OverrideMaterialPtr));
+	}
+	SplineMeshComponent->OverlayMaterial = GetMaterial(OverlayMaterial);
+
+	SplineMeshComponent->RuntimeVirtualTextures.Empty(RuntimeVirtualTextures.Num());
+	for (const TSoftObjectPtr<URuntimeVirtualTexture>& RuntimeVirtualTexturePtr : RuntimeVirtualTextures)
+	{
+		if (URuntimeVirtualTexture* RuntimeVirtualTexture = RuntimeVirtualTexturePtr.LoadSynchronous())
+		{
+			SplineMeshComponent->RuntimeVirtualTextures.Add(RuntimeVirtualTexture);
+		}
+	}
+
+	Super::InitComponent(SplineMeshComponent);
+}
