[ActorDescContainer] Don't register delegates for non-editor worlds (was causing problems when using PIE travel).

diff --git a/Engine/Source/Runtime/Engine/Private/WorldPartition/WorldPartition.cpp b/Engine/Source/Runtime/Engine/Private/WorldPartition/WorldPartition.cpp
--- a/Engine/Source/Runtime/Engine/Private/WorldPartition/WorldPartition.cpp
+++ b/Engine/Source/Runtime/Engine/Private/WorldPartition/WorldPartition.cpp
@@ -620,8 +620,11 @@ void UWorldPartition::Initialize(UWorld* InWorld, const FTransform& InTransform)
 
 	if (bIsEditor || bIsGame || bIsPIEWorldTravel || bIsDedicatedServer)
 	{
-		FName ContainerPackageName = UActorDescContainerInstance::GetContainerPackageNameFromWorld(OuterWorld);
-		ActorDescContainerInstance = RegisterActorDescContainerInstance(UActorDescContainerInstance::FInitializeParams(ContainerPackageName));
+		UActorDescContainerInstance::FInitializeParams ContainerInitParams(UActorDescContainerInstance::GetContainerPackageNameFromWorld(OuterWorld));
+		ContainerInitParams.SetShouldRegisterEditorDeletages(bIsEditor);
+
+		ActorDescContainerInstance = RegisterActorDescContainerInstance(ContainerInitParams);
+
 		CreateAndInitializeDataLayerManager();
 		InitializeActorDescContainerEditorStreaming(ActorDescContainerInstance);
 	}
@@ -1963,6 +1966,7 @@ UActorDescContainerInstance* UWorldPartition::RegisterActorDescContainerInstance
 		UActorDescContainerInstance::FInitializeParams InitParams(InParams.ContainerPackageName, bCreateContainerInstanceHierarchy);
 		InitParams.ContentBundleGuid = InParams.ContentBundleGuid;
 		InitParams.ExternalDataLayerAsset = InParams.ExternalDataLayerAsset;
+		InitParams.bShouldRegisterEditorDeletages = InParams.bShouldRegisterEditorDeletages;
 		
 		const FWorldDataLayersActorDesc* WorldDataLayerActorsDesc = nullptr;
 		InitParams.FilterActorDescFunc = [this, &WorldDataLayerActorsDesc, &InParams](const FWorldPartitionActorDesc* ActorDesc)
