Removing several stall detector scopes because they don't currently produce enough context in the crash reporter to act on them.

diff --git a/Engine/Source/Editor/UnrealEd/Private/Subsystems/AssetEditorSubsystem.cpp b/Engine/Source/Editor/UnrealEd/Private/Subsystems/AssetEditorSubsystem.cpp
--- a/Engine/Source/Editor/UnrealEd/Private/Subsystems/AssetEditorSubsystem.cpp
+++ b/Engine/Source/Editor/UnrealEd/Private/Subsystems/AssetEditorSubsystem.cpp
@@ -493,8 +493,6 @@ bool UAssetEditorSubsystem::IsAssetEditable(const UObject* Asset)
 
 bool UAssetEditorSubsystem::OpenEditorForAsset(UObject* Asset, const EToolkitMode::Type ToolkitMode, TSharedPtr< IToolkitHost > OpenedFromLevelEditor, const bool bShowProgressWindow, const EAssetTypeActivationOpenedMethod OpenedMethod)
 {
-	SCOPE_STALL_REPORTER(UAssetEditorSubsystem::OpenEditorForAsset, 2.0);
-
 	FText ErrorMessage;
 	if(!CanOpenEditorForAsset(Asset, OpenedMethod, &ErrorMessage))
 	{
diff --git a/Engine/Source/Runtime/Engine/Private/AsyncCompilationHelpers.cpp b/Engine/Source/Runtime/Engine/Private/AsyncCompilationHelpers.cpp
--- a/Engine/Source/Runtime/Engine/Private/AsyncCompilationHelpers.cpp
+++ b/Engine/Source/Runtime/Engine/Private/AsyncCompilationHelpers.cpp
@@ -94,7 +94,6 @@ namespace AsyncCompilationHelpers
 		TFunctionRef<void(ICompilable*)> PostCompileSingle
 	)
 	{
-		SCOPE_STALL_REPORTER(AsyncCompilationHelpers::FinishCompilation, 10.0);
 		if (Num <= 0)
 		{
 			return;
diff --git a/Engine/Source/Runtime/Engine/Private/ShaderCompiler/ShaderCompiler.cpp b/Engine/Source/Runtime/Engine/Private/ShaderCompiler/ShaderCompiler.cpp
--- a/Engine/Source/Runtime/Engine/Private/ShaderCompiler/ShaderCompiler.cpp
+++ b/Engine/Source/Runtime/Engine/Private/ShaderCompiler/ShaderCompiler.cpp
@@ -6241,7 +6241,6 @@ void FShaderCompilingManager::BlockOnShaderMapCompletion(const TArray<int32>& Sh
 	}
 
 	TRACE_CPUPROFILER_EVENT_SCOPE(FShaderCompilingManager::BlockOnShaderMapCompletion);
-	SCOPE_STALL_REPORTER(FShaderCompilingManager::BlockOnShaderMapCompletion, 2.0);
 
 	COOK_STAT(FScopedDurationTimer BlockingTimer(ShaderCompilerCookStats::BlockingTimeSec));
 	if (bAllowAsynchronousShaderCompiling)
@@ -6375,7 +6374,6 @@ void FShaderCompilingManager::BlockOnAllShaderMapCompletion(TMap<int32, FShaderM
 	}
 
 	TRACE_CPUPROFILER_EVENT_SCOPE(FShaderCompilingManager::BlockOnAllShaderMapCompletion);
-	SCOPE_STALL_REPORTER(FShaderCompilingManager::BlockOnAllShaderMapCompletion, 2.0);
 
 	COOK_STAT(FScopedDurationTimer BlockingTimer(ShaderCompilerCookStats::BlockingTimeSec));
 	if (bAllowAsynchronousShaderCompiling)
