Propagate AlwaysMarkDirty optional param when breaking node links (fixes MetaSounds always redirtying as they have a transient, procedurally built graph that should not alter save state)
#tests resave preset as described in bug and no redirty observedx
#rnx

diff --git a/Engine/Source/Runtime/Engine/Classes/EdGraph/EdGraph.h b/Engine/Source/Runtime/Engine/Classes/EdGraph/EdGraph.h
--- a/Engine/Source/Runtime/Engine/Classes/EdGraph/EdGraph.h
+++ b/Engine/Source/Runtime/Engine/Classes/EdGraph/EdGraph.h
@@ -160,10 +160,11 @@ public:
 	* 
 	* @param NodeToRemove		The node to remove from this graph
 	* @param bBreakAllLinks		If true, all links will be broken on the given node. Editor only. Useful for moving nodes to a different graph.
+	* @param bAlwaysMarkDirty	(Optional) If true, marks graph as dirtied upon removal.
 	* 
 	* @return True if the node has been removed from the graph
 	*/
-	ENGINE_API bool RemoveNode( UEdGraphNode* NodeToRemove, bool bBreakAllLinks = true, bool bMarkDirty = true );
+	ENGINE_API bool RemoveNode( UEdGraphNode* NodeToRemove, bool bBreakAllLinks = true, bool bAlwaysMarkDirty = true );
 
 	/** Signal to listeners that the graph has changed - prefer to use NotifyNodeChanged when updating a single node */
 	ENGINE_API virtual void NotifyGraphChanged();
diff --git a/Engine/Source/Runtime/Engine/Classes/EdGraph/EdGraphNode.h b/Engine/Source/Runtime/Engine/Classes/EdGraph/EdGraphNode.h
--- a/Engine/Source/Runtime/Engine/Classes/EdGraph/EdGraphNode.h
+++ b/Engine/Source/Runtime/Engine/Classes/EdGraph/EdGraphNode.h
@@ -648,7 +648,7 @@ public:
 	ENGINE_API UEdGraphPin* GetPinWithDirectionAt(int32 Index, EEdGraphPinDirection PinDirection) const;
 
 	/** Break all links on this node */
-	ENGINE_API void BreakAllNodeLinks();
+	ENGINE_API void BreakAllNodeLinks(bool bAlwaysMarkDirty = true);
 
 	/** Snap this node to a specified grid size */
 	ENGINE_API void SnapToGrid(uint32 GridSnapSize);
diff --git a/Engine/Source/Runtime/Engine/Private/EdGraph/EdGraph.cpp b/Engine/Source/Runtime/Engine/Private/EdGraph/EdGraph.cpp
--- a/Engine/Source/Runtime/Engine/Private/EdGraph/EdGraph.cpp
+++ b/Engine/Source/Runtime/Engine/Private/EdGraph/EdGraph.cpp
@@ -257,15 +257,15 @@ void UEdGraph::SelectNodeSet(TSet<const UEdGraphNode*> NodeSelection, bool bFrom
 	NotifyGraphChanged(SelectionAction);
 }
 
-bool UEdGraph::RemoveNode( UEdGraphNode* NodeToRemove, bool bBreakAllLinks, bool bMarkDirty )
+bool UEdGraph::RemoveNode( UEdGraphNode* NodeToRemove, bool bBreakAllLinks, bool bAlwaysMarkDirty )
 {
-	Modify(bMarkDirty);
+	Modify(bAlwaysMarkDirty);
 
 	int32 NumTimesNodeRemoved = Nodes.Remove(NodeToRemove);
 #if WITH_EDITOR
 	if (bBreakAllLinks)
 	{
-		NodeToRemove->BreakAllNodeLinks();
+		NodeToRemove->BreakAllNodeLinks(bAlwaysMarkDirty);
 	}
 #endif	//#if WITH_EDITOR
 
diff --git a/Engine/Source/Runtime/Engine/Private/EdGraph/EdGraphNode.cpp b/Engine/Source/Runtime/Engine/Private/EdGraph/EdGraphNode.cpp
--- a/Engine/Source/Runtime/Engine/Private/EdGraph/EdGraphNode.cpp
+++ b/Engine/Source/Runtime/Engine/Private/EdGraph/EdGraphNode.cpp
@@ -431,7 +431,7 @@ bool UEdGraphNode::RemovePin(UEdGraphPin* Pin)
 	return false;
 }
 
-void UEdGraphNode::BreakAllNodeLinks()
+void UEdGraphNode::BreakAllNodeLinks(bool bAlwaysMarkDirty)
 {
 	TSet<UEdGraphNode*> NodeList;
 
@@ -451,7 +451,8 @@ void UEdGraphNode::BreakAllNodeLinks()
 				}
 			}
 
-			Pin->BreakAllPinLinks();
+			constexpr bool bNotifyNodes = false;
+			Pin->BreakAllPinLinks(bNotifyNodes, bAlwaysMarkDirty);
 		}
 	}
 
