[Material Editor] Fixed crash in FHLSLMaterialTranslator::InExpressionConstantValue() triggering when a non-float constant is evaluated.

#tests Editor, linking float and non-float constant/non constant values to a Mul operator does not cause crash

diff --git a/Engine/Source/Runtime/Engine/Private/Materials/HLSLMaterialTranslator.cpp b/Engine/Source/Runtime/Engine/Private/Materials/HLSLMaterialTranslator.cpp
--- a/Engine/Source/Runtime/Engine/Private/Materials/HLSLMaterialTranslator.cpp
+++ b/Engine/Source/Runtime/Engine/Private/Materials/HLSLMaterialTranslator.cpp
@@ -8957,6 +8957,7 @@ bool FHLSLMaterialTranslator::IsExpressionConstantValue(int Code, float Constant
 	if (GetConstParameterValue(Expression, Value))
 	{
 		EMaterialValueType CodeType = (*CurrentScopeChunks)[Code].Type;
+		int64_t ConstantValueInt = (int64_t)ConstantValue;
 
 		switch (CodeType)
 		{
@@ -8969,8 +8970,17 @@ bool FHLSLMaterialTranslator::IsExpressionConstantValue(int Code, float Constant
 			return Value.R == ConstantValue && Value.G == ConstantValue && Value.B == ConstantValue;
 		case MCT_Float4:
 			return Value.R == ConstantValue && Value.G == ConstantValue && Value.B == ConstantValue && Value.A == ConstantValue;
+		case MCT_Bool:
+		case MCT_UInt:
+		case MCT_UInt1:
+			return (int64_t)Value.R == ConstantValueInt;
+		case MCT_UInt2:
+			return (int64_t)Value.R == ConstantValueInt && (int64_t)Value.G == ConstantValueInt;
+		case MCT_UInt3:
+			return (int64_t)Value.R == ConstantValueInt && (int64_t)Value.G == ConstantValueInt && (int64_t)Value.B == ConstantValueInt;
+		case MCT_UInt4:
+			return (int64_t)Value.R == ConstantValueInt && (int64_t)Value.G == ConstantValueInt && (int64_t)Value.B == ConstantValueInt && (int64_t)Value.A == ConstantValueInt;
 		default:
-			ensureAlwaysMsgf(false, TEXT("FHLSLMaterialTranslator::IsExpressionConstantValue - Unknown constant expression type"));
 			return false;
 		}
 	}
