Fix event dispatcher ASAN crash

diff --git a/Engine/Source/Runtime/Engine/Private/PhysicsEngine/Experimental/PhysScene_Chaos.cpp b/Engine/Source/Runtime/Engine/Private/PhysicsEngine/Experimental/PhysScene_Chaos.cpp
--- a/Engine/Source/Runtime/Engine/Private/PhysicsEngine/Experimental/PhysScene_Chaos.cpp
+++ b/Engine/Source/Runtime/Engine/Private/PhysicsEngine/Experimental/PhysScene_Chaos.cpp
@@ -1024,7 +1024,7 @@ void FPhysScene_Chaos::HandleCollisionEvents(const Chaos::FCollisionEventData& E
 		// Iterate through the smallest between events registration and physics proxies
 		if (PhysicsProxyToCollisionIndicesMap.Num() <= CollisionEventRegistrations.Num())
 		{
-			for (TPair<IPhysicsProxyBase*, TArray<int32>> Pair : PhysicsProxyToCollisionIndicesMap)
+			for (const TPair<IPhysicsProxyBase*, TArray<int32>>& Pair : PhysicsProxyToCollisionIndicesMap)
 			{
 				IPhysicsProxyBase*  PhysicsProxy0 = Pair.Key;
 				const TArray<int32>& CollisionIndices = Pair.Value;
diff --git a/Engine/Source/Runtime/Experimental/Chaos/Public/EventManager.h b/Engine/Source/Runtime/Experimental/Chaos/Public/EventManager.h
--- a/Engine/Source/Runtime/Experimental/Chaos/Public/EventManager.h
+++ b/Engine/Source/Runtime/Experimental/Chaos/Public/EventManager.h
@@ -446,13 +446,13 @@ private:
 		template<typename PayloadType>
 		void ClearEvents(const EEventType& EventType, TFunction<void(PayloadType & EventData)> InFunction)
 		{
-			ContainerLock.ReadLock();
+			ContainerLock.WriteLock();
 
 			if (TEventContainer<PayloadType>* EventContainer = StaticCast<TEventContainer<PayloadType>*>(EventContainers[FEventID(EventType)]))
 			{
 				EventContainer->DestroyStaleEvents(InFunction);
 			}
-			ContainerLock.ReadUnlock();
+			ContainerLock.WriteUnlock();
 		}
 
 		/**
