[PSO]
 - Clear strong ptr from worker now that it's supported

#rnx

diff --git a/Engine/Source/Runtime/Engine/Private/PSOPrecacheMaterial.cpp b/Engine/Source/Runtime/Engine/Private/PSOPrecacheMaterial.cpp
--- a/Engine/Source/Runtime/Engine/Private/PSOPrecacheMaterial.cpp
+++ b/Engine/Source/Runtime/Engine/Private/PSOPrecacheMaterial.cpp
@@ -94,35 +94,6 @@ FPSOPrecacheRequestResultArray RequestPrecachePSOs(const FPSOPrecacheDataArray&
 	return Results;
 }
 
-/**
- * Helper task used to release the strong object reference to the material interface on the game thread
- * The release has to happen on the gamethread and the material interface can't be GCd while the PSO
- * collection is happening because it touches the material resources
- */
-class FMaterialInterfaceReleaseTask
-{
-public:
-	explicit FMaterialInterfaceReleaseTask(TStrongObjectPtr<UMaterialInterface>* InMaterialInterface)
-		: MaterialInterface(InMaterialInterface)
-	{
-	}
-
-	void DoTask(ENamedThreads::Type CurrentThread, const FGraphEventRef& MyCompletionGraphEvent)
-	{
-		check(IsInGameThread());
-		delete MaterialInterface;
-	}
-
-public:
-
-	TStrongObjectPtr<UMaterialInterface>* MaterialInterface;
-
-	static ESubsequentsMode::Type	GetSubsequentsMode() { return ESubsequentsMode::FireAndForget; }
-	ENamedThreads::Type				GetDesiredThread() { return ENamedThreads::GameThread; }
-	FORCEINLINE TStatId				GetStatId() const { return TStatId(); }
-};
-
-
 /**
  * Helper task used to offload the PSO collection from the GameThread. The shader decompression
  * takes too long to run this on the GameThread and it isn't blocking anything crucial.
@@ -520,7 +491,7 @@ void FMaterialPSOPrecacheCollectionTask::DoTask(ENamedThreads::Type CurrentThrea
 	GMaterialPSORequestManager.MarkCollectionComplete(PrecacheParams, PSOPrecacheData, PrecacheResults, RequestLifecycleID);
 
 	// Won't touch the material interface anymore - PSO compile jobs take refs to all RHI resources while creating the task
-	TGraphTask<FMaterialInterfaceReleaseTask>::CreateTask().ConstructAndDispatchWhenReady(MaterialInterface);
+	MaterialInterface->Reset();
 
 	// Extend MyCompletionGraphEvent to wait for all the async compile events
 	if (PrecacheResults.Num() > 0)
