#UE Disable assert for now. It is theoretically possible for this delegate to be fired after the request id has been set to INDEX_NONE

diff --git a/Engine/Source/Runtime/Engine/Private/WorldPartition/WorldPartitionActorDescInstance.cpp b/Engine/Source/Runtime/Engine/Private/WorldPartition/WorldPartitionActorDescInstance.cpp
--- a/Engine/Source/Runtime/Engine/Private/WorldPartition/WorldPartitionActorDescInstance.cpp
+++ b/Engine/Source/Runtime/Engine/Private/WorldPartition/WorldPartitionActorDescInstance.cpp
@@ -126,7 +126,7 @@ bool FWorldPartitionActorDescInstance::StartAsyncLoad()
 
 		AsyncLoadID = LoadPackageAsync(PackagePath, PackageName, FLoadPackageAsyncDelegate::CreateLambda([this, ActorPackage](const FName& PackageName, UPackage* Package, EAsyncLoadingResult::Type Result)
 		{
-			check(AsyncLoadID != INDEX_NONE);
+			//check(AsyncLoadID != INDEX_NONE);
 			AsyncLoadID = INDEX_NONE;
 
 			if ((Result != EAsyncLoadingResult::Succeeded) || !Package)
