Changes to critical section usage in FRayTracingGeometryManager.
- separate FCriticalSection to prevent race conditions when managing registered geometries/groups, instead of reusing the existing RequestCS.
- removed critical sections from FRayTracingGeometryManager::Tick(...) since it runs on render thread after all views are rendered.

diff --git a/Engine/Source/Runtime/Engine/Private/Rendering/RayTracingGeometryManager.cpp b/Engine/Source/Runtime/Engine/Private/Rendering/RayTracingGeometryManager.cpp
--- a/Engine/Source/Runtime/Engine/Private/Rendering/RayTracingGeometryManager.cpp
+++ b/Engine/Source/Runtime/Engine/Private/Rendering/RayTracingGeometryManager.cpp
@@ -12,15 +12,6 @@
 #include "RayTracingGeometry.h"
 #include "RenderUtils.h"
 
-/*
-
-TODO:
-- Investigate whether it's necessary to lock RequestCS to access RegisteredGeometries, etc
-	- this lock should only be used by GeometryBuildRequests and related logic
-	- either we can avoid the lock or use a different one specific to that
-
-*/
-
 #if RHI_RAYTRACING
 
 static bool bHasRayTracingEnableChanged = false;
@@ -137,7 +128,7 @@ void FRayTracingGeometryManager::RemoveBuildRequest(BuildRequestIndex InRequestI
 
 RayTracing::GeometryGroupHandle FRayTracingGeometryManager::RegisterRayTracingGeometryGroup(uint32 NumLODs)
 {
-	checkf(IsInRenderingThread(), TEXT("Can only access RegisteredGroups on render thread otherwise need a critical section"));
+	FScopeLock ScopeLock(&MainCS);
 
 	FRayTracingGeometryGroup Group;
 	Group.Geometries.AddDefaulted(NumLODs);
@@ -151,7 +142,7 @@ RayTracing::GeometryGroupHandle FRayTracingGeometryManager::RegisterRayTracingGe
 
 void FRayTracingGeometryManager::ReleaseRayTracingGeometryGroup(RayTracing::GeometryGroupHandle Handle)
 {
-	checkf(IsInRenderingThread(), TEXT("Can only access RegisteredGroups on render thread otherwise need a critical section"));
+	FScopeLock ScopeLock(&MainCS);
 
 	check(RegisteredGroups.IsValidIndex(Handle));
 
@@ -170,12 +161,13 @@ void FRayTracingGeometryManager::ReleaseRayTracingGeometryGroup(RayTracing::Geom
 }
 
 FRayTracingGeometryManager::RayTracingGeometryHandle FRayTracingGeometryManager::RegisterRayTracingGeometry(FRayTracingGeometry* InGeometry)
-{	
+{
 	if (GetRayTracingMode() == ERayTracingMode::Dynamic)
 	{
 		check(InGeometry);
+		
+		FScopeLock ScopeLock(&MainCS);
 
-		FScopeLock ScopeLock(&RequestCS);
 		RayTracingGeometryHandle Handle = RegisteredGeometries.Add(InGeometry);
 
 		if (InGeometry->GroupHandle != INDEX_NONE)
@@ -194,6 +186,7 @@ FRayTracingGeometryManager::RayTracingGeometryHandle FRayTracingGeometryManager:
 
 		return Handle;
 	}
+
 	return INDEX_NONE;
 }
 
@@ -202,7 +195,9 @@ void FRayTracingGeometryManager::ReleaseRayTracingGeometryHandle(RayTracingGeome
 	if (GetRayTracingMode() == ERayTracingMode::Dynamic)
 	{
 		check(Handle != INDEX_NONE);
-		FScopeLock ScopeLock(&RequestCS);
+
+		FScopeLock ScopeLock(&MainCS);
+
 		RegisteredGeometries.RemoveAt(Handle);
 		ReferencedGeometryHandles.Remove(Handle);
 
@@ -217,6 +212,8 @@ void FRayTracingGeometryManager::PreRender()
 
 void FRayTracingGeometryManager::Tick(FRHICommandList& RHICmdList)
 {
+	check(IsInRenderingThread());
+
 	if (GetRayTracingMode() != ERayTracingMode::Dynamic)
 	{
 		return;
@@ -239,7 +236,6 @@ void FRayTracingGeometryManager::Tick(FRHICommandList& RHICmdList)
 		if (bHasRayTracingEnableChanged)
 		{
 			// evict all geometries
-			FScopeLock ScopeLock(&RequestCS);
 			for (FRayTracingGeometry* Geometry : RegisteredGeometries)
 			{
 				if (Geometry->GetRHI() != nullptr)
@@ -252,7 +248,6 @@ void FRayTracingGeometryManager::Tick(FRHICommandList& RHICmdList)
 		{
 #if DO_CHECK
 			// otherwise just check that everything is evicted
-			FScopeLock ScopeLock(&RequestCS);
 			for (FRayTracingGeometry* Geometry : RegisteredGeometries)
 			{
 				checkf(Geometry->IsEvicted() || Geometry->GetRHI() == nullptr, TEXT("Ray tracing geometry should be evicted when ray tracing is disabled."));
@@ -274,8 +269,6 @@ void FRayTracingGeometryManager::Tick(FRHICommandList& RHICmdList)
 
 		bRenderedFrame = false;
 
-		FScopeLock ScopeLock(&RequestCS);
-
 		TSet<FRayTracingGeometry*> NotReferencedResidentGeometries; // TODO: Keep track of this to avoid the following loop
 
 		for (FRayTracingGeometry* Geometry : RegisteredGeometries)
@@ -340,7 +333,6 @@ void FRayTracingGeometryManager::Tick(FRHICommandList& RHICmdList)
 		if (bHasRayTracingEnableChanged)
 		{
 			// make all geometries resident
-			FScopeLock ScopeLock(&RequestCS);
 			for (FRayTracingGeometry* Geometry : RegisteredGeometries)
 			{
 				if (Geometry->IsEvicted())
@@ -353,7 +345,6 @@ void FRayTracingGeometryManager::Tick(FRHICommandList& RHICmdList)
 		{
 #if DO_CHECK
 			// otherwise just check that all geometries are resident
-			FScopeLock ScopeLock(&RequestCS);
 			for (FRayTracingGeometry* Geometry : RegisteredGeometries)
 			{
 				checkf(!Geometry->IsEvicted(), TEXT("Ray tracing geometry should not be evicted when ray tracing is enabled."));
diff --git a/Engine/Source/Runtime/Engine/Public/Rendering/RayTracingGeometryManager.h b/Engine/Source/Runtime/Engine/Public/Rendering/RayTracingGeometryManager.h
--- a/Engine/Source/Runtime/Engine/Public/Rendering/RayTracingGeometryManager.h
+++ b/Engine/Source/Runtime/Engine/Public/Rendering/RayTracingGeometryManager.h
@@ -69,6 +69,10 @@ private:
 	TArray<FBuildRequest> SortedRequests;
 	TArray<FRayTracingGeometryBuildParams> BuildParams;
 
+	// Operations such as registering geometry/groups can be done from different render command pipes (eg: SkeletalMesh)
+	// so need to use critical section in relevant functions
+	FCriticalSection MainCS;
+
 	struct FRayTracingGeometryGroup
 	{
 		TArray<FRayTracingGeometry*> Geometries;
