Fixed PSO validation logging of mesh pass names on mobile platforms

diff --git a/Engine/Source/Runtime/Engine/Private/PSOPrecacheValidation.cpp b/Engine/Source/Runtime/Engine/Private/PSOPrecacheValidation.cpp
--- a/Engine/Source/Runtime/Engine/Private/PSOPrecacheValidation.cpp
+++ b/Engine/Source/Runtime/Engine/Private/PSOPrecacheValidation.cpp
@@ -82,7 +82,7 @@ void ConditionalBreakOnPSOPrecacheMaterial(const FMaterial& Material, int32 PSOC
 {
 	if (!GPSOPrecachingBreakOnMaterialName.IsEmpty())
 	{
-		int PSOCollectorIndexToDebug = FPSOCollectorCreateManager::GetIndex(EShadingPath::Deferred, *GPSOPrecachingBreakOnPassName);
+		int PSOCollectorIndexToDebug = FPSOCollectorCreateManager::GetIndex(GetFeatureLevelShadingPath(GMaxRHIFeatureLevel), *GPSOPrecachingBreakOnPassName);
 		FString MaterialName = Material.GetAssetName();
 		if (MaterialName == GPSOPrecachingBreakOnMaterialName && PSOCollectorIndex == PSOCollectorIndexToDebug)
 		{
@@ -448,7 +448,7 @@ static void LogGeneralPSOMissInfo(
 #if MESH_DRAW_COMMAND_STATS
 	StringBuilder.Appendf(TEXT("\n\tMDCStatsCategory:\t\t%s"), PrimitiveSceneProxy ? *PrimitiveSceneProxy->GetMeshDrawCommandStatsCategory().ToString() : TEXT("Unknown"));
 #endif // MESH_DRAW_COMMAND_STATS
-	StringBuilder.Appendf(TEXT("\n\tPassName:\t\t\t\t%s"), FPSOCollectorCreateManager::GetName(EShadingPath::Deferred, PSOCollectorIndex));
+	StringBuilder.Appendf(TEXT("\n\tPassName:\t\t\t\t%s"), FPSOCollectorCreateManager::GetName(GetFeatureLevelShadingPath(GMaxRHIFeatureLevel), PSOCollectorIndex));
 	StringBuilder.Appendf(TEXT("\n\tShader Hashes:"));
 	const auto LogShaderInfo = [&](const TCHAR* ShaderTypeName, FRHIShader* RHIShader)
 		{
@@ -787,7 +787,7 @@ void LogMinimalPSOStateMissInfo(
 	StringBuilder.Appendf(TEXT("\n\n\tShadersOnly precache information:"));
 	StringBuilder.Appendf(TEXT("\n\tMaterial:\t\t\t\t%s"), *PrecachedMaterialName);
 	StringBuilder.Appendf(TEXT("\n\tVertexFactoryType:\t\t%s"), PrecachedVertexFactoryType ? PrecachedVertexFactoryType->GetName() : TEXT("None"));
-	StringBuilder.Appendf(TEXT("\n\tPassName:\t\t\t\t%s"), FPSOCollectorCreateManager::GetName(EShadingPath::Deferred, PrecachedPSOCollectorIndex));
+	StringBuilder.Appendf(TEXT("\n\tPassName:\t\t\t\t%s"), FPSOCollectorCreateManager::GetName(GetFeatureLevelShadingPath(GMaxRHIFeatureLevel), PrecachedPSOCollectorIndex));
 
 	StringBuilder.Appendf(TEXT("\n\n\tMissed Info:"));
 
@@ -799,7 +799,7 @@ void LogMinimalPSOStateMissInfo(
 		{
 			StringBuilder.Appendf(TEXT("\n\t\t- Found PSO With same shaders & different state:"));
 			StringBuilder.Appendf(TEXT("\n\t\t\tVertexFactoryType:\t\t%s"), PSOPrecacheData.VertexFactoryType ? PSOPrecacheData.VertexFactoryType->GetName() : TEXT("None"));
-			StringBuilder.Appendf(TEXT("\n\t\t\tPassName:\t\t\t\t%s"), FPSOCollectorCreateManager::GetName(EShadingPath::Deferred, PSOPrecacheData.PSOCollectorIndex));
+			StringBuilder.Appendf(TEXT("\n\t\t\tPassName:\t\t\t\t%s"), FPSOCollectorCreateManager::GetName(GetFeatureLevelShadingPath(GMaxRHIFeatureLevel), PSOPrecacheData.PSOCollectorIndex));
 			StringBuilder.Appendf(TEXT("\n\t\t  Differences:"));
 
 			CompareVertexDeclarationAndLogChanges(PSOPrecacheData.GraphicsPSOInitializer.BoundShaderState.VertexDeclarationRHI, Initializer.BoundShaderState.VertexDeclarationRHI, StringBuilder);
@@ -965,7 +965,7 @@ void LogPSOMissInfo(
 	StringBuilder.Appendf(TEXT("\n\tType:\t\t\t\t\t%s"), TEXT("Compute"));
 	StringBuilder.Appendf(TEXT("\n\tPSOPrecachingState:\t\t%s"), LexToString(PrecacheResult));
 	StringBuilder.Appendf(TEXT("\n\tMaterial:\t\t\t\t%s"), Material ? *Material->GetAssetName() : TEXT("Unknown"));
-	StringBuilder.Appendf(TEXT("\n\tPassName:\t\t\t\t%s"), FPSOCollectorCreateManager::GetName(EShadingPath::Deferred, PSOCollectorIndex));
+	StringBuilder.Appendf(TEXT("\n\tPassName:\t\t\t\t%s"), FPSOCollectorCreateManager::GetName(GetFeatureLevelShadingPath(GMaxRHIFeatureLevel), PSOCollectorIndex));
 	StringBuilder.Appendf(TEXT("\n\tCompute Shader Hash:\t%s"), *(ComputeShader.GetHash().ToString()));
 
 	// Not sure yet if this is interesting data or not
diff --git a/Engine/Source/Runtime/Engine/Public/PSOPrecache.h b/Engine/Source/Runtime/Engine/Public/PSOPrecache.h
--- a/Engine/Source/Runtime/Engine/Public/PSOPrecache.h
+++ b/Engine/Source/Runtime/Engine/Public/PSOPrecache.h
@@ -220,8 +220,8 @@ struct FPSOPrecacheData
 	};
 
 #if PSO_PRECACHING_VALIDATE
-	int32 PSOCollectorIndex : 31;
-	int32 bDefaultMaterial : 1;
+	uint32 PSOCollectorIndex : 31;
+	uint32 bDefaultMaterial : 1;
 	const FVertexFactoryType* VertexFactoryType;
 #endif // PSO_PRECACHING_VALIDATE
 };
