Undo some of the changes from CL# 33544148

diff --git a/Engine/Source/Runtime/Engine/Public/PSOPrecache.h b/Engine/Source/Runtime/Engine/Public/PSOPrecache.h
--- a/Engine/Source/Runtime/Engine/Public/PSOPrecache.h
+++ b/Engine/Source/Runtime/Engine/Public/PSOPrecache.h
@@ -220,8 +220,8 @@ struct FPSOPrecacheData
 	};
 
 #if PSO_PRECACHING_VALIDATE
-	uint32 PSOCollectorIndex : 31;
-	uint32 bDefaultMaterial : 1;
+	int32 PSOCollectorIndex : 31;
+	int32 bDefaultMaterial : 1;
 	const FVertexFactoryType* VertexFactoryType;
 #endif // PSO_PRECACHING_VALIDATE
 };
