Fixed WorldPartition levels to respect the 'LevelStreaming.ShouldReuseUnloadedButStillAroundLevels' CVar in PIE.

diff --git a/Engine/Source/Runtime/Engine/Classes/Engine/Level.h b/Engine/Source/Runtime/Engine/Classes/Engine/Level.h
--- a/Engine/Source/Runtime/Engine/Classes/Engine/Level.h
+++ b/Engine/Source/Runtime/Engine/Classes/Engine/Level.h
@@ -591,6 +591,9 @@ public:
 private:
 	/** Whether the level is set not to be reusable after unload (editor-only) */
 	uint8										bForceCantReuseUnloadedButStillAround:1;
+
+	/** Whether the level will force trash its packages in CleanupLevel (including object packages of objects outered to the level) (editor-only) */
+	uint8										bForcePackageTrashingAtCleanup:1;
 #endif
 	
 public:
diff --git a/Engine/Source/Runtime/Engine/Private/Level.cpp b/Engine/Source/Runtime/Engine/Private/Level.cpp
--- a/Engine/Source/Runtime/Engine/Private/Level.cpp
+++ b/Engine/Source/Runtime/Engine/Private/Level.cpp
@@ -414,6 +414,7 @@ ULevel::ULevel( const FObjectInitializer& ObjectInitializer )
 	bPromptWhenAddingToLevelOutsideBounds = true;
 	bUseActorFolders = false;
 	bFixupActorFoldersAtLoad = IsActorFolderObjectsFeatureAvailable();
+	bForcePackageTrashingAtCleanup = false;
 	bForceCantReuseUnloadedButStillAround = false;
 #endif	
 	bActorClusterCreated = false;
@@ -464,7 +465,14 @@ void ULevel::CleanupLevel(bool bCleanupResources, bool bUnloadFromEditor)
 		}
 	}
 
-	const bool bTrashPackage = !ULevelStreaming::ShouldReuseUnloadedButStillAroundLevels(this);
+	bool bTrashPackage = !ULevelStreaming::ShouldReuseUnloadedButStillAroundLevels(this);
+#if WITH_EDITOR
+	if (bForcePackageTrashingAtCleanup)
+	{
+		bTrashPackage = true;
+	}
+#endif
+
 	TSet<UPackage*> ProcessedPackages;
 	auto ProcessPackage = [&ProcessedPackages, bTrashPackage](UPackage* InPackage, bool bInClearStandaloneFlag = false)
 	{
diff --git a/Engine/Source/Runtime/Engine/Private/WorldPartition/WorldPartitionLevelHelper.cpp b/Engine/Source/Runtime/Engine/Private/WorldPartition/WorldPartitionLevelHelper.cpp
--- a/Engine/Source/Runtime/Engine/Private/WorldPartition/WorldPartitionLevelHelper.cpp
+++ b/Engine/Source/Runtime/Engine/Private/WorldPartition/WorldPartitionLevelHelper.cpp
@@ -883,4 +883,9 @@ bool FWorldPartitionLevelHelper::LoadActorsInternal(FLoadActorsParams&& InParams
 	return (LoadProgress->NumPendingLoadRequests == 0);
 }
 
+void FWorldPartitionLevelHelper::SetForcePackageTrashingAtCleanup(ULevel* Level, bool bForcePackageTrashingAtCleanup)
+{
+	Level->bForcePackageTrashingAtCleanup = bForcePackageTrashingAtCleanup;
+}
+
 #endif
diff --git a/Engine/Source/Runtime/Engine/Private/WorldPartition/WorldPartitionLevelStreamingDynamic.cpp b/Engine/Source/Runtime/Engine/Private/WorldPartition/WorldPartitionLevelStreamingDynamic.cpp
--- a/Engine/Source/Runtime/Engine/Private/WorldPartition/WorldPartitionLevelStreamingDynamic.cpp
+++ b/Engine/Source/Runtime/Engine/Private/WorldPartition/WorldPartitionLevelStreamingDynamic.cpp
@@ -175,8 +175,8 @@ void UWorldPartitionLevelStreamingDynamic::CreateRuntimeLevel()
 	RuntimeLevel = FWorldPartitionLevelHelper::CreateEmptyLevelForRuntimeCell(StreamingCell.Get(), World, GetWorldAsset().ToString());
 	check(RuntimeLevel);
 
-	// Force world partition level/actor packages not to be reused
-	RuntimeLevel->SetForceCantReuseUnloadedButStillAround(true);
+	// Force world partition level/actor packages to be trashed at cleanup
+	FWorldPartitionLevelHelper::SetForcePackageTrashingAtCleanup(RuntimeLevel, true);
 
 	// Make sure Actor Folders is disabled on generated runtime levels to avoid any problems with duplicate folders that
 	// can be caused by level instances injecting their actors, which can cause duplicate folders (which only happens during PIE).
diff --git a/Engine/Source/Runtime/Engine/Public/WorldPartition/WorldPartitionLevelHelper.h b/Engine/Source/Runtime/Engine/Public/WorldPartition/WorldPartitionLevelHelper.h
--- a/Engine/Source/Runtime/Engine/Public/WorldPartition/WorldPartitionLevelHelper.h
+++ b/Engine/Source/Runtime/Engine/Public/WorldPartition/WorldPartitionLevelHelper.h
@@ -117,6 +117,8 @@ private:
 
 	TSet<TWeakObjectPtr<UPackage>> PreGCPackagesToUnload;
 
+	static void SetForcePackageTrashingAtCleanup(ULevel* Level, bool bForcePackageTrashingAtCleanup);
+
 	friend class UWorldPartitionLevelStreamingDynamic;
 	friend class UWorldPartitionSubsystem;
 	friend class UWorldPartitionRuntimeLevelStreamingCell;
