Fix pose watches not being drawn in PIE

diff --git a/Engine/Source/Runtime/Engine/Private/SkeletalMesh.cpp b/Engine/Source/Runtime/Engine/Private/SkeletalMesh.cpp
--- a/Engine/Source/Runtime/Engine/Private/SkeletalMesh.cpp
+++ b/Engine/Source/Runtime/Engine/Private/SkeletalMesh.cpp
@@ -7160,6 +7160,7 @@ FPrimitiveViewRelevance FSkeletalMeshSceneProxy::GetViewRelevance(const FSceneVi
 #if WITH_EDITORONLY_DATA
 			|| MeshObject->SelectedEditorMaterial != -1
 			|| MeshObject->SelectedEditorSection != -1
+			|| (PoseWatchDynamicData && !PoseWatchDynamicData->PoseWatches.IsEmpty())
 #endif
 			;
 #else
