[Material Translator] Removed $TemplateVersion from MaterialTranslator.ush.

Please instead assign a different GUID to your change's stream MaterialTranslationDDCVersion entry in FDevSystemGuids.

diff --git a/Engine/Shaders/Private/MaterialTemplate.ush b/Engine/Shaders/Private/MaterialTemplate.ush
--- a/Engine/Shaders/Private/MaterialTemplate.ush
+++ b/Engine/Shaders/Private/MaterialTemplate.ush
@@ -2,10 +2,6 @@
 
 /**
  * MaterialTemplate.usf: Filled in by FHLSLMaterialTranslator::GetMaterialShaderCode for each material being compiled.
- 
-  $TemplateVersion{1}
-	Bump this number up when you make a change to this file that requires the material
-	translation DDC cache to be invalidated, forcing all materials to re-translate.
  */
 
 #ifndef ENABLE_NEW_HLSL_GENERATOR
diff --git a/Engine/Source/Runtime/Engine/Private/Materials/MaterialSourceTemplate.cpp b/Engine/Source/Runtime/Engine/Private/Materials/MaterialSourceTemplate.cpp
--- a/Engine/Source/Runtime/Engine/Private/Materials/MaterialSourceTemplate.cpp
+++ b/Engine/Source/Runtime/Engine/Private/Materials/MaterialSourceTemplate.cpp
@@ -93,17 +93,8 @@ bool FMaterialSourceTemplate::Preload(EShaderPlatform ShaderPlatform)
 		return false;
 	}
 
-	// Extract the material template string TemplateVersion parameter
+	// Calculate a template hash based on the parameter contained (regardless of order)
 	FSHA1 TemplateHash = {};
-	const TStringView TemplateVersionKeyword = TEXT("$TemplateVersion{");
-	int Begin = MaterialTemplateString.Find(TemplateVersionKeyword.GetData());
-	int End = MaterialTemplateString.Find(TEXT("}"), ESearchCase::CaseSensitive, ESearchDir::FromStart, Begin + TemplateVersionKeyword.Len());
-	if (Begin > 0 && End > 0)
-	{
-		Begin += TemplateVersionKeyword.Len();
-		TemplateHash.UpdateWithString(*MaterialTemplateString + Begin, End - Begin);
-	}
-
 	TArray<FStringView> Parameters;
 	Templates[ShaderPlatform].GetParameters(Parameters);
 	Parameters.Sort();
