Add Level Streaming CSV stats for number of loading, making visible and making invisible levels.
Also unified number of pending purge levels into the same category as the new ones.

diff --git a/Engine/Source/Runtime/Engine/Private/LevelTick.cpp b/Engine/Source/Runtime/Engine/Private/LevelTick.cpp
--- a/Engine/Source/Runtime/Engine/Private/LevelTick.cpp
+++ b/Engine/Source/Runtime/Engine/Private/LevelTick.cpp
@@ -12,6 +12,7 @@
 #include "EngineStats.h"
 #include "RenderingThread.h"
 #include "Materials/MaterialParameterCollectionInstance.h"
+#include "ProfilingDebugging/LevelStreamingProfilingSubsystem.h"
 #include "AI/NavigationSystemBase.h"
 #include "GameFramework/PlayerController.h"
 #include "ParticleHelper.h"
@@ -19,6 +20,7 @@
 #include "Engine/NetConnection.h"
 #include "SceneInterface.h"
 #include "UnrealEngine.h"
+#include "Engine/LevelStreamingGCHelper.h"
 #include "Engine/LevelStreamingVolume.h"
 #include "IXRTrackingSystem.h"
 #include "Camera/CameraPhotography.h"
@@ -1721,6 +1723,27 @@ void UWorld::Tick( ELevelTick TickType, float DeltaSeconds )
 
 #if (CSV_PROFILER && !UE_BUILD_SHIPPING)
 	RecordWorldCountsToCSV(this, bDoingActorTicks);
+
+	if (IsGameWorld() && FCsvProfiler::Get()->IsCapturing())
+	{
+		for (ULevelStreaming* LevelStreaming : StreamingLevels)
+		{
+			switch (LevelStreaming->GetLevelStreamingStatus())
+			{
+			case LEVEL_Loading:
+				CSV_CUSTOM_STAT(LevelStreaming, NumLevelsLoading, 1, ECsvCustomStatOp::Accumulate);
+				break;
+			case LEVEL_MakingVisible:
+				CSV_CUSTOM_STAT(LevelStreaming, NumLevelsMakingVisible, 1, ECsvCustomStatOp::Accumulate);
+				break;
+			case LEVEL_MakingInvisible:
+				CSV_CUSTOM_STAT(LevelStreaming, NumLevelsMakingInvisible, 1, ECsvCustomStatOp::Accumulate);
+				break;
+			}
+		}
+
+		CSV_CUSTOM_STAT(LevelStreaming, NumLevelsPendingPurge, FLevelStreamingGCHelper::GetNumLevelsPendingPurge(), ECsvCustomStatOp::Set);
+	}
 #endif // (CSV_PROFILER && !UE_BUILD_SHIPPING)
 
 #if !(UE_BUILD_SHIPPING || UE_BUILD_TEST)
diff --git a/Engine/Source/Runtime/Engine/Private/World.cpp b/Engine/Source/Runtime/Engine/Private/World.cpp
--- a/Engine/Source/Runtime/Engine/Private/World.cpp
+++ b/Engine/Source/Runtime/Engine/Private/World.cpp
@@ -144,7 +144,6 @@ CSV_DEFINE_CATEGORY(LevelStreamingProfiling, true);
 CSV_DEFINE_CATEGORY(LevelStreamingAdaptive, true);
 CSV_DEFINE_CATEGORY(LevelStreamingAdaptiveDetail, false);
 CSV_DEFINE_CATEGORY(LevelStreamingDetail, false);
-CSV_DEFINE_CATEGORY(LevelStreamingPendingPurge, (!UE_BUILD_SHIPPING));
 
 TRACE_DECLARE_INT_COUNTER(NumStreamingLevelsToConsider, TEXT("LevelStreamingProfiling/NumStreamingLevelsToConsider"));
 
@@ -4434,8 +4433,6 @@ void UWorld::UpdateLevelStreaming()
 	const int32 LevelStreamingContinuouslyIncrementalGCWhileLevelsPendingPurge = GLevelStreamingContinuouslyIncrementalGCWhileLevelsPendingPurgeOverride ? 1 : GLevelStreamingContinuouslyIncrementalGCWhileLevelsPendingPurge;
 	const bool bShouldPurgeLevels = LevelStreamingContinuouslyIncrementalGCWhileLevelsPendingPurge && CurrentNumLevelsPendingPurge >= LevelStreamingContinuouslyIncrementalGCWhileLevelsPendingPurge;
 
-	CSV_CUSTOM_STAT(LevelStreamingPendingPurge, NumlevelsPendingPurge, CurrentNumLevelsPendingPurge, ECsvCustomStatOp::Set);
-
 	// Are we currently in a low memory situation and number of pending levels to purge meets or exceeds our threshold?
 	const bool bShouldDoLowMemoryGC = GIsLowMemory && CurrentNumLevelsPendingPurge >= GLevelStreamingLowMemoryPendingPurgeCount;
 
