Fix invalid VB/IB pointers when initializing FSkeletalMeshObjectStatic::StaticRayTracingGeometry
- get VertexBufferRHI/IndexBufferRHI on render thread instead of game thread.

diff --git a/Engine/Source/Runtime/Engine/Private/SkeletalRenderStatic.cpp b/Engine/Source/Runtime/Engine/Private/SkeletalRenderStatic.cpp
--- a/Engine/Source/Runtime/Engine/Private/SkeletalRenderStatic.cpp
+++ b/Engine/Source/Runtime/Engine/Private/SkeletalRenderStatic.cpp
@@ -35,9 +35,14 @@ void FSkeletalMeshObjectStatic::InitResources(USkinnedMeshComponent* InMeshCompo
 	for( int32 LODIndex=0;LODIndex < LODs.Num();LODIndex++ )
 	{
 		FSkeletalMeshObjectLOD& SkelLOD = LODs[LODIndex];
+
+		check(SkelLOD.SkelMeshRenderData);
+		check(SkelLOD.SkelMeshRenderData->LODRenderData.IsValidIndex(LODIndex));
+
+		FSkeletalMeshLODRenderData& LODData = SkelLOD.SkelMeshRenderData->LODRenderData[LODIndex];
 		
 		// Skip LODs that have their render data stripped
-		if (SkelLOD.SkelMeshRenderData->LODRenderData[LODIndex].GetNumVertices() > 0)
+		if (LODData.GetNumVertices() > 0)
 		{
 			FSkelMeshComponentLODInfo* CompLODInfo = nullptr;
 			if (InMeshComponent->LODInfo.IsValidIndex(LODIndex))
@@ -50,36 +55,31 @@ void FSkeletalMeshObjectStatic::InitResources(USkinnedMeshComponent* InMeshCompo
 #if RHI_RAYTRACING
 			if (IsRayTracingAllowed() && SkelLOD.SkelMeshRenderData->bSupportRayTracing)
 			{
-				if (SkelLOD.SkelMeshRenderData->LODRenderData[LODIndex].NumReferencingStaticSkeletalMeshObjects == 0)
+				if (LODData.NumReferencingStaticSkeletalMeshObjects == 0)
 				{
-					check(SkelLOD.SkelMeshRenderData);
-					check(SkelLOD.SkelMeshRenderData->LODRenderData.IsValidIndex(LODIndex));
-
-					FSkeletalMeshLODRenderData& LODModel = SkelLOD.SkelMeshRenderData->LODRenderData[LODIndex];
-					FBufferRHIRef VertexBufferRHI = LODModel.StaticVertexBuffers.PositionVertexBuffer.VertexBufferRHI;
-					FBufferRHIRef IndexBufferRHI = LODModel.MultiSizeIndexContainer.GetIndexBuffer()->IndexBufferRHI;
-					uint32 VertexBufferStride = LODModel.StaticVertexBuffers.PositionVertexBuffer.GetStride();
+					FPositionVertexBuffer* PositionVertexBufferPtr = &LODData.StaticVertexBuffers.PositionVertexBuffer;
+					FRawStaticIndexBuffer16or32Interface* IndexBufferPtr = LODData.MultiSizeIndexContainer.GetIndexBuffer();
 
 					uint32 TrianglesCount = 0;
-					for (int32 SectionIndex = 0; SectionIndex < LODModel.RenderSections.Num(); SectionIndex++)
+					for (int32 SectionIndex = 0; SectionIndex < LODData.RenderSections.Num(); SectionIndex++)
 					{
-						const FSkelMeshRenderSection& Section = LODModel.RenderSections[SectionIndex];
+						const FSkelMeshRenderSection& Section = LODData.RenderSections[SectionIndex];
 						TrianglesCount += Section.NumTriangles;
 					}
 
-					TArray<FSkelMeshRenderSection>* RenderSections = &LODModel.RenderSections;
+					TArray<FSkelMeshRenderSection>* RenderSections = &LODData.RenderSections;
 					ENQUEUE_RENDER_COMMAND(InitSkeletalRenderStaticRayTracingGeometry)(UE::RenderCommandPipe::SkeletalMesh,
-						[this, VertexBufferRHI, IndexBufferRHI, VertexBufferStride, TrianglesCount, RenderSections, 
+						[PositionVertexBufferPtr, IndexBufferPtr, TrianglesCount, RenderSections,
 						LODIndex = LODIndex, 
 						SkelMeshRenderData = SkelLOD.SkelMeshRenderData, 
-						&RayTracingGeometry = SkelLOD.SkelMeshRenderData->LODRenderData[LODIndex].StaticRayTracingGeometry,
-						&bReferencedByStaticSkeletalMeshObjects_RenderThread = SkelLOD.SkelMeshRenderData->LODRenderData[LODIndex].bReferencedByStaticSkeletalMeshObjects_RenderThread](FRHICommandList& RHICmdList)
+						&RayTracingGeometry = LODData.StaticRayTracingGeometry,
+						&bReferencedByStaticSkeletalMeshObjects_RenderThread = LODData.bReferencedByStaticSkeletalMeshObjects_RenderThread](FRHICommandList& RHICmdList)
 						{
 							FRayTracingGeometryInitializer Initializer;
 							static const FName DebugName("FSkeletalMeshObjectLOD");
 							static int32 DebugNumber = 0;
 							Initializer.DebugName = FDebugName(DebugName, DebugNumber++);
-							Initializer.IndexBuffer = IndexBufferRHI;
+							Initializer.IndexBuffer = IndexBufferPtr->IndexBufferRHI;
 							Initializer.TotalPrimitiveCount = TrianglesCount;
 							Initializer.GeometryType = RTGT_Triangles;
 							Initializer.bFastBuild = false;
@@ -96,10 +96,10 @@ void FSkeletalMeshObjectStatic::InitResources(USkinnedMeshComponent* InMeshCompo
 							for (const FSkelMeshRenderSection& Section : *RenderSections)
 							{
 								FRayTracingGeometrySegment Segment;
-								Segment.VertexBuffer = VertexBufferRHI;
+								Segment.VertexBuffer = PositionVertexBufferPtr->VertexBufferRHI;
 								Segment.VertexBufferElementType = VET_Float3;
 								Segment.VertexBufferOffset = 0;
-								Segment.VertexBufferStride = VertexBufferStride;
+								Segment.VertexBufferStride = PositionVertexBufferPtr->GetStride();
 								Segment.MaxVertices = TotalNumVertices;
 								Segment.FirstPrimitive = Section.BaseIndex / 3;
 								Segment.NumPrimitives = Section.NumTriangles;
@@ -123,7 +123,7 @@ void FSkeletalMeshObjectStatic::InitResources(USkinnedMeshComponent* InMeshCompo
 					);
 				}
 
-				SkelLOD.SkelMeshRenderData->LODRenderData[LODIndex].NumReferencingStaticSkeletalMeshObjects++;
+				LODData.NumReferencingStaticSkeletalMeshObjects++;
 			}
 #endif
 		}
