Performance and memory improvements to world tick structures, which improves tick performance on world partition maps:
Remove TickTaskLevel from UWorld because it is no longer in use, and only create it for ULevels that actually have tick functions.
Add s.World.CreateStaticLevelCollection cvar and default to 0, this stops it from creating the static level collection which is only used for the disabled client world duplication feature.
General efficiency cleanup for code that uses level collections.

diff --git a/Engine/Source/Runtime/Engine/Classes/Engine/Engine.h b/Engine/Source/Runtime/Engine/Classes/Engine/Engine.h
--- a/Engine/Source/Runtime/Engine/Classes/Engine/Engine.h
+++ b/Engine/Source/Runtime/Engine/Classes/Engine/Engine.h
@@ -2963,6 +2963,7 @@ public:
 	 *	If this function returns true, the DynamicSourceLevels collection will be duplicated for the given map.
 	 *	This is necessary to do outside of the editor when we don't have the original editor world, and it's 
 	 *	not safe to copy the dynamic levels once they've been fully initialized, so we pre-duplicate them when the original levels are first created.
+	 *	If you implement this, enable s.World.CreateStaticLevelCollection to stop it from duplicating static streaming levels.
 	 */
 	virtual bool Experimental_ShouldPreDuplicateMap(const FName MapName) const { return false; }
 
diff --git a/Engine/Source/Runtime/Engine/Classes/Engine/EngineTypes.h b/Engine/Source/Runtime/Engine/Classes/Engine/EngineTypes.h
--- a/Engine/Source/Runtime/Engine/Classes/Engine/EngineTypes.h
+++ b/Engine/Source/Runtime/Engine/Classes/Engine/EngineTypes.h
@@ -4024,16 +4024,21 @@ enum class ELevelCollectionType : uint8
 	/**
 	 * The dynamic levels that are used for normal gameplay and the source for any duplicated collections.
 	 * Will contain a world's persistent level and any streaming levels that contain dynamic or replicated gameplay actors.
+	 * This collection will always exist for gameplay and editor worlds.
 	 */
 	DynamicSourceLevels,
 
-	/** Gameplay relevant levels that have been duplicated from DynamicSourceLevels if requested by the game. */
+	/** 
+	 * Gameplay relevant levels that have been duplicated from DynamicSourceLevels if requested by the game.
+	 * This collection only exists if levels have actually been duplicated.
+	 */
 	DynamicDuplicatedLevels,
 
 	/**
 	 * These levels are shared between the source levels and the duplicated levels, and should contain
 	 * only static geometry and other visuals that are not replicated or affected by gameplay.
 	 * These will not be duplicated in order to save memory.
+	 * If s.World.CreateStaticLevelCollection is 0, this will not be created and static levels will be treated as dynamic.
 	 */
 	StaticLevels,
 
diff --git a/Engine/Source/Runtime/Engine/Classes/Engine/World.h b/Engine/Source/Runtime/Engine/Classes/Engine/World.h
--- a/Engine/Source/Runtime/Engine/Classes/Engine/World.h
+++ b/Engine/Source/Runtime/Engine/Classes/Engine/World.h
@@ -1629,9 +1629,6 @@ public:
 	/** Interface to the FX system managing particles and related effects for this world.										*/
 	class FFXSystemInterface*					FXSystem;
 
-	/** Data structures for holding the tick functions that are associated with the world (line batcher, etc) **/
-	class FTickTaskLevel*						TickTaskLevel;
-
 	/** Tick function for starting physics																						*/
 	FStartPhysicsTickFunction StartPhysicsTickFunction;
 	/** Tick function for ending physics																						*/
@@ -3206,6 +3203,9 @@ public:
 	 */
 	bool RemoveLevel( ULevel* InLevel );
 
+	/** Returns the FLevelCollection for the passed in streaming level, which could depend on configuration */
+	FLevelCollection& FindOrAddCollectionForLevelStreaming(const ULevelStreaming* Level);
+
 	/** Returns the FLevelCollection for the given InType. If one does not exist, it is created. */
 	FLevelCollection& FindOrAddCollectionByType(const ELevelCollectionType InType);
 
diff --git a/Engine/Source/Runtime/Engine/Private/Level.cpp b/Engine/Source/Runtime/Engine/Private/Level.cpp
--- a/Engine/Source/Runtime/Engine/Private/Level.cpp
+++ b/Engine/Source/Runtime/Engine/Private/Level.cpp
@@ -395,8 +395,8 @@ const FName ULevel::DontLoadExternalFoldersTag(TEXT("DontLoadExternalFoldersTag"
 ULevel::ULevel( const FObjectInitializer& ObjectInitializer )
 	:	UObject( ObjectInitializer )
 	,	Actors()
-	,	OwningWorld(NULL)
-	,	TickTaskLevel(FTickTaskManagerInterface::Get().AllocateTickTaskLevel())
+	,	OwningWorld(nullptr)
+	,	TickTaskLevel(nullptr)
 	,	PrecomputedLightVolume(new FPrecomputedLightVolume())
 	,	PrecomputedVolumetricLightmap(new FPrecomputedVolumetricLightmap())
 	,	RouteActorInitializationState(ERouteActorInitializationState::Preinitialize)
@@ -431,8 +431,11 @@ void ULevel::Initialize(const FURL& InURL)
 
 ULevel::~ULevel()
 {
-	FTickTaskManagerInterface::Get().FreeTickTaskLevel(TickTaskLevel);
-	TickTaskLevel = NULL;
+	if (TickTaskLevel)
+	{
+		FTickTaskManagerInterface::Get().FreeTickTaskLevel(TickTaskLevel);
+		TickTaskLevel = nullptr;
+	}
 }
 
 void ULevel::AddReferencedObjects(UObject* InThis, FReferenceCollector& Collector)
diff --git a/Engine/Source/Runtime/Engine/Private/LevelStreaming.cpp b/Engine/Source/Runtime/Engine/Private/LevelStreaming.cpp
--- a/Engine/Source/Runtime/Engine/Private/LevelStreaming.cpp
+++ b/Engine/Source/Runtime/Engine/Private/LevelStreaming.cpp
@@ -1287,8 +1287,7 @@ void ULevelStreaming::AddLevelToCollectionAfterReload()
 			LoadedLevel->GetCachedLevelCollection()->RemoveLevel(LoadedLevel);
 		}
 		// Add this level to the correct collection
-		const ELevelCollectionType CollectionType = bIsStatic ? ELevelCollectionType::StaticLevels : ELevelCollectionType::DynamicSourceLevels;
-		FLevelCollection& LC = GetWorld()->FindOrAddCollectionByType(CollectionType);
+		FLevelCollection& LC = GetWorld()->FindOrAddCollectionForLevelStreaming(this);
 		LC.AddLevel(LoadedLevel);
 	}
 }
@@ -1329,11 +1328,8 @@ void ULevelStreaming::SetLoadedLevel(ULevel* Level)
 	FLevelStreamingGCHelper::CancelUnloadRequest(LoadedLevel);
 
 	// Add this level to the correct collection
-	const ELevelCollectionType CollectionType =	bIsStatic ? ELevelCollectionType::StaticLevels : ELevelCollectionType::DynamicSourceLevels;
-
 	UWorld* World = GetWorld();
-
-	FLevelCollection& LC = World->FindOrAddCollectionByType(CollectionType);
+	FLevelCollection& LC = World->FindOrAddCollectionForLevelStreaming(this);
 	LC.RemoveLevel(PendingUnloadLevel);
 
 	if (PendingUnloadLevel)
@@ -2246,17 +2242,15 @@ void ULevelStreaming::PostEditChangeProperty(FPropertyChangedEvent& PropertyChan
 		}
 		else if (PropertyName == GET_MEMBER_NAME_CHECKED(ULevelStreaming, bIsStatic))
 		{
-			if (LoadedLevel)
+			UWorld* World = GetWorld();
+			if (LoadedLevel && World)
 			{
-				const ELevelCollectionType NewCollectionType = bIsStatic ? ELevelCollectionType::StaticLevels : ELevelCollectionType::DynamicSourceLevels;
 				FLevelCollection* PreviousCollection = LoadedLevel->GetCachedLevelCollection();
+				FLevelCollection& LC = World->FindOrAddCollectionForLevelStreaming(this);
 
-				if (PreviousCollection && PreviousCollection->GetType() != NewCollectionType)
+				if (PreviousCollection && PreviousCollection != &LC)
 				{
 					PreviousCollection->RemoveLevel(LoadedLevel);
-
-					UWorld* World = GetWorld();
-					FLevelCollection& LC = World->FindOrAddCollectionByType(NewCollectionType);
 					LC.AddLevel(LoadedLevel);
 				}
 			}
diff --git a/Engine/Source/Runtime/Engine/Private/LevelTick.cpp b/Engine/Source/Runtime/Engine/Private/LevelTick.cpp
--- a/Engine/Source/Runtime/Engine/Private/LevelTick.cpp
+++ b/Engine/Source/Runtime/Engine/Private/LevelTick.cpp
@@ -1459,6 +1459,7 @@ void UWorld::Tick( ELevelTick TickType, float DeltaSeconds )
 
 	// If only the DynamicLevel collection has entries, we can skip the validation and tick all levels.
 	bool bValidateLevelList = false;
+	bool bHasTickedACollection = false;
 
 	for (const FLevelCollection& LevelCollection : LevelCollections)
 	{
@@ -1487,12 +1488,17 @@ void UWorld::Tick( ELevelTick TickType, float DeltaSeconds )
 			}
 		}
 
+		if (LevelsToTick.Num() == 0)
+		{
+			// Nothing to do, early out before creating the slow scope
+			continue;
+		}
+
 		// Set up context on the world for this level collection
 		FScopedLevelCollectionContextSwitch LevelContext(i, this);
 
 		// If caller wants time update only, or we are paused, skip the rest.
-		const bool bShouldSkipTick = (LevelsToTick.Num() == 0);
-		if (bDoingActorTicks && !bShouldSkipTick)
+		if (bDoingActorTicks)
 		{
 			// Actually tick actors now that context is set up
 			SetupPhysicsTickFunctions(DeltaSeconds);
@@ -1542,9 +1548,11 @@ void UWorld::Tick( ELevelTick TickType, float DeltaSeconds )
 			FTickTaskManagerInterface::Get().RunPauseFrame(this, DeltaSeconds, LEVELTICK_PauseTick, LevelsToTick);
 		}
 		
-		// We only want to run the following once, so only run it for the source level collection.
-		if (LevelCollections[i].GetType() == ELevelCollectionType::DynamicSourceLevels)
+		// Run this on the first collection with levels, which will be DynamicSourceLevels by default
+		if (!bHasTickedACollection)
 		{
+			bHasTickedACollection = true;
+
 			// Process any remaining latent actions
 			if( !bIsPaused )
 			{
@@ -1609,7 +1617,7 @@ void UWorld::Tick( ELevelTick TickType, float DeltaSeconds )
 			}
 		}
 
-		if (bDoingActorTicks && !bShouldSkipTick)
+		if (bDoingActorTicks)
 		{
 			SCOPE_CYCLE_COUNTER(STAT_TickTime);
 			{
diff --git a/Engine/Source/Runtime/Engine/Private/TickTaskManager.cpp b/Engine/Source/Runtime/Engine/Private/TickTaskManager.cpp
--- a/Engine/Source/Runtime/Engine/Private/TickTaskManager.cpp
+++ b/Engine/Source/Runtime/Engine/Private/TickTaskManager.cpp
@@ -1633,8 +1633,8 @@ public:
 	/** Return true if this tick function is in the primary list **/
 	bool HasTickFunction(ULevel* InLevel, FTickFunction* TickFunction)
 	{
-		FTickTaskLevel* Level = TickTaskLevelForLevel(InLevel);
-		return Level->HasTickFunction(TickFunction);
+		FTickTaskLevel* Level = TickTaskLevelForLevel(InLevel, false);
+		return Level && Level->HasTickFunction(TickFunction);
 	}
 	/** Add the tick function to the primary list **/
 	void AddTickFunction(ULevel* InLevel, FTickFunction* TickFunction)
@@ -1662,30 +1662,30 @@ private:
 		IConsoleManager::Get().RegisterConsoleCommand(TEXT("dumpticks"), TEXT("Dumps all tick functions registered with FTickTaskManager to log."));
 	}
 
-	/** Fill the level list **/
+	/** Fill the level list, only with levels that can actually tick */
 	void FillLevelList(const TArray<ULevel*>& Levels)
 	{
 		check(!LevelList.Num());
-		if (!Context.World->GetActiveLevelCollection() || Context.World->GetActiveLevelCollection()->GetType() == ELevelCollectionType::DynamicSourceLevels)
-		{
-			check(Context.World->TickTaskLevel);
-			LevelList.Add(Context.World->TickTaskLevel);
-		}
 		for( int32 LevelIndex = 0; LevelIndex < Levels.Num(); LevelIndex++ )
 		{
 			ULevel* Level = Levels[LevelIndex];
-			if (Level && Level->bIsVisible)
+			if (Level && Level->bIsVisible && Level->TickTaskLevel)
 			{
-				check(Level->TickTaskLevel);
 				LevelList.Add(Level->TickTaskLevel);
 			}
 		}
 	}
 
 	/** Find the tick level for this actor **/
-	FTickTaskLevel* TickTaskLevelForLevel(ULevel* Level)
+	FTickTaskLevel* TickTaskLevelForLevel(ULevel* Level, bool bCreateIfNeeded = true)
 	{
 		check(Level);
+
+		if (bCreateIfNeeded && Level->TickTaskLevel == nullptr)
+		{
+			Level->TickTaskLevel = AllocateTickTaskLevel();
+		}
+
 		check(Level->TickTaskLevel);
 		return Level->TickTaskLevel;
 	}
@@ -1701,7 +1701,6 @@ private:
 		Ar.Logf(TEXT("============================ Tick Functions (%s) ============================"), bGrouped ? TEXT("GROUPED") : ((bEnabled && bDisabled) ? TEXT("All") : (bEnabled ? TEXT("Enabled") : TEXT("Disabled"))));
 
 		check(InWorld);
-		check(InWorld->TickTaskLevel);
 
 		if (bGrouped)
 		{
@@ -1739,13 +1738,11 @@ private:
 		}
 		else
 		{
-			InWorld->TickTaskLevel->DumpAllTickFunctions(Ar, EnabledCount, DisabledCount, bEnabled, bDisabled);
 			for (int32 LevelIndex = 0; LevelIndex < InWorld->GetNumLevels(); LevelIndex++)
 			{
 				ULevel* Level = InWorld->GetLevel(LevelIndex);
-				if (Level->bIsVisible)
+				if (Level->bIsVisible && Level->TickTaskLevel)
 				{
-					check(Level->TickTaskLevel);
 					Level->TickTaskLevel->DumpAllTickFunctions(Ar, EnabledCount, DisabledCount, bEnabled, bDisabled);
 				}
 			}
@@ -1760,7 +1757,6 @@ private:
 	virtual void GetEnabledTickFunctionCounts(UWorld* InWorld, TSortedMap<FName, int32, FDefaultAllocator, FNameFastLess>& TickContextToCountMap, int32& EnabledCount, bool bDetailed, bool bFilterCoolingDown=false)
 	{
 		check(InWorld);
-		check(InWorld->TickTaskLevel);
 
 		if (bFilterCoolingDown && InWorld->TickGroup >= 0 && InWorld->TickGroup < TG_NewlySpawned)
 		{
@@ -1770,14 +1766,12 @@ private:
 		const float WorldTimeSeconds = InWorld->GetTimeSeconds();
 		const float WorldUnpausedTimeSeconds = InWorld->GetUnpausedTimeSeconds();
 
-		InWorld->TickTaskLevel->AddTickFunctionsToMap(TickContextToCountMap, EnabledCount, bDetailed, bFilterCoolingDown, WorldTimeSeconds, WorldUnpausedTimeSeconds);
 
 		for (int32 LevelIndex = 0; LevelIndex < InWorld->GetNumLevels(); LevelIndex++)
 		{
 			ULevel* Level = InWorld->GetLevel(LevelIndex);
-			if (Level->bIsVisible)
+			if (Level->bIsVisible && Level->TickTaskLevel)
 			{
-				check(Level->TickTaskLevel);
 				Level->TickTaskLevel->AddTickFunctionsToMap(TickContextToCountMap, EnabledCount, bDetailed, bFilterCoolingDown, WorldTimeSeconds, WorldUnpausedTimeSeconds);
 			}
 		}
diff --git a/Engine/Source/Runtime/Engine/Private/TimerManager.cpp b/Engine/Source/Runtime/Engine/Private/TimerManager.cpp
--- a/Engine/Source/Runtime/Engine/Private/TimerManager.cpp
+++ b/Engine/Source/Runtime/Engine/Private/TimerManager.cpp
@@ -771,9 +771,9 @@ FTimerData::FTimerData()
 	, bMaxOncePerFrame(false)
 	, bRequiresDelegate(false)
 	, Status(ETimerStatus::Active)
+	, LevelCollection(ELevelCollectionType::DynamicSourceLevels)
 	, Rate(0)
 	, ExpireTime(0)
-	, LevelCollection(ELevelCollectionType::DynamicSourceLevels)
 {}
 
 // ---------------------------------
@@ -835,7 +835,6 @@ void FTimerManager::Tick(float DeltaTime)
 	InternalTime += DeltaTime;
 
 	UWorld* const OwningWorld = OwningGameInstance ? OwningGameInstance->GetWorld() : nullptr;
-	UWorld* const LevelCollectionWorld = OwningWorld;
 
 #if UE_ENABLE_DUMPALLTIMERLOGSTHRESHOLD
 	// Dump timer info to logs if we have way too many timers active.
@@ -900,7 +899,7 @@ void FTimerManager::Tick(float DeltaTime)
 			// Set the relevant level context for this timer
 			const int32 LevelCollectionIndex = OwningWorld ? OwningWorld->FindCollectionIndexByType(Top->LevelCollection) : INDEX_NONE;
 			
-			FScopedLevelCollectionContextSwitch LevelContext(LevelCollectionIndex, LevelCollectionWorld);
+			FScopedLevelCollectionContextSwitch LevelContext(LevelCollectionIndex, OwningWorld);
 
 			// Remove it from the heap and store it while we're executing
 			ActiveTimerHeap.HeapPop(CurrentlyExecutingTimer, FTimerHeapOrder(Timers), EAllowShrinking::No);
diff --git a/Engine/Source/Runtime/Engine/Private/UnrealEngine.cpp b/Engine/Source/Runtime/Engine/Private/UnrealEngine.cpp
--- a/Engine/Source/Runtime/Engine/Private/UnrealEngine.cpp
+++ b/Engine/Source/Runtime/Engine/Private/UnrealEngine.cpp
@@ -15816,8 +15816,10 @@ void UEngine::MovePendingLevel(FWorldContext &Context)
 		FLevelCollection& SourceLevels = Context.World()->FindOrAddCollectionByType(ELevelCollectionType::DynamicSourceLevels);
 		SourceLevels.SetNetDriver(NetDriver);
 
-		FLevelCollection& StaticLevels = Context.World()->FindOrAddCollectionByType(ELevelCollectionType::StaticLevels);
-		StaticLevels.SetNetDriver(NetDriver);
+		if (FLevelCollection* StaticLevels = Context.World()->FindCollectionByType(ELevelCollectionType::StaticLevels))
+		{
+			StaticLevels->SetNetDriver(NetDriver);
+		}
 	}
 
 	// Attach the DemoNetDriver to the world if there is one
diff --git a/Engine/Source/Runtime/Engine/Private/World.cpp b/Engine/Source/Runtime/Engine/Private/World.cpp
--- a/Engine/Source/Runtime/Engine/Private/World.cpp
+++ b/Engine/Source/Runtime/Engine/Private/World.cpp
@@ -155,6 +155,11 @@ FAutoConsoleVariableRef CVarDisableRemapScriptActors(TEXT("net.DisableRemapScrip
 static bool bDisableInGamePerfTrackersForUninitializedWorlds = true;
 FAutoConsoleVariableRef CVarDisableInGamePerfTrackersForUninitializedWorlds(TEXT("s.World.SkipPerfTrackerForUninitializedWorlds"), bDisableInGamePerfTrackersForUninitializedWorlds, TEXT("When set, disables allocation of InGamePerformanceTrackers for Worlds that aren't initialized."));
 
+static bool bCreateStaticLevelCollection = false;
+FAutoConsoleVariableRef CVarCreateStaticLevelCollection(TEXT("s.World.CreateStaticLevelCollection"), bCreateStaticLevelCollection,
+	TEXT("When set, create a separate level collection for static streaming levels that will not be duplicated by DuplicateRequestedLevels.\n")
+	TEXT("If this is 0, static streaming levels will be part of the main DynamicSourceLevels collection."));
+
 // Now that it's possible for subclasses of ULevelStreaming to indicate which async loads are necessary for loading,
 // it's possible existing subclasses haven't added their required loads to their StreamingLevel->GetAsyncRequestIDs() array. As a fallback,
 // allow users to force flushing of all async loads during level streaming, as was done in UE 5.4 and lower.
@@ -641,7 +646,6 @@ UWorld::UWorld( const FObjectInitializer& ObjectInitializer )
 #endif
 , URL(FURL(NULL))
 ,	FXSystem(NULL)
-,	TickTaskLevel(FTickTaskManagerInterface::Get().AllocateTickTaskLevel())
 ,	FlushLevelStreamingType(EFlushLevelStreamingType::None)
 ,	NextTravelType(TRAVEL_Relative)
 ,	CleanupWorldTag(0)
@@ -1338,8 +1342,6 @@ void UWorld::FinishDestroy()
 	{
 		GWorld = NULL;
 	}
-	FTickTaskManagerInterface::Get().FreeTickTaskLevel(TickTaskLevel);
-	TickTaskLevel = NULL;
 
 	if (TimerManager)
 	{
@@ -2304,7 +2306,18 @@ const FName UWorld::KeepInitializedDuringLoadTag(TEXT("KeepInitializedDuringLoad
 
 void UWorld::ConditionallyCreateDefaultLevelCollections()
 {
-	LevelCollections.Reserve((int32)ELevelCollectionType::MAX);
+	if (WorldType == EWorldType::Inactive)
+	{
+		return;
+	}
+	else if (bCreateStaticLevelCollection)
+	{
+		LevelCollections.Reserve((int32)ELevelCollectionType::MAX);
+	}
+	else
+	{
+		LevelCollections.Reserve(1);
+	}
 
 	// Create main level collection. The persistent level will always be considered dynamic.
 	if (!FindCollectionByType(ELevelCollectionType::DynamicSourceLevels))
@@ -2322,7 +2335,7 @@ void UWorld::ConditionallyCreateDefaultLevelCollections()
 		}
 	}
 
-	if (!FindCollectionByType(ELevelCollectionType::StaticLevels))
+	if (bCreateStaticLevelCollection && !FindCollectionByType(ELevelCollectionType::StaticLevels))
 	{
 		FLevelCollection& StaticCollection = FindOrAddCollectionByType(ELevelCollectionType::StaticLevels);
 		StaticCollection.SetPersistentLevel(PersistentLevel);
@@ -8694,8 +8707,10 @@ void UWorld::SetGameState(AGameStateBase* NewGameState)
 		        // For now the static levels use the same GameState as the source dynamic levels.
 		        if (FoundCollection->GetType() == ELevelCollectionType::DynamicSourceLevels)
 		        {
-			        FLevelCollection& StaticLevels = FindOrAddCollectionByType(ELevelCollectionType::StaticLevels);
-			        StaticLevels.SetGameState(NewGameState);
+					if (FLevelCollection* StaticLevels = FindCollectionByType(ELevelCollectionType::StaticLevels))
+					{
+						StaticLevels->SetGameState(NewGameState);
+					}
 		        }
 	        }
 		}
@@ -8818,21 +8833,20 @@ void UWorld::CreateFXSystem()
 	}
 }
 
-FLevelCollection& UWorld::FindOrAddCollectionByType(const ELevelCollectionType InType)
+FLevelCollection& UWorld::FindOrAddCollectionForLevelStreaming(const ULevelStreaming* Level)
 {
-	for (FLevelCollection& LC : LevelCollections)
+	ELevelCollectionType Type = ELevelCollectionType::DynamicSourceLevels;
+	if (bCreateStaticLevelCollection && Level->bIsStatic)
 	{
-		if (LC.GetType() == InType)
-		{
-			return LC;
-		}
+		Type = ELevelCollectionType::StaticLevels;
 	}
+	
+	return FindOrAddCollectionByType(Type);
+}
 
-	// Not found, add a new one.
-	FLevelCollection NewLC;
-	NewLC.SetType(InType);
-	LevelCollections.Add(MoveTemp(NewLC));
-	return LevelCollections.Last();
+FLevelCollection& UWorld::FindOrAddCollectionByType(const ELevelCollectionType InType)
+{
+	return LevelCollections[FindOrAddCollectionByType_Index(InType)];
 }
 
 int32 UWorld::FindOrAddCollectionByType_Index(const ELevelCollectionType InType)
@@ -8844,6 +8858,9 @@ int32 UWorld::FindOrAddCollectionByType_Index(const ELevelCollectionType InType)
 		return FoundIndex;
 	}
 
+	// Static collections should not be created if that is disabled
+	ensure(InType != ELevelCollectionType::StaticLevels || bCreateStaticLevelCollection);
+
 	// Not found, add a new one.
 	FLevelCollection NewLC;
 	NewLC.SetType(InType);
@@ -8852,15 +8869,7 @@ int32 UWorld::FindOrAddCollectionByType_Index(const ELevelCollectionType InType)
 
 FLevelCollection* UWorld::FindCollectionByType(const ELevelCollectionType InType)
 {
-	for (FLevelCollection& LC : LevelCollections)
-	{
-		if (LC.GetType() == InType)
-		{
-			return &LC;
-		}
-	}
-
-	return nullptr;
+	return const_cast<FLevelCollection*>(const_cast<const UWorld*>(this)->FindCollectionByType(InType));
 }
 
 const FLevelCollection* UWorld::FindCollectionByType(const ELevelCollectionType InType) const
@@ -8896,6 +8905,12 @@ const FLevelCollection* UWorld::GetActiveLevelCollection() const
 
 void UWorld::SetActiveLevelCollection(int32 LevelCollectionIndex)
 {
+	// Only check if collection actually changes
+	if (LevelCollectionIndex == ActiveLevelCollectionIndex)
+	{
+		return;
+	}
+
 	ActiveLevelCollectionIndex = LevelCollectionIndex;
 	const FLevelCollection* const ActiveLevelCollection = GetActiveLevelCollection();
 
@@ -8915,7 +8930,7 @@ void UWorld::SetActiveLevelCollection(int32 LevelCollectionIndex)
 	NetDriver = ActiveLevelCollection->GetNetDriver();
 	DemoNetDriver = ActiveLevelCollection->GetDemoNetDriver();
 
-	// TODO: START TEMP FIX FOR UE-42508
+	// Our net drivers may have been destroyed during the scope
 	if (NetDriver && NetDriver->NetDriverName != NAME_None)
 	{
 		UNetDriver* TempNetDriver = GEngine->FindNamedNetDriver(this, NetDriver->NetDriverName);
@@ -8935,7 +8950,6 @@ void UWorld::SetActiveLevelCollection(int32 LevelCollectionIndex)
 			DemoNetDriver = TempDemoNetDriver;
 		}
 	}
-	// TODO: END TEMP FIX FOR UE-42508
 }
 
 static ULevel* DuplicateLevelWithPrefix(ULevel* InLevel, int32 InstanceID )
