Tick optimizations to avoid redundant object checks, ActorComponent and Actor now use the same rules.
Optimize map access in the LatentActionManager to avoid slow weak pointer Get() calls.

diff --git a/Engine/Source/Runtime/Engine/Classes/Engine/LatentActionManager.h b/Engine/Source/Runtime/Engine/Classes/Engine/LatentActionManager.h
--- a/Engine/Source/Runtime/Engine/Classes/Engine/LatentActionManager.h
+++ b/Engine/Source/Runtime/Engine/Classes/Engine/LatentActionManager.h
@@ -64,24 +64,10 @@ struct FLatentActionManager
 {
 	GENERATED_USTRUCT_BODY()
 
-	/** Map of UUID->Action(s). */
-	typedef TMultiMap<int32, class FPendingLatentAction*> FActionList;
-
-	struct FObjectActions
-	{
-		/** Map of UUID->Action(s). */
-		FActionList ActionList;
-		bool bProcessedThisFrame = false;
-	};
-	
-	/** Map to convert from object to FActionList. */
-	typedef TMap< TWeakObjectPtr<UObject>, TSharedPtr<FObjectActions> > FObjectToActionListMap;
-	FObjectToActionListMap ObjectToActionListMap;
-
+public: 
 	/** @return A delegate that will be broadcast when a latent action is added or removed from the manager */
 	static FOnLatentActionsChanged& OnLatentActionsChanged() { return LatentActionsChangedDelegate; }
 
-public:
 	/** 
 	 * Advance pending latent actions by DeltaTime.
  	 * If no object is specified it will process any outstanding actions for objects that have not been processed for this frame.
@@ -180,6 +166,20 @@ public:
 	ENGINE_API ~FLatentActionManager();
 
 protected:
+	/** Map of UUID->Action(s). */
+	typedef TMultiMap<int32, class FPendingLatentAction*> FActionList;
+
+	struct FObjectActions
+	{
+		/** Map of UUID->Action(s). */
+		FActionList ActionList;
+		bool bProcessedThisFrame = false;
+	};
+
+	/** Map to convert from object to FActionList. */
+	typedef TMap< TWeakObjectPtr<UObject>, TSharedPtr<FObjectActions>, FDefaultSetAllocator, TWeakObjectPtrMapKeyFuncs<TWeakObjectPtr<UObject>, TSharedPtr<FObjectActions> > > FObjectToActionListMap;
+	FObjectToActionListMap ObjectToActionListMap;
+
 	/** 
 	 * Finds the action instance for the supplied object will return NULL if one does not exist.
 	 *
diff --git a/Engine/Source/Runtime/Engine/Classes/GameFramework/Actor.h b/Engine/Source/Runtime/Engine/Classes/GameFramework/Actor.h
--- a/Engine/Source/Runtime/Engine/Classes/GameFramework/Actor.h
+++ b/Engine/Source/Runtime/Engine/Classes/GameFramework/Actor.h
@@ -4643,11 +4643,11 @@ private:
 };
 #endif
 
-/** Helper function for executing tick functions based on the normal conditions previous found in UActorComponent::ConditionalTick */
+/** Helper function for executing component tick functions using the same conditions as FActorTickFunction */
 template <typename ExecuteTickLambda>
 void FActorComponentTickFunction::ExecuteTickHelper(UActorComponent* Target, bool bTickInEditor, float DeltaTime, ELevelTick TickType, const ExecuteTickLambda& ExecuteTickFunc)
 {
-	if (Target && IsValidChecked(Target) && !Target->IsUnreachable())
+	if (IsValid(Target))
 	{
 		FScopeCycleCounterUObject ComponentScope(Target);
 		FScopeCycleCounterUObject AdditionalScope(Target->AdditionalStatObject());
@@ -4655,11 +4655,8 @@ void FActorComponentTickFunction::ExecuteTickHelper(UActorComponent* Target, boo
 		if (Target->bRegistered)
 		{
 			AActor* MyOwner = Target->GetOwner();
-			//@optimization, I imagine this is all unnecessary in a shipping game with no editor
-			if (TickType != LEVELTICK_ViewportsOnly ||
-				(bTickInEditor && TickType == LEVELTICK_ViewportsOnly) ||
-				(MyOwner && MyOwner->ShouldTickIfViewportsOnly())
-				)
+			if (TickType != LEVELTICK_ViewportsOnly || bTickInEditor ||
+				(MyOwner && MyOwner->ShouldTickIfViewportsOnly()))
 			{
 				const float TimeDilation = (MyOwner ? MyOwner->CustomTimeDilation : 1.f);
 				ExecuteTickFunc(DeltaTime * TimeDilation);
diff --git a/Engine/Source/Runtime/Engine/Private/Actor.cpp b/Engine/Source/Runtime/Engine/Private/Actor.cpp
--- a/Engine/Source/Runtime/Engine/Private/Actor.cpp
+++ b/Engine/Source/Runtime/Engine/Private/Actor.cpp
@@ -250,7 +250,7 @@ void AActor::InitializeDefaults()
 
 void FActorTickFunction::ExecuteTick(float DeltaTime, enum ELevelTick TickType, ENamedThreads::Type CurrentThread, const FGraphEventRef& MyCompletionGraphEvent)
 {
-	if (Target && IsValidChecked(Target) && !Target->IsUnreachable())
+	if (IsValid(Target))
 	{
 		if (TickType != LEVELTICK_ViewportsOnly || Target->ShouldTickIfViewportsOnly())
 		{
@@ -1507,11 +1507,8 @@ bool AActor::DestroyNetworkActorHandled()
 
 void AActor::TickActor( float DeltaSeconds, ELevelTick TickType, FActorTickFunction& ThisTickFunction )
 {
-	//root of tick hierarchy
-
-	// Non-player update.
-	// If an Actor has been Destroyed or its level has been unloaded don't execute any queued ticks
-	if (IsValidChecked(this) && GetWorld())
+	// Actor validity was checked before this
+	if (GetWorld())
 	{
 		Tick(DeltaSeconds);	// perform any tick functions unique to an actor subclass
 	}
@@ -1521,16 +1518,13 @@ void AActor::Tick( float DeltaSeconds )
 {
 	if (GetClass()->HasAnyClassFlags(CLASS_CompiledFromBlueprint) || !GetClass()->HasAnyClassFlags(CLASS_Native))
 	{
-		// Blueprint code outside of the construction script should not run in the editor
 		// Allow tick if we are not a dedicated server, or we allow this tick on dedicated servers
-		if (GetWorldSettings() != nullptr && (bAllowReceiveTickEventOnDedicatedServer || !IsRunningDedicatedServer()))
+		if (bAllowReceiveTickEventOnDedicatedServer || !IsRunningDedicatedServer())
 		{
 			ReceiveTick(DeltaSeconds);
 		}
 
-
 		// Update any latent actions we have for this actor
-
 		// If this tick is skipped on a frame because we've got a TickInterval, our latent actions will be ticked
 		// anyway by UWorld::Tick(). Given that, our latent actions don't need to be passed a larger
 		// DeltaSeconds to make up the frames that they missed (because they wouldn't have missed any).
@@ -1544,7 +1538,6 @@ void AActor::Tick( float DeltaSeconds )
 	}
 }
 
-
 /** If true, actor is ticked even if TickType==LEVELTICK_ViewportsOnly */
 bool AActor::ShouldTickIfViewportsOnly() const
 {
diff --git a/Engine/Source/Runtime/Engine/Private/LatentActionManager.cpp b/Engine/Source/Runtime/Engine/Private/LatentActionManager.cpp
--- a/Engine/Source/Runtime/Engine/Private/LatentActionManager.cpp
+++ b/Engine/Source/Runtime/Engine/Private/LatentActionManager.cpp
@@ -47,7 +47,7 @@ void FLatentActionManager::RemoveActionsForObject(TWeakObjectPtr<UObject> InObje
 	FObjectActions* ObjectActions = GetActionsForObject(InObject);
 	if (ObjectActions)
 	{
-		FWeakObjectAndActions* FoundEntry = ActionsToRemoveMap.FindByPredicate([InObject](const FWeakObjectAndActions& Entry) { return Entry.Key == InObject; });
+		FWeakObjectAndActions* FoundEntry = ActionsToRemoveMap.FindByPredicate([InObject](const FWeakObjectAndActions& Entry) { return Entry.Key.HasSameIndexAndSerialNumber(InObject); });
 
 		TSharedPtr<TArray<FUuidAndAction>> ActionToRemoveListPtr;
 		if (FoundEntry)
@@ -84,9 +84,9 @@ DECLARE_CYCLE_STAT(TEXT("Remove Latent Actions"), STAT_RemoveLatentActions, STAT
 
 void FLatentActionManager::BeginFrame()
 {
-	for (FObjectToActionListMap::TIterator ObjIt(ObjectToActionListMap); ObjIt; ++ObjIt)
+	for (TPair<TWeakObjectPtr<UObject>, TSharedPtr<FObjectActions>>& ActionPair : ObjectToActionListMap)
 	{
-		FObjectActions* ObjectActions = ObjIt.Value().Get();
+		FObjectActions* ObjectActions = ActionPair.Value.Get();
 		check(ObjectActions);
 		ObjectActions->bProcessedThisFrame = false;
 	}
@@ -160,11 +160,6 @@ struct FScopedLatentActionTimer
 
 void FLatentActionManager::ProcessLatentActions(UObject* InObject, float DeltaTime)
 {
-	if (InObject && !InObject->GetClass()->HasAnyClassFlags(CLASS_CompiledFromBlueprint))
-	{
-		return;
-	}
-
 #if LATENT_ACTION_PROFILING_ENABLED
 	GLatentActionStats.Reset();
 	const double StartTime = FPlatformTime::Seconds();
@@ -218,8 +213,7 @@ void FLatentActionManager::ProcessLatentActions(UObject* InObject, float DeltaTi
 		SCOPE_CYCLE_COUNTER(STAT_TickLatentActions);
 		for (FObjectToActionListMap::TIterator ObjIt(ObjectToActionListMap); ObjIt; ++ObjIt)
 		{	
-			TWeakObjectPtr<UObject> WeakPtr = ObjIt.Key();
-			UObject* Object = WeakPtr.Get();
+			UObject* Object = ObjIt.Key().Get();
 			FObjectActions* ObjectActions = ObjIt.Value().Get();
 			check(ObjectActions);
 			FActionList& ObjectActionList = ObjectActions->ActionList;
@@ -390,7 +384,7 @@ void FLatentActionManager::GetActiveUUIDs(UObject* InObject, TSet<int32>& UUIDLi
 
 FLatentActionManager::~FLatentActionManager()
 {
-	for (auto& ObjectActionListIterator : ObjectToActionListMap)
+	for (TPair<TWeakObjectPtr<UObject>, TSharedPtr<FObjectActions>>& ObjectActionListIterator : ObjectToActionListMap)
 	{
 		TSharedPtr<FObjectActions>& ObjectActions = ObjectActionListIterator.Value;
 		if (ObjectActions.IsValid())
