Fixes VirtualTexture mipmap generation for non-square blocks

In case VT blocks are non square, need to use smaller dimension to calculate mipmap count. Otherwise miptail block is created from wrong source data.

diff --git a/Engine/Source/Runtime/Engine/Private/VT/VirtualTextureDataBuilder.cpp b/Engine/Source/Runtime/Engine/Private/VT/VirtualTextureDataBuilder.cpp
--- a/Engine/Source/Runtime/Engine/Private/VT/VirtualTextureDataBuilder.cpp
+++ b/Engine/Source/Runtime/Engine/Private/VT/VirtualTextureDataBuilder.cpp
@@ -809,13 +809,13 @@ void FVirtualTextureDataBuilder::BuildBlockTiles(uint32 LayerIndex, uint32 Block
 				{
 					FString DebugName = FPaths::MakeValidFileName(*DebugTexturePathName, TEXT('_'));
 					FString BasePath = FPaths::ProjectUserDir();
-					FString TileFileName = BasePath / FString::Format(TEXT("{0}_{1}_{2}_{3}_{4}_{5}.png"), TArray<FStringFormatArg>({ *DebugName, BlockIndex, TileX, TileY, Mip, LayerIndex }));
+					FString TileFileName = BasePath / FString::Format(TEXT("{0}_L{1}_B{2}_M{3}_X{4}_Y{5}.png"), TArray<FStringFormatArg>({ *DebugName, LayerIndex, BlockIndex, Mip, TileX, TileY }));
 					TileData.Save(TileFileName, ImageWrapper);
 				}
 #endif // SAVE_TILES
 
 				// give each tile a unique DebugTexturePathName for DebugDump option :
-				FString DebugTilePathName = FString::Printf(TEXT("%s_L%d_VT%04d"), *DebugTexturePathName, LayerIndex, TileIndex);
+				FString DebugTilePathName = FString::Printf(TEXT("%s_L%d_VT%04d_B%d_M%d_X%d_Y%d"), *DebugTexturePathName, LayerIndex, TileIndex, BlockIndex, Mip, TileX, TileY);
 
 				TArray<FCompressedImage2D> CompressedMip;
 				TArray<FImage> EmptyList;
@@ -852,14 +852,13 @@ void FVirtualTextureDataBuilder::BuildLayerBlocks(FSlowTask& BuildTask, uint32 L
 	// Miptail
 	TArray<FImage> MiptailInputImages;
 	FPixelDataRectangle MiptailPixelData{ LayerData.SourceFormat, 0, 0, 0 };
-	const uint32 BlockSize = FMath::Max(DerivedInfo.BlockSizeX, DerivedInfo.BlockSizeY);
+	const uint32 BlockSize = FMath::Min(DerivedInfo.BlockSizeX, DerivedInfo.BlockSizeY);
 	const uint32 BlockSizeInTiles = FMath::DivideAndRoundUp<uint32>(BlockSize, TileSize);
 	const uint32 MaxMipInBlock = FMath::CeilLogTwo(BlockSizeInTiles);
 	const uint32 MipWidthInBlock = FMath::Max<uint32>(DerivedInfo.BlockSizeX >> MaxMipInBlock, 1);
 	const uint32 MipHeightInBlock = FMath::Max<uint32>(DerivedInfo.BlockSizeY >> MaxMipInBlock, 1);
 	const uint32 MipInputSizeX = FMath::RoundUpToPowerOfTwo(DerivedInfo.SizeInBlocksX * MipWidthInBlock);
 	const uint32 MipInputSizeY = FMath::RoundUpToPowerOfTwo(DerivedInfo.SizeInBlocksY * MipHeightInBlock);
-	const uint32 MipInputSize = FMath::Max(MipInputSizeX, MipInputSizeY);
 
 	// If we have more than 1 block and we can produce more mips than each block has
 	// then need to create miptail that contains mips made from multiple blocks2
@@ -1032,7 +1031,7 @@ void FVirtualTextureDataBuilder::BuildLayerBlocks(FSlowTask& BuildTask, uint32 L
 
 		// Use actual block size (not the the one UDIM's passe here) to determine how many
 		// mips you'll have. As different blocks can be smaller than full UDIM block size
-		const uint32 BlockSizeXY = FMath::Max(BlockData.SizeX, BlockData.SizeY);
+		const uint32 BlockSizeXY = FMath::Min(BlockData.SizeX, BlockData.SizeY);
 		if (NumBlocks == 1u)
 		{
 			const uint32 MaxMipInBlockXY = FMath::CeilLogTwo(BlockSizeXY);
