- Fix crash in Cascade GPU with incorrect buffer size calculation
- The new tiled allocator uses 4xfloat not 2xfloat, ClearTiles not updated to reflect this resulting in incorrect buffer num bytes calculation

#rnx

diff --git a/Engine/Source/Runtime/Engine/Private/Particles/ParticleGpuSimulation.cpp b/Engine/Source/Runtime/Engine/Private/Particles/ParticleGpuSimulation.cpp
--- a/Engine/Source/Runtime/Engine/Private/Particles/ParticleGpuSimulation.cpp
+++ b/Engine/Source/Runtime/Engine/Private/Particles/ParticleGpuSimulation.cpp
@@ -1679,7 +1679,7 @@ void ClearTiles(FRHICommandList& RHICmdList, FGraphicsPipelineStateInitializer&
 
 	SetGraphicsPipelineState(RHICmdList, GraphicsPSOInit, 0);
 	
-	const int32 MaxTilesPerDrawCallUnaligned = GParticleScratchVertexBufferSize / sizeof(FVector2f);
+	const int32 MaxTilesPerDrawCallUnaligned = GParticleScratchVertexBufferSize / sizeof(FVector4f);
 	const int32 MaxTilesPerDrawCall = (FeatureLevel <= ERHIFeatureLevel::ES3_1 ? MaxTilesPerDrawCallUnaligned : MaxTilesPerDrawCallUnaligned & (~(TILES_PER_INSTANCE-1)));
 	int32 TileCount = Tiles.Num();
 	int32 FirstTile = 0;
