Partial backout of 33180448.

The displacement scaling material node doesn't work correctly with material instances that do not create static permutations, and we can't assume that displacement scaling of the material can be known as final during material translation. As a result, we stop trying to determine if a constant displacement could effectively disable tessellation and instead return to disabling displacement when the displacement pin has not been connected, which can still be used (via Material Attributes) to disable displacement based on static switches.

#rnx

diff --git a/Engine/Source/Developer/MaterialBaking/Private/ExportMaterialProxy.h b/Engine/Source/Developer/MaterialBaking/Private/ExportMaterialProxy.h
--- a/Engine/Source/Developer/MaterialBaking/Private/ExportMaterialProxy.h
+++ b/Engine/Source/Developer/MaterialBaking/Private/ExportMaterialProxy.h
@@ -76,16 +76,6 @@ struct FExportMaterialCompiler : public FProxyMaterialCompiler
 		return Compiler->VertexNormal();
 	}
 
-	virtual int32 DisplacementCenter() override
-	{		
-		return Compiler->DisplacementCenter();
-	}
-
-	virtual int32 DisplacementMagnitude() override
-	{		
-		return Compiler->DisplacementMagnitude();
-	}
-
 	virtual int32 ParticleRandom() override
 	{
 		return Compiler->Constant(0.0f);
diff --git a/Engine/Source/Developer/MaterialUtilities/Private/MaterialUtilities.cpp b/Engine/Source/Developer/MaterialUtilities/Private/MaterialUtilities.cpp
--- a/Engine/Source/Developer/MaterialUtilities/Private/MaterialUtilities.cpp
+++ b/Engine/Source/Developer/MaterialUtilities/Private/MaterialUtilities.cpp
@@ -289,16 +289,6 @@ struct FExportMaterialCompiler : public FProxyMaterialCompiler
 		return Compiler->VertexNormal();
 	}
 	
-	virtual int32 DisplacementCenter() override
-	{
-		return Compiler->DisplacementCenter();
-	}
-	
-	virtual int32 DisplacementMagnitude() override
-	{
-		return Compiler->DisplacementMagnitude();
-	}
-
 	virtual int32 ParticleRandom() override
 	{
 		return Compiler->Constant(0.0f);
diff --git a/Engine/Source/Runtime/Core/Private/UObject/FortniteMainBranchObjectVersion.cpp b/Engine/Source/Runtime/Core/Private/UObject/FortniteMainBranchObjectVersion.cpp
--- a/Engine/Source/Runtime/Core/Private/UObject/FortniteMainBranchObjectVersion.cpp
+++ b/Engine/Source/Runtime/Core/Private/UObject/FortniteMainBranchObjectVersion.cpp
@@ -14,6 +14,6 @@ TMap<FGuid, FGuid> FFortniteMainBranchObjectVersion::GetSystemGuids()
 	SystemGuids.Add(DevGuids.Niagara_LatestScriptCompileVersion, FGuid("C2285DEBF1B14D48BC67C27764A13F16"));
 	SystemGuids.Add(DevGuids.SkeletalMeshDerivedDataVersion, FGuid("B9AEAA3EE7AC4FCFB3860E8F3E06DD75"));
 	SystemGuids.Add(DevGuids.STATICMESH_DERIVEDDATA_VER, FGuid("543843B103794E0E9BA4BE80FB602F79"));
-	SystemGuids.Add(DevGuids.MaterialTranslationDDCVersion, FGuid("7A75688F394C4F27837B35000234B087"));
+	SystemGuids.Add(DevGuids.MaterialTranslationDDCVersion, FGuid("F43DE2FD43D353A9ABE260AD4CE9F3D4"));
 	return SystemGuids;
 }
diff --git a/Engine/Source/Runtime/Engine/Classes/Materials/MaterialExpressionDisplacementScaling.h b/Engine/Source/Runtime/Engine/Classes/Materials/MaterialExpressionDisplacementScaling.h
deleted file mode 100644
diff --git a/Engine/Source/Runtime/Engine/Private/Materials/HLSLMaterialTranslator.cpp b/Engine/Source/Runtime/Engine/Private/Materials/HLSLMaterialTranslator.cpp
--- a/Engine/Source/Runtime/Engine/Private/Materials/HLSLMaterialTranslator.cpp
+++ b/Engine/Source/Runtime/Engine/Private/Materials/HLSLMaterialTranslator.cpp
@@ -1729,11 +1729,10 @@ void FHLSLMaterialTranslator::DoTranslate()
 	bool bUsesWorldPositionOffsetPrevious = IsMaterialPropertyUsed(MP_WorldPositionOffset, Chunk[CompiledMP_PrevWorldPositionOffset], FLinearColor(0, 0, 0, 0), 3);
 	bUsesWorldPositionOffset = bUsesWorldPositionOffsetCurrent || bUsesWorldPositionOffsetPrevious;
 
-	const FDisplacementScaling DisplacementScaling = Material->GetDisplacementScaling();
+	// Check if displacement is actually used (disable if it was left at the default, invalid constant value of -1)
 	bUsesDisplacement = DoesPlatformSupportNanite(Platform) &&
 		Material->IsTessellationEnabled() &&
-		!FMath::IsNearlyZero(DisplacementScaling.Magnitude, UE_KINDA_SMALL_NUMBER) &&
-		IsMaterialPropertyUsed(MP_Displacement, Chunk[MP_Displacement], FLinearColor(DisplacementScaling.Center, 0, 0, 0), 1);
+		IsMaterialPropertyUsed(MP_Displacement, Chunk[MP_Displacement], FLinearColor(-1, 0, 0, 0), 1);
 
 	MaterialCompilationOutput.bModifiesMeshPosition = bUsesPixelDepthOffset || bUsesWorldPositionOffset || bUsesDisplacement;
 	MaterialCompilationOutput.bUsesWorldPositionOffset = bUsesWorldPositionOffset;
@@ -11130,16 +11129,6 @@ int32 FHLSLMaterialTranslator::PixelNormalWS()
 	return AddInlinedCodeChunk(MCT_Float3,TEXT("Parameters.WorldNormal"));	
 }
 
-int32 FHLSLMaterialTranslator::DisplacementCenter()
-{
-	return Constant(Material->GetDisplacementScaling().Center);
-}
-
-int32 FHLSLMaterialTranslator::DisplacementMagnitude()
-{
-	return Constant(Material->GetDisplacementScaling().Magnitude);
-}
-
 int32 FHLSLMaterialTranslator::Derivative(int32 A, EDervativeComponent Component)
 {
 	if (A == INDEX_NONE)
diff --git a/Engine/Source/Runtime/Engine/Private/Materials/HLSLMaterialTranslator.h b/Engine/Source/Runtime/Engine/Private/Materials/HLSLMaterialTranslator.h
--- a/Engine/Source/Runtime/Engine/Private/Materials/HLSLMaterialTranslator.h
+++ b/Engine/Source/Runtime/Engine/Private/Materials/HLSLMaterialTranslator.h
@@ -1162,9 +1162,6 @@ protected:
 	virtual int32 VertexTangent() override;
 	virtual int32 PixelNormalWS() override;
 
-	virtual int32 DisplacementCenter() override;
-	virtual int32 DisplacementMagnitude() override;
-
 	virtual int32 DDX(int32 A) override;
 	virtual int32 DDY(int32 A) override;
 
diff --git a/Engine/Source/Runtime/Engine/Private/Materials/MaterialAttributeDefinitionMap.cpp b/Engine/Source/Runtime/Engine/Private/Materials/MaterialAttributeDefinitionMap.cpp
--- a/Engine/Source/Runtime/Engine/Private/Materials/MaterialAttributeDefinitionMap.cpp
+++ b/Engine/Source/Runtime/Engine/Private/Materials/MaterialAttributeDefinitionMap.cpp
@@ -73,7 +73,8 @@ int32 FMaterialAttributeDefintion::CompileDefaultValue(FMaterialCompiler* Compil
 
 	if (Property == MP_Displacement)
 	{
-		return Compiler->DisplacementCenter();
+		// Initialize displacement with an out of bounds value so we can detect that it's not been hooked up
+		return Compiler->Constant(-1.0f);
 	}
 
 	if (TexCoordIndex == INDEX_NONE)
@@ -267,7 +268,7 @@ void FMaterialAttributeDefinitionMap::InitializeAttributeMap()
 
 	// Advanced attributes
 	Add(FGuid(0xF905F895, 0xD5814314, 0x916D2434, 0x8C40CE9E), TEXT("WorldPositionOffset"),		MP_WorldPositionOffset,		MCT_Float3,	FVector4(0,0,0,0),	SF_Vertex);
-	Add(FGuid(0x199A7166, 0xC67041DC, 0xA68EAD0D, 0x7017D0AD), TEXT("Displacement"),			MP_Displacement,			MCT_Float,	FVector4(.5,0,0,0), SF_Pixel);
+	Add(FGuid(0x199A7166, 0xC67041DC, 0xA68EAD0D, 0x7017D0AD), TEXT("Displacement"),			MP_Displacement,			MCT_Float,	FVector4(0.5,0,0,0), SF_Pixel);
 	Add(FGuid(0x5B8FC679, 0x51CE4082, 0x9D777BEE, 0xF4F72C44), TEXT("SubsurfaceColor"),			MP_SubsurfaceColor,			MCT_Float3,	FVector4(1,1,1,0),	SF_Pixel);
 	Add(FGuid(0x9E502E69, 0x3C8F48FA, 0x94645CFD, 0x28E5428D), TEXT("ClearCoat"),				MP_CustomData0,				MCT_Float,	FVector4(1,0,0,0),	SF_Pixel);
 	Add(FGuid(0xBE4F2FFD, 0x12FC4296, 0xB0124EEA, 0x12C28D92), TEXT("ClearCoatRoughness"),		MP_CustomData1,				MCT_Float,	FVector4(.1,0,0,0),	SF_Pixel);
diff --git a/Engine/Source/Runtime/Engine/Private/Materials/MaterialExpressionHLSL.cpp b/Engine/Source/Runtime/Engine/Private/Materials/MaterialExpressionHLSL.cpp
--- a/Engine/Source/Runtime/Engine/Private/Materials/MaterialExpressionHLSL.cpp
+++ b/Engine/Source/Runtime/Engine/Private/Materials/MaterialExpressionHLSL.cpp
@@ -65,7 +65,6 @@
 #include "Materials/MaterialExpressionDepthOfFieldFunction.h"
 #include "Materials/MaterialExpressionDeriveNormalZ.h"
 #include "Materials/MaterialExpressionDesaturation.h"
-#include "Materials/MaterialExpressionDisplacementScaling.h"
 #include "Materials/MaterialExpressionDistance.h"
 #include "Materials/MaterialExpressionDistanceCullFade.h"
 #include "Materials/MaterialExpressionDistanceFieldApproxAO.h"
@@ -5205,22 +5204,4 @@ bool UMaterialExpressionNeuralNetworkOutput::GenerateHLSLExpression(FMaterialHLS
 	return true;
 }
 
-bool UMaterialExpressionDisplacementScaling::GenerateHLSLExpression(FMaterialHLSLGenerator& Generator, UE::HLSLTree::FScope& Scope, int32 OutputIndex, UE::HLSLTree::FExpression const*& OutExpression) const
-{
-	const UMaterial* BaseMaterial = Generator.GetTargetMaterial();
-	const FDisplacementScaling Scaling = BaseMaterial->GetDisplacementScaling();
-	const UE::HLSLTree::FExpression* Expr = nullptr;
-
-	if (OutputIndex == 0)
-	{
-		Expr = Generator.NewConstant(Scaling.Center);
-	}
-	else if (OutputIndex == 1)
-	{
-		Expr = Generator.NewConstant(Scaling.Magnitude);
-	}
-
-	return Expr != nullptr;
-}
-
 #endif // WITH_EDITOR
diff --git a/Engine/Source/Runtime/Engine/Private/Materials/MaterialExpressions.cpp b/Engine/Source/Runtime/Engine/Private/Materials/MaterialExpressions.cpp
--- a/Engine/Source/Runtime/Engine/Private/Materials/MaterialExpressions.cpp
+++ b/Engine/Source/Runtime/Engine/Private/Materials/MaterialExpressions.cpp
@@ -95,7 +95,6 @@
 #include "Materials/MaterialExpressionDistance.h"
 #include "Materials/MaterialExpressionDistanceCullFade.h"
 #include "Materials/MaterialExpressionDistanceFieldsRenderingSwitch.h"
-#include "Materials/MaterialExpressionDisplacementScaling.h"
 #include "Materials/MaterialExpressionDivide.h"
 #include "Materials/MaterialExpressionDotProduct.h"
 #include "Materials/MaterialExpressionDynamicParameter.h"
@@ -29296,41 +29295,4 @@ uint32 UMaterialExpressionNeuralNetworkOutput::GetInputType(int32 InputIndex)
 }
 #endif // WITH_EDITOR
 
-///////////////////////////////////////////////////////////////////////////////
-// UMaterialExpressionActorPositionWS
-///////////////////////////////////////////////////////////////////////////////
-
-UMaterialExpressionDisplacementScaling::UMaterialExpressionDisplacementScaling(const FObjectInitializer& ObjectInitializer)
-	: Super(ObjectInitializer)
-{
-#if WITH_EDITOR
-	Outputs.Reset();
-	Outputs.Add(FExpressionOutput(TEXT("Center")));
-	Outputs.Add(FExpressionOutput(TEXT("Magnitude")));
-
-	bShowOutputNameOnPin = true;
-#endif
-}
-
-#if WITH_EDITOR
-int32 UMaterialExpressionDisplacementScaling::Compile(FMaterialCompiler* Compiler, int32 OutputIndex)
-{
-	switch (OutputIndex)
-	{
-	case 0:
-		return Compiler->DisplacementCenter();
-	case 1:
-		return Compiler->DisplacementMagnitude();
-	default:
-		return INDEX_NONE;
-	}
-}
-
-void UMaterialExpressionDisplacementScaling::GetCaption(TArray<FString>& OutCaptions) const
-{
-	OutCaptions.Add(TEXT("Displacement Scaling"));
-}
-#endif // WITH_EDITOR
-
-
 #undef LOCTEXT_NAMESPACE
diff --git a/Engine/Source/Runtime/Engine/Public/MaterialCompiler.h b/Engine/Source/Runtime/Engine/Public/MaterialCompiler.h
--- a/Engine/Source/Runtime/Engine/Public/MaterialCompiler.h
+++ b/Engine/Source/Runtime/Engine/Public/MaterialCompiler.h
@@ -515,9 +515,6 @@ public:
 	virtual int32 VertexNormal() = 0;
 	virtual int32 VertexTangent() = 0;
 	virtual int32 PixelNormalWS() = 0;
-	
-	virtual int32 DisplacementCenter() = 0;
-	virtual int32 DisplacementMagnitude() = 0;
 
 	virtual int32 CustomExpression(class UMaterialExpressionCustom* Custom, int32 OutputIndex, TArray<int32>& CompiledInputs) = 0;
 	virtual int32 CustomOutput(class UMaterialExpressionCustomOutput* Custom, int32 OutputIndex, int32 OutputCode) = 0;
@@ -1069,9 +1066,6 @@ public:
 	virtual int32 VertexTangent() override { return Compiler->VertexTangent(); }
 	virtual int32 PixelNormalWS() override { return Compiler->PixelNormalWS(); }
 
-	virtual int32 DisplacementCenter() override { return Compiler->DisplacementCenter(); }
-	virtual int32 DisplacementMagnitude() override { return Compiler->DisplacementMagnitude(); }
-
 	virtual int32 CustomExpression(class UMaterialExpressionCustom* Custom, int32 OutputIndex, TArray<int32>& CompiledInputs) override { return Compiler->CustomExpression(Custom, OutputIndex, CompiledInputs); }
 	virtual int32 CustomOutput(class UMaterialExpressionCustomOutput* Custom, int32 OutputIndex, int32 OutputCode) override { return Compiler->CustomOutput(Custom, OutputIndex, OutputCode); }
 	virtual int32 VirtualTextureOutput(uint8 AttributeMask) override { return Compiler->VirtualTextureOutput(AttributeMask); }
