Include StringView.h instead of StringFwd.h since FWorldPartitionCookPackage needs to know the definition (can't be a forward)
Fix for non unity build

diff --git a/Engine/Source/Runtime/Engine/Public/WorldPartition/Cook/WorldPartitionCookPackage.h b/Engine/Source/Runtime/Engine/Public/WorldPartition/Cook/WorldPartitionCookPackage.h
--- a/Engine/Source/Runtime/Engine/Public/WorldPartition/Cook/WorldPartitionCookPackage.h
+++ b/Engine/Source/Runtime/Engine/Public/WorldPartition/Cook/WorldPartitionCookPackage.h
@@ -8,7 +8,7 @@
 #include "WorldPartition/WorldPartitionLog.h"
 #include "UObject/Package.h"
 #else
-#include "Containers/StringFwd.h"
+#include "Containers/StringView.h"
 #endif
 
 struct FWorldPartitionCookPackage
