WorldPartition DataLayerLogicOperator now also affects loading of actors using data layers

diff --git a/Engine/Source/Runtime/Engine/Private/WorldPartition/DataLayer/DataLayerLoadingPolicy.cpp b/Engine/Source/Runtime/Engine/Private/WorldPartition/DataLayer/DataLayerLoadingPolicy.cpp
--- a/Engine/Source/Runtime/Engine/Private/WorldPartition/DataLayer/DataLayerLoadingPolicy.cpp
+++ b/Engine/Source/Runtime/Engine/Private/WorldPartition/DataLayer/DataLayerLoadingPolicy.cpp
@@ -2,19 +2,26 @@
 
 #include "WorldPartition/DataLayer/DataLayerLoadingPolicy.h"
 #include "WorldPartition/DataLayer/DataLayerInstance.h"
+#include "Algo/AllOf.h"
+#include "Algo/AnyOf.h"
 
 #include UE_INLINE_GENERATED_CPP_BY_NAME(DataLayerLoadingPolicy)
 
 #if WITH_EDITOR
-bool UDataLayerLoadingPolicy::ResolveIsLoadedInEditor(TArray<const UDataLayerInstance*>& InDataLayers) const
+EWorldPartitionDataLayersLogicOperator UDataLayerLoadingPolicy::GetDataLayersLogicOperator() const
 {
-	check(!InDataLayers.IsEmpty());
-	for (const UDataLayerInstance* DataLayerInstance : InDataLayers)
+	return GetOuterUDataLayerManager()->GetWorld()->GetWorldPartition()->GetDataLayersLogicOperator();
+}
+
+bool UDataLayerLoadingPolicy::ResolveIsLoadedInEditor(TArray<const UDataLayerInstance*>& InDataLayerInstances) const
+{
+	check(!InDataLayerInstances.IsEmpty());
+	switch (GetDataLayersLogicOperator())
 	{
-		if (DataLayerInstance->IsEffectiveLoadedInEditor())
-		{
-			return true;
-		}
+		case EWorldPartitionDataLayersLogicOperator::Or:
+			return Algo::AnyOf(InDataLayerInstances, [](const UDataLayerInstance* DataLayerInstance) { return DataLayerInstance->IsEffectiveLoadedInEditor(); });
+		case EWorldPartitionDataLayersLogicOperator::And:
+			return Algo::AllOf(InDataLayerInstances, [](const UDataLayerInstance* DataLayerInstance) { return DataLayerInstance->IsEffectiveLoadedInEditor(); });
 	}
 	return false;
 }
diff --git a/Engine/Source/Runtime/Engine/Public/WorldPartition/DataLayer/DataLayerLoadingPolicy.h b/Engine/Source/Runtime/Engine/Public/WorldPartition/DataLayer/DataLayerLoadingPolicy.h
--- a/Engine/Source/Runtime/Engine/Public/WorldPartition/DataLayer/DataLayerLoadingPolicy.h
+++ b/Engine/Source/Runtime/Engine/Public/WorldPartition/DataLayer/DataLayerLoadingPolicy.h
@@ -4,16 +4,20 @@
 
 #include "CoreMinimal.h"
 #include "UObject/Object.h"
+#include "WorldPartition/WorldPartition.h"
 #include "DataLayerLoadingPolicy.generated.h"
 
 class UDataLayerInstance;
 
-UCLASS()
+UCLASS(Within = DataLayerManager)
 class ENGINE_API UDataLayerLoadingPolicy : public UObject
 {
 	GENERATED_BODY()
-public:
 #if WITH_EDITOR
+public:
 	virtual bool ResolveIsLoadedInEditor(TArray<const UDataLayerInstance*>& InDataLayers) const;
+
+protected:
+	EWorldPartitionDataLayersLogicOperator GetDataLayersLogicOperator() const;
 #endif
 };
\ No newline at end of file
