Fix for LOCTEXT duplicates

diff --git a/Engine/Source/Runtime/Engine/Private/EdGraph/EdGraphNode.cpp b/Engine/Source/Runtime/Engine/Private/EdGraph/EdGraphNode.cpp
--- a/Engine/Source/Runtime/Engine/Private/EdGraph/EdGraphNode.cpp
+++ b/Engine/Source/Runtime/Engine/Private/EdGraph/EdGraphNode.cpp
@@ -338,7 +338,7 @@ void UEdGraphNode::DiffProperties(UStruct* StructA, UStruct* StructB, uint8* Dat
 
 				if (StringBuilder.Len() > 0)
 				{
-					Diff.DisplayString = FText::Format(LOCTEXT("DIF_NodePropertyFmt", "Property Changed: {0} - {1}"),
+					Diff.DisplayString = FText::Format(LOCTEXT("DIF_NodePropertyWithSubPropertiesFmt", "Property Changed: {0} - {1}"),
 						FText::FromString(Prop->GetName()),
 						FText::FromString(StringBuilder.ToString()));
 				}
