Keep track of resident ray tracing geometries in FRayTracingGeometryManager.
- reduces FRayTracingGeometryManager::Tick(...) overhead since it doesn't need to loop over all registered geometries.
- also added a new stat to show how much memory is being used.

diff --git a/Engine/Source/Runtime/Engine/Private/Rendering/RayTracingGeometryManager.cpp b/Engine/Source/Runtime/Engine/Private/Rendering/RayTracingGeometryManager.cpp
--- a/Engine/Source/Runtime/Engine/Private/Rendering/RayTracingGeometryManager.cpp
+++ b/Engine/Source/Runtime/Engine/Private/Rendering/RayTracingGeometryManager.cpp
@@ -71,6 +71,8 @@ DECLARE_STATS_GROUP(TEXT("Ray Tracing Geometry"), STATGROUP_RayTracingGeometry,
 DECLARE_DWORD_ACCUMULATOR_STAT(TEXT("Geometry Count"), STAT_RayTracingGeometryCount, STATGROUP_RayTracingGeometry);
 DECLARE_DWORD_ACCUMULATOR_STAT(TEXT("Geometry Group Count"), STAT_RayTracingGeometryGroupCount, STATGROUP_RayTracingGeometry);
 
+DECLARE_MEMORY_STAT(TEXT("Resident Memory"), STAT_RayTracingGeometryResidentMemory, STATGROUP_RayTracingGeometry);
+
 DECLARE_DWORD_ACCUMULATOR_STAT(TEXT("Pending Builds"), STAT_RayTracingPendingBuilds, STATGROUP_RayTracingGeometry);
 DECLARE_DWORD_ACCUMULATOR_STAT(TEXT("Pending Build Primitives"), STAT_RayTracingPendingBuildPrimitives, STATGROUP_RayTracingGeometry);
 
@@ -168,7 +170,10 @@ FRayTracingGeometryManager::RayTracingGeometryHandle FRayTracingGeometryManager:
 		
 		FScopeLock ScopeLock(&MainCS);
 
-		RayTracingGeometryHandle Handle = RegisteredGeometries.Add(InGeometry);
+		FRegisteredGeometry RegisteredGeometry;
+		RegisteredGeometry.Geometry = InGeometry;
+
+		RayTracingGeometryHandle Handle = RegisteredGeometries.Add(RegisteredGeometry);
 
 		if (InGeometry->GroupHandle != INDEX_NONE)
 		{
@@ -184,6 +189,8 @@ FRayTracingGeometryManager::RayTracingGeometryHandle FRayTracingGeometryManager:
 		
 		INC_DWORD_STAT(STAT_RayTracingGeometryCount);
 
+		GRayTracingGeometryManager->RefreshRegisteredGeometry(Handle);
+
 		return Handle;
 	}
 
@@ -198,6 +205,15 @@ void FRayTracingGeometryManager::ReleaseRayTracingGeometryHandle(RayTracingGeome
 
 		FScopeLock ScopeLock(&MainCS);
 
+		FRegisteredGeometry& RegisteredGeometry = RegisteredGeometries[Handle];
+
+		int32 NumRemoved = ResidentGeometries.Remove(RegisteredGeometry.Geometry);
+
+		if (NumRemoved > 0)
+		{
+			TotalResidentSize -= RegisteredGeometry.Size;
+		}
+
 		RegisteredGeometries.RemoveAt(Handle);
 		ReferencedGeometryHandles.Remove(Handle);
 
@@ -205,6 +221,41 @@ void FRayTracingGeometryManager::ReleaseRayTracingGeometryHandle(RayTracingGeome
 	}	
 }
 
+void FRayTracingGeometryManager::RefreshRegisteredGeometry(RayTracingGeometryHandle Handle)
+{
+	FScopeLock ScopeLock(&MainCS);
+
+	if (RegisteredGeometries.IsValidIndex(Handle))
+	{
+		FRegisteredGeometry& RegisteredGeometry = RegisteredGeometries[Handle];
+
+		if (RegisteredGeometry.Geometry->IsValid() && !RegisteredGeometry.Geometry->IsEvicted())
+		{
+			const uint32 OldSize = RegisteredGeometry.Size;
+			RegisteredGeometry.Size = RegisteredGeometry.Geometry->GetRHI()->GetSizeInfo().ResultSize;
+
+			bool bAlreadyInSet;
+			ResidentGeometries.Add(RegisteredGeometry.Geometry, &bAlreadyInSet);
+
+			if (bAlreadyInSet)
+			{
+				TotalResidentSize -= OldSize;
+			}
+
+			TotalResidentSize += RegisteredGeometry.Size;
+		}
+		else
+		{
+			int32 NumRemoved = ResidentGeometries.Remove(RegisteredGeometry.Geometry);
+
+			if (NumRemoved > 0)
+			{
+				TotalResidentSize -= RegisteredGeometry.Size;
+			}
+		}
+	}
+}
+
 void FRayTracingGeometryManager::PreRender()
 {
 	bRenderedFrame = true;
@@ -236,11 +287,11 @@ void FRayTracingGeometryManager::Tick(FRHICommandList& RHICmdList)
 		if (bHasRayTracingEnableChanged)
 		{
 			// evict all geometries
-			for (FRayTracingGeometry* Geometry : RegisteredGeometries)
+			for (FRegisteredGeometry& RegisteredGeometry : RegisteredGeometries)
 			{
-				if (Geometry->GetRHI() != nullptr)
+				if (RegisteredGeometry.Geometry->GetRHI() != nullptr)
 				{
-					Geometry->Evict();
+					RegisteredGeometry.Geometry->Evict();
 				}
 			}
 		}
@@ -248,12 +299,18 @@ void FRayTracingGeometryManager::Tick(FRHICommandList& RHICmdList)
 		{
 #if DO_CHECK
 			// otherwise just check that everything is evicted
-			for (FRayTracingGeometry* Geometry : RegisteredGeometries)
+			for (FRegisteredGeometry& RegisteredGeometry : RegisteredGeometries)
 			{
-				checkf(Geometry->IsEvicted() || Geometry->GetRHI() == nullptr, TEXT("Ray tracing geometry should be evicted when ray tracing is disabled."));
+				checkf(RegisteredGeometry.Geometry->IsEvicted() || RegisteredGeometry.Geometry->GetRHI() == nullptr, TEXT("Ray tracing geometry should be evicted when ray tracing is disabled."));
 			}
 #endif
 		}
+
+		checkf(TotalResidentSize == 0,
+			TEXT("TotalResidentSize should be 0 when ray tracing is disabled but is currently %lld.\n")
+			TEXT("There's likely some issue tracking resident geometries or not all geometries have been evicted."),
+			TotalResidentSize
+		);
 	}
 	else if (IsRayTracingUsingReferenceBasedResidency())
 	{
@@ -269,40 +326,45 @@ void FRayTracingGeometryManager::Tick(FRHICommandList& RHICmdList)
 
 		bRenderedFrame = false;
 
-		TSet<FRayTracingGeometry*> NotReferencedResidentGeometries; // TODO: Keep track of this to avoid the following loop
+		TSet<FRayTracingGeometry*> NotReferencedResidentGeometries = ResidentGeometries;
 
-		for (FRayTracingGeometry* Geometry : RegisteredGeometries)
-		{
-			NotReferencedResidentGeometries.Add(Geometry);
-		}
-
-		TSet<FRayTracingGeometry*> ReferencedGeometries;
+		TArray<FRayTracingGeometry*> ReferencedGeometries;
 
 		// 1st step
 		// - map RayTracingGeometryHandle to FRayTracingGeometry
 		//		- we use handles to track referenced geometries because some of the higher level code uses const FRayTracingGeometry*
 		for (RayTracingGeometryHandle GeometryHandle : ReferencedGeometryHandles)
 		{
-			ReferencedGeometries.Add(RegisteredGeometries[GeometryHandle]);
+			FRegisteredGeometry& RegisteredGeometry = RegisteredGeometries[GeometryHandle];
+
+			ReferencedGeometries.Add(RegisteredGeometry.Geometry);
+			NotReferencedResidentGeometries.Remove(RegisteredGeometry.Geometry);
 		}
 
 		// 2nd step
 		// - add all geometries in referenced groups to ReferencedGeometries
 		//		- need to make all geometries in group resident otherwise might not have valid geometry when reducing LOD
 		//		- TODO: Could track TargetLOD and only make [TargetLOD ... LastLOD] range resident
-		for (RayTracing::GeometryGroupHandle Group : ReferencedGeometryGroups)
+		for (RayTracing::GeometryGroupHandle GroupHandle : ReferencedGeometryGroups)
 		{
-			checkf(RegisteredGroups.IsValidIndex(Group), TEXT("RayTracingGeometryGroupHandle must be valid"));
+			checkf(RegisteredGroups.IsValidIndex(GroupHandle), TEXT("RayTracingGeometryGroupHandle must be valid"));
+
+			const FRayTracingGeometryGroup& Group = RegisteredGroups[GroupHandle];
 
-			for (FRayTracingGeometry* Geometry : RegisteredGroups[Group].Geometries)
+			for (FRayTracingGeometry* Geometry : Group.Geometries)
 			{
 				if (Geometry != nullptr) // some LODs might be stripped during cook
 				{
 					ReferencedGeometries.Add(Geometry);
+					NotReferencedResidentGeometries.Remove(Geometry);
 				}
 			}
 		}
 
+#if DO_CHECK
+		ensure(ReferencedGeometries.Num() == TSet(ReferencedGeometries).Num());
+#endif
+
 		// 3rd step
 		// - make referenced geometries resident
 		for (FRayTracingGeometry* Geometry : ReferencedGeometries)
@@ -333,11 +395,11 @@ void FRayTracingGeometryManager::Tick(FRHICommandList& RHICmdList)
 		if (bHasRayTracingEnableChanged)
 		{
 			// make all geometries resident
-			for (FRayTracingGeometry* Geometry : RegisteredGeometries)
+			for (FRegisteredGeometry& RegisteredGeometry : RegisteredGeometries)
 			{
-				if (Geometry->IsEvicted())
+				if (RegisteredGeometry.Geometry->IsEvicted())
 				{
-					Geometry->MakeResident(RHICmdList);
+					RegisteredGeometry.Geometry->MakeResident(RHICmdList);
 				}
 			}
 		}
@@ -345,9 +407,9 @@ void FRayTracingGeometryManager::Tick(FRHICommandList& RHICmdList)
 		{
 #if DO_CHECK
 			// otherwise just check that all geometries are resident
-			for (FRayTracingGeometry* Geometry : RegisteredGeometries)
+			for (FRegisteredGeometry& RegisteredGeometry : RegisteredGeometries)
 			{
-				checkf(!Geometry->IsEvicted(), TEXT("Ray tracing geometry should not be evicted when ray tracing is enabled."));
+				checkf(!RegisteredGeometry.Geometry->IsEvicted(), TEXT("Ray tracing geometry should not be evicted when ray tracing is enabled."));
 			}
 #endif
 		}
@@ -357,6 +419,8 @@ void FRayTracingGeometryManager::Tick(FRHICommandList& RHICmdList)
 	ReferencedGeometryGroups.Reset();
 
 	bHasRayTracingEnableChanged = false;
+
+	SET_MEMORY_STAT(STAT_RayTracingGeometryResidentMemory, TotalResidentSize);
 }
 
 void FRayTracingGeometryManager::BoostPriority(BuildRequestIndex InRequestIndex, float InBoostValue)
diff --git a/Engine/Source/Runtime/Engine/Public/Rendering/RayTracingGeometryManager.h b/Engine/Source/Runtime/Engine/Public/Rendering/RayTracingGeometryManager.h
--- a/Engine/Source/Runtime/Engine/Public/Rendering/RayTracingGeometryManager.h
+++ b/Engine/Source/Runtime/Engine/Public/Rendering/RayTracingGeometryManager.h
@@ -30,6 +30,8 @@ public:
 	ENGINE_API virtual RayTracing::GeometryGroupHandle RegisterRayTracingGeometryGroup(uint32 NumLODs) override;
 	ENGINE_API virtual void ReleaseRayTracingGeometryGroup(RayTracing::GeometryGroupHandle Handle) override;
 
+	ENGINE_API virtual void RefreshRegisteredGeometry(RayTracingGeometryHandle Handle) override;
+
 	ENGINE_API virtual void PreRender() override;
 	ENGINE_API virtual void Tick(FRHICommandList& RHICmdList) override;
 
@@ -86,10 +88,19 @@ private:
 		// TODO: Implement use-after-free checks in RayTracing::GeometryGroupHandle using some bits to identify generation
 	};
 
+	struct FRegisteredGeometry
+	{
+		FRayTracingGeometry* Geometry = nullptr;
+		uint32 Size = 0;
+	};
+
 	TSparseArray<FRayTracingGeometryGroup> RegisteredGroups;
 
 	// Used for keeping track of geometries when ray tracing is dynamic
-	TSparseArray<FRayTracingGeometry*> RegisteredGeometries;
+	TSparseArray<FRegisteredGeometry> RegisteredGeometries;
+
+	TSet<FRayTracingGeometry*> ResidentGeometries;
+	uint64 TotalResidentSize = 0;
 
 	TSet<RayTracingGeometryHandle> ReferencedGeometryHandles;
 	TSet<RayTracing::GeometryGroupHandle> ReferencedGeometryGroups;
diff --git a/Engine/Source/Runtime/RenderCore/Private/RayTracingGeometry.cpp b/Engine/Source/Runtime/RenderCore/Private/RayTracingGeometry.cpp
--- a/Engine/Source/Runtime/RenderCore/Private/RayTracingGeometry.cpp
+++ b/Engine/Source/Runtime/RenderCore/Private/RayTracingGeometry.cpp
@@ -39,6 +39,8 @@ void FRayTracingGeometry::InitRHIForStreaming(FRHIRayTracingGeometry* Intermedia
 	{
 		Batcher.EnqueueReplace(RayTracingGeometryRHI, IntermediateGeometry);
 		EnumAddFlags(GeometryState, EGeometryStateFlags::Valid);
+
+		GRayTracingGeometryManager->RefreshRegisteredGeometry(RayTracingGeometryHandle);
 	}
 	else
 	{
@@ -59,6 +61,8 @@ void FRayTracingGeometry::ReleaseRHIForStreaming(FRHIResourceReplaceBatcher& Bat
 	{
 		Batcher.EnqueueReplace(RayTracingGeometryRHI, nullptr);
 	}
+
+	GRayTracingGeometryManager->RefreshRegisteredGeometry(RayTracingGeometryHandle);
 }
 
 void FRayTracingGeometry::RequestBuildIfNeeded(FRHICommandListBase& RHICmdList, ERTAccelerationStructureBuildPriority InBuildPriority)
@@ -130,7 +134,9 @@ void FRayTracingGeometry::Evict()
 	RemoveBuildRequest();
 	RayTracingGeometryRHI.SafeRelease();
 	EnumAddFlags(GeometryState, EGeometryStateFlags::Evicted);
-		
+	
+	GRayTracingGeometryManager->RefreshRegisteredGeometry(RayTracingGeometryHandle);
+	
 	if (GroupHandle != INDEX_NONE)
 	{
 		GRayTracingGeometryManager->RequestUpdateCachedRenderState(GroupHandle);
@@ -212,6 +218,8 @@ void FRayTracingGeometry::CreateRayTracingGeometry(FRHICommandListBase& RHICmdLi
 			Initializer.OfflineData = nullptr;
 		}
 	}
+
+	GRayTracingGeometryManager->RefreshRegisteredGeometry(RayTracingGeometryHandle);
 }
 
 bool FRayTracingGeometry::IsValid() const
diff --git a/Engine/Source/Runtime/RenderCore/Public/RayTracingGeometryManagerInterface.h b/Engine/Source/Runtime/RenderCore/Public/RayTracingGeometryManagerInterface.h
--- a/Engine/Source/Runtime/RenderCore/Public/RayTracingGeometryManagerInterface.h
+++ b/Engine/Source/Runtime/RenderCore/Public/RayTracingGeometryManagerInterface.h
@@ -53,6 +53,8 @@ public:
 
 	RENDERCORE_API virtual void RequestUpdateCachedRenderState(RayTracing::GeometryGroupHandle InRayTracingGeometryGroupHandle) = 0;
 
+	RENDERCORE_API virtual void RefreshRegisteredGeometry(RayTracingGeometryHandle Handle) = 0;
+
 	RENDERCORE_API virtual void PreRender() = 0;
 	RENDERCORE_API virtual void Tick(FRHICommandList& RHICmdList) = 0;
 };
