Add engine API to disable anim instance update
Disable most linked anim instances update for vehicle that do not feed on the main graph pose

diff --git a/Engine/Source/Runtime/Engine/Classes/Animation/AnimInstance.h b/Engine/Source/Runtime/Engine/Classes/Animation/AnimInstance.h
--- a/Engine/Source/Runtime/Engine/Classes/Animation/AnimInstance.h
+++ b/Engine/Source/Runtime/Engine/Classes/Animation/AnimInstance.h
@@ -369,6 +369,9 @@ class UAnimInstance : public UObject
 	uint8 bUseMainInstanceMontageEvaluationData: 1;
 
 private:
+	// Should UpdateAnimation be called
+	uint8 bUpdateAnimationEnabled : 1 = true;
+
 	/** True when Montages are being ticked, and Montage Events should be queued. 
 	 * When Montage are being ticked, we queue AnimNotifies and Events. We trigger notifies first, then Montage events. */
 	UPROPERTY(Transient)
@@ -1307,6 +1310,10 @@ public:
 	// Native Uninitialize override point
 	ENGINE_API virtual void NativeUninitializeAnimation();
 
+	// Enable / Disable animation update. This is provided as an optimization to disable linked instances that aren't relevant. Disabling an instance whose graph is still evaluated will assert and causes issues.
+	void EnableUpdateAnimation(bool bEnable) { bUpdateAnimationEnabled = bEnable; }
+	bool IsUpdateAnimationEnabled() const { return bUpdateAnimationEnabled; }
+
 	// Executed when begin play is called on the owning component
 	ENGINE_API virtual void NativeBeginPlay();
 
diff --git a/Engine/Source/Runtime/Engine/Private/Animation/AnimInstance.cpp b/Engine/Source/Runtime/Engine/Private/Animation/AnimInstance.cpp
--- a/Engine/Source/Runtime/Engine/Private/Animation/AnimInstance.cpp
+++ b/Engine/Source/Runtime/Engine/Private/Animation/AnimInstance.cpp
@@ -491,6 +491,12 @@ void UAnimInstance::UpdateAnimation(float DeltaSeconds, bool bNeedsValidRootMoti
 	checkf(!bUpdatingAnimation, TEXT("UpdateAnimation already in progress, circular detected for SkeletalMeshComponent [%s], AnimInstance [%s]"), *GetNameSafe(GetOwningComponent()),  *GetName());
 	TGuardValue<bool> CircularGuard(bUpdatingAnimation, true);
 #endif
+
+	if (!IsUpdateAnimationEnabled())
+	{
+		return;
+	}
+
 	SCOPE_CYCLE_COUNTER(STAT_UpdateAnimation);
 	FScopeCycleCounterUObject AnimScope(this);
 
diff --git a/Engine/Source/Runtime/Engine/Private/Animation/AnimInstanceProxy.cpp b/Engine/Source/Runtime/Engine/Private/Animation/AnimInstanceProxy.cpp
--- a/Engine/Source/Runtime/Engine/Private/Animation/AnimInstanceProxy.cpp
+++ b/Engine/Source/Runtime/Engine/Private/Animation/AnimInstanceProxy.cpp
@@ -1232,6 +1232,15 @@ void FAnimInstanceProxy::UpdateAnimation_WithRoot(const FAnimationUpdateContext&
 	ANIM_MT_SCOPE_CYCLE_COUNTER(ProxyUpdateAnimation, !IsInGameThread());
 	FScopeCycleCounterUObject AnimScope(bUpdatingRoot ? nullptr : GetAnimInstanceObject());
 
+
+#if DO_CHECK
+	ensureMsgf(CastChecked<UAnimInstance>(GetAnimInstanceObject())->IsUpdateAnimationEnabled(), TEXT("Updating AnimInstance with UpdateAnimation disabled: %s"), *CastChecked<UAnimInstance>(GetAnimInstanceObject())->GetName());
+#endif
+	if (!CastChecked<UAnimInstance>(GetAnimInstanceObject())->IsUpdateAnimationEnabled())
+	{
+		return;
+	}
+
 	if(InRootNode == RootNode)
 	{
 		if(bInitializeSubsystems && AnimClassInterface)
