Remove C++20 extension from AnimInstance header

diff --git a/Engine/Source/Runtime/Engine/Classes/Animation/AnimInstance.h b/Engine/Source/Runtime/Engine/Classes/Animation/AnimInstance.h
--- a/Engine/Source/Runtime/Engine/Classes/Animation/AnimInstance.h
+++ b/Engine/Source/Runtime/Engine/Classes/Animation/AnimInstance.h
@@ -370,7 +370,7 @@ class UAnimInstance : public UObject
 
 private:
 	// Should UpdateAnimation be called
-	uint8 bUpdateAnimationEnabled : 1 = true;
+	uint8 bUpdateAnimationEnabled;
 
 	/** True when Montages are being ticked, and Montage Events should be queued. 
 	 * When Montage are being ticked, we queue AnimNotifies and Events. We trigger notifies first, then Montage events. */
@@ -1312,7 +1312,7 @@ public:
 
 	// Enable / Disable animation update. This is provided as an optimization to disable linked instances that aren't relevant. Disabling an instance whose graph is still evaluated will assert and causes issues.
 	void EnableUpdateAnimation(bool bEnable) { bUpdateAnimationEnabled = bEnable; }
-	bool IsUpdateAnimationEnabled() const { return bUpdateAnimationEnabled; }
+	bool IsUpdateAnimationEnabled() const { return bUpdateAnimationEnabled != 0; }
 
 	// Executed when begin play is called on the owning component
 	ENGINE_API virtual void NativeBeginPlay();
diff --git a/Engine/Source/Runtime/Engine/Private/Animation/AnimInstance.cpp b/Engine/Source/Runtime/Engine/Private/Animation/AnimInstance.cpp
--- a/Engine/Source/Runtime/Engine/Private/Animation/AnimInstance.cpp
+++ b/Engine/Source/Runtime/Engine/Private/Animation/AnimInstance.cpp
@@ -119,7 +119,8 @@ UAnimInstance::UAnimInstance(const FObjectInitializer& ObjectInitializer)
 	bReceiveNotifiesFromLinkedInstances = false;
 	bPropagateNotifiesToLinkedInstances = false;
 	bUseMainInstanceMontageEvaluationData = false;
-
+	
+	bUpdateAnimationEnabled = true;
 #if DO_CHECK
 	bInitializing = false;
 #endif
