Fixed HISM instance data mismatch crashes.
1. Make sure that externally managed HISM (AKA landscape grass components) do not have their updates cancelled when being registered.
2. Make sure the instance count is validated when instance data proxy updates are performed (vs what was promised to the render thread in the header).
3. Changed the way the visible primitives is updated to handle the case where even in the initial build there are missing instances.
4. Fixed the HISMC InstanceCountToRender getting out of sync if AddInstances & RemoveInstances are used while a HISM hierarchy build is in progress.
5. As a last resort, in case updates are lost despite this (somehow), added TestAndApplyInstanceBufferSizeFixup that makes sure the instance data buffer sizes match after the update.

diff --git a/Engine/Source/Runtime/Engine/Private/HierarchicalInstancedStaticMesh.cpp b/Engine/Source/Runtime/Engine/Private/HierarchicalInstancedStaticMesh.cpp
--- a/Engine/Source/Runtime/Engine/Private/HierarchicalInstancedStaticMesh.cpp
+++ b/Engine/Source/Runtime/Engine/Private/HierarchicalInstancedStaticMesh.cpp
@@ -2493,12 +2493,12 @@ TArray<int32> UHierarchicalInstancedStaticMeshComponent::AddInstances(const TArr
 		for (const int32 InstanceIndex : InstanceIndices)
 		{
 			InstanceReorderTable.Add(InitialBufferOffset + InstanceIndex);
+			++InstanceCountToRender;
 
 			const FBox NewInstanceBounds = GetStaticMesh()->GetBounds().GetBox().TransformBy(PerInstanceSMData[InstanceIndex].Transform);
 			UnbuiltInstanceBounds += NewInstanceBounds;
 			UnbuiltInstanceBoundsList.Add(NewInstanceBounds);
 		}
-		InstanceCountToRender = InstanceReorderTable.Num();
 
 		if (bAutoRebuildTreeOnInstanceChanges)
 		{
diff --git a/Engine/Source/Runtime/Engine/Private/InstancedStaticMesh/ISMInstanceDataManager.cpp b/Engine/Source/Runtime/Engine/Private/InstancedStaticMesh/ISMInstanceDataManager.cpp
--- a/Engine/Source/Runtime/Engine/Private/InstancedStaticMesh/ISMInstanceDataManager.cpp
+++ b/Engine/Source/Runtime/Engine/Private/InstancedStaticMesh/ISMInstanceDataManager.cpp
@@ -418,6 +418,7 @@ void FPrimitiveInstanceDataManager::DispatchUpdateTask(bool bIsUnattached, const
 		check(!InstanceDataUpdateTaskInfo || InstanceDataUpdateTaskInfo->GetHeader() == InstanceDataBufferHeader);
 		InnerTaskLambda();
 		check(!InstanceDataUpdateTaskInfo || InstanceDataUpdateTaskInfo->GetHeader() == InstanceDataBufferHeader);
+		check(!InstanceDataUpdateTaskInfo || InstanceDataBufferHeader.NumInstances == Proxy->GetData().GetNumInstances());
 		// check(InstanceDataBufferHeader.Flags == Proxy->GetData().GetFlags());
 		const FInstanceDataFlags HeaderFlags = InstanceDataBufferHeader.Flags;
 		const bool bHasAnyPayloadData = HeaderFlags.bHasPerInstanceHierarchyOffset || HeaderFlags.bHasPerInstanceLocalBounds || HeaderFlags.bHasPerInstanceDynamicData || HeaderFlags.bHasPerInstanceLMSMUVBias || HeaderFlags.bHasPerInstanceCustomData || HeaderFlags.bHasPerInstancePayloadExtension || HeaderFlags.bHasPerInstanceEditorData;
@@ -749,9 +750,11 @@ bool FPrimitiveInstanceDataManager::FlushChanges(FInstanceUpdateComponentDesc &&
 		if (Mode == EMode::ExternalLegacyData && ChangeDesc.bPrimitiveTransformChanged)
 		{
 			LOG_INST_DATA(TEXT("Primitive Transform Update %s"), TEXT(""));
-			DispatchUpdateTask(bIsUnattached, InstanceDataBufferHeader, [ChangeSet = MoveTemp(ChangeSet), Proxy = Proxy] () mutable 
+			SuccessorTrackingState = ETrackingState::Initial;
+			DispatchUpdateTask(bIsUnattached, InstanceDataBufferHeader, [ChangeSet = MoveTemp(ChangeSet), Proxy = Proxy, NumHeaderInstances = InstanceDataBufferHeader.NumInstances] () mutable 
 			{
 				Proxy->UpdatePrimitiveTransform(MoveTemp(ChangeSet));
+				check(NumHeaderInstances == Proxy->GeInstanceSceneDataBuffers()->GetNumInstances());
 			});
 		}
 		else if (bNeedFullUpdate)
@@ -768,18 +771,20 @@ bool FPrimitiveInstanceDataManager::FlushChanges(FInstanceUpdateComponentDesc &&
 				SuccessorTrackingState = ETrackingState::Optimized;
 			}
 
-			DispatchUpdateTask(bIsUnattached, InstanceDataBufferHeader, [ChangeSet = MoveTemp(ChangeSet), Proxy = Proxy] () mutable 
+			DispatchUpdateTask(bIsUnattached, InstanceDataBufferHeader, [ChangeSet = MoveTemp(ChangeSet), Proxy = Proxy, NumHeaderInstances = InstanceDataBufferHeader.NumInstances] () mutable 
 			{
 				Proxy->Build(MoveTemp(ChangeSet));
+				check(NumHeaderInstances == Proxy->GeInstanceSceneDataBuffers()->GetNumInstances());
 			});
 		}
 		else
 		{
 			LOG_INST_DATA(TEXT("Delta Update %s"), TEXT(""));
 			check(!ChangeDesc.bUntrackedState);
-			DispatchUpdateTask(bIsUnattached, InstanceDataBufferHeader, [ChangeSet = MoveTemp(ChangeSet), Proxy = Proxy] () mutable 
+			DispatchUpdateTask(bIsUnattached, InstanceDataBufferHeader, [ChangeSet = MoveTemp(ChangeSet), Proxy = Proxy, NumHeaderInstances = InstanceDataBufferHeader.NumInstances] () mutable 
 			{
 				Proxy->Update(MoveTemp(ChangeSet));
+				check(NumHeaderInstances == Proxy->GeInstanceSceneDataBuffers()->GetNumInstances());
 			});
 		}
 	}
@@ -1059,7 +1064,8 @@ void FPrimitiveInstanceDataManager::OnRegister(int32 InNumInstances)
 
 	if (CVarInstanceDataResetTrackingOnRegister.GetValueOnGameThread())
 	{
-		if (InNumInstances != NumInstances)
+		// We don't clear for external legacy, AKA landscape grass, because it manages its own data, doesn't use the tracking & we don't want to clear the externall queued update for this case.
+		if (Mode != EMode::ExternalLegacyData && InNumInstances != NumInstances)
 		{
 			ClearIdTracking(InNumInstances);
 		}
diff --git a/Engine/Source/Runtime/Engine/Private/InstancedStaticMesh/ISMInstanceDataSceneProxy.cpp b/Engine/Source/Runtime/Engine/Private/InstancedStaticMesh/ISMInstanceDataSceneProxy.cpp
--- a/Engine/Source/Runtime/Engine/Private/InstancedStaticMesh/ISMInstanceDataSceneProxy.cpp
+++ b/Engine/Source/Runtime/Engine/Private/InstancedStaticMesh/ISMInstanceDataSceneProxy.cpp
@@ -454,6 +454,65 @@ void FISMCInstanceDataSceneProxy::BuildFromOptimizedDataBuffers(FISMInstanceUpda
 	}
 }
 
+template <typename ValueArrayType>
+static void PadArrayWithZeroData(bool bUsed, int32 DesiredSize, ValueArrayType &Array)
+{
+	if (bUsed)
+	{
+		Array.SetNumZeroed(DesiredSize);
+	}
+}
+
+void FISMCInstanceDataSceneProxy::TestAndApplyInstanceBufferSizeFixup(int32 PostUpdateNumInstances)
+{
+	// This is a fixup (that should never run, hopefully) to make sure the number of instances in the buffer matches the PostUpdateNumInstances. This is important because PostUpdateNumInstances is what 
+	// the renderer expects and will use that to index.
+	if (!ensureMsgf(PostUpdateNumInstances == InstanceSceneDataBuffers.GetNumInstances(), TEXT("The number of instances does not match what was promised in the update dispatch! Padding with dummy transforms to prevent crashes!")))
+	{
+		FInstanceSceneDataBuffers::FAccessTag AccessTag(PointerHash(this));
+		FInstanceSceneDataBuffers::FWriteView ProxyData = InstanceSceneDataBuffers.BeginWriteAccess(AccessTag);
+		if (PostUpdateNumInstances != 0 && ProxyData.InstanceLocalBounds.IsEmpty())
+		{
+			static const FRenderBounds ZeroRenderBounds(FVector3f::ZeroVector, FVector3f::ZeroVector);
+			ProxyData.InstanceLocalBounds.Add(ZeroRenderBounds);
+		}
+
+		// If it is empty, we can't figure out the number of float4s for each, so in this case the only safe thing is to remove all of them and mark as no existing.
+		if (ProxyData.Flags.bHasPerInstancePayloadExtension && 
+			(ProxyData.InstancePayloadExtension.IsEmpty() || ProxyData.InstanceToPrimitiveRelative.IsEmpty()
+				// or it is not divisible, then it is likely messed up somehow.
+				|| (ProxyData.InstancePayloadExtension.Num() % ProxyData.InstanceToPrimitiveRelative.Num()) != 0))
+		{
+			ProxyData.Flags.bHasPerInstancePayloadExtension = false;
+		}
+		else
+		{
+			int32 StrideInFloat4s = ProxyData.InstancePayloadExtension.Num() / ProxyData.InstanceToPrimitiveRelative.Num();
+			PadArrayWithZeroData(ProxyData.Flags.bHasPerInstancePayloadExtension, PostUpdateNumInstances * StrideInFloat4s, ProxyData.InstancePayloadExtension);
+		}
+
+		// Fill the array with zero transforms, these are not particularly useful but should prevent memory violations in shipping builds.
+		ProxyData.InstanceToPrimitiveRelative.SetNumZeroed(PostUpdateNumInstances);
+
+		PadArrayWithZeroData(ProxyData.Flags.bHasPerInstanceCustomData, PostUpdateNumInstances * ProxyData.NumCustomDataFloats, ProxyData.InstanceCustomData);
+		PadArrayWithZeroData(ProxyData.Flags.bHasPerInstanceDynamicData, PostUpdateNumInstances, ProxyData.PrevInstanceToPrimitiveRelative);
+#if WITH_EDITOR
+		PadArrayWithZeroData(ProxyData.Flags.bHasPerInstanceEditorData, PostUpdateNumInstances, ProxyData.InstanceEditorData);
+#endif
+		PadArrayWithZeroData(ProxyData.Flags.bHasPerInstanceHierarchyOffset, PostUpdateNumInstances, ProxyData.InstanceHierarchyOffset);
+		PadArrayWithZeroData(ProxyData.Flags.bHasPerInstanceLMSMUVBias, PostUpdateNumInstances, ProxyData.InstanceLightShadowUVBias);
+		PadArrayWithZeroData(ProxyData.Flags.bHasPerInstanceLocalBounds, PostUpdateNumInstances, ProxyData.InstanceLocalBounds);
+		PadArrayWithZeroData(ProxyData.Flags.bHasPerInstanceRandom, PostUpdateNumInstances, ProxyData.InstanceRandomIDs);
+		if (ProxyData.Flags.bHasPerInstanceVisible)
+		{
+			ProxyData.VisibleInstances.SetNum(PostUpdateNumInstances, false);
+		}
+
+		InstanceSceneDataBuffers.EndWriteAccess(AccessTag);
+	}
+
+}
+
 void FISMCInstanceDataSceneProxy::Build(FISMInstanceUpdateChangeSet&& ChangeSet)
 {
 	SCOPED_NAMED_EVENT(FISMCInstanceDataSceneProxy_Build, FColor::Emerald);
@@ -625,23 +684,6 @@ void FISMCInstanceDataSceneProxyLegacyReordered::Update(FISMInstanceUpdateChange
 	FInstanceSceneDataBuffers::FWriteView ProxyData = InstanceSceneDataBuffers.BeginWriteAccess(AccessTag);
 
 	ProxyData.Flags = ChangeSet.Flags;
-	// Handle deletions before updating the data.
-	{
-		FReorderTableIndexRemap IndexRemapOld(LegacyInstanceReorderTable, InstanceSceneDataBuffers.GetNumInstances(AccessTag));
-
-		ProxyData.VisibleInstances.SetNum(ChangeSet.PostUpdateNumInstances, true);
-		ProxyData.Flags.bHasPerInstanceVisible = true;
-		for (auto It = ChangeSet.InstanceAttributeTracker.GetRemovedIterator(); It; ++It)
-		{
-			// This is somewhat nonintuitive, but the current instance->index map is where we retain knowledge of where the instance used to be placed (in the component address space at last update)
-			int32 InstanceIndex = InstanceIdIndexMap.IdToIndex(FPrimitiveInstanceId{It.GetIndex()});
-			if (IndexRemapOld.RemapDestIndex(InstanceIndex))
-			{
-				LOG_INST_DATA(TEXT("Update/HideInstance, ID: %d, IDX: %d"), It.GetIndex(), InstanceIndex);
-				ProxyData.VisibleInstances[InstanceIndex] = false;
-			}
-		}
-	}
 	UpdateIdMapping(ChangeSet, FIdentityIndexRemap());
 
 	LegacyInstanceReorderTable = MoveTemp(ChangeSet.LegacyInstanceReorderTable);
@@ -650,6 +692,26 @@ void FISMCInstanceDataSceneProxyLegacyReordered::Update(FISMInstanceUpdateChange
 	// Use the index reorder table to scatter the data to the correct locations.
 	ApplyDataChanges(ChangeSet, IndexRemap, ChangeSet.PostUpdateNumInstances, ProxyData);
 	
+	if (bLegacyReordered && ChangeSet.PostUpdateNumInstances != LegacyInstanceReorderTable.Num())
+	{
+		// Make sure any instance no longer represented in the reorder table is hidden.
+		ProxyData.VisibleInstances.Init(false, ChangeSet.PostUpdateNumInstances);
+		for (int32 InstanceIndex : LegacyInstanceReorderTable)
+		{
+			if (IndexRemap.RemapDestIndex(InstanceIndex))
+			{
+				ProxyData.VisibleInstances[InstanceIndex] = true;
+			}
+		}
+		ProxyData.Flags.bHasPerInstanceVisible = true;
+	}
+	else
+	{
+		// Mark everything as visible from the start.
+		ProxyData.VisibleInstances.Reset();
+		ProxyData.Flags.bHasPerInstanceVisible = false;
+	}	
+
 	InstanceSceneDataBuffers.EndWriteAccess(AccessTag);
 
 	InstanceSceneDataBuffers.ValidateData();
@@ -683,8 +745,7 @@ void FISMCInstanceDataSceneProxyLegacyReordered::Build(FISMInstanceUpdateChangeS
 	// Is there is a reorder table and it does not have the same number as the instances, some must be hidden
 	if (bLegacyReordered && ChangeSet.PostUpdateNumInstances != LegacyInstanceReorderTable.Num())
 	{
-		ProxyData.VisibleInstances.Reset();
-		ProxyData.VisibleInstances.SetNum(ChangeSet.PostUpdateNumInstances, false);
+		ProxyData.VisibleInstances.Init(false, ChangeSet.PostUpdateNumInstances);
 		for (int32 InstanceIndex : LegacyInstanceReorderTable)
 		{
 			if (IndexRemap.RemapDestIndex(InstanceIndex))
@@ -818,6 +879,8 @@ void FISMCInstanceDataSceneProxyLegacyReordered::UpdatePrimitiveTransform(FISMIn
 
 		InstanceSceneDataBuffers.ValidateData();
 	}
+
+	TestAndApplyInstanceBufferSizeFixup(ChangeSet.PostUpdateNumInstances);
 }
 
 FISMCInstanceDataSceneProxyNoGPUScene::FISMCInstanceDataSceneProxyNoGPUScene(FStaticShaderPlatform InShaderPlatform, ERHIFeatureLevel::Type InFeatureLevel, bool bInLegacyReordered) 
diff --git a/Engine/Source/Runtime/Engine/Public/InstancedStaticMesh/ISMInstanceDataSceneProxy.h b/Engine/Source/Runtime/Engine/Public/InstancedStaticMesh/ISMInstanceDataSceneProxy.h
--- a/Engine/Source/Runtime/Engine/Public/InstancedStaticMesh/ISMInstanceDataSceneProxy.h
+++ b/Engine/Source/Runtime/Engine/Public/InstancedStaticMesh/ISMInstanceDataSceneProxy.h
@@ -90,6 +90,8 @@ protected:
 	template <typename IndexRemapType>
 	void ApplyAttributeChanges(FISMInstanceUpdateChangeSet &ChangeSet, const IndexRemapType &IndexRemap, FInstanceSceneDataBuffers::FWriteView &ProxyData);
 
+	void TestAndApplyInstanceBufferSizeFixup(int32 PostUpdateNumInstances);
+
 	FStaticShaderPlatform ShaderPlatform;
 	ERHIFeatureLevel::Type FeatureLevel;
 	bool bUseLegacyRenderingPath = false;
