Allow AddCurveMetaData to skip recording new transactions.

diff --git a/Engine/Plugins/Runtime/SkeletalMerging/Source/SkeletalMerging/Private/SkeletalMergingLibrary.cpp b/Engine/Plugins/Runtime/SkeletalMerging/Source/SkeletalMerging/Private/SkeletalMergingLibrary.cpp
--- a/Engine/Plugins/Runtime/SkeletalMerging/Source/SkeletalMerging/Private/SkeletalMergingLibrary.cpp
+++ b/Engine/Plugins/Runtime/SkeletalMerging/Source/SkeletalMerging/Private/SkeletalMergingLibrary.cpp
@@ -332,7 +332,7 @@ void USkeletalMergingLibrary::AddCurveNames(USkeleton* InSkeleton, const TMap<FN
 	{
 		if (CurveMetaDataPair.Value)
 		{
-			InSkeleton->AddCurveMetaData(CurveMetaDataPair.Key);
+			InSkeleton->AddCurveMetaData(CurveMetaDataPair.Key, false);
 
 			FCurveMetaData& SkeletonCurveMetaData = *InSkeleton->GetCurveMetaData(CurveMetaDataPair.Key);
 			SkeletonCurveMetaData = *CurveMetaDataPair.Value;
diff --git a/Engine/Source/Runtime/Engine/Classes/Animation/Skeleton.h b/Engine/Source/Runtime/Engine/Classes/Animation/Skeleton.h
--- a/Engine/Source/Runtime/Engine/Classes/Animation/Skeleton.h
+++ b/Engine/Source/Runtime/Engine/Classes/Animation/Skeleton.h
@@ -391,9 +391,10 @@ public:
 	/**
 	 * Adds a curve metadata entry with the specified name
 	 * @param	InCurveName			The name of the curve to find
+	 * @param   bTransact           If true record a new transaction
 	 * @return true if an entry was added, false if an entry already existed
 	 */
-	ENGINE_API bool AddCurveMetaData(FName CurveName);
+	ENGINE_API bool AddCurveMetaData(FName CurveName, bool bTransact = true);
 
 	/**
 	 * Get an array of all curve metadata names
diff --git a/Engine/Source/Runtime/Engine/Private/Animation/AnimCurveMetaData.cpp b/Engine/Source/Runtime/Engine/Private/Animation/AnimCurveMetaData.cpp
--- a/Engine/Source/Runtime/Engine/Private/Animation/AnimCurveMetaData.cpp
+++ b/Engine/Source/Runtime/Engine/Private/Animation/AnimCurveMetaData.cpp
@@ -5,6 +5,7 @@
 #if WITH_EDITOR
 #include "ScopedTransaction.h"
 #include "UObject/AssetRegistryTagsContext.h"
+#include "Editor.h"
 #endif
 
 #define LOCTEXT_NAMESPACE "AnimCurveMetaData" 
@@ -78,12 +79,15 @@ const FCurveMetaData* UAnimCurveMetaData::GetCurveMetaData(FName InCurveName) co
 	return CurveMetaData.Find(InCurveName);
 }
 
-bool UAnimCurveMetaData::AddCurveMetaData(FName InCurveName)
+bool UAnimCurveMetaData::AddCurveMetaData(FName InCurveName, bool bInTransact)
 {
 	if(!CurveMetaData.Contains(InCurveName))
 	{
 #if WITH_EDITOR
-		FScopedTransaction Transaction(LOCTEXT("AddCurveMetaData", "Add Curve Metadata"));
+		if (bInTransact)
+		{
+			GEditor->BeginTransaction(LOCTEXT("AddCurveMetaData", "Add Curve Metadata"));
+		}
 		Modify();
 #endif
 
@@ -91,6 +95,13 @@ bool UAnimCurveMetaData::AddCurveMetaData(FName InCurveName)
 
 		IncreaseVersionNumber();
 		OnCurveMetaDataChanged.Broadcast();
+
+#if WITH_EDITOR
+		if (bInTransact)
+		{
+			GEditor->EndTransaction();
+		}
+#endif
 		return true;
 	}
 	return false;
diff --git a/Engine/Source/Runtime/Engine/Private/Animation/Skeleton.cpp b/Engine/Source/Runtime/Engine/Private/Animation/Skeleton.cpp
--- a/Engine/Source/Runtime/Engine/Private/Animation/Skeleton.cpp
+++ b/Engine/Source/Runtime/Engine/Private/Animation/Skeleton.cpp
@@ -1766,10 +1766,10 @@ void USkeleton::RenameSlotName(const FName& OldName, const FName& NewName)
 	SetSlotGroupName(NewName, GroupName);
 }
 
-bool USkeleton::AddCurveMetaData(FName CurveName)
+bool USkeleton::AddCurveMetaData(FName CurveName, bool bTransact)
 {
 	UAnimCurveMetaData* AnimCurveMetaData = GetOrCreateCurveMetaDataObject();
-	return AnimCurveMetaData->AddCurveMetaData(CurveName);
+	return AnimCurveMetaData->AddCurveMetaData(CurveName, bTransact);
 }
 
 #if WITH_EDITOR
diff --git a/Engine/Source/Runtime/Engine/Public/Animation/AnimCurveMetadata.h b/Engine/Source/Runtime/Engine/Public/Animation/AnimCurveMetadata.h
--- a/Engine/Source/Runtime/Engine/Public/Animation/AnimCurveMetadata.h
+++ b/Engine/Source/Runtime/Engine/Public/Animation/AnimCurveMetadata.h
@@ -89,9 +89,10 @@ public:
 	/**
 	 * Adds a curve metadata entry with the specified name
 	 * @param	InCurveName			The name of the curve to find
+	 * @param   bInTransact         If true record a new transaction
 	 * @return true if an entry was added, false if an entry already existed
 	 */
-	virtual bool AddCurveMetaData(FName InCurveName) = 0;
+	virtual bool AddCurveMetaData(FName InCurveName, bool bInTransact = true) = 0;
 
 	/**
 	 * Get the curve metadata entry with the specified name
@@ -201,7 +202,7 @@ public:
 	virtual int32 GetNumCurveMetaData() const override { return CurveMetaData.Num(); }
 	virtual const FCurveMetaData* GetCurveMetaData(FName InCurveName) const override;
 	virtual FCurveMetaData* GetCurveMetaData(FName InCurveName) override;
-	virtual bool AddCurveMetaData(FName InCurveName) override;
+	virtual bool AddCurveMetaData(FName InCurveName, bool bInTransact = true) override;
 	virtual void GetCurveMetaDataNames(TArray<FName>& OutNames) const override;
 	virtual void RefreshBoneIndices(USkeleton* InSkeleton) override;
 	virtual uint16 GetVersionNumber() const override; 
