Prevent race condition between FRayTracingGeometryManager::Tick(...) and render command pipes.

diff --git a/Engine/Source/Runtime/Engine/Private/Rendering/RayTracingGeometryManager.cpp b/Engine/Source/Runtime/Engine/Private/Rendering/RayTracingGeometryManager.cpp
--- a/Engine/Source/Runtime/Engine/Private/Rendering/RayTracingGeometryManager.cpp
+++ b/Engine/Source/Runtime/Engine/Private/Rendering/RayTracingGeometryManager.cpp
@@ -273,6 +273,14 @@ void FRayTracingGeometryManager::Tick(FRHICommandList& RHICmdList)
 	TRACE_CPUPROFILER_EVENT_SCOPE(FRayTracingGeometryManager::Tick);
 	QUICK_SCOPE_CYCLE_COUNTER(STAT_FRayTracingGeometryManager_Tick);
 
+	// TODO: investigate fine grained locking to minimize blocking progress on render command pipes
+	// - Don't touch registered geometry/group arrays from render command pipes
+	//   - Separate arrays of free geometry/group handles + HandleAllocationCS
+	//   - delay actual registration until PreRender() which happens on Render Thread
+	//	 - Tick() doesn't need to lock at all
+	// - Refresh requests could be queued and processed during Tick()
+	FScopeLock ScopeLock(&MainCS);
+
 #if DO_CHECK
 	static uint64 PreviousFrameCounter = GFrameCounterRenderThread - 1;
 	checkf(GFrameCounterRenderThread != PreviousFrameCounter, TEXT("FRayTracingGeometryManager::Tick() should only be called once per frame"));
