Fixed changes to Data Layers on child actors of a LevelInstance that were ignored until the map was reloaded.
Regression introduced by 31254015 which added resolving of the DataLayers of level instances actordesccontainers in editor only at load time. This led to streaming generation process reusing these resolved (but never updated) data layers.

diff --git a/Engine/Source/Runtime/Engine/Private/WorldPartition/ActorDescContainerInstance.cpp b/Engine/Source/Runtime/Engine/Private/WorldPartition/ActorDescContainerInstance.cpp
--- a/Engine/Source/Runtime/Engine/Private/WorldPartition/ActorDescContainerInstance.cpp
+++ b/Engine/Source/Runtime/Engine/Private/WorldPartition/ActorDescContainerInstance.cpp
@@ -10,6 +10,7 @@
 #include "WorldPartition/WorldPartition.h"
 #include "WorldPartition/WorldPartitionActorDesc.h"
 #include "WorldPartition/WorldPartitionActorDescUtils.h"
+#include "WorldPartition/DataLayer/DataLayerManager.h"
 #include "Editor.h"
 #endif
 
@@ -540,6 +541,13 @@ void UActorDescContainerInstance::OnActorDescAdded(FWorldPartitionActorDesc* InA
 	{
 		WorldPartition->OnActorDescInstanceAdded(NewActorDescInstance);
 	}
+	else if (UWorldPartition* TopWorldPartition = GetTopWorldPartition())
+	{
+		if (const UDataLayerManager* ResolvingDataLayerManager = TopWorldPartition->GetResolvingDataLayerManager())
+		{
+			ResolvingDataLayerManager->ResolveActorDescInstanceDataLayers(NewActorDescInstance);
+		}
+	}
 
 	OnActorDescInstanceAddedEvent.Broadcast(NewActorDescInstance);
 }
@@ -592,6 +600,13 @@ void UActorDescContainerInstance::OnActorDescUpdated(FWorldPartitionActorDesc* I
 		{
 			WorldPartition->OnActorDescInstanceUpdated(ActorDescInstance->Get());
 		}
+		else if (UWorldPartition* TopWorldPartition = GetTopWorldPartition())
+		{
+			if (const UDataLayerManager* ResolvingDataLayerManager = TopWorldPartition->GetResolvingDataLayerManager())
+			{
+				ResolvingDataLayerManager->ResolveActorDescInstanceDataLayers(ActorDescInstance->Get());
+			}
+		}
 
 		OnActorDescInstanceUpdatedEvent.Broadcast(ActorDescInstance->Get());
 	}
