[NetPhysics] Don't correct connected particles that also replicate physics.

diff --git a/Engine/Source/Runtime/Engine/Private/PhysicsEngine/PhysicsReplication.cpp b/Engine/Source/Runtime/Engine/Private/PhysicsEngine/PhysicsReplication.cpp
--- a/Engine/Source/Runtime/Engine/Private/PhysicsEngine/PhysicsReplication.cpp
+++ b/Engine/Source/Runtime/Engine/Private/PhysicsEngine/PhysicsReplication.cpp
@@ -737,7 +737,7 @@ bool FPhysicsReplication::ApplyRigidBodyState(float DeltaSeconds, FBodyInstance*
 								
 									// Set XRVW to hard snap dynamic object and force recalculation of friction
 									Solver->GetEvolution()->ApplyParticleTransformCorrection(Proxy->GetHandle_LowLevel(), IdealWorldTM.GetLocation(), IdealWorldTM.GetRotation(), bCorrectConnectedBodies, bCorrectConnectedBodiesFriction);
-								
+
 									Handle->SetV(NewState.LinVel);
 									Handle->SetW(FMath::DegreesToRadians(NewState.AngVel));
 								});
@@ -840,7 +840,9 @@ void FPhysicsReplication::PrepareAsyncData_External(const FRigidBodyErrorCorrect
 #pragma region FPhysicsReplicationAsync
 void FPhysicsReplicationAsync::OnPhysicsObjectUnregistered_Internal(Chaos::FConstPhysicsObjectHandle PhysicsObject)
 {
-	ObjectToTarget.Remove(PhysicsObject);
+	RemoveObjectFromReplication(PhysicsObject);
+
+	// Only clear Settings when PhysicsObject unregister (not when it stops replicating, hence why it's not baked into RemoveObjectFromReplication())
 	ObjectToSettings.Remove(PhysicsObject);
 }
 
@@ -909,8 +911,8 @@ void FPhysicsReplicationAsync::OnPreSimulate_Internal()
 		{
 			if (Input.TargetState.Flags == ERigidBodyFlags::None)
 			{
-				// Remove replication target 
-				ObjectToTarget.Remove(Input.PhysicsObject);
+				// Remove replication target
+				RemoveObjectFromReplication(Input.PhysicsObject);
 				continue;
 			}
 
@@ -937,6 +939,41 @@ void FPhysicsReplicationAsync::OnPreSimulate_Internal()
 	}
 }
 
+FReplicatedPhysicsTargetAsync* FPhysicsReplicationAsync::AddObjectToReplication(Chaos::FConstPhysicsObjectHandle PhysicsObject)
+{
+	if (ensure(PhysicsObject))
+	{
+		// Cache ParticleID in array of replicated objects
+		Chaos::FReadPhysicsObjectInterface_Internal Interface = Chaos::FPhysicsObjectInternalInterface::GetRead();
+		if (Chaos::FGeometryParticleHandle* Handle = Interface.GetParticle(PhysicsObject))
+		{
+			ReplicatedParticleIDs.Add(Handle->ParticleID());
+		}
+
+		// Add to Object-Target map
+		return &ObjectToTarget.Add(PhysicsObject, FReplicatedPhysicsTargetAsync());
+	}
+	return nullptr;
+}
+
+void FPhysicsReplicationAsync::RemoveObjectFromReplication(Chaos::FConstPhysicsObjectHandle PhysicsObject)
+{
+	if (PhysicsObject == nullptr)
+	{
+		return;
+	}
+
+	// Remove from Object-Target map
+	ObjectToTarget.Remove(PhysicsObject);
+
+	// Remove cached replicated ParticleID
+	Chaos::FReadPhysicsObjectInterface_Internal Interface = Chaos::FPhysicsObjectInternalInterface::GetRead();
+	if (Chaos::FGeometryParticleHandle* Handle = Interface.GetParticle(PhysicsObject))
+	{
+		ReplicatedParticleIDs.Remove(Handle->ParticleID());
+	}
+}
+
 void FPhysicsReplicationAsync::UpdateRewindDataTarget(const FPhysicsRepAsyncInputData& Input)
 {
 	if (Input.PhysicsObject == nullptr)
@@ -979,7 +1016,7 @@ void FPhysicsReplicationAsync::UpdateAsyncTarget(const FPhysicsRepAsyncInputData
 	if (bFirstTarget)
 	{
 		// First time we add a target, set previous state to current input
-		Target = &ObjectToTarget.Add(Input.PhysicsObject, FReplicatedPhysicsTargetAsync());
+		Target = AddObjectToReplication(Input.PhysicsObject);
 		Target->PrevPos = Input.TargetState.Position;
 		Target->PrevPosTarget = Input.TargetState.Position;
 		Target->PrevRotTarget = Input.TargetState.Quaternion;
@@ -1173,6 +1210,7 @@ void FPhysicsReplicationAsync::ApplyTargetStatesAsync(const float DeltaSeconds,
 	for (auto Itr = ObjectToTarget.CreateIterator(); Itr; ++Itr)
 	{
 		bool bRemoveItr = true; // Remove current cached replication target unless replication logic tells us to store it for next tick
+		FParticleID ParticleID;
 
 		Chaos::FConstPhysicsObjectHandle& POHandle = Itr.Key();
 		if (FGeometryParticleHandle* Handle = Interface.GetParticle(POHandle))
@@ -1181,6 +1219,8 @@ void FPhysicsReplicationAsync::ApplyTargetStatesAsync(const float DeltaSeconds,
 
 			if (FPBDRigidParticleHandle* RigidHandle = Handle->CastToRigidParticle())
 			{
+				ParticleID = RigidHandle->ParticleID();
+
 				// Cache custom settings for this object if there are any
 				FetchObjectSettings(POHandle);
 
@@ -1205,6 +1245,7 @@ void FPhysicsReplicationAsync::ApplyTargetStatesAsync(const float DeltaSeconds,
 
 		if (bRemoveItr)
 		{
+			ReplicatedParticleIDs.Remove(ParticleID);
 			Itr.RemoveCurrent();
 		}
 	}
@@ -1475,8 +1516,7 @@ bool FPhysicsReplicationAsync::DefaultReplication(Chaos::FPBDRigidParticleHandle
 			// Set XRVW to hard snap dynamic object and force recalculation of friction
 			const bool bCorrectConnectedBodies = SettingsCurrent.DefaultReplicationSettings.GetCorrectConnectedBodies();
 			const bool bCorrectConnectedBodiesFriction = SettingsCurrent.DefaultReplicationSettings.GetCorrectConnectedBodiesFriction();
-			RigidsSolver->GetEvolution()->ApplyParticleTransformCorrection(Handle, TargetPos, TargetQuat, bCorrectConnectedBodies, bCorrectConnectedBodiesFriction);
-
+			RigidsSolver->GetEvolution()->ApplyParticleTransformCorrection(Handle, TargetPos, TargetQuat, bCorrectConnectedBodies, bCorrectConnectedBodiesFriction, ReplicatedParticleIDs);
 			Handle->SetV(NewState.LinVel);
 			Handle->SetW(FMath::DegreesToRadians(NewState.AngVel));
 		}
@@ -1686,7 +1726,7 @@ bool FPhysicsReplicationAsync::PredictiveInterpolation(Chaos::FPBDRigidParticleH
 		{
 			// Set XRVW to hard snap dynamic object and force recalculation of friction
 			const bool bCorrectConnectedBodies = SettingsCurrent.PredictiveInterpolationSettings.GetCorrectConnectedBodies();
-			RigidsSolver->GetEvolution()->ApplyParticleTransformCorrection(Handle, Target.PrevPosTarget, Target.PrevRotTarget, bCorrectConnectedBodies, /*bInRecalculateFrictionOnConnectedBodies*/ true);
+			RigidsSolver->GetEvolution()->ApplyParticleTransformCorrection(Handle, Target.PrevPosTarget, Target.PrevRotTarget, bCorrectConnectedBodies, /*bInRecalculateFrictionOnConnectedBodies*/ true, ReplicatedParticleIDs);
 			Handle->SetV(TargetLinVel);
 			Handle->SetW(TargetAngVel);
 		}
@@ -1885,7 +1925,7 @@ bool FPhysicsReplicationAsync::PredictiveInterpolation(Chaos::FPBDRigidParticleH
 		{	
 			const bool bCorrectConnectedBodies = SettingsCurrent.PredictiveInterpolationSettings.GetCorrectConnectedBodies();
 			const bool bCorrectConnectedBodiesFriction = SettingsCurrent.PredictiveInterpolationSettings.GetCorrectConnectedBodiesFriction();
-			RigidsSolver->GetEvolution()->ApplyParticleTransformCorrection(Handle, CorrectionX, CorrectionR, bCorrectConnectedBodies, bCorrectConnectedBodiesFriction);
+			RigidsSolver->GetEvolution()->ApplyParticleTransformCorrection(Handle, CorrectionX, CorrectionR, bCorrectConnectedBodies, bCorrectConnectedBodiesFriction, ReplicatedParticleIDs);
 		}
 
 		if (bSoftSnap)
@@ -1901,7 +1941,7 @@ bool FPhysicsReplicationAsync::PredictiveInterpolation(Chaos::FPBDRigidParticleH
 			// Apply correction as a transform shift
 			const bool bCorrectConnectedBodies = SettingsCurrent.PredictiveInterpolationSettings.GetCorrectConnectedBodies();
 			const bool bCorrectConnectedBodiesFriction = SettingsCurrent.PredictiveInterpolationSettings.GetCorrectConnectedBodiesFriction();
-			RigidsSolver->GetEvolution()->ApplyParticleTransformCorrection(Handle, SoftSnapPos, SoftSnapRot, bCorrectConnectedBodies, bCorrectConnectedBodiesFriction);
+			RigidsSolver->GetEvolution()->ApplyParticleTransformCorrection(Handle, SoftSnapPos, SoftSnapRot, bCorrectConnectedBodies, bCorrectConnectedBodiesFriction, ReplicatedParticleIDs);
 		}
 	}
 
@@ -2042,7 +2082,7 @@ bool FPhysicsReplicationAsync::ResimulationReplication(Chaos::FPBDRigidParticleH
 				}
 #endif
 				// Apply correction to position and rotation
-				RigidsSolver->GetEvolution()->ApplyParticleTransformCorrection(Handle, CorrectedX, CorrectedR, SettingsCurrent.ResimulationSettings.GetRuntimeCorrectConnectedBodies());
+				RigidsSolver->GetEvolution()->ApplyParticleTransformCorrection(Handle, CorrectedX, CorrectedR, SettingsCurrent.ResimulationSettings.GetRuntimeCorrectConnectedBodies(), /*bInRecalculateFrictionOnConnectedBodies*/true, ReplicatedParticleIDs);
 			}
 
 			// Keep target for NumPredictedFrames time to perform runtime corrections with until a new target is received
@@ -2057,6 +2097,7 @@ FName FPhysicsReplicationAsync::GetFNameForStatId() const
 	const static FLazyName StaticName("FPhysicsReplicationAsyncCallback");
 	return StaticName;
 }
+
 #pragma endregion // FPhysicsReplicationAsync
 
 
diff --git a/Engine/Source/Runtime/Engine/Public/PhysicsReplication.h b/Engine/Source/Runtime/Engine/Public/PhysicsReplication.h
--- a/Engine/Source/Runtime/Engine/Public/PhysicsReplication.h
+++ b/Engine/Source/Runtime/Engine/Public/PhysicsReplication.h
@@ -13,6 +13,7 @@
 #include "Physics/PhysicsInterfaceDeclares.h"
 #include "PhysicsProxy/SingleParticlePhysicsProxyFwd.h"
 #include "Chaos/PhysicsObject.h"
+#include "Chaos/ParticleDirtyFlags.h"
 #include "Chaos/SimCallbackObject.h"
 #include "Physics/PhysicsInterfaceUtils.h"
 #include "Physics/NetworkPhysicsSettingsComponent.h"
@@ -45,9 +46,6 @@ namespace PhysicsReplicationCVars
 }
 #endif
 
-class FPhysScene_PhysX;
-
-
 #pragma region FPhysicsReplicationAsync
 
 struct FPhysicsRepErrorCorrectionData
@@ -224,8 +222,11 @@ private:
 	TMap<Chaos::FConstPhysicsObjectHandle, FReplicatedPhysicsTargetAsync> ObjectToTarget;
 	TMap<Chaos::FConstPhysicsObjectHandle, FNetworkPhysicsSettingsAsync> ObjectToSettings;
 	TArray<int32> ParticlesInResimIslands;
+	TArray<Chaos::FParticleID> ReplicatedParticleIDs;
 
 private:
+	FReplicatedPhysicsTargetAsync* AddObjectToReplication(Chaos::FConstPhysicsObjectHandle PhysicsObject);
+	void RemoveObjectFromReplication(Chaos::FConstPhysicsObjectHandle PhysicsObject);
 	void UpdateAsyncTarget(const FPhysicsRepAsyncInputData& Input, Chaos::FPBDRigidsSolver* RigidsSolver);
 	void UpdateRewindDataTarget(const FPhysicsRepAsyncInputData& Input);
 	void CacheResimInteractions();
diff --git a/Engine/Source/Runtime/Experimental/Chaos/Private/Chaos/PBDRigidsEvolutionGBF.cpp b/Engine/Source/Runtime/Experimental/Chaos/Private/Chaos/PBDRigidsEvolutionGBF.cpp
--- a/Engine/Source/Runtime/Experimental/Chaos/Private/Chaos/PBDRigidsEvolutionGBF.cpp
+++ b/Engine/Source/Runtime/Experimental/Chaos/Private/Chaos/PBDRigidsEvolutionGBF.cpp
@@ -1480,16 +1480,16 @@ void FPBDRigidsEvolutionGBF::OnParticleMoved(FGeometryParticleHandle* InParticle
 	}
 }
 
-void FPBDRigidsEvolutionGBF::ApplyParticleTransformCorrectionDelta(FGeometryParticleHandle* InParticle, const FVec3& InPosDelta, const FVec3& InRotDelta, const bool bApplyToConnectedBodies, const bool bInRecalculateFrictionOnConnectedBodies)
+void FPBDRigidsEvolutionGBF::ApplyParticleTransformCorrectionDelta(FGeometryParticleHandle* InParticle, const FVec3& InPosDelta, const FVec3& InRotDelta, const bool bApplyToConnectedBodies, const bool bInRecalculateFrictionOnConnectedBodies, const TArray<FParticleID>& ExcludeConnections)
 {
 	ApplyParticleTransformCorrection(
 		InParticle, 
 		InParticle->GetX() + InPosDelta, 
 		FRotation3::IntegrateRotationWithAngularVelocity(InParticle->GetR(), InRotDelta, FReal(1.0)),
-		bApplyToConnectedBodies, bInRecalculateFrictionOnConnectedBodies);
+		bApplyToConnectedBodies, bInRecalculateFrictionOnConnectedBodies, ExcludeConnections);
 }
 
-void FPBDRigidsEvolutionGBF::ApplyParticleTransformCorrection(FGeometryParticleHandle* InParticle, const FVec3& InPos, const FRotation3& InRot, const bool bApplyToConnectedBodies, const bool bInRecalculateFrictionOnConnectedBodies)
+void FPBDRigidsEvolutionGBF::ApplyParticleTransformCorrection(FGeometryParticleHandle* InParticle, const FVec3& InPos, const FRotation3& InRot, const bool bApplyToConnectedBodies, const bool bInRecalculateFrictionOnConnectedBodies, const TArray<FParticleID>& ExcludeConnections)
 {
 	const FRigidTransform3 OldParticleTransform = InParticle->GetTransformXR();
 	const FRigidTransform3 NewParticleTransform = FRigidTransform3(InPos, InRot);
@@ -1501,7 +1501,7 @@ void FPBDRigidsEvolutionGBF::ApplyParticleTransformCorrection(FGeometryParticleH
 	{
 		// Find all the connected particles and move them to retain their relative transform
 		// NOTE: ConnectedParticles will not include InParticle
-		TArray<FGeometryParticleHandle*> ConnectedParticles = GetConnectedParticles(InParticle);
+		TArray<FGeometryParticleHandle*> ConnectedParticles = GetConnectedParticles(InParticle, ExcludeConnections);
 		for (FGeometryParticleHandle* ConnectedParticle : ConnectedParticles)
 		{
 			if (FGenericParticleHandle(ConnectedParticle)->IsDynamic())
@@ -1550,7 +1550,7 @@ void FPBDRigidsEvolutionGBF::ApplySleepOnConnectedParticles(FGeometryParticleHan
 	}
 }
 
-TArray<FGeometryParticleHandle*> FPBDRigidsEvolutionGBF::GetConnectedParticles(FGeometryParticleHandle* InParticle)
+TArray<FGeometryParticleHandle*> FPBDRigidsEvolutionGBF::GetConnectedParticles(FGeometryParticleHandle* InParticle, const TArray<FParticleID>& ExcludeConnections)
 {
 	if (InParticle->ParticleConstraints().IsEmpty())
 	{
@@ -1579,7 +1579,7 @@ TArray<FGeometryParticleHandle*> FPBDRigidsEvolutionGBF::GetConnectedParticles(F
 				{
 					FParticlePair JointParticles = Joint->GetConstrainedParticles();
 					FGeometryParticleHandle* OtherParticle = (JointParticles[0] != NextParticle) ? JointParticles[0] : JointParticles[1];
-					if ((OtherParticle != InParticle) && (ConnectedParticles.Find(OtherParticle) == nullptr))
+					if ((OtherParticle != InParticle) && (ConnectedParticles.Find(OtherParticle) == nullptr) && !ExcludeConnections.Contains(OtherParticle->ParticleID()))
 					{
 						ConnectedParticles.Add(OtherParticle, OtherParticle);
 					}
diff --git a/Engine/Source/Runtime/Experimental/Chaos/Public/Chaos/PBDRigidsEvolutionGBF.h b/Engine/Source/Runtime/Experimental/Chaos/Public/Chaos/PBDRigidsEvolutionGBF.h
--- a/Engine/Source/Runtime/Experimental/Chaos/Public/Chaos/PBDRigidsEvolutionGBF.h
+++ b/Engine/Source/Runtime/Experimental/Chaos/Public/Chaos/PBDRigidsEvolutionGBF.h
@@ -178,13 +178,14 @@ namespace Chaos
 		 * NOTE: must be called prior to Integrate() to be effective.
 		 * NOTE: be careful with bApplyToConnectedBodies - only one particle in the connected graph should have ApplyParticleTransformCorrectionDelta called on 
 		 * it, otherwise you will get multiple particles trying to recorrect each other leading to very strange behaviour.
+		 * Use @param ExcludeConnections to exclude particles from being corrected, this includes particles connected to the excluded particle as long as they are not connected to a non-excluded particle.
 		 */
-		CHAOS_API void ApplyParticleTransformCorrectionDelta(FGeometryParticleHandle* InParticle, const FVec3& InPosDelta, const FVec3& InRotDelta, const bool bApplyToConnectedBodies, const bool bInRecalculateFrictionOnConnectedBodies = true);
+		CHAOS_API void ApplyParticleTransformCorrectionDelta(FGeometryParticleHandle* InParticle, const FVec3& InPosDelta, const FVec3& InRotDelta, const bool bApplyToConnectedBodies, const bool bInRecalculateFrictionOnConnectedBodies = true, const TArray<FParticleID>& ExcludeConnections = TArray<FParticleID>());
 
 		/*
 		 * [EXPERIMENTAL] Similar to SetParticleTransformCorrectionDelta, but supplied an absolute transform to jump to. This is used for snaps.
 		 */
-		CHAOS_API void ApplyParticleTransformCorrection(FGeometryParticleHandle* InParticle, const FVec3& InPos, const FRotation3& InRot, const bool bApplyToConnectedBodies, const bool bInRecalculateFrictionOnConnectedBodies = true);
+		CHAOS_API void ApplyParticleTransformCorrection(FGeometryParticleHandle* InParticle, const FVec3& InPos, const FRotation3& InRot, const bool bApplyToConnectedBodies, const bool bInRecalculateFrictionOnConnectedBodies = true, const TArray<FParticleID>& ExcludeConnections = TArray<FParticleID>());
 
 		/*
 		* [EXPERIMENTAL] Apply sleep state on connected particles that are dynamic
@@ -274,8 +275,9 @@ namespace Chaos
 		// Update the particle transform and fix collision anchors (used by client corrections)
 		void ApplyParticleTransformCorrectionImpl(FGeometryParticleHandle* InParticle, const FRigidTransform3& InTransform, const bool bInRecalculateFriction = true);
 
-		// Get all the particles that are connected to InParticle by a joint with locked position limits
-		TArray<FGeometryParticleHandle*> GetConnectedParticles(FGeometryParticleHandle* InParticle);
+		/* Get all the particles that are connected to InParticle by a joint with locked position limits
+		 * Use @param ExcludeConnections to exclude particles, this includes particles connected to the excluded particle as long as they are not connected to a non-excluded particle. */
+		TArray<FGeometryParticleHandle*> GetConnectedParticles(FGeometryParticleHandle* InParticle, const TArray<FParticleID>& ExcludeConnections = TArray<FParticleID>());
 
 		void UpdateInertiaConditioning();
 
