Changed the validation of WorldDatalayers to rely on default project asset reference restriction filtering.
Only restrict the External Data Layer asset and its associated WorldDataLayers to be part of the same plugin.

diff --git a/Engine/Source/Runtime/Engine/Private/WorldPartition/DataLayer/WorldDataLayers.cpp b/Engine/Source/Runtime/Engine/Private/WorldPartition/DataLayer/WorldDataLayers.cpp
--- a/Engine/Source/Runtime/Engine/Private/WorldPartition/DataLayer/WorldDataLayers.cpp
+++ b/Engine/Source/Runtime/Engine/Private/WorldPartition/DataLayer/WorldDataLayers.cpp
@@ -1376,7 +1376,8 @@ bool AWorldDataLayers::CanReferenceDataLayerAsset(const UDataLayerAsset* InDataL
 	{
 		return false;
 	}
-	else if (InDataLayerAsset)
+	// Validate that AWorldDataLayers and its EDL Asset are part of the same plugin
+	else if (InDataLayerAsset && (RootExternalDataLayerAsset == InDataLayerAsset))
 	{
 		IPluginManager& PluginManager = IPluginManager::Get();
 		const FString AssetMountPoint = FPackageName::GetPackageMountPoint(InDataLayerAsset->GetPackage()->GetName()).ToString();
