Iris - Fixed typo.

#rnx

diff --git a/Engine/Source/Runtime/Engine/Public/Net/Iris/ReplicationSystem/ReplicationSystemUtil.h b/Engine/Source/Runtime/Engine/Public/Net/Iris/ReplicationSystem/ReplicationSystemUtil.h
--- a/Engine/Source/Runtime/Engine/Public/Net/Iris/ReplicationSystem/ReplicationSystemUtil.h
+++ b/Engine/Source/Runtime/Engine/Public/Net/Iris/ReplicationSystem/ReplicationSystemUtil.h
@@ -132,7 +132,7 @@ struct FReplicationSystemUtil
 	ENGINE_API static void ClearCullDistanceSqrOverride(const AActor* Actor);
 
 	/** Set the poll frequency for an object and its subobjects. */
-	ENGINE_API static void SetPollFrequency(const UObject* Object, float CullDistSqr);
+	ENGINE_API static void SetPollFrequency(const UObject* Object, float PollFrequency);
 };
 
 }
