[World Partition Runtime HashSet]
- Runtime partitions can now output 2D streaming cells that will be stored in a 2D version of the spatial hash for faster queries.
- Runtime partitions can now choose between 3 different methods for bounds computing: cell bounds (for fixed cell size partitions), content bounds, or the minimum of the previous 2.
- Updated spatial hash unit tests to include 2D tests.

diff --git a/Engine/Source/Runtime/Engine/Private/Tests/WorldPartition/WorldPartitionStaticSpatialIndexTests.cpp b/Engine/Source/Runtime/Engine/Private/Tests/WorldPartition/WorldPartitionStaticSpatialIndexTests.cpp
--- a/Engine/Source/Runtime/Engine/Private/Tests/WorldPartition/WorldPartitionStaticSpatialIndexTests.cpp
+++ b/Engine/Source/Runtime/Engine/Private/Tests/WorldPartition/WorldPartitionStaticSpatialIndexTests.cpp
@@ -16,8 +16,21 @@ namespace WorldPartitionTests
 	IMPLEMENT_SIMPLE_AUTOMATION_TEST(FWorldPartitionStaticSpatialIndexTest, TEST_NAME_ROOT ".StaticSpatialIndex", EAutomationTestFlags::EditorContext | EAutomationTestFlags::EngineFilter)
 
 #if WITH_EDITOR
-	template <class Class>
-	void PerformTests(FWorldPartitionStaticSpatialIndexTest* Test, const TCHAR* Name, const TArray<TPair<FBox, int32>>& Elements, const TArray<FSphere>& Tests, TArray<int32>& Results)
+	template <typename BoxType>
+	const TCHAR* GetSpaceString()
+	{
+		if constexpr (std::is_same<BoxType, FBox2D>::value)
+		{
+			return TEXT("2d");
+		}
+		else
+		{
+			return TEXT("3d");
+		}
+	}
+
+	template <typename Profile, class Class>
+	FORCENOINLINE void PerformTests(FWorldPartitionStaticSpatialIndexTest* Test, const TCHAR* Name, const TArray<TPair<typename Profile::FBox, int32>>& Elements, const TArray<FSphere>& Tests, TArray<int32>& Results)
 	{
 		Class SpatialIndex;
 		SpatialIndex.Init(Elements);
@@ -32,103 +45,126 @@ namespace WorldPartitionTests
 		}
 
 		const double RunTime = FPlatformTime::Seconds() - StartTime;
-		Test->AddInfo(FString::Printf(TEXT("%s: %d tests in %s (%.2f/s, %s)"), Name, Tests.Num(), *FPlatformTime::PrettyTime(RunTime), Tests.Num() / RunTime, *FGenericPlatformMemory::PrettyMemory(SpatialIndex.GetAllocatedSize())));
+		Test->AddInfo(FString::Printf(TEXT("%s(%s): %d tests in %s (%.2f/s, %s)"), Name, GetSpaceString<typename Profile::FBox>(), Tests.Num(), *FPlatformTime::PrettyTime(RunTime), Tests.Num() / RunTime, *FGenericPlatformMemory::PrettyMemory(SpatialIndex.GetAllocatedSize())));
 	}
 
+	template <typename Profile, int32 MaxNumElementsPerNode, int32 MaxNumElementsPerLeaf>
+	FORCENOINLINE void PerformMinXTest(FWorldPartitionStaticSpatialIndexTest* Test, const TArray<TPair<typename Profile::FBox, int32>>& Elements, const TArray<FSphere>& Tests, const TArray<int32>& ReferenceResults)
+	{
+		TArray<int32> RTreeMinXResults;
+		RTreeMinXResults.Reserve(ReferenceResults.Num());
+		FString TestName = FString::Printf(TEXT("TStaticSpatialIndexRTree(%s-minx-%d-%d)"), GetSpaceString<typename Profile::FBox>(), MaxNumElementsPerNode, MaxNumElementsPerLeaf);
+		PerformTests<Profile, TStaticSpatialIndexRTree<int32, FStaticSpatialIndex::TNodeSorterMinX<Profile>, Profile>>(Test, *TestName, Elements, Tests, RTreeMinXResults);
+		RTreeMinXResults.Sort();
+		Test->TestTrue(TestName, RTreeMinXResults == ReferenceResults);
+	}
 
-	template <int32 MaxNumElementsPerNode, int32 MaxNumElementsPerLeaf>
-	void PerformMinXTest(FWorldPartitionStaticSpatialIndexTest* Test, const TArray<TPair<FBox, int32>>& Elements, const TArray<FSphere>& Tests, const TArray<int32>& ReferenceResults)
+	template <typename Profile>
+	FORCENOINLINE void PerformMinXTests(FWorldPartitionStaticSpatialIndexTest* Test, const TArray<TPair<typename Profile::FBox, int32>>& Elements, const TArray<FSphere>& Tests, const TArray<int32>& ReferenceResults)
 	{
-		TArray<int32> FRTreeMinXResults;
-		FRTreeMinXResults.Reserve(ReferenceResults.Num());
-		FString TestName = FString::Printf(TEXT("TStaticSpatialIndexRTree(minx-%d-%d)"), MaxNumElementsPerNode, MaxNumElementsPerLeaf);
-		PerformTests<TStaticSpatialIndexRTree<int32, FStaticSpatialIndex::FNodeSorterMinX>>(Test, *TestName, Elements, Tests, FRTreeMinXResults);
-		FRTreeMinXResults.Sort();
-		Test->TestTrue(TestName, FRTreeMinXResults == ReferenceResults);
+		PerformMinXTest<Profile, 16, 16>(Test, Elements, Tests, ReferenceResults);
+		PerformMinXTest<Profile, 16, 64>(Test, Elements, Tests, ReferenceResults);
+		PerformMinXTest<Profile, 16, 256>(Test, Elements, Tests, ReferenceResults);
+		PerformMinXTest<Profile, 16, 1024>(Test, Elements, Tests, ReferenceResults);
+		PerformMinXTest<Profile, 64, 16>(Test, Elements, Tests, ReferenceResults);
+		PerformMinXTest<Profile, 64, 64>(Test, Elements, Tests, ReferenceResults);
+		PerformMinXTest<Profile, 64, 256>(Test, Elements, Tests, ReferenceResults);
+		PerformMinXTest<Profile, 64, 1024>(Test, Elements, Tests, ReferenceResults);
+		PerformMinXTest<Profile, 256, 16>(Test, Elements, Tests, ReferenceResults);
+		PerformMinXTest<Profile, 256, 64>(Test, Elements, Tests, ReferenceResults);
+		PerformMinXTest<Profile, 256, 256>(Test, Elements, Tests, ReferenceResults);
+		PerformMinXTest<Profile, 256, 1024>(Test, Elements, Tests, ReferenceResults);
+		PerformMinXTest<Profile, 1024, 16>(Test, Elements, Tests, ReferenceResults);
+		PerformMinXTest<Profile, 1024, 64>(Test, Elements, Tests, ReferenceResults);
+		PerformMinXTest<Profile, 1024, 256>(Test, Elements, Tests, ReferenceResults);
+		PerformMinXTest<Profile, 1024, 1024>(Test, Elements, Tests, ReferenceResults);
 	}
 
-	void PerformMinXTests(FWorldPartitionStaticSpatialIndexTest* Test, const TArray<TPair<FBox, int32>>& Elements, const TArray<FSphere>& Tests, const TArray<int32>& ReferenceResults)
+	template <typename Profile, int32 BucketSize, int32 MaxNumElementsPerNode, int32 MaxNumElementsPerLeaf>
+	FORCENOINLINE void PerformMortonTest(FWorldPartitionStaticSpatialIndexTest* Test, const TArray<TPair<typename Profile::FBox, int32>>& Elements, const TArray<FSphere>& Tests, const TArray<int32>& ReferenceResults)
 	{
-		PerformMinXTest<16, 16>(Test, Elements, Tests, ReferenceResults);
-		PerformMinXTest<16, 64>(Test, Elements, Tests, ReferenceResults);
-		PerformMinXTest<16, 256>(Test, Elements, Tests, ReferenceResults);
-		PerformMinXTest<16, 1024>(Test, Elements, Tests, ReferenceResults);
-		PerformMinXTest<64, 16>(Test, Elements, Tests, ReferenceResults);
-		PerformMinXTest<64, 64>(Test, Elements, Tests, ReferenceResults);
-		PerformMinXTest<64, 256>(Test, Elements, Tests, ReferenceResults);
-		PerformMinXTest<64, 1024>(Test, Elements, Tests, ReferenceResults);
-		PerformMinXTest<256, 16>(Test, Elements, Tests, ReferenceResults);
-		PerformMinXTest<256, 64>(Test, Elements, Tests, ReferenceResults);
-		PerformMinXTest<256, 256>(Test, Elements, Tests, ReferenceResults);
-		PerformMinXTest<256, 1024>(Test, Elements, Tests, ReferenceResults);
-		PerformMinXTest<1024, 16>(Test, Elements, Tests, ReferenceResults);
-		PerformMinXTest<1024, 64>(Test, Elements, Tests, ReferenceResults);
-		PerformMinXTest<1024, 256>(Test, Elements, Tests, ReferenceResults);
-		PerformMinXTest<1024, 1024>(Test, Elements, Tests, ReferenceResults);
+		TArray<int32> RTreeMortonResults;
+		RTreeMortonResults.Reserve(ReferenceResults.Num());
+		FString TestName = FString::Printf(TEXT("TStaticSpatialIndexRTree(%s-morton-%dk-%d-%d)"), GetSpaceString<typename Profile::FBox>(), BucketSize >> 10, MaxNumElementsPerNode, MaxNumElementsPerLeaf);
+		PerformTests<Profile, TStaticSpatialIndexRTree<int32, FStaticSpatialIndex::TNodeSorterMorton<Profile, BucketSize>, Profile>>(Test, *TestName, Elements, Tests, RTreeMortonResults);
+		RTreeMortonResults.Sort();
+		Test->TestTrue(TestName, RTreeMortonResults == ReferenceResults);
 	}
 
-	template <int32 BucketSize, int32 MaxNumElementsPerNode, int32 MaxNumElementsPerLeaf>
-	void PerformMortonTest(FWorldPartitionStaticSpatialIndexTest* Test, const TArray<TPair<FBox, int32>>& Elements, const TArray<FSphere>& Tests, const TArray<int32>& ReferenceResults)
+	template <typename Profile, int32 BucketSize>
+	FORCENOINLINE void PerformMortonTests(FWorldPartitionStaticSpatialIndexTest* Test, const TArray<TPair<typename Profile::FBox, int32>>& Elements, const TArray<FSphere>& Tests, const TArray<int32>& ReferenceResults)
 	{
-		TArray<int32> FRTreeMortonResults;
-		FRTreeMortonResults.Reserve(ReferenceResults.Num());
-		FString TestName = FString::Printf(TEXT("TStaticSpatialIndexRTree(morton-%dk-%d-%d)"), BucketSize >> 10, MaxNumElementsPerNode, MaxNumElementsPerLeaf);
-		PerformTests<TStaticSpatialIndexRTree<int32, FStaticSpatialIndex::TNodeSorterMorton<BucketSize>>>(Test, *TestName, Elements, Tests, FRTreeMortonResults);
-		FRTreeMortonResults.Sort();
-		Test->TestTrue(TestName, FRTreeMortonResults == ReferenceResults);
+		PerformMortonTest<Profile, BucketSize, 16, 16>(Test, Elements, Tests, ReferenceResults);
+		PerformMortonTest<Profile, BucketSize, 16, 64>(Test, Elements, Tests, ReferenceResults);
+		PerformMortonTest<Profile, BucketSize, 16, 256>(Test, Elements, Tests, ReferenceResults);
+		PerformMortonTest<Profile, BucketSize, 16, 1024>(Test, Elements, Tests, ReferenceResults);
+		PerformMortonTest<Profile, BucketSize, 64, 16>(Test, Elements, Tests, ReferenceResults);
+		PerformMortonTest<Profile, BucketSize, 64, 64>(Test, Elements, Tests, ReferenceResults);
+		PerformMortonTest<Profile, BucketSize, 64, 256>(Test, Elements, Tests, ReferenceResults);
+		PerformMortonTest<Profile, BucketSize, 64, 1024>(Test, Elements, Tests, ReferenceResults);
+		PerformMortonTest<Profile, BucketSize, 256, 16>(Test, Elements, Tests, ReferenceResults);
+		PerformMortonTest<Profile, BucketSize, 256, 64>(Test, Elements, Tests, ReferenceResults);
+		PerformMortonTest<Profile, BucketSize, 256, 256>(Test, Elements, Tests, ReferenceResults);
+		PerformMortonTest<Profile, BucketSize, 256, 1024>(Test, Elements, Tests, ReferenceResults);
+		PerformMortonTest<Profile, BucketSize, 1024, 16>(Test, Elements, Tests, ReferenceResults);
+		PerformMortonTest<Profile, BucketSize, 1024, 64>(Test, Elements, Tests, ReferenceResults);
+		PerformMortonTest<Profile, BucketSize, 1024, 256>(Test, Elements, Tests, ReferenceResults);
+		PerformMortonTest<Profile, BucketSize, 1024, 1024>(Test, Elements, Tests, ReferenceResults);
 	}
 
-	template <int32 BucketSize>
-	void PerformMortonTests(FWorldPartitionStaticSpatialIndexTest* Test, const TArray<TPair<FBox, int32>>& Elements, const TArray<FSphere>& Tests, const TArray<int32>& ReferenceResults)
+	template <typename Profile, int32 BucketSize, int32 MaxNumElementsPerNode, int32 MaxNumElementsPerLeaf>
+	FORCENOINLINE void PerformHilbertTest(FWorldPartitionStaticSpatialIndexTest* Test, const TArray<TPair<typename Profile::FBox, int32>>& Elements, const TArray<FSphere>& Tests, const TArray<int32>& ReferenceResults)
 	{
-		PerformMortonTest<BucketSize, 16, 16>(Test, Elements, Tests, ReferenceResults);
-		PerformMortonTest<BucketSize, 16, 64>(Test, Elements, Tests, ReferenceResults);
-		PerformMortonTest<BucketSize, 16, 256>(Test, Elements, Tests, ReferenceResults);
-		PerformMortonTest<BucketSize, 16, 1024>(Test, Elements, Tests, ReferenceResults);
-		PerformMortonTest<BucketSize, 64, 16>(Test, Elements, Tests, ReferenceResults);
-		PerformMortonTest<BucketSize, 64, 64>(Test, Elements, Tests, ReferenceResults);
-		PerformMortonTest<BucketSize, 64, 256>(Test, Elements, Tests, ReferenceResults);
-		PerformMortonTest<BucketSize, 64, 1024>(Test, Elements, Tests, ReferenceResults);
-		PerformMortonTest<BucketSize, 256, 16>(Test, Elements, Tests, ReferenceResults);
-		PerformMortonTest<BucketSize, 256, 64>(Test, Elements, Tests, ReferenceResults);
-		PerformMortonTest<BucketSize, 256, 256>(Test, Elements, Tests, ReferenceResults);
-		PerformMortonTest<BucketSize, 256, 1024>(Test, Elements, Tests, ReferenceResults);
-		PerformMortonTest<BucketSize, 1024, 16>(Test, Elements, Tests, ReferenceResults);
-		PerformMortonTest<BucketSize, 1024, 64>(Test, Elements, Tests, ReferenceResults);
-		PerformMortonTest<BucketSize, 1024, 256>(Test, Elements, Tests, ReferenceResults);
-		PerformMortonTest<BucketSize, 1024, 1024>(Test, Elements, Tests, ReferenceResults);
+		TArray<int32> RTreeHilbertResults;
+		RTreeHilbertResults.Reserve(ReferenceResults.Num());
+		FString TestName = FString::Printf(TEXT("TStaticSpatialIndexRTree(%s-hilbert-%dk-%d-%d)"), GetSpaceString<typename Profile::FBox>(), BucketSize >> 10, MaxNumElementsPerNode, MaxNumElementsPerLeaf);
+		PerformTests<Profile, TStaticSpatialIndexRTree<int32, FStaticSpatialIndex::TNodeSorterHilbert<Profile, BucketSize>, Profile>>(Test, *TestName, Elements, Tests, RTreeHilbertResults);
+		RTreeHilbertResults.Sort();
+		Test->TestTrue(TestName, RTreeHilbertResults == ReferenceResults);
 	}
 
-	template <int32 BucketSize, int32 MaxNumElementsPerNode, int32 MaxNumElementsPerLeaf>
-	void PerformHilbertTest(FWorldPartitionStaticSpatialIndexTest* Test, const TArray<TPair<FBox, int32>>& Elements, const TArray<FSphere>& Tests, const TArray<int32>& ReferenceResults)
+	template <typename Profile, int32 BucketSize>
+	FORCENOINLINE void PerformHilbertTests(FWorldPartitionStaticSpatialIndexTest* Test, const TArray<TPair<typename Profile::FBox, int32>>& Elements, const TArray<FSphere>& Tests, const TArray<int32>& ReferenceResults)
 	{
-		TArray<int32> FRTreeHilbertResults;
-		FRTreeHilbertResults.Reserve(ReferenceResults.Num());
-		FString TestName = FString::Printf(TEXT("TStaticSpatialIndexRTree(hilbert-%dk-%d-%d)"), BucketSize >> 10, MaxNumElementsPerNode, MaxNumElementsPerLeaf);
-		PerformTests<TStaticSpatialIndexRTree<int32, FStaticSpatialIndex::TNodeSorterHilbert<BucketSize>>>(Test, *TestName, Elements, Tests, FRTreeHilbertResults);
-		FRTreeHilbertResults.Sort();
-		Test->TestTrue(TestName, FRTreeHilbertResults == ReferenceResults);
+		PerformHilbertTest<Profile, BucketSize, 16, 16>(Test, Elements, Tests, ReferenceResults);
+		PerformHilbertTest<Profile, BucketSize, 16, 64>(Test, Elements, Tests, ReferenceResults);
+		PerformHilbertTest<Profile, BucketSize, 16, 256>(Test, Elements, Tests, ReferenceResults);
+		PerformHilbertTest<Profile, BucketSize, 16, 1024>(Test, Elements, Tests, ReferenceResults);
+		PerformHilbertTest<Profile, BucketSize, 64, 16>(Test, Elements, Tests, ReferenceResults);
+		PerformHilbertTest<Profile, BucketSize, 64, 64>(Test, Elements, Tests, ReferenceResults);
+		PerformHilbertTest<Profile, BucketSize, 64, 256>(Test, Elements, Tests, ReferenceResults);
+		PerformHilbertTest<Profile, BucketSize, 64, 1024>(Test, Elements, Tests, ReferenceResults);
+		PerformHilbertTest<Profile, BucketSize, 256, 16>(Test, Elements, Tests, ReferenceResults);
+		PerformHilbertTest<Profile, BucketSize, 256, 64>(Test, Elements, Tests, ReferenceResults);
+		PerformHilbertTest<Profile, BucketSize, 256, 256>(Test, Elements, Tests, ReferenceResults);
+		PerformHilbertTest<Profile, BucketSize, 256, 1024>(Test, Elements, Tests, ReferenceResults);
+		PerformHilbertTest<Profile, BucketSize, 1024, 16>(Test, Elements, Tests, ReferenceResults);
+		PerformHilbertTest<Profile, BucketSize, 1024, 64>(Test, Elements, Tests, ReferenceResults);
+		PerformHilbertTest<Profile, BucketSize, 1024, 256>(Test, Elements, Tests, ReferenceResults);
+		PerformHilbertTest<Profile, BucketSize, 1024, 1024>(Test, Elements, Tests, ReferenceResults);
 	}
 
-	template <int32 BucketSize>
-	void PerformHilbertTests(FWorldPartitionStaticSpatialIndexTest* Test, const TArray<TPair<FBox, int32>>& Elements, const TArray<FSphere>& Tests, const TArray<int32>& ReferenceResults)
+	template <typename Profile>
+	FORCENOINLINE void PerformTests(FWorldPartitionStaticSpatialIndexTest* Test, const TArray<TPair<typename Profile::FBox, int32>>& Elements, const TArray<FSphere>& Tests)
 	{
-		PerformHilbertTest<BucketSize, 16, 16>(Test, Elements, Tests, ReferenceResults);
-		PerformHilbertTest<BucketSize, 16, 64>(Test, Elements, Tests, ReferenceResults);
-		PerformHilbertTest<BucketSize, 16, 256>(Test, Elements, Tests, ReferenceResults);
-		PerformHilbertTest<BucketSize, 16, 1024>(Test, Elements, Tests, ReferenceResults);
-		PerformHilbertTest<BucketSize, 64, 16>(Test, Elements, Tests, ReferenceResults);
-		PerformHilbertTest<BucketSize, 64, 64>(Test, Elements, Tests, ReferenceResults);
-		PerformHilbertTest<BucketSize, 64, 256>(Test, Elements, Tests, ReferenceResults);
-		PerformHilbertTest<BucketSize, 64, 1024>(Test, Elements, Tests, ReferenceResults);
-		PerformHilbertTest<BucketSize, 256, 16>(Test, Elements, Tests, ReferenceResults);
-		PerformHilbertTest<BucketSize, 256, 64>(Test, Elements, Tests, ReferenceResults);
-		PerformHilbertTest<BucketSize, 256, 256>(Test, Elements, Tests, ReferenceResults);
-		PerformHilbertTest<BucketSize, 256, 1024>(Test, Elements, Tests, ReferenceResults);
-		PerformHilbertTest<BucketSize, 1024, 16>(Test, Elements, Tests, ReferenceResults);
-		PerformHilbertTest<BucketSize, 1024, 64>(Test, Elements, Tests, ReferenceResults);
-		PerformHilbertTest<BucketSize, 1024, 256>(Test, Elements, Tests, ReferenceResults);
-		PerformHilbertTest<BucketSize, 1024, 1024>(Test, Elements, Tests, ReferenceResults);
+		TArray<int32> ListResults;
+		PerformTests<Profile, TStaticSpatialIndexList<int32, FStaticSpatialIndex::TNodeSorterNoSort<Profile>, Profile>>(Test, TEXT("TStaticSpatialIndexList"), Elements, Tests, ListResults);
+		ListResults.Sort();
+
+		PerformMinXTests<Profile>(Test, Elements, Tests, ListResults);
+		PerformMinXTests<Profile>(Test, Elements, Tests, ListResults);
+		PerformMinXTests<Profile>(Test, Elements, Tests, ListResults);
+		PerformMinXTests<Profile>(Test, Elements, Tests, ListResults);
+
+		PerformMortonTests<Profile, 4096>(Test, Elements, Tests, ListResults);
+		PerformMortonTests<Profile, 16384>(Test, Elements, Tests, ListResults);
+		PerformMortonTests<Profile, 65536>(Test, Elements, Tests, ListResults);
+		PerformMortonTests<Profile, 262144>(Test, Elements, Tests, ListResults);
+
+		PerformHilbertTests<Profile, 4096>(Test, Elements, Tests, ListResults);
+		PerformHilbertTests<Profile, 16384>(Test, Elements, Tests, ListResults);
+		PerformHilbertTests<Profile, 65536>(Test, Elements, Tests, ListResults);
+		PerformHilbertTests<Profile, 262144>(Test, Elements, Tests, ListResults);
 	}
 #endif
 
@@ -139,11 +175,15 @@ namespace WorldPartitionTests
 		const int32 NumTests = 10000;
 
 		TArray<TPair<FBox, int32>> Elements;
+		TArray<TPair<FBox2D, int32>> Elements2D;
 		Elements.Reserve(NumBoxes);
+		Elements2D.Reserve(NumBoxes);
 		for (int32 i=0; i<NumBoxes; i++)
 		{
 			const FSphere Sphere(FMath::VRand() * 10000000, FMath::RandRange(1, 100000));
-			Elements.Emplace(FBoxSphereBounds(Sphere).GetBox(), i);
+			const FBox ElementBox(FBoxSphereBounds(Sphere).GetBox());
+			Elements.Emplace(ElementBox, i);
+			Elements2D.Emplace(FBox2D(FVector2D(ElementBox.Min), FVector2D(ElementBox.Max)), i);
 		}
 
 		TArray<FSphere> Tests;
@@ -154,24 +194,8 @@ namespace WorldPartitionTests
 			Tests.Add(Sphere);
 		}
 
-		TArray<int32> ListResults;
-		PerformTests<TStaticSpatialIndexList<int32, FStaticSpatialIndex::FNodeSorterNoSort>>(this, TEXT("TStaticSpatialIndexList"), Elements, Tests, ListResults);
-		ListResults.Sort();
-
-		PerformMinXTests(this, Elements, Tests, ListResults);
-		PerformMinXTests(this, Elements, Tests, ListResults);
-		PerformMinXTests(this, Elements, Tests, ListResults);
-		PerformMinXTests(this, Elements, Tests, ListResults);
-
-		PerformMortonTests<4096>(this, Elements, Tests, ListResults);
-		PerformMortonTests<16384>(this, Elements, Tests, ListResults);
-		PerformMortonTests<65536>(this, Elements, Tests, ListResults);
-		PerformMortonTests<262144>(this, Elements, Tests, ListResults);
-
-		PerformHilbertTests<4096>(this, Elements, Tests, ListResults);
-		PerformHilbertTests<16384>(this, Elements, Tests, ListResults);
-		PerformHilbertTests<65536>(this, Elements, Tests, ListResults);
-		PerformHilbertTests<262144>(this, Elements, Tests, ListResults);
+		PerformTests<FStaticSpatialIndex::FSpatialIndexProfile3D>(this, Elements, Tests);
+		PerformTests<FStaticSpatialIndex::FSpatialIndexProfile2D>(this, Elements2D, Tests);
 #endif
 		return true;
 	}
diff --git a/Engine/Source/Runtime/Engine/Private/WorldPartition/RuntimeHashSet/RuntimePartition.cpp b/Engine/Source/Runtime/Engine/Private/WorldPartition/RuntimeHashSet/RuntimePartition.cpp
--- a/Engine/Source/Runtime/Engine/Private/WorldPartition/RuntimeHashSet/RuntimePartition.cpp
+++ b/Engine/Source/Runtime/Engine/Private/WorldPartition/RuntimeHashSet/RuntimePartition.cpp
@@ -19,6 +19,7 @@ void URuntimePartition::SetDefaultValues()
 	bBlockOnSlowStreaming = false;
 	bClientOnlyVisible = false;
 	Priority = 0;
+	BoundsMethod = ERuntimePartitionCellBoundsMethod::UseMinContentCellBounds;
 	LoadingRange = 25600;
 	DebugColor = FLinearColor::MakeRandomSeededColor(GetTypeHash(GetName()));
 	HLODIndex = INDEX_NONE;
@@ -72,12 +73,6 @@ URuntimePartition::FCellDesc URuntimePartition::CreateCellDesc(const FString& In
 	// Add actor set instances
 	CellDesc.ActorSetInstances = InActorSetInstances;
 
-	// Update cell bounds
-	for (const IStreamingGenerationContext::FActorSetInstance* ActorSetInstance : InActorSetInstances)
-	{
-		CellDesc.Bounds += ActorSetInstance->Bounds;
-	}
-
 	return CellDesc;
 }
 
diff --git a/Engine/Source/Runtime/Engine/Private/WorldPartition/RuntimeHashSet/RuntimePartitionLHGrid.cpp b/Engine/Source/Runtime/Engine/Private/WorldPartition/RuntimeHashSet/RuntimePartitionLHGrid.cpp
--- a/Engine/Source/Runtime/Engine/Private/WorldPartition/RuntimeHashSet/RuntimePartitionLHGrid.cpp
+++ b/Engine/Source/Runtime/Engine/Private/WorldPartition/RuntimeHashSet/RuntimePartitionLHGrid.cpp
@@ -28,8 +28,8 @@ struct FCellCoord
 
 	static inline int32 GetLevelForBox(const FBox& InBox, int32 InCellSize)
 	{
-		const FVector Extent = InBox.GetExtent();
-		const FVector::FReal MaxLength = Extent.GetMax() * 2.0;
+		const FVector Extent = InBox.GetSize();
+		const FVector::FReal MaxLength = Extent.GetMax();
 		return FMath::CeilToInt32(FMath::Max<FVector::FReal>(FMath::Log2(MaxLength / InCellSize), 0));
 	}
 
@@ -145,14 +145,16 @@ bool URuntimePartitionLHGrid::GenerateStreaming(const FGenerateStreamingParams&
 	UWorld* World = WorldPartition->GetWorld();
 	UWorld* OuterWorld = GetTypedOuter<UWorld>();
 	const bool bIsMainWorldPartition = (World == OuterWorld);
+	const FVector SpaceMask = FVector(1, 1, bIs2D ? 0 : 1);
 
 	TMap<FCellCoord, TArray<const IStreamingGenerationContext::FActorSetInstance*>> CellsActorSetInstances;
 	for (const IStreamingGenerationContext::FActorSetInstance* ActorSetInstance : *InParams.ActorSetInstances)
 	{
 		if (ActorSetInstance->bIsSpatiallyLoaded)
 		{
-			const int32 GridLevel = FCellCoord::GetLevelForBox(ActorSetInstance->Bounds, CellSize);
-			const FCellCoord CellCoord = FCellCoord::GetCellCoords(ActorSetInstance->Bounds.GetCenter(), CellSize, GridLevel, Origin);
+			const FBox ActorSetInstanceBounds = FBox(ActorSetInstance->Bounds.Min * SpaceMask, ActorSetInstance->Bounds.Max * SpaceMask);
+			const int32 GridLevel = FCellCoord::GetLevelForBox(ActorSetInstanceBounds, CellSize);
+			const FCellCoord CellCoord = FCellCoord::GetCellCoords(ActorSetInstanceBounds.GetCenter(), CellSize, GridLevel, Origin * SpaceMask);
 			CellsActorSetInstances.FindOrAdd(CellCoord).Add(ActorSetInstance);
 		}
 		else
@@ -169,8 +171,8 @@ bool URuntimePartitionLHGrid::GenerateStreaming(const FGenerateStreamingParams&
 
 		if (bIsSpatiallyLoaded)
 		{
-			CellDesc.CellBounds = FCellCoord::GetCellBounds(CellCoord, CellSize, Origin);
-			CellDesc.Bounds = CellDesc.CellBounds.GetValue();
+			CellDesc.CellBounds = FCellCoord::GetCellBounds(CellCoord, CellSize, Origin * SpaceMask);
+			CellDesc.bIs2D = bIs2D;
 		}
 	}
 
diff --git a/Engine/Source/Runtime/Engine/Private/WorldPartition/RuntimeHashSet/StaticSpatialIndex.cpp b/Engine/Source/Runtime/Engine/Private/WorldPartition/RuntimeHashSet/StaticSpatialIndex.cpp
deleted file mode 100644
diff --git a/Engine/Source/Runtime/Engine/Private/WorldPartition/RuntimeHashSet/WorldPartitionRuntimeHashSet.cpp b/Engine/Source/Runtime/Engine/Private/WorldPartition/RuntimeHashSet/WorldPartitionRuntimeHashSet.cpp
--- a/Engine/Source/Runtime/Engine/Private/WorldPartition/RuntimeHashSet/WorldPartitionRuntimeHashSet.cpp
+++ b/Engine/Source/Runtime/Engine/Private/WorldPartition/RuntimeHashSet/WorldPartitionRuntimeHashSet.cpp
@@ -1,6 +1,7 @@
 // Copyright Epic Games, Inc. All Rights Reserved.
 
 #include "WorldPartition/RuntimeHashSet/WorldPartitionRuntimeHashSet.h"
+#include "WorldPartition/RuntimeHashSet/WorldPartitionRuntimeCellDataHashSet.h"
 #include "WorldPartition/RuntimeHashSet/RuntimePartition.h"
 #include "WorldPartition/RuntimeHashSet/RuntimePartitionLHGrid.h"
 #include "WorldPartition/RuntimeHashSet/RuntimePartitionPersistent.h"
@@ -20,27 +21,62 @@ void FRuntimePartitionStreamingData::CreatePartitionsSpatialIndex() const
 {
 	if (!SpatialIndex)
 	{
-		SpatialIndex = MakeUnique<FStaticSpatialIndexType>();
+		TArray<UWorldPartitionRuntimeCell*> SpatiallyLoadedCells3D;
+		TArray<UWorldPartitionRuntimeCell*> SpatiallyLoadedCells2D;
+		
+		SpatiallyLoadedCells3D.Reserve(SpatiallyLoadedCells.Num());		
+		SpatiallyLoadedCells2D.Reserve(SpatiallyLoadedCells.Num());
+
+		for (UWorldPartitionRuntimeCell* Cell : SpatiallyLoadedCells)
 		{
-			TArray<TPair<FBox, TObjectPtr<UWorldPartitionRuntimeCell>>> PartitionsElements;
-			Algo::Transform(SpatiallyLoadedCells, PartitionsElements, [](UWorldPartitionRuntimeCell* Cell)
+			if (CastChecked<UWorldPartitionRuntimeCellDataHashSet>(Cell->RuntimeCellData)->bIs2D)
 			{
-				return TPair<FBox, TObjectPtr<UWorldPartitionRuntimeCell>>(Cell->GetStreamingBounds(), Cell);
-			});
-			SpatialIndex->Init(MoveTemp(PartitionsElements));
+				SpatiallyLoadedCells2D.Add(Cell);
+			}
+			else
+			{
+				SpatiallyLoadedCells3D.Add(Cell);
+			}
 		}
+
+		if (SpatiallyLoadedCells3D.Num())
+		{
+			SpatialIndex = MakeUnique<FStaticSpatialIndexType>();
+			{
+				TArray<TPair<FBox, TObjectPtr<UWorldPartitionRuntimeCell>>> PartitionsElements;
+				Algo::Transform(SpatiallyLoadedCells3D, PartitionsElements, [](UWorldPartitionRuntimeCell* Cell)
+				{
+					return TPair<FBox, TObjectPtr<UWorldPartitionRuntimeCell>>(Cell->GetStreamingBounds(), Cell);
+				});
+				SpatialIndex->Init(MoveTemp(PartitionsElements));
+			}
 		
-		SpatialIndex2D = MakeUnique<FStaticSpatialIndexType>();
+			SpatialIndexForce2D = MakeUnique<FStaticSpatialIndexType2D>();
+			{
+				TArray<TPair<FBox2D, TObjectPtr<UWorldPartitionRuntimeCell>>> PartitionsElements;
+				Algo::Transform(SpatiallyLoadedCells3D, PartitionsElements, [](UWorldPartitionRuntimeCell* Cell)
+				{
+					const FBox CellBounds = Cell->GetStreamingBounds();
+					const FBox2D CellBounds2D = FBox2D(FVector2D(CellBounds.Min), FVector2D(CellBounds.Max));
+					return TPair<FBox2D, TObjectPtr<UWorldPartitionRuntimeCell>>(CellBounds2D, Cell);
+				});
+				SpatialIndexForce2D->Init(MoveTemp(PartitionsElements));
+			}
+		}
+
+		if (SpatiallyLoadedCells2D.Num())
 		{
-			TArray<TPair<FBox, TObjectPtr<UWorldPartitionRuntimeCell>>> PartitionsElements;
-			Algo::Transform(SpatiallyLoadedCells, PartitionsElements, [](UWorldPartitionRuntimeCell* Cell)
+			SpatialIndex2D = MakeUnique<FStaticSpatialIndexType2D>();
 			{
-				FBox CellBounds = Cell->GetStreamingBounds();
-				CellBounds.Min.Z = -HALF_WORLD_MAX;
-				CellBounds.Max.Z = HALF_WORLD_MAX;
-				return TPair<FBox, TObjectPtr<UWorldPartitionRuntimeCell>>(CellBounds, Cell);
-			});
-			SpatialIndex2D->Init(MoveTemp(PartitionsElements));
+				TArray<TPair<FBox2D, TObjectPtr<UWorldPartitionRuntimeCell>>> PartitionsElements;
+				Algo::Transform(SpatiallyLoadedCells2D, PartitionsElements, [](UWorldPartitionRuntimeCell* Cell)
+				{
+					const FBox CellBounds = Cell->GetStreamingBounds();
+					const FBox2D CellBounds2D = FBox2D(FVector2D(CellBounds.Min), FVector2D(CellBounds.Max));
+					return TPair<FBox2D, TObjectPtr<UWorldPartitionRuntimeCell>>(CellBounds2D, Cell);
+				});
+				SpatialIndex2D->Init(MoveTemp(PartitionsElements));
+			}
 		}
 	}
 }
@@ -48,6 +84,7 @@ void FRuntimePartitionStreamingData::CreatePartitionsSpatialIndex() const
 void FRuntimePartitionStreamingData::DestroyPartitionsSpatialIndex() const
 {
 	SpatialIndex.Reset();
+	SpatialIndexForce2D.Reset();
 	SpatialIndex2D.Reset();
 }
 
@@ -90,6 +127,11 @@ void URuntimeHashSetExternalStreamingObject::AddReferencedObjects(UObject* InThi
 			StreamingData.SpatialIndex->AddReferencedObjects(Collector);
 		}
 		
+		if (StreamingData.SpatialIndexForce2D.IsValid())
+		{
+			StreamingData.SpatialIndexForce2D->AddReferencedObjects(Collector);
+		}
+
 		if (StreamingData.SpatialIndex2D.IsValid())
 		{
 			StreamingData.SpatialIndex2D->AddReferencedObjects(Collector);
@@ -557,13 +599,21 @@ void UWorldPartitionRuntimeHashSet::ForEachStreamingCellsSources(const TArray<FW
 							}
 						};
 
-						if (Source.bForce2D)
+						if (StreamingData->SpatialIndex.IsValid())
 						{
-							StreamingData->SpatialIndex2D->ForEachIntersectingElement(ShapeSphere, ForEachIntersectingElementFunc);
+							if (Source.bForce2D)
+							{
+								StreamingData->SpatialIndexForce2D->ForEachIntersectingElement(ShapeSphere, ForEachIntersectingElementFunc);
+							}
+							else
+							{
+								StreamingData->SpatialIndex->ForEachIntersectingElement(ShapeSphere, ForEachIntersectingElementFunc);
+							}
 						}
-						else
+
+						if (StreamingData->SpatialIndex2D.IsValid())
 						{
-							StreamingData->SpatialIndex->ForEachIntersectingElement(ShapeSphere, ForEachIntersectingElementFunc);
+							StreamingData->SpatialIndex2D->ForEachIntersectingElement(ShapeSphere, ForEachIntersectingElementFunc);
 						}
 					});
 				}
@@ -667,7 +717,7 @@ FGuid UWorldPartitionRuntimeHashSet::RegisterWorldAssetStreaming(const UWorldPar
 		const TSoftObjectPtr<UWorld>& WorldAsset = bIsHLODPass ? InParams.WorldAssetHLOD : InParams.WorldAsset;
 
 		if (UWorldPartitionRuntimeLevelStreamingCell* RuntimeCell = Cast<UWorldPartitionRuntimeLevelStreamingCell>(
-			CreateRuntimeCell(UWorldPartitionRuntimeLevelStreamingCell::StaticClass(), UWorldPartitionRuntimeCellData::StaticClass(), CellName, InParams.CellInstanceSuffix, StreamingObject)))
+			CreateRuntimeCell(UWorldPartitionRuntimeLevelStreamingCell::StaticClass(), UWorldPartitionRuntimeCellDataHashSet::StaticClass(), CellName, InParams.CellInstanceSuffix, StreamingObject)))
 		{
 			RuntimeCell->SetClientOnlyVisible(bClientOnlyVisible);
 			RuntimeCell->SetBlockOnSlowLoading(bBlockOnSlowStreaming);
@@ -680,13 +730,14 @@ FGuid UWorldPartitionRuntimeHashSet::RegisterWorldAssetStreaming(const UWorldPar
 				RuntimeCell->SetSourceCellGuid(SourceCellGuid);
 			}
 
-			UWorldPartitionRuntimeCellData* RuntimeCellData = RuntimeCell->RuntimeCellData;
+			UWorldPartitionRuntimeCellDataHashSet* RuntimeCellData = CastChecked<UWorldPartitionRuntimeCellDataHashSet>(RuntimeCell->RuntimeCellData);
 			RuntimeCellData->DebugName = CellName + InParams.CellInstanceSuffix;
 			RuntimeCellData->CellBounds = InParams.Bounds;
 			RuntimeCellData->ContentBounds = InParams.Bounds;
 			RuntimeCellData->HierarchicalLevel = MAX_int32;
 			RuntimeCellData->Priority = InParams.Priority;
 			RuntimeCellData->GridName = TargetGrid;
+			RuntimeCellData->bIs2D = false;
 
 			if (RuntimeCell->CreateAndSetLevelStreaming(WorldAsset, InParams.Transform))
 			{
diff --git a/Engine/Source/Runtime/Engine/Private/WorldPartition/RuntimeHashSet/WorldPartitionRuntimeHashSetConversions.cpp b/Engine/Source/Runtime/Engine/Private/WorldPartition/RuntimeHashSet/WorldPartitionRuntimeHashSetConversions.cpp
--- a/Engine/Source/Runtime/Engine/Private/WorldPartition/RuntimeHashSet/WorldPartitionRuntimeHashSetConversions.cpp
+++ b/Engine/Source/Runtime/Engine/Private/WorldPartition/RuntimeHashSet/WorldPartitionRuntimeHashSetConversions.cpp
@@ -69,6 +69,7 @@ UWorldPartitionRuntimeHashSet* UWorldPartitionRuntimeHashSet::CreateFrom(const U
 				LHGrid->Name = RuntimePartitionDesc.Name;
 				LHGrid->CellSize = Grid.CellSize;
 				LHGrid->Origin = bUseAlignedGridLevels ? FVector::ZeroVector : FVector(LHGrid->CellSize * -0.5f);
+				LHGrid->bIs2D = true;
 				LHGrid->bBlockOnSlowStreaming = Grid.bBlockOnSlowStreaming;
 				LHGrid->bClientOnlyVisible = Grid.bClientOnlyVisible;
 				LHGrid->Priority = Grid.Priority;
@@ -91,6 +92,7 @@ UWorldPartitionRuntimeHashSet* UWorldPartitionRuntimeHashSet::CreateFrom(const U
 						URuntimePartitionLHGrid* HLODLHGrid = NewObject<URuntimePartitionLHGrid>(HashSet, NAME_None);
 						HLODLHGrid->CellSize = HLODSetup.HLODLayers[0]->GetCellSize();
 						HLODLHGrid->Origin = bUseAlignedGridLevels ? FVector::ZeroVector : FVector(HLODLHGrid->CellSize * -0.5f);
+						HLODLHGrid->bIs2D = true;
 						HLODLHGrid->LoadingRange = HLODSetup.HLODLayers[0]->GetLoadingRange();
 						HLODSetup.PartitionLayer = HLODLHGrid;
 					}
diff --git a/Engine/Source/Runtime/Engine/Private/WorldPartition/RuntimeHashSet/WorldPartitionRuntimeHashSetDebugDraw.cpp b/Engine/Source/Runtime/Engine/Private/WorldPartition/RuntimeHashSet/WorldPartitionRuntimeHashSetDebugDraw.cpp
--- a/Engine/Source/Runtime/Engine/Private/WorldPartition/RuntimeHashSet/WorldPartitionRuntimeHashSetDebugDraw.cpp
+++ b/Engine/Source/Runtime/Engine/Private/WorldPartition/RuntimeHashSet/WorldPartitionRuntimeHashSetDebugDraw.cpp
@@ -120,16 +120,26 @@ bool UWorldPartitionRuntimeHashSet::Draw2D(FWorldPartitionDraw2DContext& DrawCon
 		};
 
 		TArray<const UWorldPartitionRuntimeCell*> FilteredCells;
-		const FBox Region3D(FVector(Region.Min.X, Region.Min.Y, -HALF_WORLD_MAX), FVector(Region.Max.X, Region.Max.Y, HALF_WORLD_MAX));		
-		StreamingData->SpatialIndex->ForEachIntersectingElement(Region3D, [&FilteredCells](UWorldPartitionRuntimeCell* Cell)
+		const FBox Region3D(FVector(Region.Min.X, Region.Min.Y, -HALF_WORLD_MAX), FVector(Region.Max.X, Region.Max.Y, HALF_WORLD_MAX));
+
+		auto FilterCell = [&FilteredCells](UWorldPartitionRuntimeCell* Cell)
 		{
 			UWorldPartitionRuntimeCellData* RuntimeCellData = Cell->RuntimeCellData;
-
 			if ((RuntimeCellData->HierarchicalLevel >= GShowRuntimeHashSetDebugDisplayLevel) && (RuntimeCellData->HierarchicalLevel < (GShowRuntimeHashSetDebugDisplayLevel + GShowRuntimeHashSetDebugDisplayLevelCount)))
 			{
 				FilteredCells.Add(Cell);
 			}
-		});
+		};
+		
+		if (StreamingData->SpatialIndex.IsValid())
+		{
+			StreamingData->SpatialIndex->ForEachIntersectingElement(Region3D, FilterCell);
+		}
+
+		if (StreamingData->SpatialIndex2D.IsValid())
+		{
+			StreamingData->SpatialIndex2D->ForEachIntersectingElement(Region2D, FilterCell);
+		}
 
 		if (FilteredCells.Num())
 		{
diff --git a/Engine/Source/Runtime/Engine/Private/WorldPartition/RuntimeHashSet/WorldPartitionRuntimeHashSetHLODGeneration.cpp b/Engine/Source/Runtime/Engine/Private/WorldPartition/RuntimeHashSet/WorldPartitionRuntimeHashSetHLODGeneration.cpp
--- a/Engine/Source/Runtime/Engine/Private/WorldPartition/RuntimeHashSet/WorldPartitionRuntimeHashSetHLODGeneration.cpp
+++ b/Engine/Source/Runtime/Engine/Private/WorldPartition/RuntimeHashSet/WorldPartitionRuntimeHashSetHLODGeneration.cpp
@@ -204,12 +204,13 @@ bool UWorldPartitionRuntimeHashSet::SetupHLODActors(const IStreamingGenerationCo
 
 				UE_LOG(LogWorldPartition, Display, TEXT("[%d / %d] Processing cell %s..."), ++CellDescInstanceIndex, CellDescInstances.Num(), *CellUniqueId.Name);
 
+				FBox CellBounds;
 				TArray<IStreamingGenerationContext::FActorInstance> ActorInstances;
 				for (const IStreamingGenerationContext::FActorSetInstance* ActorSetInstance : CellDescInstance.ActorSetInstances)
 				{
-					ActorSetInstance->ForEachActor([this, ActorSetInstance, &ActorInstances](const FGuid& ActorGuid)
+					ActorSetInstance->ForEachActor([this, ActorSetInstance, &ActorInstances, &CellBounds](const FGuid& ActorGuid)
 					{
-						ActorInstances.Emplace(ActorGuid, ActorSetInstance);
+						CellBounds += ActorInstances.Emplace_GetRef(ActorGuid, ActorSetInstance).GetBounds();
 					});
 				}
 
@@ -224,12 +225,12 @@ bool UWorldPartitionRuntimeHashSet::SetupHLODActors(const IStreamingGenerationCo
 				{
 					MainPartitionTokens[0] = RuntimePartitions[0].Name;
 				}
-					
+
 				FHLODCreationParams HLODCreationParams;
 				HLODCreationParams.WorldPartition = WorldPartition;
 				HLODCreationParams.CellName = CellUniqueId.Name;
-				HLODCreationParams.CellGuid = CellUniqueId.Guid;			
-				HLODCreationParams.CellBounds = CellDescInstance.Bounds;
+				HLODCreationParams.CellGuid = CellUniqueId.Guid;
+				HLODCreationParams.CellBounds = CellBounds;
 				HLODCreationParams.GetRuntimeGrid = [&MainPartitionTokens](const UHLODLayer* InHLODLayer) { return FName(*FString::Printf(TEXT("%s:%s"), *FString::JoinBy(MainPartitionTokens, TEXT("."), [](const FName Token) { return Token.ToString(); }), *InHLODLayer->GetName())); };
 				HLODCreationParams.HLODLevel = HLODLevel;
 				HLODCreationParams.MinVisibleDistance = RuntimePartition->LoadingRange;
diff --git a/Engine/Source/Runtime/Engine/Private/WorldPartition/RuntimeHashSet/WorldPartitionRuntimeHashSetStreamingGeneration.cpp b/Engine/Source/Runtime/Engine/Private/WorldPartition/RuntimeHashSet/WorldPartitionRuntimeHashSetStreamingGeneration.cpp
--- a/Engine/Source/Runtime/Engine/Private/WorldPartition/RuntimeHashSet/WorldPartitionRuntimeHashSetStreamingGeneration.cpp
+++ b/Engine/Source/Runtime/Engine/Private/WorldPartition/RuntimeHashSet/WorldPartitionRuntimeHashSetStreamingGeneration.cpp
@@ -1,6 +1,7 @@
 // Copyright Epic Games, Inc. All Rights Reserved.
 
 #include "WorldPartition/RuntimeHashSet/WorldPartitionRuntimeHashSet.h"
+#include "WorldPartition/RuntimeHashSet/WorldPartitionRuntimeCellDataHashSet.h"
 #include "WorldPartition/RuntimeHashSet/RuntimePartition.h"
 #include "WorldPartition/RuntimeHashSet/RuntimePartitionPersistent.h"
 #include "WorldPartition/ContentBundle/ContentBundleDescriptor.h"
@@ -144,7 +145,7 @@ bool UWorldPartitionRuntimeHashSet::GenerateStreaming(UWorldPartitionStreamingPo
 	//
 
 	// Generate runtime cells
-	auto CreateRuntimeCellFromCellDesc = [this](const URuntimePartition::FCellDescInstance& CellDescInstance, TSubclassOf<UWorldPartitionRuntimeCell> CellClass, TSubclassOf<UWorldPartitionRuntimeCellData> CellDataClass)
+	auto CreateRuntimeCellFromCellDesc = [this](const URuntimePartition::FCellDescInstance& CellDescInstance, TSubclassOf<UWorldPartitionRuntimeCell> CellClass, TSubclassOf<UWorldPartitionRuntimeCellDataHashSet> CellDataClass)
 	{
 		const FCellUniqueId CellUniqueId = GetCellUniqueId(CellDescInstance);
 
@@ -160,12 +161,13 @@ bool UWorldPartitionRuntimeHashSet::GenerateStreaming(UWorldPartitionStreamingPo
 		RuntimeCell->SetGuid(CellUniqueId.Guid);
 		RuntimeCell->SetCellDebugColor(CellDescInstance.SourcePartition->DebugColor);
 
-		UWorldPartitionRuntimeCellData* RuntimeCellData = RuntimeCell->RuntimeCellData;
+		UWorldPartitionRuntimeCellDataHashSet* RuntimeCellData = CastChecked<UWorldPartitionRuntimeCellDataHashSet>(RuntimeCell->RuntimeCellData);
 		RuntimeCellData->DebugName = CellUniqueId.Name;
 		RuntimeCellData->CellBounds = CellDescInstance.CellBounds;
 		RuntimeCellData->HierarchicalLevel = CellDescInstance.bIsSpatiallyLoaded ? CellDescInstance.Level : MAX_int32;
 		RuntimeCellData->Priority = CellDescInstance.Priority;
 		RuntimeCellData->GridName = CellDescInstance.SourcePartition->Name;
+		RuntimeCellData->bIs2D = CellDescInstance.bIs2D;
 
 		return RuntimeCell;
 	};
@@ -181,14 +183,24 @@ bool UWorldPartitionRuntimeHashSet::GenerateStreaming(UWorldPartitionStreamingPo
 			TArray<IStreamingGenerationContext::FActorInstance> CellActorInstances;
 			if (PopulateCellActorInstances(CellDescInstance.ActorSetInstances, bIsMainWorldPartition, bIsCellAlwaysLoaded, CellActorInstances))
 			{
-				UWorldPartitionRuntimeCell* RuntimeCell = RuntimeCells.Emplace_GetRef(CreateRuntimeCellFromCellDesc(CellDescInstance, StreamingPolicy->GetRuntimeCellClass(), UWorldPartitionRuntimeCellData::StaticClass()));
+				UWorldPartitionRuntimeCell* RuntimeCell = RuntimeCells.Emplace_GetRef(CreateRuntimeCellFromCellDesc(CellDescInstance, StreamingPolicy->GetRuntimeCellClass(), UWorldPartitionRuntimeCellDataHashSet::StaticClass()));
 				RuntimeCell->SetIsAlwaysLoaded(bIsCellAlwaysLoaded);
 				PopulateRuntimeCell(RuntimeCell, CellActorInstances, OutPackagesToGenerate);
 
-				// Override the cell bounds if the runtime partition provided one
-				if (CellDescInstance.Bounds.IsValid)
+				if (CellDescInstance.CellBounds.IsSet())
 				{
-					RuntimeCell->RuntimeCellData->ContentBounds = CellDescInstance.Bounds;
+					switch (CellDescInstance.SourcePartition->BoundsMethod)
+					{
+					case ERuntimePartitionCellBoundsMethod::UseCellBounds:
+						RuntimeCell->RuntimeCellData->ContentBounds = CellDescInstance.CellBounds.GetValue();
+						break;
+					case ERuntimePartitionCellBoundsMethod::UseMinContentCellBounds:
+						if (RuntimeCell->RuntimeCellData->ContentBounds.IsValid)
+						{
+							RuntimeCell->RuntimeCellData->ContentBounds = RuntimeCell->RuntimeCellData->ContentBounds.Overlap(CellDescInstance.CellBounds.GetValue());
+						}
+						break;
+					}
 				}
 
 				// Create partition streaming data
diff --git a/Engine/Source/Runtime/Engine/Public/WorldPartition/RuntimeHashSet/RuntimePartition.h b/Engine/Source/Runtime/Engine/Public/WorldPartition/RuntimeHashSet/RuntimePartition.h
--- a/Engine/Source/Runtime/Engine/Public/WorldPartition/RuntimeHashSet/RuntimePartition.h
+++ b/Engine/Source/Runtime/Engine/Public/WorldPartition/RuntimeHashSet/RuntimePartition.h
@@ -5,6 +5,15 @@
 #include "WorldPartition/WorldPartitionStreamingGenerationContext.h"
 #include "RuntimePartition.generated.h"
 
+/** Chooses a method for how to compute streaming cells bounds */
+UENUM()
+enum class ERuntimePartitionCellBoundsMethod : uint8
+{
+	UseContent,
+	UseCellBounds,
+	UseMinContentCellBounds
+};
+
 UCLASS(Abstract, CollapseCategories)
 class URuntimePartition : public UObject
 {
@@ -24,10 +33,10 @@ public:
 	struct FCellDesc
 	{
 		FName Name;
-		FBox Bounds = FBox(ForceInit);
 		bool bIsSpatiallyLoaded;
 		bool bBlockOnSlowStreaming;
 		bool bClientOnlyVisible;
+		bool bIs2D;
 		int32 Priority;
 
 		/** Optional level value that can be used to filter debug display */
@@ -75,7 +84,7 @@ public:
 	FName Name;
 
 	UPROPERTY(EditAnywhere, Category = RuntimeSettings, Meta = (EditCondition = "HLODIndex == INDEX_NONE", EditConditionHides, HideEditConditionToggle))
-	bool bBlockOnSlowStreaming;
+	bool bBlockOnSlowStreaming = false;
 
 	UPROPERTY(EditAnywhere, Category = RuntimeSettings, Meta = (EditCondition = "HLODIndex == INDEX_NONE", EditConditionHides, HideEditConditionToggle))
 	bool bClientOnlyVisible;
@@ -83,6 +92,9 @@ public:
 	UPROPERTY(EditAnywhere, Category = RuntimeSettings, Meta = (EditCondition = "HLODIndex == INDEX_NONE", EditConditionHides, HideEditConditionToggle))
 	int32 Priority;
 
+	UPROPERTY(EditAnywhere, Category = RuntimeSettings, Meta = (EditCondition = "HLODIndex == INDEX_NONE", EditConditionHides, HideEditConditionToggle))
+	ERuntimePartitionCellBoundsMethod BoundsMethod;
+
 	UPROPERTY(EditAnywhere, Category = RuntimeSettings)
 	int32 LoadingRange;
 
diff --git a/Engine/Source/Runtime/Engine/Public/WorldPartition/RuntimeHashSet/RuntimePartitionLHGrid.h b/Engine/Source/Runtime/Engine/Public/WorldPartition/RuntimeHashSet/RuntimePartitionLHGrid.h
--- a/Engine/Source/Runtime/Engine/Public/WorldPartition/RuntimeHashSet/RuntimePartitionLHGrid.h
+++ b/Engine/Source/Runtime/Engine/Public/WorldPartition/RuntimeHashSet/RuntimePartitionLHGrid.h
@@ -41,6 +41,9 @@ private:
 	UPROPERTY(EditAnywhere, Category = RuntimeSettings)
 	FVector Origin = FVector::ZeroVector;
 
+	UPROPERTY(EditAnywhere, Category = RuntimeSettings)
+	bool bIs2D = false;
+
 	UPROPERTY(EditAnywhere, Category = RuntimeSettings, Transient, SkipSerialization)
 	bool bShowGridPreview = false;
 #endif
diff --git a/Engine/Source/Runtime/Engine/Public/WorldPartition/RuntimeHashSet/StaticSpatialIndex.h b/Engine/Source/Runtime/Engine/Public/WorldPartition/RuntimeHashSet/StaticSpatialIndex.h
--- a/Engine/Source/Runtime/Engine/Public/WorldPartition/RuntimeHashSet/StaticSpatialIndex.h
+++ b/Engine/Source/Runtime/Engine/Public/WorldPartition/RuntimeHashSet/StaticSpatialIndex.h
@@ -9,16 +9,48 @@
 #include "UObject/UObjectGlobals.h"
 #include "OverrideVoidReturnInvoker.h"
 
-struct IStaticSpatialIndexDataInterface
+namespace FStaticSpatialIndex
+{
+	struct FSpatialIndexProfile2D
+	{
+		enum { Is3D = 0 };
+		using FReal = FVector2D::FReal;
+		using FVector = FVector2D;
+		using FIntPoint = FIntVector2;
+		using FBox = FBox2D;
+	};
+
+	struct FSpatialIndexProfile3D
+	{
+		enum { Is3D = 1 };
+		using FReal = FVector::FReal;
+		using FVector = FVector;
+		using FIntPoint = FIntVector;
+		using FBox = FBox;
+	};
+
+	template <typename Profile>
+	inline bool FastSphereAABBIntersection(const typename Profile::FVector& InCenter, const typename Profile::FReal InRadiusSquared, const typename Profile::FBox& InBox)
+	{
+		const typename Profile::FVector ClosestPoint = Profile::FVector::Max(InBox.Min, Profile::FVector::Min(InCenter, InBox.Max));
+		return (ClosestPoint - InCenter).SizeSquared() <= InRadiusSquared;
+	}
+}
+
+template <typename Profile>
+struct TStaticSpatialIndexDataInterface
 {
-	virtual ~IStaticSpatialIndexDataInterface() {}
-	virtual const FBox& GetBox(uint32 InIndex) const = 0;
+	virtual ~TStaticSpatialIndexDataInterface() {}
+	virtual int32 GetNumBox() const = 0;
+	virtual const typename Profile::FBox& GetBox(uint32 InIndex) const = 0;
 	virtual uint32 GetAllocatedSize() const = 0;
 };
 
-template <typename ValueType, class SpatialIndexType, class ElementsSorter>
-class TStaticSpatialIndex : public IStaticSpatialIndexDataInterface
+template <typename ValueType, typename Profile, class SpatialIndexType, class ElementsSorter>
+class TStaticSpatialIndex : public TStaticSpatialIndexDataInterface<Profile>
 {
+	using FBox = typename Profile::FBox;
+
 public:
 	TStaticSpatialIndex()
 		: SpatialIndex(*this)
@@ -80,7 +112,12 @@ public:
 		}
 	}
 
-	// IStaticSpatialIndexDataInterface interface
+	// TStaticSpatialIndexDataInterface interface
+	virtual int32 GetNumBox() const override
+	{
+		return Elements.Num();
+	}
+
 	virtual const FBox& GetBox(uint32 InIndex) const override
 	{
 		return Elements[InIndex].Key;
@@ -105,15 +142,8 @@ private:
 			Elements.Sort([&Sorter](const TPair<FBox, ValueType>& A, const TPair<FBox, ValueType>B) { return Sorter.Sort(A.Key, B.Key); });
 		}
 
-		// Build index data for type agnostic implementation
-		TArray<TPair<FBox, uint32>> IndexData;
-		
-		int32 ElemIndex = 0;
-		IndexData.Reserve(Elements.Num());
-		Algo::Transform(Elements, IndexData, [&ElemIndex](const TPair<FBox, ValueType>& Element) { return TPair<FBox, uint32>(Element.Key, ElemIndex++); });
-
 		// Initialize spatial index implementation
-		SpatialIndex.Init(IndexData);
+		SpatialIndex.Init();
 	}
 
 	TArray<TPair<FBox, ValueType>> Elements;
@@ -122,95 +152,120 @@ private:
 
 namespace FStaticSpatialIndex
 {
-	class FImpl
+	template <typename Profile>
+	class TImpl
 	{
 	public:
-		FImpl(const IStaticSpatialIndexDataInterface& InDataInterface)
+		TImpl(const TStaticSpatialIndexDataInterface<Profile>& InDataInterface)
 			: DataInterface(InDataInterface)
 		{}
 
 	protected:
-		const IStaticSpatialIndexDataInterface& DataInterface;
+		const TStaticSpatialIndexDataInterface<Profile>& DataInterface;
 	};
 
-	class FListImpl : public FImpl
+	template <typename Profile>
+	class TListImpl : public TImpl<Profile>
 	{
+		using FVector = typename Profile::FVector;
+		using FBox = typename Profile::FBox;
+
 	public:
-		FListImpl(const IStaticSpatialIndexDataInterface& InDataInterface)
-			: FImpl(InDataInterface)
+		TListImpl(const TStaticSpatialIndexDataInterface<Profile>& InDataInterface)
+			: TImpl<Profile>(InDataInterface)
 		{}
 
-		void Init(const TArray<TPair<FBox, uint32>>& InElements);
-
-		bool ForEachElement(TFunctionRef<bool(uint32 InValueIndex)> InFunc) const;
-		bool ForEachIntersectingElement(const FBox& InBox, TFunctionRef<bool(uint32 InValueIndex)> InFunc) const;
-		bool ForEachIntersectingElement(const FSphere& InSphere, TFunctionRef<bool(uint32 InValueIndex)> InFunc) const;
+		void Init() {}
 
-		uint32 GetAllocatedSize() const;
-
-	private:
-		TArray<uint32> Elements;
-	};
+		bool ForEachElement(TFunctionRef<bool(uint32 InValueIndex)> InFunc) const
+		{
+			for (int32 ValueIndex = 0; ValueIndex < this->DataInterface.GetNumBox(); ValueIndex++)
+			{
+				if (!InFunc(ValueIndex))
+				{
+					return false;
+				}
+			}
+			return true;
+		}
 
-	class FRTreeBaseImpl : public FImpl
-	{
-	public:
-		FRTreeBaseImpl(const IStaticSpatialIndexDataInterface& InDataInterface)
-			: FImpl(InDataInterface)
-		{}
+		bool ForEachIntersectingElement(const FBox& InBox, TFunctionRef<bool(uint32 InValueIndex)> InFunc) const
+		{
+			for (int32 ValueIndex = 0; ValueIndex < this->DataInterface.GetNumBox(); ValueIndex++)
+			{
+				const FBox& Box = this->DataInterface.GetBox(ValueIndex);
 
-		bool ForEachElement(TFunctionRef<bool(uint32 InValueIndex)> InFunc) const;
-		bool ForEachIntersectingElement(const FBox& InBox, TFunctionRef<bool(uint32 InValueIndex)> InFunc) const;
-		bool ForEachIntersectingElement(const FSphere& InSphere, TFunctionRef<bool(uint32 InValueIndex)> InFunc) const;
+				if (Box.Intersect(InBox))
+				{
+					if (!InFunc(ValueIndex))
+					{
+						return false;
+					}
+				}
+			}
+			return true;
 
-		uint32 GetAllocatedSize() const;
+		}
 
-	protected:
-		struct FNode
+		bool ForEachIntersectingElement(const FSphere& InSphere, TFunctionRef<bool(uint32 InValueIndex)> InFunc) const
 		{
-			using FNodeType = TArray<FNode>;
-			using FLeafType = TArray<uint32>;
+			const FSphere::FReal RadiusSquared = FMath::Square(InSphere.W);
 
-			FBox Box = FBox(ForceInit);
-			TVariant<FNodeType, FLeafType> Content;
-		};
+			for (int32 ValueIndex = 0; ValueIndex < this->DataInterface.GetNumBox(); ValueIndex++)
+			{
+				const FBox& Box = this->DataInterface.GetBox(ValueIndex);
 
-		bool ForEachElementRecursive(const FNode* InNode, TFunctionRef<bool(uint32 InValueIndex)> InFunc) const;
-		bool ForEachIntersectingElementRecursive(const FNode* InNode, const FBox& InBox, TFunctionRef<bool(uint32 InValueIndex)> InFunc) const;
-		bool ForEachIntersectingElementRecursive(const FNode* InNode, const FVector& InSphereCenter, FSphere::FReal InRadiusSquared, TFunctionRef<bool(uint32 InValueIndex)> InFunc) const;
+				if (FastSphereAABBIntersection<Profile>(FVector(InSphere.Center), RadiusSquared, Box))
+				{
+					if (!InFunc(ValueIndex))
+					{
+						return false;
+					}
+				}
+			}
+			return true;
+		}
 
-		FNode RootNode;
+		uint32 GetAllocatedSize() const
+		{
+			return sizeof(*this);
+		}
 	};
 
-	template <int32 MaxNumElementsPerNode = 64, int32 MaxNumElementsPerLeaf = 64>
-	class FRTreeImpl : public FRTreeBaseImpl
+	template <typename Profile, int32 MaxNumElementsPerNode = 16, int32 MaxNumElementsPerLeaf = 64>
+	class TRTreeImpl : public TImpl<Profile>
 	{
+		using FVector = typename Profile::FVector;
+		using FBox = typename Profile::FBox;
+
 	public:
-		FRTreeImpl(const IStaticSpatialIndexDataInterface& InDataInterface)
-			: FRTreeBaseImpl(InDataInterface)
+		TRTreeImpl(const TStaticSpatialIndexDataInterface<Profile>& InDataInterface)
+			: TImpl<Profile>(InDataInterface)
 		{}
 
-		void Init(const TArray<TPair<FBox, uint32>>& InElements)
+		void Init()
 		{
-			if (InElements.Num())
+			if (this->DataInterface.GetNumBox())
 			{
-				// build leaves
+				// Build leaves
 				FNode* CurrentNode = nullptr;
 				TArray<FNode> Nodes;
 
-				for (const TPair<FBox, uint32>& Element : InElements)
+				for (int32 ElementIndex = 0; ElementIndex < this->DataInterface.GetNumBox(); ElementIndex++)
 				{
-					if (!CurrentNode || (CurrentNode->Content.Get<FNode::FLeafType>().Num() >= MaxNumElementsPerLeaf))
+					const FBox& Element = this->DataInterface.GetBox(ElementIndex);
+
+					if (!CurrentNode || (CurrentNode->Content.template Get<typename FNode::FLeafType>().Num() >= MaxNumElementsPerLeaf))
 					{
 						CurrentNode = &Nodes.Emplace_GetRef();
-						CurrentNode->Content.Emplace<FNode::FLeafType>();
+						CurrentNode->Content.template Emplace<typename FNode::FLeafType>();
 					}
 
-					CurrentNode->Box += Element.Key;
-					CurrentNode->Content.Get<FNode::FLeafType>().Add(Element.Value);
+					CurrentNode->Box += Element;
+					CurrentNode->Content.template Get<typename FNode::FLeafType>().Add(ElementIndex);
 				}
 
-				// build nodes
+				// Build nodes
 				while (Nodes.Num() > 1)
 				{
 					CurrentNode = nullptr;
@@ -218,14 +273,14 @@ namespace FStaticSpatialIndex
 
 					for (FNode& Node : Nodes)
 					{
-						if (!CurrentNode || (CurrentNode->Content.Get<FNode::FNodeType>().Num() >= MaxNumElementsPerNode))
+						if (!CurrentNode || (CurrentNode->Content.template Get<typename FNode::FNodeType>().Num() >= MaxNumElementsPerNode))
 						{
 							CurrentNode = &TopNodes.Emplace_GetRef();
-							CurrentNode->Content.Emplace<FNode::FNodeType>();
+							CurrentNode->Content.template Emplace<typename FNode::FNodeType>();
 						}
 
 						CurrentNode->Box += Node.Box;
-						CurrentNode->Content.Get<FNode::FNodeType>().Add(MoveTemp(Node));
+						CurrentNode->Content.template Get<typename FNode::FNodeType>().Add(MoveTemp(Node));
 					}
 
 					Nodes = MoveTemp(TopNodes);
@@ -235,70 +290,249 @@ namespace FStaticSpatialIndex
 				RootNode = MoveTemp(Nodes[0]);
 			}
 		}
+
+		bool ForEachElement(TFunctionRef<bool(uint32 InValueIndex)> InFunc) const
+		{
+			return ForEachElementRecursive(&RootNode, InFunc);
+		}
+
+		bool ForEachIntersectingElement(const FBox& InBox, TFunctionRef<bool(uint32 InValueIndex)> InFunc) const
+		{
+			return ForEachIntersectingElementRecursive(&RootNode, InBox, InFunc);
+		}
+
+		bool ForEachIntersectingElement(const FSphere& InSphere, TFunctionRef<bool(uint32 InValueIndex)> InFunc) const
+		{
+			const FSphere::FReal RadiusSquared = FMath::Square(InSphere.W);
+			return ForEachIntersectingElementRecursive(&RootNode, FVector(InSphere.Center), RadiusSquared, InFunc);
+		}
+
+		uint32 GetAllocatedSize() const
+		{
+			TFunction<uint32(const FNode*, uint32)> GetAllocatedSizeRecursive = [this, &GetAllocatedSizeRecursive](const FNode* Node, uint32 BaseSize)
+			{
+				uint32 AllocatedSize = BaseSize;
+
+				if (Node->Content.template IsType<typename FNode::FNodeType>())
+				{
+					AllocatedSize += Node->Content.template Get<typename FNode::FNodeType>().GetAllocatedSize();
+
+					for (auto& ChildNode : Node->Content.template Get<typename FNode::FNodeType>())
+					{
+						AllocatedSize += GetAllocatedSizeRecursive(&ChildNode, sizeof(FNode));
+					}
+				}
+				else
+				{
+					AllocatedSize += Node->Content.template Get<typename FNode::FLeafType>().GetAllocatedSize();
+				}
+
+				return AllocatedSize;
+			};
+	
+			return GetAllocatedSizeRecursive(&RootNode, sizeof(*this));
+		}
+
+	protected:
+		struct FNode
+		{
+			using FNodeType = TArray<FNode>;
+			using FLeafType = TArray<uint32>;
+			using FBoxType = typename Profile::FBox;
+
+			FBoxType Box = FBoxType(ForceInit);
+			TVariant<FNodeType, FLeafType> Content;
+		};
+
+		bool ForEachElementRecursive(const FNode* InNode, TFunctionRef<bool(uint32 InValueIndex)> InFunc) const
+		{
+			if (InNode->Content.template IsType<typename FNode::FNodeType>())
+			{
+				for (auto& ChildNode : InNode->Content.template Get<typename FNode::FNodeType>())
+				{
+					if (!ForEachElementRecursive(&ChildNode, InFunc))
+					{
+						return false;
+					}
+				}
+			}
+			else
+			{
+				for (uint32 ValueIndex : InNode->Content.template Get<typename FNode::FLeafType>())
+				{
+					if (!InFunc(ValueIndex))
+					{
+						return false;
+					}
+				}
+			}
+			return true;
+
+		}
+
+		bool ForEachIntersectingElementRecursive(const FNode* InNode, const FBox& InBox, TFunctionRef<bool(uint32 InValueIndex)> InFunc) const
+		{
+			if (InNode->Content.template IsType<typename FNode::FNodeType>())
+			{
+				for (auto& ChildNode : InNode->Content.template Get<typename FNode::FNodeType>())
+				{
+					if (ChildNode.Box.Intersect(InBox))
+					{
+						if (!ForEachIntersectingElementRecursive(&ChildNode, InBox, InFunc))
+						{
+							return false;
+						}
+					}
+				}
+			}
+			else
+			{
+				for (uint32 ValueIndex : InNode->Content.template Get<typename FNode::FLeafType>())
+				{
+					const FBox& Box = this->DataInterface.GetBox(ValueIndex);
+
+					if (Box.Intersect(InBox))
+					{
+						if (!InFunc(ValueIndex))
+						{
+							return false;
+						}
+					}
+				}
+			}
+			return true;
+		}
+
+		bool ForEachIntersectingElementRecursive(const FNode* InNode, const FVector& InSphereCenter, FSphere::FReal InRadiusSquared, TFunctionRef<bool(uint32 InValueIndex)> InFunc) const
+		{
+			if (InNode->Content.template IsType<typename FNode::FNodeType>())
+			{
+				for (auto& ChildNode : InNode->Content.template Get<typename FNode::FNodeType>())
+				{
+					if (FastSphereAABBIntersection<Profile>(FVector(InSphereCenter), InRadiusSquared, ChildNode.Box))
+					{
+						if (!ForEachIntersectingElementRecursive(&ChildNode, InSphereCenter, InRadiusSquared, InFunc))
+						{
+							return false;
+						}
+					}
+				}
+			}
+			else
+			{
+				for (uint32 ValueIndex : InNode->Content.template Get<typename FNode::FLeafType>())
+				{
+					const FBox& Box = this->DataInterface.GetBox(ValueIndex);
+
+					if (FastSphereAABBIntersection<Profile>(FVector(InSphereCenter), InRadiusSquared, Box))
+					{
+						if (!InFunc(ValueIndex))
+						{
+							return false;
+						}
+					}
+				}
+			}
+			return true;
+		}
+
+		FNode RootNode;
 	};
 
-	class FNodeSorterNoSort
+	template <typename Profile>
+	class TNodeSorterNoSort
 	{
 	public:
 		enum { NeedSort = 0 };
+		using FBox = typename Profile::FBox;
 		void Init(const FBox& SortBox) {}
 		bool Sort(const FBox& A, const FBox& B) { return false; }
 	};
 
-	class FNodeSorterMinX
+	template <typename Profile>
+	class TNodeSorterMinX
 	{
 	public:
 		enum { NeedSort = 1 };
-
+		using FBox = typename Profile::FBox;
 		void Init(const FBox& SortBox) {}
 		bool Sort(const FBox& A, const FBox& B) { return A.Min.X < B.Min.X; }
 	};
 
-	template <uint32 BucketSize>
+	template <typename Profile, int32 BucketSize>
 	class TNodeSorterMorton
 	{
 	public:
 		enum { NeedSort = 1 };
+		using FBox = typename Profile::FBox;
+		using FReal = typename Profile::FReal;
+		using FIntPoint = typename Profile::FIntPoint;
 
 		void Init(const FBox& SortBox)
 		{}
 		
 		bool Sort(const FBox& A, const FBox& B)
 		{
-			const uint32 MortonCodeA = MortonEncode({ int32(A.GetCenter().X / (FBox::FReal)BucketSize), int32(A.GetCenter().Y / (FBox::FReal)BucketSize), int32(A.GetCenter().Z / (FBox::FReal)BucketSize) });
-			const uint32 MortonCodeB = MortonEncode({ int32(B.GetCenter().X / (FBox::FReal)BucketSize), int32(B.GetCenter().Y / (FBox::FReal)BucketSize), int32(B.GetCenter().Z / (FBox::FReal)BucketSize) });
+			FIntPoint CellPosA;
+			CellPosA.X = int32(A.GetCenter().X / (FReal)BucketSize);
+			CellPosA.Y = int32(A.GetCenter().Y / (FReal)BucketSize);
+			if constexpr (Profile::Is3D)
+			{
+				CellPosA.Z = int32(A.GetCenter().Z / (FReal)BucketSize);
+			}
+			const uint32 MortonCodeA = MortonEncode(CellPosA);
+
+			FIntPoint CellPosB;
+			CellPosB.X = int32(B.GetCenter().X / (FReal)BucketSize);
+			CellPosB.Y = int32(B.GetCenter().Y / (FReal)BucketSize);
+			if constexpr (Profile::Is3D)
+			{
+				CellPosB.Z = int32(A.GetCenter().Z / (FReal)BucketSize);
+			}
+			const uint32 MortonCodeB = MortonEncode(CellPosB);
+
 			return MortonCodeA < MortonCodeB;
 		}
 
 	private:
-		inline uint32 MortonEncode(const FInt32Vector3& Point)
+		inline uint32 MortonEncode(const FIntPoint& Point)
 		{
-			return FMath::MortonCode3(Point.X) | (FMath::MortonCode3(Point.Y) << 1) | (FMath::MortonCode3(Point.Z) << 2);
+			if constexpr (Profile::Is3D)
+			{
+				return FMath::MortonCode3(Point.X) | (FMath::MortonCode3(Point.Y) << 1) | (FMath::MortonCode3(Point.Z) << 2);
+			}
+			else
+			{
+				return FMath::MortonCode2(Point.X) | (FMath::MortonCode2(Point.Y) << 1);
+			}
 		}
 	};
 
-	template <uint32 BucketSize>
+	template <typename Profile, int32 BucketSize>
 	class TNodeSorterHilbert
 	{
 	public:
 		enum { NeedSort = 1 };
+		using FBox = typename Profile::FBox;
+		using FReal = typename Profile::FReal;
+		using FIntPoint = typename Profile::FIntPoint;
 
 		void Init(const FBox& SortBox)
 		{
-			const FBox::FReal MaxExtent = SortBox.GetExtent().GetMax();
-			const uint32 NumBuckets = FMath::CeilToInt32(MaxExtent / (FBox::FReal)BucketSize);			
+			const FReal MaxExtent = SortBox.GetExtent().GetMax();
+			const uint32 NumBuckets = FMath::CeilToInt32(MaxExtent / (FReal)BucketSize);			
 			HilbertOrder = 1 + FMath::FloorLog2(NumBuckets);
 		}
 		
 		bool Sort(const FBox& A, const FBox& B)
 		{
-			const int32 HilbertCodeA = HilbertEncode({ int32(A.GetCenter().X / (FBox::FReal)BucketSize), int32(A.GetCenter().Y / (FBox::FReal)BucketSize) }, HilbertOrder);
-			const int32 HilbertCodeB = HilbertEncode({ int32(B.GetCenter().X / (FBox::FReal)BucketSize), int32(B.GetCenter().Y / (FBox::FReal)BucketSize) }, HilbertOrder);
+			const int32 HilbertCodeA = HilbertEncode({ int32(A.GetCenter().X / (FReal)BucketSize), int32(A.GetCenter().Y / (FReal)BucketSize) }, HilbertOrder);
+			const int32 HilbertCodeB = HilbertEncode({ int32(B.GetCenter().X / (FReal)BucketSize), int32(B.GetCenter().Y / (FReal)BucketSize) }, HilbertOrder);
 			return HilbertCodeA < HilbertCodeB;
 		}
 
 	private:
-		uint32 HilbertEncode(FIntPoint Point, uint32 Order)
+		uint32 HilbertEncode(const FIntVector2 Point, uint32 Order)
 		{
 			uint32 Result = 0;
 
@@ -317,5 +551,10 @@ namespace FStaticSpatialIndex
 	};
 }
 
-template <class ValueType, class NodeSorter> class TStaticSpatialIndexList : public TStaticSpatialIndex<ValueType, FStaticSpatialIndex::FListImpl, NodeSorter> {};
-template <class ValueType, class NodeSorter> class TStaticSpatialIndexRTree : public TStaticSpatialIndex<ValueType, FStaticSpatialIndex::FRTreeImpl<16, 64>, NodeSorter> {};
\ No newline at end of file
+template <typename ValueType, typename NodeSorter, typename Profile>
+class TStaticSpatialIndexList : public TStaticSpatialIndex<ValueType, Profile, FStaticSpatialIndex::TListImpl<Profile>, NodeSorter>
+{};
+
+template <typename ValueType, typename NodeSorter, typename Profile>
+class TStaticSpatialIndexRTree : public TStaticSpatialIndex<ValueType, Profile, FStaticSpatialIndex::TRTreeImpl<Profile>, NodeSorter>
+{};
\ No newline at end of file
diff --git a/Engine/Source/Runtime/Engine/Public/WorldPartition/RuntimeHashSet/WorldPartitionRuntimeCellDataHashSet.h b/Engine/Source/Runtime/Engine/Public/WorldPartition/RuntimeHashSet/WorldPartitionRuntimeCellDataHashSet.h
new file mode 100644
--- /dev/null
+++ b/Engine/Source/Runtime/Engine/Public/WorldPartition/RuntimeHashSet/WorldPartitionRuntimeCellDataHashSet.h
@@ -0,0 +1,16 @@
+// Copyright Epic Games, Inc. All Rights Reserved.
+#pragma once
+
+#include "CoreMinimal.h"
+#include "WorldPartition/WorldPartitionRuntimeCellData.h"
+#include "WorldPartitionRuntimeCellDataHashSet.generated.h"
+
+UCLASS(Within = WorldPartitionRuntimeCell, MinimalAPI)
+class UWorldPartitionRuntimeCellDataHashSet : public UWorldPartitionRuntimeCellData
+{
+	GENERATED_BODY()
+
+public:
+	UPROPERTY()
+	bool bIs2D;
+};
diff --git a/Engine/Source/Runtime/Engine/Public/WorldPartition/RuntimeHashSet/WorldPartitionRuntimeHashSet.h b/Engine/Source/Runtime/Engine/Public/WorldPartition/RuntimeHashSet/WorldPartitionRuntimeHashSet.h
--- a/Engine/Source/Runtime/Engine/Public/WorldPartition/RuntimeHashSet/WorldPartitionRuntimeHashSet.h
+++ b/Engine/Source/Runtime/Engine/Public/WorldPartition/RuntimeHashSet/WorldPartitionRuntimeHashSet.h
@@ -11,8 +11,11 @@ class UHLODLayer;
 class URuntimePartitionPersistent;
 struct FPropertyChangedChainEvent;
 
-using FStaticSpatialIndexSorter = FStaticSpatialIndex::TNodeSorterMorton<65536>;
-using FStaticSpatialIndexType = TStaticSpatialIndexRTree<TObjectPtr<UWorldPartitionRuntimeCell>, FStaticSpatialIndexSorter>;
+using FStaticSpatialIndexSorter = FStaticSpatialIndex::TNodeSorterMorton<FStaticSpatialIndex::FSpatialIndexProfile3D, 65536>;
+using FStaticSpatialIndexType = TStaticSpatialIndexRTree<TObjectPtr<UWorldPartitionRuntimeCell>, FStaticSpatialIndexSorter, FStaticSpatialIndex::FSpatialIndexProfile3D>;
+
+using FStaticSpatialIndexSorter2D = FStaticSpatialIndex::TNodeSorterMorton<FStaticSpatialIndex::FSpatialIndexProfile2D, 65536>;
+using FStaticSpatialIndexType2D = TStaticSpatialIndexRTree<TObjectPtr<UWorldPartitionRuntimeCell>, FStaticSpatialIndexSorter2D, FStaticSpatialIndex::FSpatialIndexProfile2D>;
 
 /** Holds an HLOD setup for a particular partition class. */
 USTRUCT()
@@ -93,7 +96,8 @@ protected:
 
 	// Transient
 	mutable TUniquePtr<FStaticSpatialIndexType> SpatialIndex;
-	mutable TUniquePtr<FStaticSpatialIndexType> SpatialIndex2D;
+	mutable TUniquePtr<FStaticSpatialIndexType2D> SpatialIndexForce2D;
+	mutable TUniquePtr<FStaticSpatialIndexType2D> SpatialIndex2D;
 };
 
 template<>
