Remove gc.SizeOfPermanentObjectPool reference from UGarbageCollectionSettings as the cvar was removed.

diff --git a/Engine/Source/Runtime/Engine/Classes/Engine/CoreSettings.h b/Engine/Source/Runtime/Engine/Classes/Engine/CoreSettings.h
--- a/Engine/Source/Runtime/Engine/Classes/Engine/CoreSettings.h
+++ b/Engine/Source/Runtime/Engine/Classes/Engine/CoreSettings.h
@@ -244,11 +244,6 @@ protected:
 		ToolTip = "Maximum Object Count Not Considered By GC. Works only in cooked builds."))
 	int32 MaxObjectsNotConsideredByGC;
 
-	UPROPERTY(EditAnywhere, config, Category = Optimization, meta = (
-		ConsoleVariable = "gc.SizeOfPermanentObjectPool", DisplayName = "Size Of Permanent Object Pool",
-		ToolTip = "Size Of Permanent Object Pool (bytes). Works only in cooked builds."))
-	int32 SizeOfPermanentObjectPool;
-
 	UPROPERTY(EditAnywhere, config, Category = Optimization, meta = (
 		ConsoleVariable = "gc.MaxObjectsInGame", DisplayName = "Maximum number of UObjects that can exist in cooked game",
 		ToolTip = "Maximum number of UObjects that can exist in cooked game. Keep this as small as possible."))
diff --git a/Engine/Source/Runtime/Engine/Private/CoreSettings.cpp b/Engine/Source/Runtime/Engine/Private/CoreSettings.cpp
--- a/Engine/Source/Runtime/Engine/Private/CoreSettings.cpp
+++ b/Engine/Source/Runtime/Engine/Private/CoreSettings.cpp
@@ -204,7 +204,6 @@ UGarbageCollectionSettings::UGarbageCollectionSettings()
 	MultithreadedDestructionEnabled = true;
 	NumRetriesBeforeForcingGC = 0;
 	MaxObjectsNotConsideredByGC = 0;
-	SizeOfPermanentObjectPool = 0;
 	MaxObjectsInEditor = 12 * 1024 * 1024;
 	MaxObjectsInGame = 2 * 1024 * 1024;	
 	CreateGCClusters = true;	
