disable Object Trace world subsystem in shipping builds

diff --git a/Engine/Source/Runtime/Engine/Public/ObjectTrace.h b/Engine/Source/Runtime/Engine/Public/ObjectTrace.h
--- a/Engine/Source/Runtime/Engine/Public/ObjectTrace.h
+++ b/Engine/Source/Runtime/Engine/Public/ObjectTrace.h
@@ -25,6 +25,13 @@ class UObjectTraceWorldSubsystem : public UWorldSubsystem
 		: FrameIndex(0), RecordingIndex(0), ElapsedTime(0.0) 
 	{}
 
+#if !OBJECT_TRACE_ENABLED
+	virtual bool ShouldCreateSubsystem(UObject* Outer) const override
+	{
+		return false;
+	}
+#endif
+
 public:
 	// The frame index incremented each tick
 	uint16 FrameIndex;
