[World Partition Runtime HashSet] Reflect the partitions Is2D flag on HLOD partition layers so they always match the main partition layer value.

diff --git a/Engine/Source/Runtime/Engine/Private/WorldPartition/RuntimeHashSet/RuntimePartitionLHGrid.cpp b/Engine/Source/Runtime/Engine/Private/WorldPartition/RuntimeHashSet/RuntimePartitionLHGrid.cpp
--- a/Engine/Source/Runtime/Engine/Private/WorldPartition/RuntimeHashSet/RuntimePartitionLHGrid.cpp
+++ b/Engine/Source/Runtime/Engine/Private/WorldPartition/RuntimeHashSet/RuntimePartitionLHGrid.cpp
@@ -73,6 +73,20 @@ struct FCellCoord
 
 FCellCoord FCellCoord::Invalid(0, 0, 0, -1);
 
+bool URuntimePartitionLHGrid::CanEditChange(const FProperty* InProperty) const
+{
+	if (InProperty)
+	{
+		const FString PropertyName = InProperty->GetName();
+		if (PropertyName == GET_MEMBER_NAME_STRING_CHECKED(URuntimePartitionLHGrid, bIs2D))
+		{
+			return HLODIndex == INDEX_NONE;
+		}
+	}
+
+	return Super::CanEditChange(InProperty);
+}
+
 static bool GPackageWasDirty = false;
 void URuntimePartitionLHGrid::PreEditChange(FProperty* InPropertyAboutToChange)
 {
@@ -124,6 +138,14 @@ void URuntimePartitionLHGrid::InitHLODRuntimePartitionFrom(const URuntimePartiti
 	Super::InitHLODRuntimePartitionFrom(InRuntimePartition, InHLODIndex);
 	const URuntimePartitionLHGrid* RuntimePartitionLHGrid = CastChecked<const URuntimePartitionLHGrid>(InRuntimePartition);
 	CellSize = RuntimePartitionLHGrid->CellSize * 2;
+	bIs2D = RuntimePartitionLHGrid->bIs2D;
+}
+
+void URuntimePartitionLHGrid::UpdateHLODRuntimePartitionFrom(const URuntimePartition* InRuntimePartition)
+{
+	Super::UpdateHLODRuntimePartitionFrom(InRuntimePartition);
+	const URuntimePartitionLHGrid* RuntimePartitionLHGrid = CastChecked<const URuntimePartitionLHGrid>(InRuntimePartition);
+	bIs2D = RuntimePartitionLHGrid->bIs2D;
 }
 
 void URuntimePartitionLHGrid::SetDefaultValues()
diff --git a/Engine/Source/Runtime/Engine/Private/WorldPartition/RuntimeHashSet/WorldPartitionRuntimeHashSet.cpp b/Engine/Source/Runtime/Engine/Private/WorldPartition/RuntimeHashSet/WorldPartitionRuntimeHashSet.cpp
--- a/Engine/Source/Runtime/Engine/Private/WorldPartition/RuntimeHashSet/WorldPartitionRuntimeHashSet.cpp
+++ b/Engine/Source/Runtime/Engine/Private/WorldPartition/RuntimeHashSet/WorldPartitionRuntimeHashSet.cpp
@@ -971,6 +971,20 @@ void UWorldPartitionRuntimeHashSet::PostEditChangeChainProperty(FPropertyChanged
 			}
 		}
 	}
+	else if (PropertyName == GET_MEMBER_NAME_CHECKED(FRuntimePartitionDesc, MainLayer))
+	{
+		int32 RuntimePartitionIndex = PropertyChangedEvent.GetArrayIndex(NAME_RuntimePartitions.ToString());
+		if (RuntimePartitions.IsValidIndex(RuntimePartitionIndex))
+		{
+			for (const FRuntimePartitionHLODSetup& HLODSetup : RuntimePartitions[RuntimePartitionIndex].HLODSetups)
+			{
+				if (HLODSetup.PartitionLayer)
+				{
+					HLODSetup.PartitionLayer->UpdateHLODRuntimePartitionFrom(RuntimePartitions[RuntimePartitionIndex].MainLayer);
+				}
+			}
+		}
+	}
 }
 
 void UWorldPartitionRuntimeHashSet::PostDuplicate(EDuplicateMode::Type DuplicateMode)
diff --git a/Engine/Source/Runtime/Engine/Public/WorldPartition/RuntimeHashSet/RuntimePartition.h b/Engine/Source/Runtime/Engine/Public/WorldPartition/RuntimeHashSet/RuntimePartition.h
--- a/Engine/Source/Runtime/Engine/Public/WorldPartition/RuntimeHashSet/RuntimePartition.h
+++ b/Engine/Source/Runtime/Engine/Public/WorldPartition/RuntimeHashSet/RuntimePartition.h
@@ -73,6 +73,7 @@ public:
 	virtual void SetDefaultValues();
 	virtual bool SupportsHLODs() const PURE_VIRTUAL(URuntimePartition::SupportsHLODs, return true;);
 	virtual void InitHLODRuntimePartitionFrom(const URuntimePartition* InRuntimePartition, int32 InHLODIndex);
+	virtual void UpdateHLODRuntimePartitionFrom(const URuntimePartition* InRuntimePartition) {}
 #endif
 	virtual bool IsValidPartitionTokens(const TArray<FName>& InPartitionTokens) const PURE_VIRTUAL(URuntimePartition::IsValidPartitionTokens, return false;);
 #if WITH_EDITOR
diff --git a/Engine/Source/Runtime/Engine/Public/WorldPartition/RuntimeHashSet/RuntimePartitionLHGrid.h b/Engine/Source/Runtime/Engine/Public/WorldPartition/RuntimeHashSet/RuntimePartitionLHGrid.h
--- a/Engine/Source/Runtime/Engine/Public/WorldPartition/RuntimeHashSet/RuntimePartitionLHGrid.h
+++ b/Engine/Source/Runtime/Engine/Public/WorldPartition/RuntimeHashSet/RuntimePartitionLHGrid.h
@@ -17,13 +17,15 @@ class ENGINE_API URuntimePartitionLHGrid : public URuntimePartition
 public:
 #if WITH_EDITOR
 	//~ Begin UObject Interface.
+	virtual bool CanEditChange(const FProperty* InProperty) const override;
 	virtual void PreEditChange(FProperty* InPropertyAboutToChange) override;
 	virtual void PostEditChangeProperty(FPropertyChangedEvent& InPropertyChangedEvent) override;
 	//~ End UObject Interface.
 
 	//~ Begin URuntimePartition interface
 	virtual bool SupportsHLODs() const override { return true; }
-	virtual void InitHLODRuntimePartitionFrom(const URuntimePartition* InRuntimePartition, int32 InHLODIndex);
+	virtual void InitHLODRuntimePartitionFrom(const URuntimePartition* InRuntimePartition, int32 InHLODIndex) override;
+	virtual void UpdateHLODRuntimePartitionFrom(const URuntimePartition* InRuntimePartition) override;
 	virtual void SetDefaultValues() override;
 #endif
 	virtual bool IsValidPartitionTokens(const TArray<FName>& InPartitionTokens) const override;
