Fix typo from 33604784.

diff --git a/Engine/Source/Runtime/Engine/Public/WorldPartition/RuntimeHashSet/StaticSpatialIndex.h b/Engine/Source/Runtime/Engine/Public/WorldPartition/RuntimeHashSet/StaticSpatialIndex.h
--- a/Engine/Source/Runtime/Engine/Public/WorldPartition/RuntimeHashSet/StaticSpatialIndex.h
+++ b/Engine/Source/Runtime/Engine/Public/WorldPartition/RuntimeHashSet/StaticSpatialIndex.h
@@ -487,7 +487,7 @@ namespace FStaticSpatialIndex
 			CellPosB.Y = int32(B.GetCenter().Y / (FReal)BucketSize);
 			if constexpr (Profile::Is3D)
 			{
-				CellPosB.Z = int32(A.GetCenter().Z / (FReal)BucketSize);
+				CellPosB.Z = int32(B.GetCenter().Z / (FReal)BucketSize);
 			}
 			const uint32 MortonCodeB = MortonEncode(CellPosB);
 
