Fix a bug where invalid skeletal mesh LODs can be selected for rendering when mesh streaming is enabled.

#tests A/B tested using ReplayRun and didn't see perf or mem regression

diff --git a/Engine/Source/Runtime/Engine/Private/SkeletalMesh.cpp b/Engine/Source/Runtime/Engine/Private/SkeletalMesh.cpp
--- a/Engine/Source/Runtime/Engine/Private/SkeletalMesh.cpp
+++ b/Engine/Source/Runtime/Engine/Private/SkeletalMesh.cpp
@@ -7171,7 +7171,12 @@ FPrimitiveViewRelevance FSkeletalMeshSceneProxy::GetViewRelevance(const FSceneVi
 	FPrimitiveViewRelevance Result;
 	Result.bDrawRelevance = IsShown(View) && !!EngineShowFlags.SkeletalMeshes;
 	Result.bShadowRelevance = IsShadowCast(View);
-	Result.bStaticRelevance = (bRenderStatic || GSkeletalMeshUseCachedMDCs) && MeshObject->SupportsStaticRelevance() && !IsDynamic();
+	Result.bStaticRelevance = (bRenderStatic || GSkeletalMeshUseCachedMDCs)
+		&& MeshObject->SupportsStaticRelevance()
+		// Switch to dynamic if the mesh object is not ready. GetDynamicMeshElements won't generate any mesh batch in this case.
+		// Consequently, this mesh won't be drawn this frame but render time will be updated which triggers an update to the mesh object. 
+		&& MeshObject->GetLOD() >= GetCurrentFirstLODIdx_Internal()
+		&& !IsDynamic();
 	Result.bDynamicRelevance = ~Result.bStaticRelevance;
 	Result.bRenderCustomDepth = ShouldRenderCustomDepth();
 	Result.bRenderInMainPass = ShouldRenderInMainPass();
