Remove unnecessary runtime virtual texture component assets.
Asserts checked that the component was in an active world. For blueprints this isn't the case. The asserts would be OK (but not very useful) to keep in paths only called when the RVT is instantiated. But IsStreamingTexture() is called from details panel UI updates.

diff --git a/Engine/Source/Runtime/Engine/Private/Components/RuntimeVirtualTextureComponent.cpp b/Engine/Source/Runtime/Engine/Private/Components/RuntimeVirtualTextureComponent.cpp
--- a/Engine/Source/Runtime/Engine/Private/Components/RuntimeVirtualTextureComponent.cpp
+++ b/Engine/Source/Runtime/Engine/Private/Components/RuntimeVirtualTextureComponent.cpp
@@ -285,8 +285,6 @@ uint64 URuntimeVirtualTextureComponent::CalculateStreamingTextureSettingsHash()
 
 bool URuntimeVirtualTextureComponent::IsStreamingLowMips(EShadingPath ShadingPath) const
 {
-	checkf(IsActiveInWorld(), TEXT("This function should never be called for a world where we're inactive"));
-	
 #if WITH_EDITOR
 	if (GIsEditor)
 	{
@@ -314,8 +312,6 @@ bool IsCompatibleFormat(URuntimeVirtualTexture const& RuntimeVirtualTexture, UVi
 
 bool URuntimeVirtualTextureComponent::IsStreamingTextureInvalid(EShadingPath ShadingPath) const
 {
-	checkf(IsActiveInWorld(), TEXT("This function should never be called for a world where we're inactive"));
-
 	return 
 		VirtualTexture != nullptr && 
 		StreamingTexture != nullptr && 
