Integrate github pull request adding more comprehensive custom gravity support for CharacterMovementComponent
#rnx
PR #11829: Comprehensive support for custom gravity in CharacterMovementComponent.

diff --git a/Engine/Source/Runtime/Engine/Classes/GameFramework/CharacterMovementComponent.h b/Engine/Source/Runtime/Engine/Classes/GameFramework/CharacterMovementComponent.h
--- a/Engine/Source/Runtime/Engine/Classes/GameFramework/CharacterMovementComponent.h
+++ b/Engine/Source/Runtime/Engine/Classes/GameFramework/CharacterMovementComponent.h
@@ -1303,6 +1303,7 @@ public:
 	bool IsMovementInProgress() const { return bMovementInProgress; }
 
 	//BEGIN UNavMovementComponent Interface
+	ENGINE_API virtual FVector GetActorFeetLocation() const override;
 	ENGINE_API virtual void RequestDirectMove(const FVector& MoveVelocity, bool bForceMaxSpeed) override;
 	ENGINE_API virtual void RequestPathMove(const FVector& MoveInput) override;
 	ENGINE_API virtual bool CanStartPathFollowing() const override;
@@ -1661,6 +1662,18 @@ public:
 	/** Rotate a vector from gravity to world space. */
 	FVector RotateWorldToGravity(const FVector& Gravity) const { return GravityToWorldTransform.RotateVector(Gravity); }
 
+	/** Project a vector onto the floor defined by the gravity direction. */
+	FVector ProjectToGravityFloor(const FVector& Vector) const { return FVector::VectorPlaneProject(Vector, GetGravityDirection()); }
+
+	/** Returns the size of a vector in the gravity-space vertical direction. */
+	FVector::FReal GetGravitySpaceZ(const FVector& Vector) const { return Vector.Dot(-GetGravityDirection()); }
+
+	/** Returns the component of the vector in the gravity-space vertical direction.  */
+	FVector GetGravitySpaceComponentZ(const FVector& Vector) const { return Vector.Dot(GetGravityDirection()) * GetGravityDirection(); }
+
+	/** Set the vertical component of the vector to the given value in the gravity-space vertical direction. */
+	void SetGravitySpaceZ(FVector& Vector, const FVector::FReal Z) const { Vector = ProjectToGravityFloor(Vector) - Z * GetGravityDirection(); }
+	
 protected:
 
 	/**
diff --git a/Engine/Source/Runtime/Engine/Private/Components/CharacterMovementComponent.cpp b/Engine/Source/Runtime/Engine/Private/Components/CharacterMovementComponent.cpp
--- a/Engine/Source/Runtime/Engine/Private/Components/CharacterMovementComponent.cpp
+++ b/Engine/Source/Runtime/Engine/Private/Components/CharacterMovementComponent.cpp
@@ -1057,13 +1057,11 @@ bool UCharacterMovementComponent::DoJump(bool bReplayingMoves)
 	if ( CharacterOwner && CharacterOwner->CanJump() )
 	{
 		// Don't jump if we can't move up/down.
-		if (!bConstrainToPlane || FMath::Abs(PlaneConstraintNormal.Z) != 1.f)
+		if (!bConstrainToPlane || FMath::IsNearlyEqual(FMath::Abs(GetGravitySpaceZ(PlaneConstraintNormal)), 1.f))
 		{
 			if (HasCustomGravity())
 			{
-				FVector GravityRelativeVelocity = RotateWorldToGravity(Velocity);
-				GravityRelativeVelocity.Z = FMath::Max<FVector::FReal>(GravityRelativeVelocity.Z, JumpZVelocity);
-				Velocity = RotateGravityToWorld(GravityRelativeVelocity);
+				SetGravitySpaceZ(Velocity, FMath::Max<FVector::FReal>(GetGravitySpaceZ(Velocity), JumpZVelocity));
 			}
 			else
 			{
@@ -1098,7 +1096,7 @@ FVector UCharacterMovementComponent::GetImpartedMovementBaseVelocity() const
 			
 			if (bImpartBaseAngularVelocity)
 			{
-				const FVector CharacterBasePosition = (UpdatedComponent->GetComponentLocation() - FVector(0.f, 0.f, CharacterOwner->GetCapsuleComponent()->GetScaledCapsuleHalfHeight()));
+				const FVector CharacterBasePosition = (UpdatedComponent->GetComponentLocation() + CharacterOwner->GetCapsuleComponent()->GetScaledCapsuleHalfHeight() * GetGravityDirection());
 				const FVector BaseTangentialVel = MovementBaseUtility::GetMovementBaseTangentialVelocity(MovementBase, CharacterOwner->GetBasedMovement().BoneName, CharacterBasePosition);
 				BaseVelocity += BaseTangentialVel;
 			}
@@ -1152,16 +1150,14 @@ void UCharacterMovementComponent::JumpOff(AActor* MovementBaseActor)
 		{
 			const float MaxSpeed = GetMaxSpeed() * 0.85f;
 			Velocity += MaxSpeed * GetBestDirectionOffActor(MovementBaseActor);
-			if ( Velocity.Size2D() > MaxSpeed )
+			if (ProjectToGravityFloor(Velocity).Size() > MaxSpeed)
 			{
 				Velocity = MaxSpeed * Velocity.GetSafeNormal();
 			}
 
 			if (HasCustomGravity())
 			{
-				FVector GravityRelativeVelocity = RotateWorldToGravity(Velocity);
-				GravityRelativeVelocity.Z = JumpOffJumpZFactor * JumpZVelocity;
-				Velocity = RotateGravityToWorld(GravityRelativeVelocity);
+				SetGravitySpaceZ(Velocity, JumpOffJumpZFactor * JumpZVelocity);
 			}
 			else
 			{
@@ -1179,7 +1175,7 @@ FVector UCharacterMovementComponent::GetBestDirectionOffActor(AActor* BaseActor)
 	// By default, just pick a random direction.  Derived character classes can choose to do more complex calculations,
 	// such as finding the shortest distance to move in based on the BaseActor's Bounding Volume.
 	const float RandAngle = FMath::DegreesToRadians(GetNetworkSafeRandomAngleDegrees());
-	return FVector(FMath::Cos(RandAngle), FMath::Sin(RandAngle), 0.5f).GetSafeNormal();
+	return RotateWorldToGravity(FVector(FMath::Cos(RandAngle), FMath::Sin(RandAngle), 0.5f).GetSafeNormal());
 }
 
 float UCharacterMovementComponent::GetNetworkSafeRandomAngleDegrees() const
@@ -1222,13 +1218,13 @@ void UCharacterMovementComponent::SetDefaultMovementMode()
 	}
 	else if ( !CharacterOwner || MovementMode != DefaultLandMovementMode )
 	{
-		const float SavedVelocityZ = Velocity.Z;
+		const float SavedVelocityZ = GetGravitySpaceZ(Velocity);
 		SetMovementMode(DefaultLandMovementMode);
 
 		// Avoid 1-frame delay if trying to walk but walking fails at this location.
 		if (MovementMode == MOVE_Walking && GetMovementBase() == NULL)
 		{
-			Velocity.Z = SavedVelocityZ; // Prevent temporary walking state from zeroing Z velocity.
+			SetGravitySpaceZ(Velocity, SavedVelocityZ); // Prevent temporary walking state from zeroing Z velocity.
 			SetMovementMode(MOVE_Falling);
 		}
 	}
@@ -1315,7 +1311,7 @@ void UCharacterMovementComponent::OnMovementModeChanged(EMovementMode PreviousMo
 
 		GroundMovementMode = MovementMode;
 		// Walking uses only XY velocity
-		Velocity.Z = 0.f;
+		Velocity = ProjectToGravityFloor(Velocity);
 		SetNavWalkingPhysics(true);
 	}
 	else if (PreviousMovementMode == MOVE_NavWalking)
@@ -1447,36 +1443,35 @@ void UCharacterMovementComponent::ApplyNetworkMovementMode(const uint8 ReceivedM
 void UCharacterMovementComponent::PerformAirControlForPathFollowing(FVector Direction, float ZDiff)
 {
 	// use air control if low grav or above destination and falling towards it
-	if ( CharacterOwner && Velocity.Z < 0.f && (ZDiff < 0.f || GetGravityZ() > 0.9f * GetWorld()->GetDefaultGravityZ()))
+	if ( CharacterOwner && GetGravitySpaceZ(Velocity) < 0.f && (ZDiff < 0.f || GetGravityZ() > 0.9f * GetWorld()->GetDefaultGravityZ()))
 	{
 		if ( ZDiff < 0.f )
 		{
-			if ( (Velocity.X == 0.f) && (Velocity.Y == 0.f) )
+			const FVector LateralVelocity = ProjectToGravityFloor(Velocity);
+			if ( LateralVelocity.SizeSquared() == 0.f )
 			{
 				Acceleration = FVector::ZeroVector;
 			}
 			else
 			{
-				float Dist2D = Direction.Size2D();
+				const FVector::FReal Dist2D = ProjectToGravityFloor(Direction).Size();
 				//Direction.Z = 0.f;
 				Acceleration = Direction.GetSafeNormal() * GetMaxAcceleration();
 
-				if ( (Dist2D < 0.5f * FMath::Abs(Direction.Z)) && ((Velocity | Direction) > 0.5f*FMath::Square(Dist2D)) )
+				if ( (Dist2D < 0.5f * FMath::Abs(GetGravitySpaceZ(Direction))) && ((Velocity | Direction) > 0.5f*FMath::Square(Dist2D)) )
 				{
 					Acceleration *= -1.f;
 				}
 
 				if ( Dist2D < 1.5f*CharacterOwner->GetCapsuleComponent()->GetScaledCapsuleRadius() )
 				{
-					Velocity.X = 0.f;
-					Velocity.Y = 0.f;
+					Velocity = GetGravitySpaceComponentZ(Velocity);
 					Acceleration = FVector::ZeroVector;
 				}
 				else if ( (Velocity | Direction) < 0.f )
 				{
 					float M = FMath::Max(0.f, 0.2f - GetWorld()->DeltaTimeSeconds);
-					Velocity.X *= M;
-					Velocity.Y *= M;
+					Velocity = LateralVelocity * M + GetGravitySpaceComponentZ(Velocity);
 				}
 			}
 		}
@@ -2038,10 +2033,7 @@ FVector UCharacterMovementComponent::ConstrainAnimRootMotionVelocity(const FVect
 	{
 		if (HasCustomGravity())
 		{
-			FVector GravityRelativeResult = RotateWorldToGravity(Result);
-			const FVector GravityRelativeCurrentVelocity = RotateWorldToGravity(CurrentVelocity);
-			GravityRelativeResult.Z = GravityRelativeCurrentVelocity.Z;
-			Result = RotateGravityToWorld(GravityRelativeResult);
+			SetGravitySpaceZ(Result, GetGravitySpaceZ(CurrentVelocity));
 		}
 		else
 		{
@@ -2161,11 +2153,7 @@ void UCharacterMovementComponent::SimulateMovement(float DeltaSeconds)
 			// find floor and check if falling
 			if (IsMovingOnGround() || MovementMode == MOVE_Falling)
 			{
-				bool bShouldFindFloor = Velocity.Z <= 0.f;
-				if (HasCustomGravity())
-				{
-					bShouldFindFloor = RotateWorldToGravity(Velocity).Z <= UE_KINDA_SMALL_NUMBER;
-				}
+				const bool bShouldFindFloor = GetGravitySpaceZ(Velocity) <= UE_KINDA_SMALL_NUMBER;
 
 				if (StepDownResult.bComputedFloor)
 				{
@@ -2186,14 +2174,7 @@ void UCharacterMovementComponent::SimulateMovement(float DeltaSeconds)
 				{
 					// Follows PhysWalking approach for encroachment on floor tests
 					FHitResult Hit(CurrentFloor.HitResult);
-					if (HasCustomGravity())
-					{
-						Hit.TraceEnd = Hit.TraceStart - GetGravityDirection() * MAX_FLOOR_DIST;
-					}
-					else
-					{
-						Hit.TraceEnd = Hit.TraceStart + FVector(0.f, 0.f, MAX_FLOOR_DIST);
-					}
+					Hit.TraceEnd = Hit.TraceStart - GetGravityDirection() * MAX_FLOOR_DIST;
 					
 					const FVector RequestedAdjustment = GetPenetrationAdjustment(Hit);
 					const bool bResolved = ResolvePenetration(RequestedAdjustment, Hit, UpdatedComponent->GetComponentQuat());
@@ -2204,16 +2185,9 @@ void UCharacterMovementComponent::SimulateMovement(float DeltaSeconds)
 					if (!bSimGravityDisabled)
 					{
 						// No floor, must fall.
-						if (HasCustomGravity())
-						{
-							if (RotateWorldToGravity(Velocity).Z <= UE_KINDA_SMALL_NUMBER || bApplyGravityWhileJumping || !CharacterOwner->IsJumpProvidingForce())
-							{
-								Velocity = NewFallVelocity(Velocity, -GetGravityDirection() * GetGravityZ(), DeltaSeconds);
-							}
-						}
-						else if (Velocity.Z <= 0.f || bApplyGravityWhileJumping || !CharacterOwner->IsJumpProvidingForce())
+						if (GetGravitySpaceZ(Velocity) <= UE_KINDA_SMALL_NUMBER || bApplyGravityWhileJumping || !CharacterOwner->IsJumpProvidingForce())
 						{
-							Velocity = NewFallVelocity(Velocity, FVector(0.f, 0.f, GetGravityZ()), DeltaSeconds);
+							Velocity = NewFallVelocity(Velocity, -GetGravityDirection() * GetGravityZ(), DeltaSeconds);
 						}
 					}
 					SetMovementMode(MOVE_Falling);
@@ -2433,8 +2407,12 @@ void UCharacterMovementComponent::UpdateBasedMovement(float DeltaSeconds)
 				// @todo: This assumes only Yaw is used, currently a valid assumption. This is the only reason FaceRotation() is used above really, aside from being a virtual hook.
 				if (bOrientRotationToMovement || (bUseControllerDesiredRotation && CharacterOwner->Controller))
 				{
-					TargetRotator.Pitch = 0.f;
-					TargetRotator.Roll = 0.f;
+					// Custom gravity automatically aligns the character to the gravity direction, so we shouldn't zero out pitch and roll.
+					if (!HasCustomGravity())
+					{
+						TargetRotator.Pitch = 0.f;
+						TargetRotator.Roll = 0.f;
+					}
 					MoveUpdatedComponent(FVector::ZeroVector, TargetRotator, false);
 					FinalQuat = UpdatedComponent->GetComponentQuat();
 				}
@@ -2450,13 +2428,25 @@ void UCharacterMovementComponent::UpdateBasedMovement(float DeltaSeconds)
 			}
 		}
 
-		// We need to offset the base of the character here, not its origin, so offset by half height
-		float HalfHeight, Radius;
-		CharacterOwner->GetCapsuleComponent()->GetScaledCapsuleSize(Radius, HalfHeight);
+		FVector NewWorldPos;
+		if (HasCustomGravity())
+		{
+			const FVector RotationRadius = UpdatedComponent->GetComponentLocation() - NewBaseLocation;
+			const FVector RotationDelta = DeltaQuat.RotateVector(RotationRadius) - RotationRadius;
+			const FVector LinearDelta = NewBaseLocation - OldBaseLocation;
+			NewWorldPos = ConstrainLocationToPlane(UpdatedComponent->GetComponentLocation() + RotationDelta + LinearDelta);
+		}
+		else
+		{
+			// We need to offset the base of the character here, not its origin, so offset by half height
+			float HalfHeight, Radius;
+			CharacterOwner->GetCapsuleComponent()->GetScaledCapsuleSize(Radius, HalfHeight);
 
-		FVector const BaseOffset(0.0f, 0.0f, HalfHeight);
-		FVector const LocalBasePos = OldLocalToWorld.InverseTransformPosition(UpdatedComponent->GetComponentLocation() - BaseOffset);
-		FVector const NewWorldPos = ConstrainLocationToPlane(NewLocalToWorld.TransformPosition(LocalBasePos) + BaseOffset);
+			FVector const BaseOffset = FVector(0.0f, 0.0f, HalfHeight);
+			FVector const LocalBasePos = OldLocalToWorld.InverseTransformPosition(UpdatedComponent->GetComponentLocation() - BaseOffset);
+			NewWorldPos = ConstrainLocationToPlane(NewLocalToWorld.TransformPosition(LocalBasePos) + BaseOffset);
+		}
+		
 		DeltaPosition = ConstrainDirectionToPlane(NewWorldPos - UpdatedComponent->GetComponentLocation());
 
 		// move attached actor
@@ -3114,7 +3104,7 @@ void UCharacterMovementComponent::Crouch(bool bClientSimulation)
 			FCollisionQueryParams CapsuleParams(SCENE_QUERY_STAT(CrouchTrace), false, CharacterOwner);
 			FCollisionResponseParams ResponseParam;
 			InitCollisionParams(CapsuleParams, ResponseParam);
-			const bool bEncroached = GetWorld()->OverlapBlockingTestByChannel(UpdatedComponent->GetComponentLocation() - RotateGravityToWorld(FVector(0.f,0.f,ScaledHalfHeightAdjust)), GetWorldToGravityTransform(),
+			const bool bEncroached = GetWorld()->OverlapBlockingTestByChannel(UpdatedComponent->GetComponentLocation() + ScaledHalfHeightAdjust * GetGravityDirection(), GetWorldToGravityTransform(),
 				UpdatedComponent->GetCollisionObjectType(), GetPawnCapsuleCollisionShape(SHRINK_None), CapsuleParams, ResponseParam);
 
 			// If encroached, cancel
@@ -3128,7 +3118,7 @@ void UCharacterMovementComponent::Crouch(bool bClientSimulation)
 		if (bCrouchMaintainsBaseLocation)
 		{
 			// Intentionally not using MoveUpdatedComponent, where a horizontal plane constraint would prevent the base of the capsule from staying at the same spot.
-			UpdatedComponent->MoveComponent(RotateGravityToWorld(FVector(0.f, 0.f, -ScaledHalfHeightAdjust)), UpdatedComponent->GetComponentQuat(), true, nullptr, EMoveComponentFlags::MOVECOMP_NoFlags, ETeleportType::TeleportPhysics);
+			UpdatedComponent->MoveComponent(ScaledHalfHeightAdjust * GetGravityDirection(), UpdatedComponent->GetComponentQuat(), true, nullptr, EMoveComponentFlags::MOVECOMP_NoFlags, ETeleportType::TeleportPhysics);
 		}
 
 		CharacterOwner->bIsCrouched = true;
@@ -3151,7 +3141,7 @@ void UCharacterMovementComponent::Crouch(bool bClientSimulation)
 		FNetworkPredictionData_Client_Character* ClientData = GetPredictionData_Client_Character();
 		if (ClientData)
 		{
-			ClientData->MeshTranslationOffset -= RotateGravityToWorld(FVector(0.f, 0.f, MeshAdjust));
+			ClientData->MeshTranslationOffset -= MeshAdjust * -GetGravityDirection();
 			ClientData->OriginalMeshTranslationOffset = ClientData->MeshTranslationOffset;
 		}
 	}
@@ -3218,7 +3208,7 @@ void UCharacterMovementComponent::UnCrouch(bool bClientSimulation)
 					CharacterOwner->GetCapsuleComponent()->GetScaledCapsuleSize(PawnRadius, PawnHalfHeight);
 					const float ShrinkHalfHeight = PawnHalfHeight - PawnRadius;
 					const float TraceDist = PawnHalfHeight - ShrinkHalfHeight;
-					const FVector Down = RotateGravityToWorld(FVector(0.f, 0.f, -TraceDist));
+					const FVector Down = TraceDist * GetGravityDirection();
 
 					FHitResult Hit(1.f);
 					const FCollisionShape ShortCapsuleShape = GetPawnCapsuleCollisionShape(SHRINK_HeightCustom, ShrinkHalfHeight);
@@ -3231,8 +3221,8 @@ void UCharacterMovementComponent::UnCrouch(bool bClientSimulation)
 					{
 						// Compute where the base of the sweep ended up, and see if we can stand there
 						const float DistanceToBase = (Hit.Time * TraceDist) + ShortCapsuleShape.Capsule.HalfHeight;
-						const FVector GravityRelativeAdjustment = FVector(0, 0, -DistanceToBase + StandingCapsuleShape.Capsule.HalfHeight + SweepInflation + MIN_FLOOR_DIST / 2.f);
-						const FVector NewLoc = PawnLocation + RotateGravityToWorld(GravityRelativeAdjustment);
+						const FVector Adjustment = (-DistanceToBase + StandingCapsuleShape.Capsule.HalfHeight + SweepInflation + MIN_FLOOR_DIST / 2.f) * -GetGravityDirection();
+						const FVector NewLoc = PawnLocation + Adjustment;
 						bEncroached = MyWorld->OverlapBlockingTestByChannel(NewLoc, GetWorldToGravityTransform(), CollisionChannel, StandingCapsuleShape, CapsuleParams, ResponseParam);
 						if (!bEncroached)
 						{
@@ -3246,7 +3236,7 @@ void UCharacterMovementComponent::UnCrouch(bool bClientSimulation)
 		else
 		{
 			// Expand while keeping base location the same.
-			FVector StandingLocation = PawnLocation + RotateGravityToWorld(FVector(0.f, 0.f, StandingCapsuleShape.GetCapsuleHalfHeight() - CurrentCrouchedHalfHeight));
+			FVector StandingLocation = PawnLocation + (StandingCapsuleShape.GetCapsuleHalfHeight() - CurrentCrouchedHalfHeight) * -GetGravityDirection();
 			bEncroached = MyWorld->OverlapBlockingTestByChannel(StandingLocation, GetWorldToGravityTransform(), CollisionChannel, StandingCapsuleShape, CapsuleParams, ResponseParam);
 
 			if (bEncroached)
@@ -3257,7 +3247,7 @@ void UCharacterMovementComponent::UnCrouch(bool bClientSimulation)
 					const float MinFloorDist = UE_KINDA_SMALL_NUMBER * 10.f;
 					if (CurrentFloor.bBlockingHit && CurrentFloor.FloorDist > MinFloorDist)
 					{
-						StandingLocation -= RotateGravityToWorld(FVector(0, 0, CurrentFloor.FloorDist - MinFloorDist));
+						StandingLocation -= (CurrentFloor.FloorDist - MinFloorDist) * -GetGravityDirection();
 						bEncroached = MyWorld->OverlapBlockingTestByChannel(StandingLocation, GetWorldToGravityTransform(), CollisionChannel, StandingCapsuleShape, CapsuleParams, ResponseParam);
 					}
 				}				
@@ -3297,7 +3287,7 @@ void UCharacterMovementComponent::UnCrouch(bool bClientSimulation)
 		FNetworkPredictionData_Client_Character* ClientData = GetPredictionData_Client_Character();
 		if (ClientData)
 		{
-			ClientData->MeshTranslationOffset += RotateGravityToWorld(FVector(0.f, 0.f, MeshAdjust));
+			ClientData->MeshTranslationOffset += MeshAdjust * -GetGravityDirection();
 			ClientData->OriginalMeshTranslationOffset = ClientData->MeshTranslationOffset;
 		}
 	}
@@ -3459,81 +3449,82 @@ float UCharacterMovementComponent::SlideAlongSurface(const FVector& Delta, float
 		return 0.f;
 	}
 
-	FVector Normal(RotateWorldToGravity(InNormal));
+	FVector Normal(InNormal);
+	const FVector::FReal NormalZ = GetGravitySpaceZ(Normal);
 	if (IsMovingOnGround())
 	{
 		// We don't want to be pushed up an unwalkable surface.
-		if (Normal.Z > 0.f)
+		if (NormalZ > 0.f)
 		{
 			if (!IsWalkable(Hit))
 			{
-				Normal = Normal.GetSafeNormal2D();
+				Normal = ProjectToGravityFloor(Normal).GetSafeNormal();
 			}
 		}
-		else if (Normal.Z < -UE_KINDA_SMALL_NUMBER)
+		else if (NormalZ < -UE_KINDA_SMALL_NUMBER)
 		{
 			// Don't push down into the floor when the impact is on the upper portion of the capsule.
 			if (CurrentFloor.FloorDist < MIN_FLOOR_DIST && CurrentFloor.bBlockingHit)
 			{
-				const FVector FloorNormal = RotateWorldToGravity(CurrentFloor.HitResult.Normal);
-
-				const bool bFloorOpposedToMovement = (RotateWorldToGravity(Delta) | FloorNormal) < 0.f && (FloorNormal.Z < 1.f - UE_DELTA);
+				const FVector FloorNormal = CurrentFloor.HitResult.Normal;
+				const bool bFloorOpposedToMovement = (Delta | FloorNormal) < 0.f && (GetGravitySpaceZ(FloorNormal) < 1.f - UE_DELTA);
 				if (bFloorOpposedToMovement)
 				{
 					Normal = FloorNormal;
 				}
 				
-				Normal = Normal.GetSafeNormal2D();
+				Normal = ProjectToGravityFloor(Normal).GetSafeNormal();
 			}
 		}
 	}
 
-	return Super::SlideAlongSurface(Delta, Time, RotateGravityToWorld(Normal), Hit, bHandleImpact);
+	return Super::SlideAlongSurface(Delta, Time, Normal, Hit, bHandleImpact);
 }
 
 
 void UCharacterMovementComponent::TwoWallAdjust(FVector& WorldSpaceDelta, const FHitResult& Hit, const FVector& OldHitNormal) const
 {
+	const FVector InDelta = WorldSpaceDelta;
 	Super::TwoWallAdjust(WorldSpaceDelta, Hit, OldHitNormal);
 
-	FVector GravityRelativeDelta = RotateWorldToGravity(WorldSpaceDelta);
 	if (IsMovingOnGround())
 	{
 		// Allow slides up walkable surfaces, but not unwalkable ones (treat those as vertical barriers).
-		if (GravityRelativeDelta.Z > 0.f)
+		const FVector::FReal WorldSpaceDeltaZ = GetGravitySpaceZ(WorldSpaceDelta);
+		if (WorldSpaceDeltaZ > 0.f)
 		{
-			const FVector GravityRelativeHitNormal = RotateWorldToGravity(Hit.Normal);
-			if ((GravityRelativeHitNormal.Z >= WalkableFloorZ || IsWalkable(Hit)) && GravityRelativeHitNormal.Z > UE_KINDA_SMALL_NUMBER)
+			const FVector::FReal HitNormalZ = GetGravitySpaceZ(Hit.Normal);
+			if ((HitNormalZ >= WalkableFloorZ || IsWalkable(Hit)) && HitNormalZ > UE_KINDA_SMALL_NUMBER)
 			{
 				// Maintain horizontal velocity
 				const float Time = (1.f - Hit.Time);
-				const FVector ScaledDelta = GravityRelativeDelta.GetSafeNormal() * GravityRelativeDelta.Size();
-				GravityRelativeDelta = FVector(GravityRelativeDelta.X, GravityRelativeDelta.Y, ScaledDelta.Z / GravityRelativeHitNormal.Z) * Time;
+				const FVector ScaledDelta = WorldSpaceDelta.GetSafeNormal() * InDelta.Size();
+				const FVector::FReal NewDeltaZ = (GetGravitySpaceZ(ScaledDelta) / HitNormalZ);
+				WorldSpaceDelta = (ProjectToGravityFloor(InDelta) + NewDeltaZ * -GetGravityDirection()) * Time;
 
 				// Should never exceed MaxStepHeight in vertical component, so rescale if necessary.
 				// This should be rare (Hit.Normal.Z above would have been very small) but we'd rather lose horizontal velocity than go too high.
-				if (GravityRelativeDelta.Z > MaxStepHeight)
+				const FVector::FReal DeltaZ = GetGravitySpaceZ(WorldSpaceDelta);
+				if (DeltaZ > MaxStepHeight)
 				{
-					const float Rescale = MaxStepHeight / GravityRelativeDelta.Z;
-					GravityRelativeDelta *= Rescale;
+					const float Rescale = MaxStepHeight / DeltaZ;
+					WorldSpaceDelta *= Rescale;
 				}
 			}
 			else
 			{
-				GravityRelativeDelta.Z = 0.f;
+				WorldSpaceDelta = ProjectToGravityFloor(WorldSpaceDelta);
 			}
 		}
-		else if (GravityRelativeDelta.Z < 0.f)
+		else if (WorldSpaceDeltaZ < 0.f)
 		{
 			// Don't push down into the floor.
 			if (CurrentFloor.FloorDist < MIN_FLOOR_DIST && CurrentFloor.bBlockingHit)
 			{
-				GravityRelativeDelta.Z = 0.f;
+				WorldSpaceDelta = ProjectToGravityFloor(WorldSpaceDelta);
 			}
 		}
 	}
-
-	WorldSpaceDelta = RotateGravityToWorld(GravityRelativeDelta);
 }
 
 
@@ -3554,17 +3545,17 @@ FVector UCharacterMovementComponent::ComputeSlideVector(const FVector& Delta, co
 FVector UCharacterMovementComponent::HandleSlopeBoosting(const FVector& SlideResult, const FVector& Delta, const float Time, const FVector& Normal, const FHitResult& Hit) const
 {
 	FVector Result = SlideResult;
-
-	if (Result.Z > 0.f)
+	const float ResultZ = GetGravitySpaceZ(Result);
+	if (ResultZ > 0.f)
 	{
 		// Don't move any higher than we originally intended.
-		const float ZLimit = Delta.Z * Time;
-		if (Result.Z - ZLimit > UE_KINDA_SMALL_NUMBER)
+		const float ZLimit = GetGravitySpaceZ(Delta) * Time;
+		if (ResultZ - ZLimit > UE_KINDA_SMALL_NUMBER)
 		{
 			if (ZLimit > 0.f)
 			{
 				// Rescale the entire vector (not just the Z component) otherwise we change the direction and likely head right back into the impact.
-				const float UpPercent = ZLimit / Result.Z;
+				const float UpPercent = ZLimit / ResultZ;
 				Result *= UpPercent;
 			}
 			else
@@ -3574,8 +3565,8 @@ FVector UCharacterMovementComponent::HandleSlopeBoosting(const FVector& SlideRes
 			}
 
 			// Make remaining portion of original result horizontal and parallel to impact normal.
-			const FVector RemainderXY = (SlideResult - Result) * FVector(1.f, 1.f, 0.f);
-			const FVector NormalXY = Normal.GetSafeNormal2D();
+			const FVector RemainderXY = ProjectToGravityFloor(SlideResult - Result);
+			const FVector NormalXY = ProjectToGravityFloor(Normal).GetSafeNormal();
 			const FVector Adjust = Super::ComputeSlideVector(RemainderXY, 1.f, NormalXY, Hit);
 			Result += Adjust;
 		}
@@ -3627,8 +3618,8 @@ float UCharacterMovementComponent::ImmersionDepth() const
 			FHitResult Hit(1.f);
 			if ( VolumeBrushComp )
 			{
-				const FVector TraceStart = UpdatedComponent->GetComponentLocation() + FVector(0.f,0.f,CollisionHalfHeight);
-				const FVector TraceEnd = UpdatedComponent->GetComponentLocation() - FVector(0.f,0.f,CollisionHalfHeight);
+				const FVector TraceStart = UpdatedComponent->GetComponentLocation() + CollisionHalfHeight * -GetGravityDirection();
+				const FVector TraceEnd = UpdatedComponent->GetComponentLocation() - CollisionHalfHeight * -GetGravityDirection();
 
 				FCollisionQueryParams NewTraceParams(SCENE_QUERY_STAT(ImmersionDepth), true);
 				VolumeBrushComp->LineTraceComponent( Hit, TraceStart, TraceEnd, NewTraceParams );
@@ -3814,6 +3805,12 @@ bool UCharacterMovementComponent::ApplyRequestedMove(float DeltaTime, float MaxA
 	return false;
 }
 
+FVector UCharacterMovementComponent::GetActorFeetLocation() const
+{
+	const float HalfHeight = CharacterOwner->GetCapsuleComponent()->GetScaledCapsuleHalfHeight();
+	return UpdatedComponent->GetComponentLocation() + HalfHeight * GetGravityDirection();
+}
+
 void UCharacterMovementComponent::RequestDirectMove(const FVector& MoveVelocity, bool bForceMaxSpeed)
 {
 	if (MoveVelocity.SizeSquared() < UE_KINDA_SMALL_NUMBER)
@@ -3824,7 +3821,7 @@ void UCharacterMovementComponent::RequestDirectMove(const FVector& MoveVelocity,
 	if (ShouldPerformAirControlForPathFollowing())
 	{
 		const FVector FallVelocity = MoveVelocity.GetClampedToMaxSize(GetMaxSpeed());
-		PerformAirControlForPathFollowing(FallVelocity, FallVelocity.Z);
+		PerformAirControlForPathFollowing(FallVelocity, GetGravitySpaceZ(FallVelocity));
 		return;
 	}
 
@@ -3834,7 +3831,7 @@ void UCharacterMovementComponent::RequestDirectMove(const FVector& MoveVelocity,
 
 	if (IsMovingOnGround())
 	{
-		RequestedVelocity.Z = 0.0f;
+		RequestedVelocity = ProjectToGravityFloor(RequestedVelocity);
 	}
 }
 
@@ -3974,8 +3971,7 @@ void UCharacterMovementComponent::CalcAvoidanceVelocity(float DeltaTime)
 
 	if (bShowDebug)
 	{
-		FVector UpLine(0,0,500);
-		DrawDebugLine(GetWorld(), GetActorFeetLocation(), GetActorFeetLocation() + UpLine, (AvoidanceLockTimer > 0.01f) ? FColor::Red : FColor::Blue, false, 0.05f, SDPG_MAX, 5.0f);
+		DrawDebugLine(GetWorld(), GetActorFeetLocation(), GetActorFeetLocation() + 500.f * -GetGravityDirection(), (AvoidanceLockTimer > 0.01f) ? FColor::Red : FColor::Blue, false, 0.05f, SDPG_MAX, 5.0f);
 	}
 #endif
 }
@@ -4099,7 +4095,7 @@ void UCharacterMovementComponent::NotifyBumpedPawn(APawn* BumpedPawn)
 	const bool bShowDebug = Avoidance && Avoidance->IsDebugEnabled(AvoidanceUID);
 	if (bShowDebug)
 	{
-		DrawDebugLine(GetWorld(), GetActorFeetLocation(), GetActorFeetLocation() + FVector(0,0,500), (AvoidanceLockTimer > 0) ? FColor(255,64,64) : FColor(64,64,255), false, 2.0f, SDPG_MAX, 20.0f);
+		DrawDebugLine(GetWorld(), GetActorFeetLocation(), GetActorFeetLocation() + 500.f * -GetGravityDirection(), (AvoidanceLockTimer > 0) ? FColor(255,64,64) : FColor(64,64,255), false, 2.0f, SDPG_MAX, 20.0f);
 	}
 #endif
 
@@ -4251,18 +4247,18 @@ void UCharacterMovementComponent::PhysFlying(float deltaTime, int32 Iterations)
 
 	if (Hit.Time < 1.f)
 	{
-		const FVector GravDir = FVector(0.f, 0.f, -1.f);
 		const FVector VelDir = Velocity.GetSafeNormal();
-		const float UpDown = GravDir | VelDir;
+		const float UpDown = VelDir | GetGravityDirection();
 
 		bool bSteppedUp = false;
-		if ((FMath::Abs(Hit.ImpactNormal.Z) < 0.2f) && (UpDown < 0.5f) && (UpDown > -0.2f) && CanStepUp(Hit))
+		if ((FMath::Abs(GetGravitySpaceZ(Hit.ImpactNormal)) < 0.2f) && (UpDown < 0.5f) && (UpDown > -0.2f) && CanStepUp(Hit))
 		{
-			float stepZ = UpdatedComponent->GetComponentLocation().Z;
-			bSteppedUp = StepUp(GravDir, Adjusted * (1.f - Hit.Time), Hit);
+			const FVector::FReal StepZ = GetGravitySpaceZ(UpdatedComponent->GetComponentLocation());
+			bSteppedUp = StepUp(GetGravityDirection(), Adjusted * (1.f - Hit.Time), Hit);
 			if (bSteppedUp)
 			{
-				OldLocation.Z = UpdatedComponent->GetComponentLocation().Z + (OldLocation.Z - stepZ);
+				const FVector::FReal LocationZ = GetGravitySpaceZ(UpdatedComponent->GetComponentLocation()) + (GetGravitySpaceZ(OldLocation) - StepZ);
+				SetGravitySpaceZ(OldLocation, LocationZ);
 			}
 		}
 
@@ -4310,7 +4306,7 @@ void UCharacterMovementComponent::ApplyRootMotionToVelocity(float deltaTime)
 		Velocity = ConstrainAnimRootMotionVelocity(AnimRootMotionVelocity, Velocity);
 		if (IsFalling())
 		{
-			Velocity += FVector(DecayingFormerBaseVelocity.X, DecayingFormerBaseVelocity.Y, 0.f);
+			Velocity += ProjectToGravityFloor(DecayingFormerBaseVelocity);
 		}
 		return;
 	}
@@ -4325,7 +4321,7 @@ void UCharacterMovementComponent::ApplyRootMotionToVelocity(float deltaTime)
 		CurrentRootMotion.AccumulateOverrideRootMotionVelocity(deltaTime, *CharacterOwner, *this, Velocity);
 		if (IsFalling())
 		{
-			Velocity += CurrentRootMotion.HasOverrideVelocityWithIgnoreZAccumulate() ? FVector(DecayingFormerBaseVelocity.X, DecayingFormerBaseVelocity.Y, 0.f) : DecayingFormerBaseVelocity;
+			Velocity += CurrentRootMotion.HasOverrideVelocityWithIgnoreZAccumulate() ? ProjectToGravityFloor(DecayingFormerBaseVelocity) : DecayingFormerBaseVelocity;
 		}
 		bAppliedRootMotion = true;
 
@@ -4359,7 +4355,7 @@ void UCharacterMovementComponent::ApplyRootMotionToVelocity(float deltaTime)
 
 	// Switch to Falling if we have vertical velocity from root motion so we can lift off the ground
 	const FVector AppliedVelocityDelta = Velocity - OldVelocity;
-	if( bAppliedRootMotion && AppliedVelocityDelta.Z != 0.f && IsMovingOnGround() )
+	if( bAppliedRootMotion && GetGravitySpaceZ(AppliedVelocityDelta) != 0.f && IsMovingOnGround() )
 	{
 		float LiftoffBound;
 		if( CurrentRootMotion.LastAccumulatedSettings.HasFlag(ERootMotionSourceSettingsFlags::UseSensitiveLiftoffCheck) )
@@ -4373,7 +4369,7 @@ void UCharacterMovementComponent::ApplyRootMotionToVelocity(float deltaTime)
 			LiftoffBound = FMath::Max(-GetGravityZ() * deltaTime, UE_SMALL_NUMBER);
 		}
 
-		if( AppliedVelocityDelta.Z > LiftoffBound )
+		if(GetGravitySpaceZ(AppliedVelocityDelta) > LiftoffBound )
 		{
 			SetMovementMode(MOVE_Falling);
 		}
@@ -4408,18 +4404,18 @@ void UCharacterMovementComponent::PhysSwimming(float deltaTime, int32 Iterations
 	float NetFluidFriction  = 0.f;
 	float Depth = ImmersionDepth();
 	float NetBuoyancy = Buoyancy * Depth;
-	float OriginalAccelZ = Acceleration.Z;
+	float OriginalAccelZ = GetGravitySpaceZ(Acceleration);
 	bool bLimitedUpAccel = false;
 
-	if (!HasAnimRootMotion() && !CurrentRootMotion.HasOverrideVelocity() && (Velocity.Z > 0.33f * MaxSwimSpeed) && (NetBuoyancy != 0.f))
+	if (!HasAnimRootMotion() && !CurrentRootMotion.HasOverrideVelocity() && (GetGravitySpaceZ(Velocity) > 0.33f * MaxSwimSpeed) && (NetBuoyancy != 0.f))
 	{
 		//damp positive Z out of water
-		Velocity.Z = FMath::Max<FVector::FReal>(0.33f * MaxSwimSpeed, Velocity.Z * Depth*Depth);
+		SetGravitySpaceZ(Velocity, FMath::Max<FVector::FReal>(0.33f * MaxSwimSpeed, GetGravitySpaceZ(Velocity) * Depth*Depth));
 	}
 	else if (Depth < 0.65f)
 	{
-		bLimitedUpAccel = (Acceleration.Z > 0.f);
-		Acceleration.Z = FMath::Min<FVector::FReal>(0.1f, Acceleration.Z);
+		bLimitedUpAccel = (OriginalAccelZ > 0.f);
+		SetGravitySpaceZ(Acceleration, FMath::Min<FVector::FReal>(0.1f, OriginalAccelZ));
 	}
 
 	Iterations++;
@@ -4429,14 +4425,14 @@ void UCharacterMovementComponent::PhysSwimming(float deltaTime, int32 Iterations
 	{
 		const float Friction = 0.5f * GetPhysicsVolume()->FluidFriction * Depth;
 		CalcVelocity(deltaTime, Friction, true, GetMaxBrakingDeceleration());
-		Velocity.Z += GetGravityZ() * deltaTime * (1.f - NetBuoyancy);
+		Velocity += (GetGravityZ() * deltaTime * (1.f - NetBuoyancy)) * -GetGravityDirection();
 	}
 
 	ApplyRootMotionToVelocity(deltaTime);
 
 	FVector Adjusted = Velocity * deltaTime;
 	FHitResult Hit(1.f);
-	float remainingTime = deltaTime * Swim(Adjusted, Hit);
+	const float remainingTime = deltaTime * Swim(Adjusted, Hit);
 
 	//may have left water - if so, script might have set new physics mode
 	if ( !IsSwimming() )
@@ -4448,10 +4444,10 @@ void UCharacterMovementComponent::PhysSwimming(float deltaTime, int32 Iterations
 	if ( Hit.Time < 1.f && CharacterOwner)
 	{
 		HandleSwimmingWallHit(Hit, deltaTime);
-		if (bLimitedUpAccel && (Velocity.Z >= 0.f))
+		if (bLimitedUpAccel && (GetGravitySpaceZ(Velocity) >= 0.f))
 		{
 			// allow upward velocity at surface if against obstacle
-			Velocity.Z += OriginalAccelZ * deltaTime;
+			Velocity += OriginalAccelZ * deltaTime * -GetGravityDirection();
 			Adjusted = Velocity * (1.f - Hit.Time)*deltaTime;
 			Swim(Adjusted, Hit);
 			if (!IsSwimming())
@@ -4461,17 +4457,16 @@ void UCharacterMovementComponent::PhysSwimming(float deltaTime, int32 Iterations
 			}
 		}
 
-		const FVector GravDir = FVector(0.f,0.f,-1.f);
 		const FVector VelDir = Velocity.GetSafeNormal();
-		const float UpDown = GravDir | VelDir;
+		const float UpDown = VelDir | GetGravityDirection();
 
 		bool bSteppedUp = false;
-		if( (FMath::Abs(Hit.ImpactNormal.Z) < 0.2f) && (UpDown < 0.5f) && (UpDown > -0.2f) && CanStepUp(Hit))
+		if( (FMath::Abs(GetGravitySpaceZ(Hit.ImpactNormal)) < 0.2f) && (UpDown < 0.5f) && (UpDown > -0.2f) && CanStepUp(Hit))
 		{
-			float stepZ = UpdatedComponent->GetComponentLocation().Z;
+			const float StepZ = GetGravitySpaceZ(UpdatedComponent->GetComponentLocation());
 			const FVector RealVelocity = Velocity;
-			Velocity.Z = 1.f;	// HACK: since will be moving up, in case pawn leaves the water
-			bSteppedUp = StepUp(GravDir, Adjusted * (1.f - Hit.Time), Hit);
+			SetGravitySpaceZ(Velocity, 1.f);	// HACK: since will be moving up, in case pawn leaves the water
+			bSteppedUp = StepUp(GetGravityDirection(), Adjusted * (1.f - Hit.Time), Hit);
 			if (bSteppedUp)
 			{
 				//may have left water - if so, script might have set new physics mode
@@ -4480,7 +4475,7 @@ void UCharacterMovementComponent::PhysSwimming(float deltaTime, int32 Iterations
 					StartNewPhysics(remainingTime, Iterations);
 					return;
 				}
-				OldLocation.Z = UpdatedComponent->GetComponentLocation().Z + (OldLocation.Z - stepZ);
+				SetGravitySpaceZ(OldLocation, GetGravitySpaceZ(UpdatedComponent->GetComponentLocation()) + (GetGravitySpaceZ(OldLocation) - StepZ));
 			}
 			Velocity = RealVelocity;
 		}
@@ -4495,12 +4490,12 @@ void UCharacterMovementComponent::PhysSwimming(float deltaTime, int32 Iterations
 
 	if( !HasAnimRootMotion() && !CurrentRootMotion.HasOverrideVelocity() && !bJustTeleported && ((deltaTime - remainingTime) > UE_KINDA_SMALL_NUMBER) && CharacterOwner )
 	{
-		bool bWaterJump = !GetPhysicsVolume()->bWaterVolume;
-		float velZ = Velocity.Z;
+		const bool bWaterJump = !GetPhysicsVolume()->bWaterVolume;
+		const FVector::FReal VelZ = GetGravitySpaceZ(Velocity);
 		Velocity = (UpdatedComponent->GetComponentLocation() - OldLocation) / (deltaTime - remainingTime);
 		if (bWaterJump)
 		{
-			Velocity.Z = velZ;
+			SetGravitySpaceZ(Velocity, VelZ);
 		}
 	}
 
@@ -4542,9 +4537,10 @@ void UCharacterMovementComponent::StartSwimming(FVector OldLocation, FVector Old
 		}
 		MoveUpdatedComponent(End - UpdatedComponent->GetComponentLocation(), UpdatedComponent->GetComponentQuat(), true);
 	}
-	if ( !HasAnimRootMotion() && !CurrentRootMotion.HasOverrideVelocity() && (Velocity.Z > 2.f*CharacterMovementConstants::SWIMBOBSPEED) && (Velocity.Z < 0.f)) //allow for falling out of water
+	const FVector::FReal GravityRelativeVelocityZ = GetGravitySpaceZ(Velocity);
+	if ( !HasAnimRootMotion() && !CurrentRootMotion.HasOverrideVelocity() && (GravityRelativeVelocityZ > 2.f*CharacterMovementConstants::SWIMBOBSPEED) && (GravityRelativeVelocityZ < 0.f)) //allow for falling out of water
 	{
-		Velocity.Z = CharacterMovementConstants::SWIMBOBSPEED - Velocity.Size2D() * 0.7f; //smooth bobbing
+		SetGravitySpaceZ(Velocity, CharacterMovementConstants::SWIMBOBSPEED - ProjectToGravityFloor(Velocity).Size() * 0.7f); //smooth bobbing
 	}
 	if ( (remainingTime >= MIN_TICK_TIME) && (Iterations < MaxSimulationIterations) )
 	{
@@ -4615,11 +4611,10 @@ void UCharacterMovementComponent::NotifyJumpApex()
 FVector UCharacterMovementComponent::GetFallingLateralAcceleration(float DeltaTime)
 {
 	// No acceleration in Z
-	const FVector GravityRelativeAcceleration = RotateWorldToGravity(Acceleration);
-	FVector FallAcceleration = RotateGravityToWorld(FVector(GravityRelativeAcceleration.X, GravityRelativeAcceleration.Y, 0.f));
+	FVector FallAcceleration = ProjectToGravityFloor(Acceleration);
 
 	// bound acceleration, falling object has minimal ability to impact acceleration
-	if (!HasAnimRootMotion() && GravityRelativeAcceleration.SizeSquared2D() > 0.f)
+	if (!HasAnimRootMotion() && FallAcceleration.SizeSquared() > 0.f)
 	{
 		FallAcceleration = GetAirControl(DeltaTime, AirControl, FallAcceleration);
 		FallAcceleration = FallAcceleration.GetClampedToMaxSize(GetMaxAcceleration());
@@ -4631,8 +4626,7 @@ FVector UCharacterMovementComponent::GetFallingLateralAcceleration(float DeltaTi
 
 bool UCharacterMovementComponent::ShouldLimitAirControl(float DeltaTime, const FVector& FallAcceleration) const
 {
-	const FVector GravityRelativeFallAcceleration = RotateWorldToGravity(FallAcceleration);
-	return (GravityRelativeFallAcceleration.SizeSquared2D() > 0.f);
+	return (ProjectToGravityFloor(FallAcceleration).SizeSquared() > 0.f);
 }
 
 FVector UCharacterMovementComponent::GetAirControl(float DeltaTime, float TickAirControl, const FVector& FallAcceleration)
@@ -4665,8 +4659,7 @@ void UCharacterMovementComponent::SetGravityDirection(const FVector& InNewGravit
 float UCharacterMovementComponent::BoostAirControl(float DeltaTime, float TickAirControl, const FVector& FallAcceleration)
 {
 	// Allow a burst of initial acceleration
-	const FVector GravityRelativeVelocity = RotateWorldToGravity(Velocity);
-	if (AirControlBoostMultiplier > 0.f && GravityRelativeVelocity.SizeSquared2D() < FMath::Square(AirControlBoostVelocityThreshold))
+	if (AirControlBoostMultiplier > 0.f && ProjectToGravityFloor(Velocity).SizeSquared() < FMath::Square(AirControlBoostVelocityThreshold))
 	{
 		TickAirControl = FMath::Min(1.f, AirControlBoostMultiplier * TickAirControl);
 	}
@@ -4684,9 +4677,7 @@ void UCharacterMovementComponent::PhysFalling(float deltaTime, int32 Iterations)
 		return;
 	}
 
-	FVector FallAcceleration = GetFallingLateralAcceleration(deltaTime);
-	const FVector GravityRelativeFallAcceleration = RotateWorldToGravity(FallAcceleration);
-	FallAcceleration = RotateGravityToWorld(FVector(GravityRelativeFallAcceleration.X, GravityRelativeFallAcceleration.Y, 0));
+	const FVector FallAcceleration = ProjectToGravityFloor(GetFallingLateralAcceleration(deltaTime));
 	const bool bHasLimitedAirControl = ShouldLimitAirControl(deltaTime, FallAcceleration);
 
 	float remainingTime = deltaTime;
@@ -4716,7 +4707,7 @@ void UCharacterMovementComponent::PhysFalling(float deltaTime, int32 Iterations)
 				TGuardValue<FVector> RestoreAcceleration(Acceleration, FallAcceleration);
 				if (HasCustomGravity())
 				{
-					Velocity = FVector::VectorPlaneProject(Velocity, RotateGravityToWorld(FVector::UpVector));
+					Velocity = ProjectToGravityFloor(Velocity);
 					const FVector GravityRelativeOffset = OldVelocity - Velocity;
 					CalcVelocity(timeTick, FallingLateralFriction, false, MaxDecel);
 					Velocity += GravityRelativeOffset;
@@ -4759,14 +4750,14 @@ void UCharacterMovementComponent::PhysFalling(float deltaTime, int32 Iterations)
 		DecayFormerBaseVelocity(timeTick);
 
 		// See if we need to sub-step to exactly reach the apex. This is important for avoiding "cutting off the top" of the trajectory as framerate varies.
-		const FVector GravityRelativeOldVelocityWithRootMotion = RotateWorldToGravity(OldVelocityWithRootMotion);
-		if (CharacterMovementCVars::ForceJumpPeakSubstep && GravityRelativeOldVelocityWithRootMotion.Z > 0.f && RotateWorldToGravity(Velocity).Z <= 0.f && NumJumpApexAttempts < MaxJumpApexAttemptsPerSimulation)
+		const FVector::FReal GravityRelativeOldVelocityWithRootMotionZ = GetGravitySpaceZ(OldVelocityWithRootMotion);
+		if (CharacterMovementCVars::ForceJumpPeakSubstep && GravityRelativeOldVelocityWithRootMotionZ > 0.f && GetGravitySpaceZ(Velocity) <= 0.f && NumJumpApexAttempts < MaxJumpApexAttemptsPerSimulation)
 		{
 			const FVector DerivedAccel = (Velocity - OldVelocityWithRootMotion) / timeTick;
-			const FVector GravityRelativeDerivedAccel = RotateWorldToGravity(DerivedAccel);
-			if (!FMath::IsNearlyZero(GravityRelativeDerivedAccel.Z))
+			const FVector::FReal GravityRelativeDerivedAccelZ = GetGravitySpaceZ(DerivedAccel);
+			if (!FMath::IsNearlyZero(GravityRelativeDerivedAccelZ))
 			{
-				const float TimeToApex = -GravityRelativeOldVelocityWithRootMotion.Z / GravityRelativeDerivedAccel.Z;
+				const float TimeToApex = -GravityRelativeOldVelocityWithRootMotionZ / GravityRelativeDerivedAccelZ;
 				
 				// The time-to-apex calculation should be precise, and we want to avoid adding a substep when we are basically already at the apex from the previous iteration's work.
 				const float ApexTimeMinimum = 0.0001f;
@@ -4775,8 +4766,7 @@ void UCharacterMovementComponent::PhysFalling(float deltaTime, int32 Iterations)
 					const FVector ApexVelocity = OldVelocityWithRootMotion + (DerivedAccel * TimeToApex);
 					if (HasCustomGravity())
 					{
-						const FVector GravityRelativeApexVelocity = RotateWorldToGravity(ApexVelocity);
-						Velocity = RotateGravityToWorld(FVector(GravityRelativeApexVelocity.X, GravityRelativeApexVelocity.Y, 0)); // Should be nearly zero anyway, but this makes apex notifications consistent.
+						Velocity = ProjectToGravityFloor(ApexVelocity); // Should be nearly zero anyway, but this makes apex notifications consistent.
 					}
 					else
 					{
@@ -4802,7 +4792,7 @@ void UCharacterMovementComponent::PhysFalling(float deltaTime, int32 Iterations)
 			}
 		}
 
-		if (bNotifyApex && (RotateWorldToGravity(Velocity).Z < 0.f))
+		if (bNotifyApex && (GetGravitySpaceZ(Velocity) < 0.f))
 		{
 			// Just passed jump apex since now going down
 			bNotifyApex = false;
@@ -4890,7 +4880,7 @@ void UCharacterMovementComponent::PhysFalling(float deltaTime, int32 Iterations)
 						TGuardValue<FVector> RestoreVelocity(Velocity, OldVelocity);
 						if (HasCustomGravity())
 						{
-							Velocity = FVector::VectorPlaneProject(Velocity, RotateGravityToWorld(FVector::UpVector));
+							Velocity = ProjectToGravityFloor(Velocity);
 							const FVector GravityRelativeOffset = OldVelocity - Velocity;
 							CalcVelocity(timeTick, FallingLateralFriction, false, MaxDecel);
 							VelocityNoAirControl = Velocity + GravityRelativeOffset;
@@ -4921,12 +4911,12 @@ void UCharacterMovementComponent::PhysFalling(float deltaTime, int32 Iterations)
 				{
 					const FVector ContactVelocity = MovementBaseUtility::GetMovementBaseVelocity(HitComponent, NAME_None) + MovementBaseUtility::GetMovementBaseTangentialVelocity(HitComponent, NAME_None, Hit.ImpactPoint);
 					const FVector NewVelocity = Velocity - Hit.ImpactNormal * FVector::DotProduct(Velocity - ContactVelocity, Hit.ImpactNormal);
-					Velocity = HasAnimRootMotion() || CurrentRootMotion.HasOverrideVelocityWithIgnoreZAccumulate() ? FVector(Velocity.X, Velocity.Y, NewVelocity.Z) : NewVelocity;
+					Velocity = HasAnimRootMotion() || CurrentRootMotion.HasOverrideVelocityWithIgnoreZAccumulate() ? ProjectToGravityFloor(Velocity) + GetGravitySpaceComponentZ(NewVelocity) : NewVelocity;
 				}
 				else if (subTimeTickRemaining > UE_KINDA_SMALL_NUMBER && !bJustTeleported)
 				{
 					const FVector NewVelocity = (Delta / subTimeTickRemaining);
-					Velocity = HasAnimRootMotion() || CurrentRootMotion.HasOverrideVelocityWithIgnoreZAccumulate() ? FVector(Velocity.X, Velocity.Y, NewVelocity.Z) : NewVelocity;
+					Velocity = HasAnimRootMotion() || CurrentRootMotion.HasOverrideVelocityWithIgnoreZAccumulate() ? ProjectToGravityFloor(Velocity) + GetGravitySpaceComponentZ(NewVelocity) : NewVelocity;
 				}
 
 				if (subTimeTickRemaining > UE_KINDA_SMALL_NUMBER && (Delta | Adjusted) > 0.f)
@@ -4956,7 +4946,7 @@ void UCharacterMovementComponent::PhysFalling(float deltaTime, int32 Iterations)
 						}
 
 						// Act as if there was no air control on the last move when computing new deflection.
-						if (bHasLimitedAirControl && RotateWorldToGravity(Hit.Normal).Z > CharacterMovementConstants::VERTICAL_SLOPE_NORMAL_Z)
+						if (bHasLimitedAirControl && GetGravitySpaceZ(Hit.Normal) > CharacterMovementConstants::VERTICAL_SLOPE_NORMAL_Z)
 						{
 							const FVector LastMoveNoAirControl = VelocityNoAirControl * LastMoveTimeSlice;
 							Delta = ComputeSlideVector(LastMoveNoAirControl, 1.f, OldHitNormal, Hit);
@@ -4982,19 +4972,27 @@ void UCharacterMovementComponent::PhysFalling(float deltaTime, int32 Iterations)
 						if (subTimeTickRemaining > UE_KINDA_SMALL_NUMBER && !bJustTeleported)
 						{
 							const FVector NewVelocity = (Delta / subTimeTickRemaining);
-							Velocity = HasAnimRootMotion() || CurrentRootMotion.HasOverrideVelocityWithIgnoreZAccumulate() ? FVector(Velocity.X, Velocity.Y, NewVelocity.Z) : NewVelocity;
+							Velocity = HasAnimRootMotion() || CurrentRootMotion.HasOverrideVelocityWithIgnoreZAccumulate() ? ProjectToGravityFloor(Velocity) + GetGravitySpaceComponentZ(NewVelocity) : NewVelocity;
 						}
 
 						// bDitch=true means that pawn is straddling two slopes, neither of which it can stand on
-						bool bDitch = ( (RotateWorldToGravity(OldHitImpactNormal).Z > 0.f) && (RotateWorldToGravity(Hit.ImpactNormal).Z > 0.f) && (FMath::Abs(Delta.Z) <= UE_KINDA_SMALL_NUMBER) && ((Hit.ImpactNormal | OldHitImpactNormal) < 0.f) );
+						bool bDitch = ( (GetGravitySpaceZ(OldHitImpactNormal) > 0.f) && (GetGravitySpaceZ(Hit.ImpactNormal) > 0.f) && (FMath::Abs(GetGravitySpaceZ(Delta)) <= UE_KINDA_SMALL_NUMBER) && ((Hit.ImpactNormal | OldHitImpactNormal) < 0.f) );
 						SafeMoveUpdatedComponent( Delta, PawnRotation, true, Hit);
 						if ( Hit.Time == 0.f )
 						{
 							// if we are stuck then try to side step
-							FVector SideDelta = (OldHitNormal + Hit.ImpactNormal).GetSafeNormal2D();
+							FVector SideDelta = ProjectToGravityFloor(OldHitNormal + Hit.ImpactNormal).GetSafeNormal();
 							if ( SideDelta.IsNearlyZero() )
 							{
-								SideDelta = FVector(OldHitNormal.Y, -OldHitNormal.X, 0).GetSafeNormal();
+								if (HasCustomGravity())
+								{
+									const FVector GravityRelativeHitNormal = RotateWorldToGravity(OldHitNormal);
+									SideDelta = RotateGravityToWorld(FVector(GravityRelativeHitNormal.Y, -GravityRelativeHitNormal.X, 0.f)).GetSafeNormal();
+								}
+								else
+								{
+									SideDelta = FVector(OldHitNormal.Y, -OldHitNormal.X, 0).GetSafeNormal();	
+								}
 							}
 							SafeMoveUpdatedComponent( SideDelta, PawnRotation, true, Hit);
 						}
@@ -5005,13 +5003,13 @@ void UCharacterMovementComponent::PhysFalling(float deltaTime, int32 Iterations)
 							ProcessLanded(Hit, remainingTime, Iterations);
 							return;
 						}
-						else if (GetPerchRadiusThreshold() > 0.f && Hit.Time == 1.f && RotateWorldToGravity(OldHitImpactNormal).Z >= WalkableFloorZ)
+						else if (GetPerchRadiusThreshold() > 0.f && Hit.Time == 1.f && GetGravitySpaceZ(OldHitImpactNormal) >= WalkableFloorZ)
 						{
 							// We might be in a virtual 'ditch' within our perch radius. This is rare.
 							const FVector PawnLocation = UpdatedComponent->GetComponentLocation();
-							const float ZMovedDist = FMath::Abs(RotateWorldToGravity(PawnLocation - OldLocation).Z);
-							const float MovedDist2DSq = FVector::VectorPlaneProject(PawnLocation - OldLocation, RotateGravityToWorld(FVector::UpVector)).Size2D();
-							if (ZMovedDist <= 0.2f * timeTick && MovedDist2DSq <= 4.f * timeTick)
+							const float ZMovedDist = FMath::Abs(GetGravitySpaceZ(PawnLocation - OldLocation));
+							const float MovedDist2D = ProjectToGravityFloor(PawnLocation - OldLocation).Size();
+							if (ZMovedDist <= 0.2f * timeTick && MovedDist2D <= 4.f * timeTick)
 							{
 								FVector GravityRelativeVelocity = RotateWorldToGravity(Velocity);
 								GravityRelativeVelocity.X += 0.25f * GetMaxSpeed() * (RandomStream.FRand() - 0.5f);
@@ -5027,12 +5025,10 @@ void UCharacterMovementComponent::PhysFalling(float deltaTime, int32 Iterations)
 			}
 		}
 
-		FVector GravityRelativeVelocity = RotateWorldToGravity(Velocity);
-		if (GravityRelativeVelocity.SizeSquared2D() <= UE_KINDA_SMALL_NUMBER * 10.f)
+		const FVector GravityProjectedVelocity = ProjectToGravityFloor(Velocity);
+		if (GravityProjectedVelocity.SizeSquared() <= UE_KINDA_SMALL_NUMBER * 10.f)
 		{
-			GravityRelativeVelocity.X = 0.f;
-			GravityRelativeVelocity.Y = 0.f;
-			Velocity = RotateGravityToWorld(GravityRelativeVelocity);
+			Velocity = GetGravitySpaceComponentZ(Velocity);
 		}
 	}
 }
@@ -5041,7 +5037,7 @@ FVector UCharacterMovementComponent::LimitAirControl(float DeltaTime, const FVec
 {
 	FVector Result(FallAcceleration);
 
-	if (HitResult.IsValidBlockingHit() && HitResult.Normal.Z > CharacterMovementConstants::VERTICAL_SLOPE_NORMAL_Z)
+	if (HitResult.IsValidBlockingHit() && GetGravitySpaceZ(HitResult.Normal) > CharacterMovementConstants::VERTICAL_SLOPE_NORMAL_Z)
 	{
 		if (!bCheckForValidLandingSpot || !IsValidLandingSpot(HitResult.Location, HitResult))
 		{
@@ -5049,7 +5045,7 @@ FVector UCharacterMovementComponent::LimitAirControl(float DeltaTime, const FVec
 			if (FVector::DotProduct(FallAcceleration, HitResult.Normal) < 0.f)
 			{
 				// Allow movement parallel to the wall, but not into it because that may push us up.
-				const FVector Normal2D = HitResult.Normal.GetSafeNormal2D();
+				const FVector Normal2D = ProjectToGravityFloor(HitResult.Normal).GetSafeNormal();
 				Result = FVector::VectorPlaneProject(FallAcceleration, Normal2D);
 			}
 		}
@@ -5073,13 +5069,13 @@ bool UCharacterMovementComponent::CheckLedgeDirection(const FVector& OldLocation
 	const FCollisionShape CapsuleShape = GetPawnCapsuleCollisionShape(SHRINK_None);
 	const ECollisionChannel CollisionChannel = UpdatedComponent->GetCollisionObjectType();
 	FHitResult Result(1.f);
-	GetWorld()->SweepSingleByChannel(Result, OldLocation, SideDest, FQuat::Identity, CollisionChannel, CapsuleShape, CapsuleParams, ResponseParam);
+	GetWorld()->SweepSingleByChannel(Result, OldLocation, SideDest, GetWorldToGravityTransform(), CollisionChannel, CapsuleShape, CapsuleParams, ResponseParam);
 
 	if ( !Result.bBlockingHit || IsWalkable(Result) )
 	{
 		if ( !Result.bBlockingHit )
 		{
-			GetWorld()->SweepSingleByChannel(Result, SideDest, SideDest + GetGravityDirection() * (MaxStepHeight + LedgeCheckThreshold), FQuat::Identity, CollisionChannel, CapsuleShape, CapsuleParams, ResponseParam);
+			GetWorld()->SweepSingleByChannel(Result, SideDest, SideDest + GetGravityDirection() * (MaxStepHeight + LedgeCheckThreshold), GetWorldToGravityTransform(), CollisionChannel, CapsuleShape, CapsuleParams, ResponseParam);
 		}
 		if ( (Result.Time < 1.f) && IsWalkable(Result) )
 		{
@@ -5123,12 +5119,12 @@ FVector UCharacterMovementComponent::GetLedgeMove(const FVector& OldLocation, co
 		}
 
 		// Check previous (valid) walkable floor location to get data about the ledge
-		if (LedgeFloor.bBlockingHit && RotateWorldToGravity(LedgeFloor.HitResult.Normal).Z > UE_KINDA_SMALL_NUMBER)
+		if (LedgeFloor.bBlockingHit && GetGravitySpaceZ(LedgeFloor.HitResult.Normal) > UE_KINDA_SMALL_NUMBER)
 		{
 			// It would be nice to use HitResult.Normal, but the quality of this normal doesn't seem as good as ImpactNormal (contains off-axis values for simple shapes).
 			LedgeNormal = (LedgeFloor.HitResult.ImpactNormal ^ LedgeFloor.HitResult.Normal);
 			LedgeNormal = (LedgeFloor.HitResult.ImpactNormal ^ LedgeNormal);
-			LedgeNormal = FVector::VectorPlaneProject(LedgeNormal, -GetGravityDirection()).GetSafeNormal();
+			LedgeNormal = ProjectToGravityFloor(LedgeNormal).GetSafeNormal();
 			LedgeMove = ComputeSlideVector(Delta, 1.0f, -LedgeNormal, LedgeFloor.HitResult);
 			bUseLedgeNormal = true;
 		}
@@ -5205,7 +5201,7 @@ void UCharacterMovementComponent::StartFalling(int32 Iterations, float remaining
 {
 	// start falling 
 	const float DesiredDist = Delta.Size();
-	const float ActualDist = (UpdatedComponent->GetComponentLocation() - subLoc).Size2D();
+	const float ActualDist = ProjectToGravityFloor(UpdatedComponent->GetComponentLocation() - subLoc).Size();
 	remainingTime = (DesiredDist < UE_KINDA_SMALL_NUMBER)
 					? 0.f
 					: remainingTime + timeTick * (1.f - FMath::Min(1.f,ActualDist/DesiredDist));
@@ -5264,52 +5260,26 @@ void UCharacterMovementComponent::RevertMove(const FVector& OldLocation, UPrimit
 
 FVector UCharacterMovementComponent::ComputeGroundMovementDelta(const FVector& Delta, const FHitResult& RampHit, const bool bHitFromLineTrace) const
 {
-	if (!HasCustomGravity())
+	const FVector FloorNormal = RampHit.ImpactNormal;
+	const FVector::FReal FloorNormalZ = GetGravitySpaceZ(FloorNormal);
+	const FVector::FReal ContactNormalZ = GetGravitySpaceZ(RampHit.Normal);
+	if (FloorNormalZ < (1.f - UE_KINDA_SMALL_NUMBER) && FloorNormalZ > UE_KINDA_SMALL_NUMBER && ContactNormalZ > UE_KINDA_SMALL_NUMBER && !bHitFromLineTrace && IsWalkable(RampHit))
 	{
-		const FVector FloorNormal = RampHit.ImpactNormal;
-		const FVector ContactNormal = RampHit.Normal;
-
-		if (FloorNormal.Z < (1.f - UE_KINDA_SMALL_NUMBER) && FloorNormal.Z > UE_KINDA_SMALL_NUMBER && ContactNormal.Z > UE_KINDA_SMALL_NUMBER && !bHitFromLineTrace && IsWalkable(RampHit))
+		// Compute a vector that moves parallel to the surface, by projecting the horizontal movement direction onto the ramp.
+		const float FloorDotDelta = (FloorNormal | Delta);
+		FVector RampMovement = Delta;
+		SetGravitySpaceZ(RampMovement, -FloorDotDelta / FloorNormalZ);
+		if (bMaintainHorizontalGroundVelocity)
 		{
-			// Compute a vector that moves parallel to the surface, by projecting the horizontal movement direction onto the ramp.
-			const float FloorDotDelta = (FloorNormal | Delta);
-			FVector RampMovement(Delta.X, Delta.Y, -FloorDotDelta / FloorNormal.Z);
-
-			if (bMaintainHorizontalGroundVelocity)
-			{
-				return RampMovement;
-			}
-			else
-			{
-				return RampMovement.GetSafeNormal() * Delta.Size();
-			}
+			return RampMovement;
 		}
-		return Delta;
-	}
-	else
-	{
-		const FVector GravityRelativeDelta = RotateWorldToGravity(Delta);
-		const FVector GravityRelativeFloorNormal = RotateWorldToGravity(RampHit.ImpactNormal);
-		const FVector GravityRelativeContactNormal = RotateWorldToGravity(RampHit.Normal);
-
-		if (GravityRelativeFloorNormal.Z < (1.f - UE_KINDA_SMALL_NUMBER) && GravityRelativeFloorNormal.Z > UE_KINDA_SMALL_NUMBER && GravityRelativeContactNormal.Z > UE_KINDA_SMALL_NUMBER && !bHitFromLineTrace && IsWalkable(RampHit))
+		else
 		{
-			// Compute a vector that moves parallel to the surface, by projecting the horizontal movement direction onto the ramp.
-			const float FloorDotDelta = (GravityRelativeFloorNormal | GravityRelativeDelta);
-			FVector GravityRelativeRampMovement(GravityRelativeDelta.X, GravityRelativeDelta.Y, -FloorDotDelta / GravityRelativeFloorNormal.Z);
-
-			if (bMaintainHorizontalGroundVelocity)
-			{
-				return RotateGravityToWorld(GravityRelativeRampMovement);
-			}
-			else
-			{
-				return RotateGravityToWorld(GravityRelativeRampMovement.GetSafeNormal() * GravityRelativeDelta.Size());
-			}
+			return RampMovement.GetSafeNormal() * Delta.Size();
 		}
-
-		return RotateGravityToWorld(GravityRelativeDelta);
 	}
+
+	return Delta;
 }
 
 void UCharacterMovementComponent::OnCharacterStuckInGeometry(const FHitResult* Hit)
@@ -5359,7 +5329,7 @@ void UCharacterMovementComponent::MoveAlongFloor(const FVector& InVelocity, floa
 	}
 
 	// Move along the current floor
-	const FVector Delta = RotateGravityToWorld(RotateWorldToGravity(InVelocity) * FVector(1.0, 1.0, 0.0)) * DeltaSeconds;
+	const FVector Delta = ProjectToGravityFloor(InVelocity) * DeltaSeconds;
 	FHitResult Hit(1.f);
 	FVector RampVector = ComputeGroundMovementDelta(Delta, CurrentFloor.HitResult, CurrentFloor.bLineTrace);
 	SafeMoveUpdatedComponent(RampVector, UpdatedComponent->GetComponentQuat(), true, Hit);
@@ -5380,7 +5350,7 @@ void UCharacterMovementComponent::MoveAlongFloor(const FVector& InVelocity, floa
 	{
 		// We impacted something (most likely another ramp, but possibly a barrier).
 		float PercentTimeApplied = Hit.Time;
-		if ((Hit.Time > 0.f) && (Hit.Normal.Z > UE_KINDA_SMALL_NUMBER) && IsWalkable(Hit))
+		if ((Hit.Time > 0.f) && (GetGravitySpaceZ(Hit.Normal) > UE_KINDA_SMALL_NUMBER) && IsWalkable(Hit))
 		{
 			// Another walkable ramp.
 			const float InitialPercentRemaining = 1.f - PercentTimeApplied;
@@ -5398,8 +5368,7 @@ void UCharacterMovementComponent::MoveAlongFloor(const FVector& InVelocity, floa
 			{
 				// hit a barrier, try to step up
 				const FVector PreStepUpLocation = UpdatedComponent->GetComponentLocation();
-				const FVector GravDir = GetGravityDirection();
-				if (!StepUp(GravDir, Delta * (1.f - PercentTimeApplied), Hit, OutStepDownResult))
+				if (!StepUp(GetGravityDirection(), Delta * (1.f - PercentTimeApplied), Hit, OutStepDownResult))
 				{
 					UE_LOG(LogCharacterMovement, Verbose, TEXT("- StepUp (ImpactNormal %s, Normal %s"), *Hit.ImpactNormal.ToString(), *Hit.Normal.ToString());
 					HandleImpact(Hit, LastMoveTimeSlice, RampVector);
@@ -5416,7 +5385,7 @@ void UCharacterMovementComponent::MoveAlongFloor(const FVector& InVelocity, floa
 						if (!HasAnimRootMotion() && !CurrentRootMotion.HasOverrideVelocity() && StepUpTimeSlice >= UE_KINDA_SMALL_NUMBER)
 						{
 							Velocity = (UpdatedComponent->GetComponentLocation() - PreStepUpLocation) / StepUpTimeSlice;
-							Velocity = FVector::VectorPlaneProject(Velocity, -GravDir);
+							Velocity = ProjectToGravityFloor(Velocity);
 						}
 					}
 				}
@@ -5433,22 +5402,19 @@ void UCharacterMovementComponent::MoveAlongFloor(const FVector& InVelocity, floa
 
 void UCharacterMovementComponent::MaintainHorizontalGroundVelocity()
 {
-	FVector GravityRelativeVelocity = RotateWorldToGravity(Velocity);
-	if (GravityRelativeVelocity.Z != 0.f)
+	if (GetGravitySpaceZ(Velocity) != 0.f)
 	{
 		if (bMaintainHorizontalGroundVelocity)
 		{
 			// Ramp movement already maintained the velocity, so we just want to remove the vertical component.
-			GravityRelativeVelocity.Z = 0.f;
+			Velocity = ProjectToGravityFloor(Velocity);
 		}
 		else
 		{
 			// Rescale velocity to be horizontal but maintain magnitude of last update.
-			GravityRelativeVelocity = GravityRelativeVelocity.GetSafeNormal2D() * GravityRelativeVelocity.Size();
+			Velocity = ProjectToGravityFloor(Velocity).GetSafeNormal() * Velocity.Size();
 		}
 	}
-
-	Velocity = RotateGravityToWorld(GravityRelativeVelocity);
 }
 
 
@@ -5551,7 +5517,7 @@ void UCharacterMovementComponent::PhysWalking(float deltaTime, int32 Iterations)
 				const float DesiredDist = Delta.Size();
 				if (DesiredDist > UE_KINDA_SMALL_NUMBER)
 				{
-					const float ActualDist = (UpdatedComponent->GetComponentLocation() - OldLocation).Size2D();
+					const float ActualDist = ProjectToGravityFloor(UpdatedComponent->GetComponentLocation() - OldLocation).Size();
 					remainingTime += timeTick * (1.f - FMath::Min(1.f,ActualDist/DesiredDist));
 				}
 				StartNewPhysics(remainingTime,Iterations);
@@ -5631,7 +5597,7 @@ void UCharacterMovementComponent::PhysWalking(float deltaTime, int32 Iterations)
 				// The floor check failed because it started in penetration
 				// We do not want to try to move downward because the downward sweep failed, rather we'd like to try to pop out of the floor.
 				FHitResult Hit(CurrentFloor.HitResult);
-				Hit.TraceEnd = Hit.TraceStart + RotateGravityToWorld(FVector(0.f, 0.f, MAX_FLOOR_DIST));
+				Hit.TraceEnd = Hit.TraceStart + MAX_FLOOR_DIST * -GetGravityDirection();
 				const FVector RequestedAdjustment = GetPenetrationAdjustment(Hit);
 				ResolvePenetration(RequestedAdjustment, Hit, UpdatedComponent->GetComponentQuat());
 				bForceNextFloorCheck = true;
@@ -5709,7 +5675,7 @@ void UCharacterMovementComponent::PhysNavWalking(float deltaTime, int32 Iteratio
 	devCode(ensureMsgf(!Velocity.ContainsNaN(), TEXT("PhysNavWalking: Velocity contains NaN before CalcVelocity (%s)\n%s"), *GetPathNameSafe(this), *Velocity.ToString()));
 
 	//bound acceleration
-	Acceleration.Z = 0.f;
+	Acceleration = ProjectToGravityFloor(Acceleration);
 	if (!HasAnimRootMotion() && !CurrentRootMotion.HasOverrideVelocity())
 	{
 		CalcVelocity(deltaTime, GroundFriction, false, GetMaxBrakingDeceleration());
@@ -5727,8 +5693,7 @@ void UCharacterMovementComponent::PhysNavWalking(float deltaTime, int32 Iteratio
 
 	Iterations++;
 
-	FVector DesiredMove = Velocity;
-	DesiredMove.Z = 0.f;
+	const FVector DesiredMove = ProjectToGravityFloor(Velocity);
 
 	const FVector OldLocation = GetActorFeetLocation();
 	const FVector DeltaMove = DesiredMove * deltaTime;
@@ -5742,11 +5707,11 @@ void UCharacterMovementComponent::PhysNavWalking(float deltaTime, int32 Iteratio
 	{
 		if (bProjectNavMeshWalking)
 		{
-			const float DistSq2D = (OldLocation - CachedNavLocation.Location).SizeSquared2D();
-			const float DistZ = FMath::Abs(OldLocation.Z - CachedNavLocation.Location.Z);
+			const float DistSq2D = ProjectToGravityFloor(OldLocation - CachedNavLocation.Location).SizeSquared();
+			const float DistZ = FMath::Abs(GetGravitySpaceZ(OldLocation - CachedNavLocation.Location));
 
 			const float TotalCapsuleHeight = CharacterOwner->GetCapsuleComponent()->GetScaledCapsuleHalfHeight() * 2.0f;
-			const float ProjectionScale = (OldLocation.Z > CachedNavLocation.Location.Z) ? NavMeshProjectionHeightScaleUp : NavMeshProjectionHeightScaleDown;
+			const float ProjectionScale = (GetGravitySpaceZ(OldLocation) > GetGravitySpaceZ(CachedNavLocation.Location)) ? NavMeshProjectionHeightScaleUp : NavMeshProjectionHeightScaleDown;
 			const float DistZThr = TotalCapsuleHeight * FMath::Max(0.f, ProjectionScale);
 
 			bSameNavLocation = (DistSq2D <= UE_KINDA_SMALL_NUMBER) && (DistZ < DistZThr);
@@ -5783,7 +5748,7 @@ void UCharacterMovementComponent::PhysNavWalking(float deltaTime, int32 Iteratio
 		// we'll follow that geometry's plane out of range of valid navigation.
 		if (bSameNavLocation && bProjectNavMeshWalking)
 		{
-			AdjustedDest.Z = CachedNavLocation.Location.Z;
+			SetGravitySpaceZ(AdjustedDest, GetGravitySpaceZ(CachedNavLocation.Location));
 		}
 
 		// Find the point on the NavMesh
@@ -5799,7 +5764,7 @@ void UCharacterMovementComponent::PhysNavWalking(float deltaTime, int32 Iteratio
 
 	if (DestNavLocation.NodeRef != INVALID_NAVNODEREF)
 	{
-		FVector NewLocation(AdjustedDest.X, AdjustedDest.Y, DestNavLocation.Location.Z);
+		FVector NewLocation = ProjectToGravityFloor(AdjustedDest) + GetGravitySpaceComponentZ(DestNavLocation.Location);
 		if (bProjectNavMeshWalking)
 		{
 			SCOPE_CYCLE_COUNTER(STAT_CharNavProjectLocation);
@@ -5859,8 +5824,8 @@ FVector UCharacterMovementComponent::ProjectLocationFromNavMesh(float DeltaSecon
 		return NewLocation;
 	}
 
-	const FVector TraceStart = FVector(TargetNavLocation.X, TargetNavLocation.Y, TargetNavLocation.Z + UpOffset);
-	const FVector TraceEnd   = FVector(TargetNavLocation.X, TargetNavLocation.Y, TargetNavLocation.Z - DownOffset);
+	const FVector TraceStart = TargetNavLocation + UpOffset * -GetGravityDirection();
+	const FVector TraceEnd   = TargetNavLocation + DownOffset * GetGravityDirection();
 
 	// We can skip this trace if we are checking at the same location as the last trace (ie, we haven't moved).
 	const bool bCachedLocationStillValid = (CachedProjectedNavMeshHitResult.bBlockingHit &&
@@ -5906,31 +5871,41 @@ FVector UCharacterMovementComponent::ProjectLocationFromNavMesh(float DeltaSecon
 	// Project to last plane we found.
 	if (CachedProjectedNavMeshHitResult.bBlockingHit)
 	{
-		if (bCachedLocationStillValid && FMath::IsNearlyEqual(CurrentFeetLocation.Z, CachedProjectedNavMeshHitResult.ImpactPoint.Z, (FVector::FReal)0.01f))
+		if (bCachedLocationStillValid && FMath::IsNearlyEqual(GetGravitySpaceZ(CurrentFeetLocation), GetGravitySpaceZ(CachedProjectedNavMeshHitResult.ImpactPoint), (FVector::FReal)0.01f))
 		{
 			// Already at destination.
-			NewLocation.Z = CurrentFeetLocation.Z;
+			SetGravitySpaceZ(NewLocation, GetGravitySpaceZ(CurrentFeetLocation));
 		}
 		else
 		{
-			//const FVector ProjectedPoint = FMath::LinePlaneIntersection(TraceStart, TraceEnd, CachedProjectedNavMeshHitResult.ImpactPoint, CachedProjectedNavMeshHitResult.Normal);
-			//float ProjectedZ = ProjectedPoint.Z;
-
-			// Optimized assuming we only care about Z coordinate of result.
-			const FVector& PlaneOrigin = CachedProjectedNavMeshHitResult.ImpactPoint;
-			const FVector& PlaneNormal = CachedProjectedNavMeshHitResult.Normal;
-			FVector::FReal ProjectedZ = TraceStart.Z + ZOffset * (((PlaneOrigin - TraceStart)|PlaneNormal) / (ZOffset * PlaneNormal.Z));
-
+			FVector::FReal ProjectedZ;
+			if (HasCustomGravity())
+			{
+				const FVector ProjectedPoint = FMath::LinePlaneIntersection(TraceStart, TraceEnd, CachedProjectedNavMeshHitResult.ImpactPoint, CachedProjectedNavMeshHitResult.Normal);
+				ProjectedZ = GetGravitySpaceZ(ProjectedPoint);
+			}
+			else
+			{
+				// Optimized assuming we only care about Z coordinate of result.
+				const FVector& PlaneOrigin = CachedProjectedNavMeshHitResult.ImpactPoint;
+				const FVector& PlaneNormal = CachedProjectedNavMeshHitResult.Normal;
+				ProjectedZ = TraceStart.Z + ZOffset * (((PlaneOrigin - TraceStart)|PlaneNormal) / (ZOffset * PlaneNormal.Z));
+			}
+			
 			// Limit to not be too far above or below NavMesh location
-			ProjectedZ = FMath::Clamp(ProjectedZ, TraceEnd.Z, TraceStart.Z);
+			const FVector::FReal TraceStartZ = GetGravitySpaceZ(TraceStart);
+			const FVector::FReal TraceEndZ = GetGravitySpaceZ(TraceEnd);
+			const FVector::FReal TraceMinZ = FMath::Min(TraceStartZ, TraceEndZ);
+			const FVector::FReal TraceMaxZ = FMath::Max(TraceStartZ, TraceEndZ);
+			ProjectedZ = FMath::Clamp(ProjectedZ, TraceMinZ, TraceMaxZ);
 
 			// Interp for smoother updates (less "pop" when trace hits something new). 0 interp speed is instant.
 			const FVector::FReal InterpSpeed = FMath::Max<FVector::FReal>(0.f, NavMeshProjectionInterpSpeed);
-			ProjectedZ = FMath::FInterpTo(CurrentFeetLocation.Z, ProjectedZ, (FVector::FReal)DeltaSeconds, InterpSpeed);
-			ProjectedZ = FMath::Clamp(ProjectedZ, TraceEnd.Z, TraceStart.Z);
-
+			ProjectedZ = FMath::FInterpTo(GetGravitySpaceZ(CurrentFeetLocation), ProjectedZ, (FVector::FReal)DeltaSeconds, InterpSpeed);
+			ProjectedZ = FMath::Clamp(ProjectedZ, TraceMinZ, TraceMaxZ);
+			
 			// Final result
-			NewLocation.Z = ProjectedZ;
+			SetGravitySpaceZ(NewLocation, ProjectedZ);
 		}
 	}
 
@@ -6065,10 +6040,10 @@ void UCharacterMovementComponent::AdjustFloorHeight()
 	if (OldFloorDist < MIN_FLOOR_DIST || OldFloorDist > MAX_FLOOR_DIST)
 	{
 		FHitResult AdjustHit(1.f);
-		const double InitialZ = RotateWorldToGravity(UpdatedComponent->GetComponentLocation()).Z;
+		const double InitialZ = GetGravitySpaceZ(UpdatedComponent->GetComponentLocation());
 		const float AvgFloorDist = (MIN_FLOOR_DIST + MAX_FLOOR_DIST) * 0.5f;
 		const float MoveDist = AvgFloorDist - OldFloorDist;
-		SafeMoveUpdatedComponent(RotateGravityToWorld(FVector(0.f,0.f,MoveDist)), UpdatedComponent->GetComponentQuat(), true, AdjustHit );
+		SafeMoveUpdatedComponent(MoveDist * -GetGravityDirection(), UpdatedComponent->GetComponentQuat(), true, AdjustHit );
 		UE_LOG(LogCharacterMovement, VeryVerbose, TEXT("Adjust floor height %.3f (Hit = %d)"), MoveDist, AdjustHit.bBlockingHit);
 
 		if (!AdjustHit.IsValidBlockingHit())
@@ -6077,13 +6052,13 @@ void UCharacterMovementComponent::AdjustFloorHeight()
 		}
 		else if (MoveDist > 0.f)
 		{
-			const double CurrentZ = RotateWorldToGravity(UpdatedComponent->GetComponentLocation()).Z;
+			const double CurrentZ = GetGravitySpaceZ(UpdatedComponent->GetComponentLocation());
 			CurrentFloor.FloorDist += CurrentZ - InitialZ;
 		}
 		else
 		{
 			checkSlow(MoveDist < 0.f);
-			CurrentFloor.FloorDist = RotateWorldToGravity(UpdatedComponent->GetComponentLocation() - AdjustHit.Location).Z;
+			CurrentFloor.FloorDist = GetGravitySpaceZ(UpdatedComponent->GetComponentLocation() - AdjustHit.Location);
 			if (IsWalkable(AdjustHit))
 			{
 				CurrentFloor.SetFromSweep(AdjustHit, CurrentFloor.FloorDist, true);
@@ -6292,7 +6267,7 @@ void UCharacterMovementComponent::OnTeleported()
 	UPrimitiveComponent* OldBase = CharacterOwner->GetMovementBase();
 	UPrimitiveComponent* NewBase = NULL;
 	
-	if (OldBase && CurrentFloor.IsWalkableFloor() && CurrentFloor.FloorDist <= MAX_FLOOR_DIST && Velocity.Z <= 0.f)
+	if (OldBase && CurrentFloor.IsWalkableFloor() && CurrentFloor.FloorDist <= MAX_FLOOR_DIST && GetGravitySpaceZ(Velocity) <= 0.f)
 	{
 		// Close enough to land or just keep walking.
 		NewBase = CurrentFloor.HitResult.Component.Get();
@@ -6533,11 +6508,11 @@ void UCharacterMovementComponent::PhysicsVolumeChanged( APhysicsVolume* NewVolum
 		SetMovementMode(MOVE_Falling);
 		FVector JumpDir(0.f);
 		FVector WallNormal(0.f);
-		if ( Acceleration.Z > 0.f && ShouldJumpOutOfWater(JumpDir)
+		if ( GetGravitySpaceZ(Acceleration) > 0.f && ShouldJumpOutOfWater(JumpDir)
 			&& ((JumpDir | Acceleration) > 0.f) && CheckWaterJump(JumpDir, WallNormal) ) 
 		{
 			JumpOutOfWater(WallNormal);
-			Velocity.Z = OutofWaterZ; //set here so physics uses this for remainder of tick
+			SetGravitySpaceZ(Velocity, OutofWaterZ); //set here so physics uses this for remainder of tick
 		}
 	}
 }
@@ -6549,7 +6524,7 @@ bool UCharacterMovementComponent::ShouldJumpOutOfWater(FVector& JumpDir)
 	if (OwnerController)
 	{
 		const FRotator ControllerRot = OwnerController->GetControlRotation();
-		if ( (Velocity.Z > 0.0f) && (ControllerRot.Pitch > JumpOutOfWaterPitch) )
+		if ((GetGravitySpaceZ(Velocity) > 0.0f) && (ControllerRot.Pitch > JumpOutOfWaterPitch))
 		{
 			// if Pawn is going up and looking up, then make it jump
 			JumpDir = ControllerRot.Vector();
@@ -6571,26 +6546,25 @@ bool UCharacterMovementComponent::CheckWaterJump(FVector CheckPoint, FVector& Wa
 		return false;
 	}
 	// check if there is a wall directly in front of the swimming pawn
-	CheckPoint.Z = 0.f;
+	CheckPoint = ProjectToGravityFloor(CheckPoint);
 	FVector CheckNorm = CheckPoint.GetSafeNormal();
 	float PawnCapsuleRadius, PawnCapsuleHalfHeight;
 	CharacterOwner->GetCapsuleComponent()->GetScaledCapsuleSize(PawnCapsuleRadius, PawnCapsuleHalfHeight);
 	CheckPoint = UpdatedComponent->GetComponentLocation() + 1.2f * PawnCapsuleRadius * CheckNorm;
-	FVector Extent(PawnCapsuleRadius, PawnCapsuleRadius, PawnCapsuleHalfHeight);
 	FHitResult HitInfo(1.f);
 	FCollisionQueryParams CapsuleParams(SCENE_QUERY_STAT(CheckWaterJump), false, CharacterOwner);
 	FCollisionResponseParams ResponseParam;
 	InitCollisionParams(CapsuleParams, ResponseParam);
 	FCollisionShape CapsuleShape = GetPawnCapsuleCollisionShape(SHRINK_None);
 	const ECollisionChannel CollisionChannel = UpdatedComponent->GetCollisionObjectType();
-	bool bHit = GetWorld()->SweepSingleByChannel( HitInfo, UpdatedComponent->GetComponentLocation(), CheckPoint, FQuat::Identity, CollisionChannel, CapsuleShape, CapsuleParams, ResponseParam);
+	bool bHit = GetWorld()->SweepSingleByChannel( HitInfo, UpdatedComponent->GetComponentLocation(), CheckPoint, GetWorldToGravityTransform(), CollisionChannel, CapsuleShape, CapsuleParams, ResponseParam);
 	
 	if ( bHit && !HitInfo.HitObjectHandle.DoesRepresentClass(APawn::StaticClass()) )
 	{
 		// hit a wall - check if it is low enough
 		WallNormal = -1.f * HitInfo.ImpactNormal;
 		FVector Start = UpdatedComponent->GetComponentLocation();
-		Start.Z += MaxOutOfWaterStepHeight;
+		Start += MaxOutOfWaterStepHeight * -GetGravityDirection();
 		CheckPoint = Start + 3.2f * PawnCapsuleRadius * WallNormal;
 		FCollisionQueryParams LineParams(SCENE_QUERY_STAT(CheckWaterJump), true, CharacterOwner);
 		FCollisionResponseParams LineResponseParam;
@@ -6682,7 +6656,7 @@ void UCharacterMovementComponent::MoveSmooth(const FVector& InVelocity, const fl
 				{
 					OutStepDownResult = NULL; // No need for a floor when not walking.
 					bool bShouldAttemptStepUp = false;
-					bShouldAttemptStepUp = FMath::Abs(RotateWorldToGravity(Hit.ImpactNormal).Z) < 0.2;
+					bShouldAttemptStepUp = FMath::Abs(GetGravitySpaceZ(Hit.ImpactNormal)) < 0.2;
 					if (bShouldAttemptStepUp)
 					{
 						const FVector GravDir = GetGravityDirection();
@@ -6722,8 +6696,8 @@ bool UCharacterMovementComponent::IsWalkable(const FHitResult& Hit) const
 	}
 
 	// Never walk up vertical surfaces.
-	const FVector GravityRelativeImpactNormal = RotateWorldToGravity(Hit.ImpactNormal);
-	if (GravityRelativeImpactNormal.Z < UE_KINDA_SMALL_NUMBER)
+	const FVector::FReal ImpactNormalZ = GetGravitySpaceZ(Hit.ImpactNormal);
+	if (ImpactNormalZ < UE_KINDA_SMALL_NUMBER)
 	{
 		return false;
 	}
@@ -6739,7 +6713,7 @@ bool UCharacterMovementComponent::IsWalkable(const FHitResult& Hit) const
 	}
 
 	// Can't walk on this surface if it is too steep.
-	if (GravityRelativeImpactNormal.Z < TestWalkableZ)
+	if (ImpactNormalZ < TestWalkableZ)
 	{
 		return false;
 	}
@@ -6773,8 +6747,7 @@ float UCharacterMovementComponent::K2_GetWalkableFloorZ() const
 
 bool UCharacterMovementComponent::IsWithinEdgeTolerance(const FVector& CapsuleLocation, const FVector& TestImpactPoint, const float CapsuleRadius) const
 {
-	const FVector GravityRelativeToTestImpactPoint = RotateWorldToGravity(TestImpactPoint - CapsuleLocation);
-	const float DistFromCenterSq = GravityRelativeToTestImpactPoint.SizeSquared2D();
+	const float DistFromCenterSq = ProjectToGravityFloor(TestImpactPoint - CapsuleLocation).SizeSquared();
 	const float ReducedRadiusSq = FMath::Square(FMath::Max(SWEEP_EDGE_REJECT_DISTANCE + UE_KINDA_SMALL_NUMBER, CapsuleRadius - SWEEP_EDGE_REJECT_DISTANCE));
 	return DistFromCenterSq < ReducedRadiusSq;
 }
@@ -6792,8 +6765,8 @@ void UCharacterMovementComponent::ComputeFloorDist(const FVector& CapsuleLocatio
 	if (DownwardSweepResult != NULL && DownwardSweepResult->IsValidBlockingHit())
 	{
 		// Only if the supplied sweep was vertical and downward.
-		const bool bIsDownward = RotateWorldToGravity(DownwardSweepResult->TraceStart - DownwardSweepResult->TraceEnd).Z > 0;
-		const bool bIsVertical = RotateWorldToGravity(DownwardSweepResult->TraceStart - DownwardSweepResult->TraceEnd).SizeSquared2D() <= UE_KINDA_SMALL_NUMBER;
+		const bool bIsDownward = GetGravitySpaceZ(DownwardSweepResult->TraceStart - DownwardSweepResult->TraceEnd) > 0;
+		const bool bIsVertical = ProjectToGravityFloor(DownwardSweepResult->TraceStart - DownwardSweepResult->TraceEnd).SizeSquared() <= UE_KINDA_SMALL_NUMBER;
 		if (bIsDownward && bIsVertical)
 		{
 			// Reject hits that are barely on the cusp of the radius of the capsule
@@ -6803,7 +6776,7 @@ void UCharacterMovementComponent::ComputeFloorDist(const FVector& CapsuleLocatio
 				bSkipSweep = true;
 
 				const bool bIsWalkable = IsWalkable(*DownwardSweepResult);
-				const float FloorDist = RotateWorldToGravity(CapsuleLocation - DownwardSweepResult->Location).Z;
+				const float FloorDist = GetGravitySpaceZ(CapsuleLocation - DownwardSweepResult->Location);
 				OutFloorResult.SetFromSweep(*DownwardSweepResult, FloorDist, bIsWalkable);
 				
 				if (bIsWalkable)
@@ -6847,7 +6820,7 @@ void UCharacterMovementComponent::ComputeFloorDist(const FVector& CapsuleLocatio
 		FCollisionShape CapsuleShape = FCollisionShape::MakeCapsule(SweepRadius, PawnHalfHeight - ShrinkHeight);
 
 		FHitResult Hit(1.f);
-		bBlockingHit = FloorSweepTest(Hit, CapsuleLocation, CapsuleLocation + RotateGravityToWorld(FVector(0.f,0.f,-TraceDist)), CollisionChannel, CapsuleShape, QueryParams, ResponseParam);
+		bBlockingHit = FloorSweepTest(Hit, CapsuleLocation, CapsuleLocation + TraceDist * GetGravityDirection(), CollisionChannel, CapsuleShape, QueryParams, ResponseParam);
 
 		if (bBlockingHit)
 		{
@@ -6865,7 +6838,7 @@ void UCharacterMovementComponent::ComputeFloorDist(const FVector& CapsuleLocatio
 					CapsuleShape.Capsule.HalfHeight = FMath::Max(PawnHalfHeight - ShrinkHeight, CapsuleShape.Capsule.Radius);
 					Hit.Reset(1.f, false);
 
-					bBlockingHit = FloorSweepTest(Hit, CapsuleLocation, CapsuleLocation + RotateGravityToWorld(FVector(0.f,0.f,-TraceDist)), CollisionChannel, CapsuleShape, QueryParams, ResponseParam);
+					bBlockingHit = FloorSweepTest(Hit, CapsuleLocation, CapsuleLocation + TraceDist * GetGravityDirection(), CollisionChannel, CapsuleShape, QueryParams, ResponseParam);
 				}
 			}
 
@@ -6901,7 +6874,7 @@ void UCharacterMovementComponent::ComputeFloorDist(const FVector& CapsuleLocatio
 		const float ShrinkHeight = PawnHalfHeight;
 		const FVector LineTraceStart = CapsuleLocation;	
 		const float TraceDist = LineDistance + ShrinkHeight;
-		const FVector Down = RotateGravityToWorld(FVector(0.f, 0.f, -TraceDist));
+		const FVector Down = TraceDist * GetGravityDirection(); 
 		QueryParams.TraceTag = SCENE_QUERY_STAT_NAME_ONLY(FloorLineTrace);
 
 		FHitResult Hit(1.f);
@@ -7080,7 +7053,7 @@ bool UCharacterMovementComponent::FloorSweepTest(
 		const FCollisionShape BoxShape = FCollisionShape::MakeBox(FVector(CapsuleRadius * 0.707f, CapsuleRadius * 0.707f, CapsuleHeight));
 
 		// First test with the box rotated so the corners are along the major axes (ie rotated 45 degrees).
-		bBlockingHit = GetWorld()->SweepSingleByChannel(OutHit, Start, End, FQuat(RotateGravityToWorld(FVector(0.f, 0.f, -1.f)), UE_PI * 0.25f), TraceChannel, BoxShape, Params, ResponseParam);
+		bBlockingHit = GetWorld()->SweepSingleByChannel(OutHit, Start, End, FQuat(GetGravityDirection(), UE_PI * 0.25f), TraceChannel, BoxShape, Params, ResponseParam);
 
 		if (!bBlockingHit)
 		{
@@ -7114,10 +7087,8 @@ bool UCharacterMovementComponent::IsValidLandingSpot(const FVector& CapsuleLocat
 		CharacterOwner->GetCapsuleComponent()->GetScaledCapsuleSize(PawnRadius, PawnHalfHeight);
 
 		// Reject hits that are above our lower hemisphere (can happen when sliding down a vertical surface).
-		const FVector GravityRelativeHitLocation = RotateWorldToGravity(Hit.Location);
-		const FVector GravityRelativeHitImpactPoint = RotateWorldToGravity(Hit.ImpactPoint);
-		const float LowerHemisphereZ = GravityRelativeHitLocation.Z - PawnHalfHeight + PawnRadius;
-		if (GravityRelativeHitImpactPoint.Z >= LowerHemisphereZ)
+		const float LowerHemisphereZ = GetGravitySpaceZ(Hit.Location) - PawnHalfHeight + PawnRadius;
+		if (GetGravitySpaceZ(Hit.ImpactPoint) >= LowerHemisphereZ)
 		{
 			return false;
 		}
@@ -7131,7 +7102,7 @@ bool UCharacterMovementComponent::IsValidLandingSpot(const FVector& CapsuleLocat
 	else
 	{
 		// Penetrating
-		if (Hit.Normal.Z < UE_KINDA_SMALL_NUMBER)
+		if (GetGravitySpaceZ(Hit.Normal) < UE_KINDA_SMALL_NUMBER)
 		{
 			// Normal is nearly horizontal or downward, that's a penetration adjustment next to a vertical or overhanging wall. Don't pop to the floor.
 			return false;
@@ -7154,7 +7125,7 @@ bool UCharacterMovementComponent::ShouldCheckForValidLandingSpot(float DeltaTime
 {
 	// See if we hit an edge of a surface on the lower portion of the capsule.
 	// In this case the normal will not equal the impact normal, and a downward sweep may find a walkable surface on top of the edge.
-	if (Hit.Normal.Z > UE_KINDA_SMALL_NUMBER && !Hit.Normal.Equals(Hit.ImpactNormal))
+	if (GetGravitySpaceZ(Hit.Normal) > UE_KINDA_SMALL_NUMBER && !Hit.Normal.Equals(Hit.ImpactNormal))
 	{
 		const FVector PawnLocation = UpdatedComponent->GetComponentLocation();
 		if (IsWithinEdgeTolerance(PawnLocation, Hit.ImpactPoint, CharacterOwner->GetCapsuleComponent()->GetScaledCapsuleRadius()))
@@ -7200,7 +7171,7 @@ bool UCharacterMovementComponent::ShouldComputePerchResult(const FHitResult& InH
 
 	if (bCheckRadius)
 	{
-		const float DistFromCenterSq = (InHit.ImpactPoint - InHit.Location).SizeSquared2D();
+		const float DistFromCenterSq = ProjectToGravityFloor(InHit.ImpactPoint - InHit.Location).SizeSquared();
 		const float StandOnEdgeRadius = GetValidPerchRadius();
 		if (DistFromCenterSq <= FMath::Square(StandOnEdgeRadius))
 		{
@@ -7225,7 +7196,7 @@ bool UCharacterMovementComponent::ComputePerchResult(const float TestRadius, con
 	CharacterOwner->GetCapsuleComponent()->GetScaledCapsuleSize(PawnRadius, PawnHalfHeight);
 	const FVector CapsuleLocation = (bUseFlatBaseForFloorChecks ? InHit.TraceStart : InHit.Location);
 
-	const float InHitAboveBase = FMath::Max<float>(0.f, RotateWorldToGravity(InHit.ImpactPoint - CapsuleLocation).Z + PawnHalfHeight);
+	const float InHitAboveBase = FMath::Max<float>(0.f, GetGravitySpaceZ(InHit.ImpactPoint - CapsuleLocation) + PawnHalfHeight);
 	const float PerchLineDist = FMath::Max(0.f, InMaxFloorDist - InHitAboveBase);
 	const float PerchSweepDist = FMath::Max(0.f, InMaxFloorDist);
 
@@ -7297,8 +7268,9 @@ bool UCharacterMovementComponent::StepUp(const FVector& GravDir, const FVector&
 	CharacterOwner->GetCapsuleComponent()->GetScaledCapsuleSize(PawnRadius, PawnHalfHeight);
 
 	// Don't bother stepping up if top of capsule is hitting something.
-	const float InitialImpactZ = InHit.ImpactPoint.Z;
-	if (InitialImpactZ > OldLocation.Z + (PawnHalfHeight - PawnRadius))
+	const float InitialImpactZ = InHit.ImpactPoint | -GravDir;
+	const float OldLocationZ = OldLocation | -GravDir;
+	if (InitialImpactZ > OldLocationZ + (PawnHalfHeight - PawnRadius))
 	{
 		return false;
 	}
@@ -7314,7 +7286,7 @@ bool UCharacterMovementComponent::StepUp(const FVector& GravDir, const FVector&
 	float StepTravelUpHeight = MaxStepHeight;
 	float StepTravelDownHeight = StepTravelUpHeight;
 	const float StepSideZ = -1.f * FVector::DotProduct(InHit.ImpactNormal, GravDir);
-	float PawnInitialFloorBaseZ = OldLocation.Z - PawnHalfHeight;
+	float PawnInitialFloorBaseZ = OldLocationZ - PawnHalfHeight;
 	float PawnFloorPointZ = PawnInitialFloorBaseZ;
 
 	if (IsMovingOnGround() && CurrentFloor.IsWalkableFloor())
@@ -7328,7 +7300,7 @@ bool UCharacterMovementComponent::StepUp(const FVector& GravDir, const FVector&
 		const bool bHitVerticalFace = !IsWithinEdgeTolerance(InHit.Location, InHit.ImpactPoint, PawnRadius);
 		if (!CurrentFloor.bLineTrace && !bHitVerticalFace)
 		{
-			PawnFloorPointZ = CurrentFloor.HitResult.ImpactPoint.Z;
+			PawnFloorPointZ = CurrentFloor.HitResult.ImpactPoint | -GravDir;
 		}
 		else
 		{
@@ -7419,7 +7391,7 @@ bool UCharacterMovementComponent::StepUp(const FVector& GravDir, const FVector&
 	if (Hit.IsValidBlockingHit())
 	{	
 		// See if this step sequence would have allowed us to travel higher than our max step height allows.
-		const float DeltaZ = Hit.ImpactPoint.Z - PawnFloorPointZ;
+		const float DeltaZ = (Hit.ImpactPoint | -GravDir) - PawnFloorPointZ;
 		if (DeltaZ > MaxStepHeight)
 		{
 			//UE_LOG(LogCharacterMovement, VeryVerbose, TEXT("- Reject StepUp (too high Height %.3f) up from floor base %f to %f"), DeltaZ, PawnInitialFloorBaseZ, NewLocation.Z);
@@ -7441,7 +7413,7 @@ bool UCharacterMovementComponent::StepUp(const FVector& GravDir, const FVector&
 
 			// Also reject if we would end up being higher than our starting location by stepping down.
 			// It's fine to step down onto an unwalkable normal below us, we will just slide off. Rejecting those moves would prevent us from being able to walk off the edge.
-			if (Hit.Location.Z > OldLocation.Z)
+			if ((Hit.Location | -GravDir) > OldLocationZ)
 			{
 				//UE_LOG(LogCharacterMovement, VeryVerbose, TEXT("- Reject StepUp (unwalkable normal %s above old position)"), *Hit.ImpactNormal.ToString());
 				ScopedStepUpMovement.RevertMove();
@@ -7472,7 +7444,7 @@ bool UCharacterMovementComponent::StepUp(const FVector& GravDir, const FVector&
 
 			// Reject unwalkable normals if we end up higher than our initial height.
 			// It's fine to walk down onto an unwalkable surface, don't reject those moves.
-			if (Hit.Location.Z > OldLocation.Z)
+			if ((Hit.Location | -GravDir) > OldLocationZ)
 			{
 				// We should reject the floor result if we are trying to step up an actual step where we are not able to perch (this is rare).
 				// In those cases we should instead abort the step up and try to slide along the stair.
@@ -7553,7 +7525,7 @@ void UCharacterMovementComponent::ApplyImpactPhysicsForces(const FHitResult& Imp
 
 					if (!Extents.IsNearlyZero())
 					{
-						ForcePoint.Z = Center.Z + Extents.Z * PushForcePointZOffsetFactor;
+						SetGravitySpaceZ(ForcePoint, GetGravitySpaceZ(Center) + GetGravitySpaceZ(Extents) * PushForcePointZOffsetFactor);
 					}
 				}
 
@@ -7702,12 +7674,12 @@ float UCharacterMovementComponent::VisualizeMovement() const
 	}
 
 #if !(UE_BUILD_SHIPPING || UE_BUILD_TEST)
-	const FVector TopOfCapsule = GetActorLocation() + FVector(0.f, 0.f, CharacterOwner->GetSimpleCollisionHalfHeight());
+	const FVector TopOfCapsule = GetActorLocation() + CharacterOwner->GetSimpleCollisionHalfHeight() * -GetGravityDirection();
 	
 	// Position
 	{
 		const FColor DebugColor = FColor::White;
-		const FVector DebugLocation = TopOfCapsule + FVector(0.f,0.f,HeightOffset);
+		const FVector DebugLocation = TopOfCapsule + HeightOffset * -GetGravityDirection();
 		FString DebugText = FString::Printf(TEXT("Position: %s"), *GetActorLocation().ToCompactString());
 		DrawDebugString(GetWorld(), DebugLocation, DebugText, nullptr, DebugColor, 0.f, true);
 	}
@@ -7716,8 +7688,8 @@ float UCharacterMovementComponent::VisualizeMovement() const
 	{
 		const FColor DebugColor = FColor::Green;
 		HeightOffset += OffsetPerElement;
-		const FVector DebugLocation = TopOfCapsule + FVector(0.f,0.f,HeightOffset);
-		DrawDebugDirectionalArrow(GetWorld(), DebugLocation - FVector(0.f, 0.f, 5.0f), DebugLocation - FVector(0.f, 0.f, 5.0f) + Velocity,
+		const FVector DebugLocation = TopOfCapsule + HeightOffset * -GetGravityDirection();
+		DrawDebugDirectionalArrow(GetWorld(), DebugLocation - 5.f * -GetGravityDirection(), DebugLocation - 5.f * -GetGravityDirection() + Velocity,
 			100.f, DebugColor, false, -1.f, (uint8)'\000', 10.f);
 
 		FString DebugText = FString::Printf(TEXT("Velocity: %s (Speed: %.2f) (Max: %.2f)"), *Velocity.ToCompactString(), Velocity.Size(), GetMaxSpeed());
@@ -7731,9 +7703,9 @@ float UCharacterMovementComponent::VisualizeMovement() const
 		const float MaxAccelerationLineLength = 200.f;
 		const float CurrentMaxAccel = GetMaxAcceleration();
 		const float CurrentAccelAsPercentOfMaxAccel = CurrentMaxAccel > 0.f ? Acceleration.Size() / CurrentMaxAccel : 1.f;
-		const FVector DebugLocation = TopOfCapsule + FVector(0.f,0.f,HeightOffset);
-		DrawDebugDirectionalArrow(GetWorld(), DebugLocation - FVector(0.f, 0.f, 5.0f), 
-			DebugLocation - FVector(0.f, 0.f, 5.0f) + Acceleration.GetSafeNormal(UE_SMALL_NUMBER) * CurrentAccelAsPercentOfMaxAccel * MaxAccelerationLineLength,
+		const FVector DebugLocation = TopOfCapsule + HeightOffset * -GetGravityDirection();
+		DrawDebugDirectionalArrow(GetWorld(), DebugLocation - 5.f * -GetGravityDirection(),
+			DebugLocation - 5.f * -GetGravityDirection() + Acceleration.GetSafeNormal(UE_SMALL_NUMBER) * CurrentAccelAsPercentOfMaxAccel * MaxAccelerationLineLength,
 			25.f, DebugColor, false, -1.f, (uint8)'\000', 8.f);
 
 		FString DebugText = FString::Printf(TEXT("Acceleration: %s"), *Acceleration.ToCompactString());
@@ -7744,14 +7716,14 @@ float UCharacterMovementComponent::VisualizeMovement() const
 	{
 		const FColor DebugColor = FColor::Blue;
 		HeightOffset += OffsetPerElement;
-		FVector DebugLocation = TopOfCapsule + FVector(0.f,0.f,HeightOffset);
+		FVector DebugLocation = TopOfCapsule + HeightOffset * -GetGravityDirection();
 		FString DebugText = FString::Printf(TEXT("MovementMode: %s"), *GetMovementName());
 		DrawDebugString(GetWorld(), DebugLocation, DebugText, nullptr, DebugColor, 0.f, true);
 
 		if (IsInWater())
 		{
 			HeightOffset += OffsetPerElement;
-			DebugLocation = TopOfCapsule + FVector(0.f, 0.f, HeightOffset);
+			DebugLocation = TopOfCapsule + HeightOffset * -GetGravityDirection();
 			DebugText = FString::Printf(TEXT("ImmersionDepth: %.2f"), ImmersionDepth());
 			DrawDebugString(GetWorld(), DebugLocation, DebugText, nullptr, DebugColor, 0.f, true);
 		}
@@ -7761,7 +7733,7 @@ float UCharacterMovementComponent::VisualizeMovement() const
 	{
 		const FColor DebugColor = FColor::Blue;
 		HeightOffset += OffsetPerElement;
-		FVector DebugLocation = TopOfCapsule + FVector(0.f, 0.f, HeightOffset);
+		FVector DebugLocation = TopOfCapsule + HeightOffset * -GetGravityDirection();
 		FString DebugText = FString::Printf(TEXT("bIsJumping: %d Count: %d HoldTime: %.2f"), CharacterOwner->bPressedJump, CharacterOwner->JumpCurrentCount, CharacterOwner->JumpKeyHoldTime);
 		DrawDebugString(GetWorld(), DebugLocation, DebugText, nullptr, DebugColor, 0.f, true);
 	}
@@ -7771,7 +7743,7 @@ float UCharacterMovementComponent::VisualizeMovement() const
 	{
 		const FColor DebugColor = FColor::Cyan;
 		HeightOffset += OffsetPerElement;
-		const FVector DebugLocation = TopOfCapsule + FVector(0.f,0.f,HeightOffset);
+		const FVector DebugLocation = TopOfCapsule + HeightOffset * -GetGravityDirection();
 
 		FVector CurrentAdditiveVelocity(FVector::ZeroVector);
 		CurrentRootMotion.AccumulateAdditiveRootMotionVelocity(0.f, *CharacterOwner, *this, CurrentAdditiveVelocity);
@@ -7781,7 +7753,7 @@ float UCharacterMovementComponent::VisualizeMovement() const
 
 		FString DebugText = FString::Printf(TEXT("RootMotionAdditiveVelocity: %s (Speed: %.2f)"), 
 			*CurrentAdditiveVelocity.ToCompactString(), CurrentAdditiveVelocity.Size());
-		DrawDebugString(GetWorld(), DebugLocation + FVector(0.f,0.f,5.f), DebugText, nullptr, DebugColor, 0.f, true);
+		DrawDebugString(GetWorld(), DebugLocation + 5.f * -GetGravityDirection(), DebugText, nullptr, DebugColor, 0.f, true);
 	}
 
 	// Root motion (override)
@@ -7789,7 +7761,7 @@ float UCharacterMovementComponent::VisualizeMovement() const
 	{
 		const FColor DebugColor = FColor::Green;
 		HeightOffset += OffsetPerElement;
-		const FVector DebugLocation = TopOfCapsule + FVector(0.f,0.f,HeightOffset);
+		const FVector DebugLocation = TopOfCapsule + HeightOffset * -GetGravityDirection();
 		FString DebugText = FString::Printf(TEXT("Has Override RootMotion"));
 		DrawDebugString(GetWorld(), DebugLocation, DebugText, nullptr, DebugColor, 0.f, true);
 	}
@@ -7915,12 +7887,12 @@ void UCharacterMovementComponent::SmoothCorrection(const FVector& OldLocation, c
 
 		// The mesh doesn't move, but the capsule does so we have a new offset.
 		FVector NewToOldVector = (OldLocation - NewLocation);
-		if (bIsNavWalkingOnServer && FMath::Abs(NewToOldVector.Z) < NavWalkingFloorDistTolerance)
+		if (bIsNavWalkingOnServer && FMath::Abs(GetGravitySpaceZ(NewToOldVector)) < NavWalkingFloorDistTolerance)
 		{
 			// ignore smoothing on Z axis
 			// don't modify new location (local simulation result), since it's probably more accurate than server data
 			// and shouldn't matter as long as difference is relatively small
-			NewToOldVector.Z = 0;
+			NewToOldVector = ProjectToGravityFloor(NewToOldVector);
 		}
 
 		const float DistSq = NewToOldVector.SizeSquared();
@@ -8165,7 +8137,7 @@ void UCharacterMovementComponent::SmoothClientPosition_Interpolate(float DeltaSe
 			if ( CharacterMovementCVars::NetVisualizeSimulatedCorrections >= 1 )
 			{
 				const FColor DebugColor = FColor::White;
-				const FVector DebugLocation = CharacterOwner->GetMesh()->GetComponentLocation() + FVector( 0.f, 0.f, 130.0f ) - CharacterOwner->GetBaseTranslationOffset();
+				const FVector DebugLocation = CharacterOwner->GetMesh()->GetComponentLocation() + 130.0f * -GetGravityDirection() - CharacterOwner->GetBaseTranslationOffset();
 				FString DebugText = FString::Printf( TEXT( "Lerp: %2.2f" ), LerpPercent );
 				DrawDebugString( GetWorld(), DebugLocation, DebugText, nullptr, DebugColor, 0.f, true );
 				FString TimeText = FString::Printf( TEXT("ClientTime: %2.2f ServerTime: %2.2f" ), ClientData->SmoothingClientTimeStamp, ClientData->SmoothingServerTimeStamp);
@@ -8220,7 +8192,7 @@ void UCharacterMovementComponent::SmoothClientPosition_Interpolate(float DeltaSe
 
 		if ( CharacterMovementCVars::NetVisualizeSimulatedCorrections >= 1 )
 		{
-			const FVector DebugLocation = CharacterOwner->GetMesh()->GetComponentLocation() + FVector( 0.f, 0.f, 300.0f ) - CharacterOwner->GetBaseTranslationOffset();
+			const FVector DebugLocation = CharacterOwner->GetMesh()->GetComponentLocation() + 300.0f * -GetGravityDirection() - CharacterOwner->GetBaseTranslationOffset();
 			DrawDebugBox( GetWorld(), DebugLocation, FVector( 45, 45, 45 ), CharacterOwner->GetMesh()->GetComponentQuat(), FColor( 0, 255, 0 ) );
 
 			//DrawDebugCoordinateSystem( GetWorld(), UpdatedComponent->GetComponentLocation() + FVector( 0, 0, 300.0f ), UpdatedComponent->GetComponentRotation(), 45.0f );
@@ -9983,7 +9955,7 @@ void UCharacterMovementComponent::ServerMoveHandleClientError(float ClientTimeSt
 			{
 				const FVector LastBaseVelocity = MovementBaseUtility::GetMovementBaseVelocity(LastServerMovementBasePtr, LastServerMovementBaseBoneName);
 				RelativeVelocity = Velocity - LastBaseVelocity;
-				const FVector BaseDirection = LastBaseVelocity.GetSafeNormal2D();
+				const FVector BaseDirection = ProjectToGravityFloor(LastBaseVelocity).GetSafeNormal();
 				const FVector RelativeDirection = RelativeVelocity * (1.f / MaxWalkSpeed);
 
 				ClientForwardFactor = FMath::Clamp(FVector::DotProduct(BaseDirection, RelativeDirection), 0.f, 1.f);
@@ -10608,7 +10580,7 @@ void UCharacterMovementComponent::ClientHandleMoveResponse(const FCharacterMoveR
 					MoveResponse.RootMotionTrackPosition,
 					MoveResponse.ClientAdjustment.NewLoc,
 					MoveResponse.RootMotionRotation,
-					MoveResponse.ClientAdjustment.NewVel.Z,
+					GetGravitySpaceZ(MoveResponse.ClientAdjustment.NewVel),
 					MoveResponse.ClientAdjustment.NewBase,
 					MoveResponse.ClientAdjustment.NewBaseBoneName,
 					MoveResponse.bHasBase,
@@ -10623,7 +10595,7 @@ void UCharacterMovementComponent::ClientHandleMoveResponse(const FCharacterMoveR
 				MoveResponse.RootMotionTrackPosition,
 				MoveResponse.ClientAdjustment.NewLoc,
 				MoveResponse.RootMotionRotation,
-				MoveResponse.ClientAdjustment.NewVel.Z,
+				GetGravitySpaceZ(MoveResponse.ClientAdjustment.NewVel),
 				MoveResponse.ClientAdjustment.NewBase,
 				MoveResponse.ClientAdjustment.NewBaseBoneName,
 				MoveResponse.bHasBase,
@@ -10732,7 +10704,7 @@ void UCharacterMovementComponent::SendClientAdjustment()
 						bIsPlayingNetworkedRootMotionMontage ? CharacterOwner->GetRootMotionAnimMontageInstance()->GetPosition() : -1.f,
 						ServerData->PendingAdjustment.NewLoc,
 						CompressedRotation,
-						ServerData->PendingAdjustment.NewVel.Z,
+						GetGravitySpaceZ(ServerData->PendingAdjustment.NewVel),
 						ServerData->PendingAdjustment.NewBase,
 						ServerData->PendingAdjustment.NewBaseBoneName,
 						ServerData->PendingAdjustment.NewBase != NULL,
@@ -10750,7 +10722,7 @@ void UCharacterMovementComponent::SendClientAdjustment()
 						CharacterOwner->GetRootMotionAnimMontageInstance()->GetPosition(),
 						ServerData->PendingAdjustment.NewLoc,
 						CompressedRotation,
-						ServerData->PendingAdjustment.NewVel.Z,
+						GetGravitySpaceZ(ServerData->PendingAdjustment.NewVel),
 						ServerData->PendingAdjustment.NewBase,
 						ServerData->PendingAdjustment.NewBaseBoneName,
 						ServerData->PendingAdjustment.NewBase != NULL,
@@ -11029,7 +11001,7 @@ void UCharacterMovementComponent::ClientAdjustRootMotionPosition_Implementation(
 	}
 
 	// Call ClientAdjustPosition first. This will Ack the move if it's not outdated.
-	ClientAdjustPosition_Implementation(TimeStamp, ServerLoc, FVector(0.f, 0.f, ServerVelZ), ServerBase, ServerBaseBoneName, bHasBase, bBaseRelativePosition, ServerMovementMode);
+	ClientAdjustPosition_Implementation(TimeStamp, ServerLoc, ServerVelZ * -GetGravityDirection(), ServerBase, ServerBaseBoneName, bHasBase, bBaseRelativePosition, ServerMovementMode);
 	
 	FNetworkPredictionData_Client_Character* ClientData = GetPredictionData_Client_Character();
 	check(ClientData);
@@ -11110,7 +11082,7 @@ void UCharacterMovementComponent::ClientAdjustRootMotionSourcePosition_Implement
 #endif
 
 	// Call ClientAdjustPosition first. This will Ack the move if it's not outdated.
-	ClientAdjustPosition_Implementation(TimeStamp, ServerLoc, FVector(0.f, 0.f, ServerVelZ), ServerBase, ServerBaseBoneName, bHasBase, bBaseRelativePosition, ServerMovementMode);
+	ClientAdjustPosition_Implementation(TimeStamp, ServerLoc, ServerVelZ * -GetGravityDirection(), ServerBase, ServerBaseBoneName, bHasBase, bBaseRelativePosition, ServerMovementMode);
 	
 	FNetworkPredictionData_Client_Character* ClientData = GetPredictionData_Client_Character();
 	check(ClientData);
@@ -11228,8 +11200,10 @@ void UCharacterMovementComponent::CapsuleTouched(UPrimitiveComponent* Overlapped
 	{
 		const FVector OtherLoc = OtherComp->GetComponentLocation();
 		const FVector Loc = UpdatedComponent->GetComponentLocation();
-		FVector ImpulseDir = FVector(OtherLoc.X - Loc.X, OtherLoc.Y - Loc.Y, 0.25f).GetSafeNormal();
-		ImpulseDir = (ImpulseDir + Velocity.GetSafeNormal2D()) * 0.5f;
+		
+		FVector ImpulseDir = OtherLoc - Loc;
+		SetGravitySpaceZ(ImpulseDir, 0.25f);
+		ImpulseDir = (ImpulseDir.GetSafeNormal() + ProjectToGravityFloor(Velocity).GetSafeNormal()) * 0.5f;
 		ImpulseDir.Normalize();
 
 		FName BoneName = NAME_None;
@@ -11246,7 +11220,7 @@ void UCharacterMovementComponent::CapsuleTouched(UPrimitiveComponent* Overlapped
 			TouchForceFactorModified *= BI ? BI->GetBodyMass() : 1.0f;
 		}
 
-		float ImpulseStrength = FMath::Clamp<FVector::FReal>(Velocity.Size2D() * TouchForceFactorModified, 
+		float ImpulseStrength = FMath::Clamp<FVector::FReal>(ProjectToGravityFloor(Velocity).Size() * TouchForceFactorModified, 
 			MinTouchForce > 0.0f ? MinTouchForce : -FLT_MAX, 
 			MaxTouchForce > 0.0f ? MaxTouchForce : FLT_MAX);
 
@@ -11399,7 +11373,7 @@ void UCharacterMovementComponent::ApplyRepulsionForce(float DeltaSeconds)
 				// Trace to get the hit location on the capsule
 				FHitResult Hit;
 				bool bHasHit = UpdatedPrimitive->LineTraceComponent(Hit, BodyLocation,
-																	FVector(MyLocation.X, MyLocation.Y, BodyLocation.Z),
+																	ProjectToGravityFloor(MyLocation) + GetGravitySpaceComponentZ(BodyLocation),
 																	QueryParams);
 
 				FVector HitLoc = Hit.ImpactPoint;
@@ -11412,12 +11386,12 @@ void UCharacterMovementComponent::ApplyRepulsionForce(float DeltaSeconds)
 					bIsPenetrating = true;
 				}
 
-				const float DistanceNow = (HitLoc - BodyLocation).SizeSquared2D();
-				const float DistanceLater = (HitLoc - (BodyLocation + BodyVelocity * DeltaSeconds)).SizeSquared2D();
+				const float DistanceNow = ProjectToGravityFloor(HitLoc - BodyLocation).SizeSquared();
+				const float DistanceLater = ProjectToGravityFloor(HitLoc - (BodyLocation + BodyVelocity * DeltaSeconds)).SizeSquared();
 
 				if (bHasHit && DistanceNow < StopBodyDistance && !bIsPenetrating)
 				{
-					OverlapBody->SetLinearVelocity(FVector(0.0f, 0.0f, 0.0f), false);
+					OverlapBody->SetLinearVelocity(FVector::ZeroVector, false);
 				}
 				else if (DistanceLater <= DistanceNow || bIsPenetrating)
 				{
@@ -11425,11 +11399,12 @@ void UCharacterMovementComponent::ApplyRepulsionForce(float DeltaSeconds)
 
 					if (bHasHit)
 					{
-						ForceCenter.Z = HitLoc.Z;
+						SetGravitySpaceZ(ForceCenter, GetGravitySpaceZ(HitLoc));
 					}
 					else
 					{
-						ForceCenter.Z = FMath::Clamp(BodyLocation.Z, MyLocation.Z - CapsuleHalfHeight, MyLocation.Z + CapsuleHalfHeight);
+						const FVector::FReal MyLocationZ = GetGravitySpaceZ(MyLocation);
+						SetGravitySpaceZ(ForceCenter, FMath::Clamp(GetGravitySpaceZ(BodyLocation), MyLocationZ - CapsuleHalfHeight, MyLocationZ + CapsuleHalfHeight));
 					}
 
 					OverlapBody->AddRadialForceToBody(ForceCenter, RepulsionForceRadius, RepulsionForce * Mass, ERadialImpulseFalloff::RIF_Constant);
@@ -11441,12 +11416,12 @@ void UCharacterMovementComponent::ApplyRepulsionForce(float DeltaSeconds)
 
 void UCharacterMovementComponent::ApplyAccumulatedForces(float DeltaSeconds)
 {
-	const FVector GravityRelativePendingImpulseToApply = RotateWorldToGravity(PendingImpulseToApply);
-	const FVector GravityRelativePendingForceToApply = RotateWorldToGravity(PendingForceToApply);
-	if (GravityRelativePendingImpulseToApply.Z != 0.0 || GravityRelativePendingForceToApply.Z != 0.0)
+	const FVector::FReal ImpulseToApplyZ = GetGravitySpaceZ(PendingImpulseToApply);
+	const FVector::FReal ForceToApplyZ = GetGravitySpaceZ(PendingForceToApply);
+	if (ImpulseToApplyZ != 0.0 || ForceToApplyZ != 0.0)
 	{
 		// check to see if applied momentum is enough to overcome gravity
-		if ( IsMovingOnGround() && (GravityRelativePendingImpulseToApply.Z + (GravityRelativePendingForceToApply.Z * DeltaSeconds) + (GetGravityZ() * DeltaSeconds) > UE_SMALL_NUMBER))
+		if (IsMovingOnGround() && (ImpulseToApplyZ + (ForceToApplyZ * DeltaSeconds) + (GetGravityZ() * DeltaSeconds) > UE_SMALL_NUMBER))
 		{
 			SetMovementMode(MOVE_Falling);
 		}
@@ -13480,4 +13455,3 @@ bool UCharacterMovementComponent::IsAsyncCallbackRegistered() const
 {
 	return AsyncCallback != nullptr;
 }
-
