Fix for ASAN crash issue on shutdown where UActorReplicationBridge would read from an already cleaned up reference to PackageMap. Changing to a UPROPERTY TObjectPtr fixes the issue.
#tests Tested against Dedicated Server Win64 Cooked Data
[RN] Fixed a case for a potential rare crash when shutting down ActorReplicationBridgebmittool SubmitTool-Development-4.2.0.1191505988

diff --git a/Engine/Source/Runtime/Engine/Private/Net/Iris/ReplicationSystem/ActorReplicationBridge.cpp b/Engine/Source/Runtime/Engine/Private/Net/Iris/ReplicationSystem/ActorReplicationBridge.cpp
--- a/Engine/Source/Runtime/Engine/Private/Net/Iris/ReplicationSystem/ActorReplicationBridge.cpp
+++ b/Engine/Source/Runtime/Engine/Private/Net/Iris/ReplicationSystem/ActorReplicationBridge.cpp
@@ -96,8 +96,6 @@ bool ShouldIncludeActorInLevelGroups(const AActor* Actor)
 
 UActorReplicationBridge::UActorReplicationBridge()
 : UObjectReplicationBridge()
-, NetDriver(nullptr)
-, ObjectReferencePackageMap(nullptr)
 , SpawnInfoFlags(0U)
 {
 	SetInstancePreUpdateFunction(UE::Net::Private::ActorReplicationBridgePreUpdateFunction);
@@ -151,7 +149,6 @@ void UActorReplicationBridge::Initialize(UReplicationSystem* InReplicationSystem
 	}
 
 	ObjectReferencePackageMap = NewObject<UIrisObjectReferencePackageMap>();
-	ObjectReferencePackageMap->AddToRoot();
 
 	// Get spawn info flags from cvars
 	SpawnInfoFlags = UE::Net::Private::GetActorReplicationBridgeSpawnInfoFlags();
@@ -159,12 +156,6 @@ void UActorReplicationBridge::Initialize(UReplicationSystem* InReplicationSystem
 
 UActorReplicationBridge::~UActorReplicationBridge()
 {
-	if (ObjectReferencePackageMap)
-	{
-		ObjectReferencePackageMap->RemoveFromRoot();
-		ObjectReferencePackageMap->MarkAsGarbage();
-		ObjectReferencePackageMap = nullptr;
-	}
 }
 
 void UActorReplicationBridge::Deinitialize()
@@ -172,8 +163,10 @@ void UActorReplicationBridge::Deinitialize()
 	if (NetDriver)
 	{
 		NetDriver->OnNetServerMaxTickRateChanged.RemoveAll(this);
+		NetDriver = nullptr;
 	}
 	Super::Deinitialize();
+	ObjectReferencePackageMap = nullptr;
 }
 
 UE::Net::FNetRefHandle UActorReplicationBridge::BeginReplication(AActor* Actor, const FActorBeginReplicationParams& Params)
diff --git a/Engine/Source/Runtime/Engine/Public/Net/Iris/ReplicationSystem/ActorReplicationBridge.h b/Engine/Source/Runtime/Engine/Public/Net/Iris/ReplicationSystem/ActorReplicationBridge.h
--- a/Engine/Source/Runtime/Engine/Public/Net/Iris/ReplicationSystem/ActorReplicationBridge.h
+++ b/Engine/Source/Runtime/Engine/Public/Net/Iris/ReplicationSystem/ActorReplicationBridge.h
@@ -4,15 +4,16 @@
 
 #include "Iris/ReplicationSystem/ObjectReplicationBridge.h"
 #include "Engine/EngineTypes.h"
+#include "UObject/ObjectPtr.h"
 #include "ActorReplicationBridge.generated.h"
 
+class UNetDriver;
+class UIrisObjectReferencePackageMap;
+
 #if UE_WITH_IRIS
 
 class UActorComponent;
-class UNetDriver;
 class UWorld;
-class UNetDriver;
-class UIrisObjectReferencePackageMap;
 
 namespace UE::Net
 {
@@ -126,11 +127,14 @@ private:
 	void AddActorToLevelGroup(const AActor* Actor);
 
 private:
-	UNetDriver* NetDriver;
-
-	UIrisObjectReferencePackageMap* ObjectReferencePackageMap;
 
 	uint32 SpawnInfoFlags;
 
 #endif // UE_WITH_IRIS
+
+	UNetDriver* NetDriver = nullptr;
+
+	UPROPERTY(Transient)
+	TObjectPtr<UIrisObjectReferencePackageMap> ObjectReferencePackageMap = nullptr;
+
 };
