Fix issue where page table size wasn't being shown correctly in the runtime virtual texture asset details panel.
Displayed size was being clamped at 1024 for when not using adaptive page tables. Actual used size can go to 4096.
Also changed UI so that Enable Adaptive Page Table option is no longer in the advanced section. It's better that this option is surfaced more clearly since it is more optimal than using large tile sizes of 512 or 1024.
#tests Create a new RVT asset in editor and observe the details panel fixes/changes

diff --git a/Engine/Source/Runtime/Engine/Classes/VT/RuntimeVirtualTexture.h b/Engine/Source/Runtime/Engine/Classes/VT/RuntimeVirtualTexture.h
--- a/Engine/Source/Runtime/Engine/Classes/VT/RuntimeVirtualTexture.h
+++ b/Engine/Source/Runtime/Engine/Classes/VT/RuntimeVirtualTexture.h
@@ -60,7 +60,7 @@ protected:
 	bool bPrivateSpace = true;
 
 	/** Enable sparse adaptive page tables. This supports larger tile counts but adds an indirection cost when sampling the virtual texture. It is recommended only when very large virtual resolutions are necessary. */
-	UPROPERTY(EditAnywhere, AdvancedDisplay, Category = Layout, meta = (DisplayName = "Enable adaptive page table"))
+	UPROPERTY(EditAnywhere, Category = Layout, meta = (DisplayName = "Enable adaptive page table"))
 	bool bAdaptive = false;
 
 	/** Enable continuous update of the virtual texture pages. This round-robin updates already mapped pages and can help fix pages that are mapped before dependent textures are fully streamed in.  */
diff --git a/Engine/Source/Runtime/Engine/Private/VT/RuntimeVirtualTexture.cpp b/Engine/Source/Runtime/Engine/Private/VT/RuntimeVirtualTexture.cpp
--- a/Engine/Source/Runtime/Engine/Private/VT/RuntimeVirtualTexture.cpp
+++ b/Engine/Source/Runtime/Engine/Private/VT/RuntimeVirtualTexture.cpp
@@ -299,7 +299,9 @@ int32 URuntimeVirtualTexture::GetMaxTileCountLog2(bool InAdaptive)
 
 int32 URuntimeVirtualTexture::GetPageTableSize() const
 {
-	return 1 << GetPageTableTileCountLog2(FMath::Clamp(TileCount, 0, GetMaxTileCountLog2(bAdaptive)));
+	const int32 ClampedTileCountLog2 = FMath::Clamp(TileCount, 0, GetMaxTileCountLog2(bAdaptive));
+	const int32 PageTableSizeLog2 = bAdaptive ? GetPageTableTileCountLog2(ClampedTileCountLog2) : ClampedTileCountLog2;
+	return 1 << PageTableSizeLog2;
 }
 
 void URuntimeVirtualTexture::GetProducerDescription(FVTProducerDescription& OutDesc, FInitSettings const& InitSettings, FTransform const& VolumeToWorld) const
