Add some VeryVerbose logging to UPlayerInput::InputKey. This can be useful for quick debugging.

#rnx

diff --git a/Engine/Source/Runtime/Engine/Classes/GameFramework/PlayerInput.h b/Engine/Source/Runtime/Engine/Classes/GameFramework/PlayerInput.h
--- a/Engine/Source/Runtime/Engine/Classes/GameFramework/PlayerInput.h
+++ b/Engine/Source/Runtime/Engine/Classes/GameFramework/PlayerInput.h
@@ -398,6 +398,8 @@ struct FInputKeyParams
 	FVector Get3DAxisDelta() const { return Delta; }
 };
 
+ENGINE_API DECLARE_LOG_CATEGORY_EXTERN(LogPlayerInput, Log, All);
+
 /**
  * Object within PlayerController that processes player input.
  * Only exists on the client in network games.
diff --git a/Engine/Source/Runtime/Engine/Private/UserInterface/PlayerInput.cpp b/Engine/Source/Runtime/Engine/Private/UserInterface/PlayerInput.cpp
--- a/Engine/Source/Runtime/Engine/Private/UserInterface/PlayerInput.cpp
+++ b/Engine/Source/Runtime/Engine/Private/UserInterface/PlayerInput.cpp
@@ -23,6 +23,8 @@
 
 DECLARE_CYCLE_STAT(TEXT("    PC Gesture Recognition"), STAT_PC_GestureRecognition, STATGROUP_PlayerController);
 
+DEFINE_LOG_CATEGORY(LogPlayerInput);
+
 bool bExecutingBindCommand = false;
 
 /** for debug rendering */
@@ -55,6 +57,20 @@ namespace UE
 		static FAutoConsoleVariableRef CVarClearAxisValueIfConsumed(TEXT("Input.ClearAxisValueIfConsumed"),
 			bClearAxisValueIfConsumed,
 			TEXT("If true, we will clear the value of any FInputAxisKeyBinding whose FKey has been previously consumed.\nNote: This option will be removed in a future update."));
+
+		FString LexToString(const EInputEvent Event)
+		{
+			switch (Event)
+			{
+			case IE_Pressed: return TEXT("IE_Pressed"); break;
+			case IE_Released: return TEXT("IE_Released"); break;
+			case IE_Repeat: return TEXT("IE_Repeat"); break;
+			case IE_DoubleClick: return TEXT("IE_DoubleClick"); break;
+			case IE_Axis: return TEXT("IE_Axis"); break;
+			case IE_MAX: return TEXT("IE_MAX"); break;
+			default: return TEXT("Unknown");
+			}
+		}
 	}
 }
 
@@ -239,6 +255,14 @@ bool UPlayerInput::InputKey(FKey Key, EInputEvent Event, float AmountDepressed,
 
 bool UPlayerInput::InputKey(const FInputKeyParams& Params)
 {
+	UE_LOG(LogPlayerInput, VeryVerbose, TEXT("[%hs] %s (outer: %s) received input : Key: %s Value: %s Event:  %s"),
+		__func__,
+		*GetNameSafe(this),
+		*GetNameSafe(GetOuter()),
+		*Params.Key.GetFName().ToString(),
+		*Params.Delta.ToString(),
+		*UE::Input::LexToString(Params.Event));
+	
 	const bool bGamepad = Params.IsGamepad();
 
 	// MouseX and MouseY should not be treated as analog if there are no samples, as they need their EventAccumulator to be incremented 
