return const TCHAR* instead of fstring for better memory in this log

#rnx
#pf skip

diff --git a/Engine/Source/Runtime/Engine/Private/UserInterface/PlayerInput.cpp b/Engine/Source/Runtime/Engine/Private/UserInterface/PlayerInput.cpp
--- a/Engine/Source/Runtime/Engine/Private/UserInterface/PlayerInput.cpp
+++ b/Engine/Source/Runtime/Engine/Private/UserInterface/PlayerInput.cpp
@@ -58,7 +58,7 @@ namespace UE
 			bClearAxisValueIfConsumed,
 			TEXT("If true, we will clear the value of any FInputAxisKeyBinding whose FKey has been previously consumed.\nNote: This option will be removed in a future update."));
 
-		FString LexToString(const EInputEvent Event)
+		const TCHAR* LexToString(const EInputEvent Event)
 		{
 			switch (Event)
 			{
@@ -261,7 +261,7 @@ bool UPlayerInput::InputKey(const FInputKeyParams& Params)
 		*GetNameSafe(GetOuter()),
 		*Params.Key.GetFName().ToString(),
 		*Params.Delta.ToString(),
-		*UE::Input::LexToString(Params.Event));
+		UE::Input::LexToString(Params.Event));
 	
 	const bool bGamepad = Params.IsGamepad();
 
