Compute RequiredMemory for "DDC2" texture build worker internally to avoid it affecting ddc keys.
Update GetOutputAlphaFromKnownAlpha to have the same call structure as GetOutputAlphaFromKnownAlphaOrFail

#rnx

diff --git a/Engine/Source/Developer/TextureBuild/Private/TextureBuildFunction.cpp b/Engine/Source/Developer/TextureBuild/Private/TextureBuildFunction.cpp
--- a/Engine/Source/Developer/TextureBuild/Private/TextureBuildFunction.cpp
+++ b/Engine/Source/Developer/TextureBuild/Private/TextureBuildFunction.cpp
@@ -172,13 +172,73 @@ static FTextureBuildSettings ReadBuildSettingsFromCompactBinary(const FCbObjectV
 	return BuildSettings;
 }
 
-static ERawImageFormat::Type ComputeRawImageFormat(ETextureSourceFormat SourceFormat)
+
+static bool GetResolvedBuildSettings(const FCbObject& Settings, FTextureBuildSettings* OutBuildSettings, const ITextureFormat** OutTextureFormat)
 {
-	return FImageCoreUtils::ConvertToRawImageFormat(SourceFormat);
+	*OutBuildSettings = ReadBuildSettingsFromCompactBinary(Settings["Build"].AsObjectView());
+
+	const uint16 RequiredTextureFormatVersion = Settings["FormatVersion"].AsUInt16();
+	const ITextureFormat* TextureFormat = nullptr;
+	if (ITextureFormatManagerModule* TFM = GetTextureFormatManager())
+	{
+		TextureFormat = TFM->FindTextureFormat(OutBuildSettings->TextureFormatName);
+	}
+	else
+	{
+		UE_LOG(LogTextureBuildFunction, Error, TEXT("TextureFormatManager not found!"));
+		return false;
+	}
+
+	if (!TextureFormat)
+	{
+		UE_LOG(LogTextureBuildFunction, Error, TEXT("Texture format %s not found"), *WriteToString<128>(OutBuildSettings->TextureFormatName));
+		return false;
+	}
+
+	if (OutTextureFormat)
+	{
+		*OutTextureFormat = TextureFormat;
+	}
+
+	const uint16 CurrentTextureFormatVersion = TextureFormat->GetVersion(OutBuildSettings->TextureFormatName, OutBuildSettings);
+	if (CurrentTextureFormatVersion != RequiredTextureFormatVersion)
+	{
+		UE_LOG(LogTextureBuildFunction, Error, TEXT("%s has version %hu when version %hu is required."),
+			*OutBuildSettings->TextureFormatName.ToString(), CurrentTextureFormatVersion, RequiredTextureFormatVersion);;
+		return false;
+	}
+
+	const FChildTextureFormat* ChildTextureFormat = TextureFormat->GetChildFormat();
+	if (ChildTextureFormat)
+	{
+		OutBuildSettings->BaseTextureFormatName = ChildTextureFormat->GetBaseFormatName(OutBuildSettings->TextureFormatName);
+	}
+	else
+	{
+		OutBuildSettings->BaseTextureFormatName = OutBuildSettings->TextureFormatName;
+	}
+
+	OutBuildSettings->BaseTextureFormat = GetTextureFormatManager()->FindTextureFormat(OutBuildSettings->BaseTextureFormatName);
+	return true;
 }
 
+
+static bool GetImageInfoFromCb(FCbFieldView InSource, FImageInfo* OutImageInfo, int32* OutMipCount)
+{
+	OutImageInfo->Format = FImageCoreUtils::ConvertToRawImageFormat((ETextureSourceFormat)InSource["SourceFormat"].AsUInt8());
+	OutImageInfo->GammaSpace = (EGammaSpace)InSource["GammaSpace"].AsUInt8();
+	OutImageInfo->NumSlices = InSource["NumSlices"].AsInt32();
+	OutImageInfo->SizeX = InSource["SizeX"].AsInt32();
+	OutImageInfo->SizeY = InSource["SizeY"].AsInt32();
+
+	*OutMipCount = InSource["Mips"].AsArrayView().Num();
+
+	return true;
+}
+
+
 static bool TryReadTextureSourceFromCompactBinary(FCbFieldView Source, UE::DerivedData::FBuildContext& Context,
-												const FTextureBuildSettings & BuildSettings, TArray<FImage>& OutMips)
+												bool bVolume, TArray<FImage>& OutMips)
 {
 	FSharedBuffer InputBuffer = Context.FindInput(Source.GetName());
 	if (!InputBuffer)
@@ -195,7 +255,7 @@ static bool TryReadTextureSourceFromCompactBinary(FCbFieldView Source, UE::Deriv
 	// Source data has no CompressionFormat
 	ETextureSourceFormat SourceFormat = (ETextureSourceFormat)Source["SourceFormat"].AsUInt8();
 
-	ERawImageFormat::Type RawImageFormat = ComputeRawImageFormat(SourceFormat);
+	ERawImageFormat::Type RawImageFormat = FImageCoreUtils::ConvertToRawImageFormat(SourceFormat);
 
 	EGammaSpace GammaSpace = (EGammaSpace)Source["GammaSpace"].AsUInt8();
 	int32 NumSlices = Source["NumSlices"].AsInt32();
@@ -234,11 +294,11 @@ static bool TryReadTextureSourceFromCompactBinary(FCbFieldView Source, UE::Deriv
 			MipSize
 		);
 
-		MipSizeX = FMath::Max(MipSizeX / 2, 1);
-		MipSizeY = FMath::Max(MipSizeY / 2, 1);
-		if ( BuildSettings.bVolume )
+		MipSizeX = FEncodedTextureDescription::GetMipWidth(MipSizeX, 1);
+		MipSizeY = FEncodedTextureDescription::GetMipHeight(MipSizeY, 1);
+		if ( bVolume )
 		{
-			NumSlices = FMath::Max(NumSlices / 2, 1);
+			NumSlices = FEncodedTextureDescription::GetMipDepth(NumSlices, 1, true);
 		}
 	}
 
@@ -272,8 +332,20 @@ void FTextureBuildFunction::Configure(UE::DerivedData::FBuildConfigContext& Cont
 	Context.SetCacheBucket(UE::DerivedData::FCacheBucket(ANSITEXTVIEW("Texture")));
 
 	const FCbObject Settings = Context.FindConstant(UTF8TEXTVIEW("Settings"));
-	const int64 RequiredMemoryEstimate = Settings["RequiredMemoryEstimate"].AsInt64();
-	Context.SetRequiredMemory(RequiredMemoryEstimate);
+
+	// Bit unfortunate - we have to deserialize this entire thing in order to be able to compute
+	// the memory estimate, and we're about the deserialize the whole things again.	
+	FTextureBuildSettings BuildSettings;
+	if (GetResolvedBuildSettings(Settings, &BuildSettings, nullptr))
+	{	
+		FImageInfo SourceImageInfo;
+		int32 SourceMipCount = 0;
+		if (GetImageInfoFromCb(Settings["Source"], &SourceImageInfo, &SourceMipCount))
+		{
+			const int64 RequiredMemoryEstimate = UE::TextureBuildUtilities::GetPhysicalTextureBuildMemoryEstimate(&BuildSettings, SourceImageInfo, SourceMipCount);
+			Context.SetRequiredMemory(RequiredMemoryEstimate);
+		}
+	}
 }
 
 // All texture builds output (at least) these values.
@@ -437,48 +509,15 @@ void FTextureBuildFunction::Build(UE::DerivedData::FBuildContext& Context) const
 		return;
 	}
 
-	const FTextureBuildSettings BuildSettings = ReadBuildSettingsFromCompactBinary(Settings["Build"].AsObjectView());
-	
-	const uint16 RequiredTextureFormatVersion = Settings["FormatVersion"].AsUInt16();
 	const ITextureFormat* TextureFormat = nullptr;
-	if (ITextureFormatManagerModule* TFM = GetTextureFormatManager())
-	{
-		TextureFormat = TFM->FindTextureFormat(BuildSettings.TextureFormatName);
-	}
-	else
-	{
-		UE_LOG(LogTextureBuildFunction, Error, TEXT("TextureFormatManager not found!"));
-		return;
-	}
-
-	if (!TextureFormat)
-	{
-		UE_LOG(LogTextureBuildFunction, Error, TEXT("Texture format %s not found"), *WriteToString<128>(BuildSettings.TextureFormatName));
-		return;
-	}
-
-	const uint16 CurrentTextureFormatVersion = TextureFormat->GetVersion(BuildSettings.TextureFormatName, &BuildSettings);
-	if (CurrentTextureFormatVersion != RequiredTextureFormatVersion)
+	FTextureBuildSettings BuildSettings;
+	if (!GetResolvedBuildSettings(Settings, &BuildSettings, &TextureFormat))
 	{
-		UE_LOG(LogTextureBuildFunction, Error, TEXT("%s has version %hu when version %hu is required."),
-			*BuildSettings.TextureFormatName.ToString(), CurrentTextureFormatVersion, RequiredTextureFormatVersion);;
 		return;
 	}
-
-	const FChildTextureFormat* ChildTextureFormat = TextureFormat->GetChildFormat();
-	if (ChildTextureFormat)
-	{
-		const_cast<FTextureBuildSettings&>(BuildSettings).BaseTextureFormatName = ChildTextureFormat->GetBaseFormatName(BuildSettings.TextureFormatName);
-	}
-	else
-	{
-		const_cast<FTextureBuildSettings&>(BuildSettings).BaseTextureFormatName = BuildSettings.TextureFormatName;
-	}
-
-	const_cast<FTextureBuildSettings&>(BuildSettings).BaseTextureFormat = GetTextureFormatManager()->FindTextureFormat(BuildSettings.BaseTextureFormatName);
-
+	
 	TArray<FImage> SourceMips;
-	if (!TryReadTextureSourceFromCompactBinary(Settings["Source"], Context,BuildSettings, SourceMips))
+	if (!TryReadTextureSourceFromCompactBinary(Settings["Source"], Context, BuildSettings.bVolume, SourceMips))
 	{
 		return;
 	}
@@ -497,7 +536,7 @@ void FTextureBuildFunction::Build(UE::DerivedData::FBuildContext& Context) const
 
 	TArray<FImage> AssociatedNormalSourceMips;
 	if (FCbFieldView CompositeSource = Settings["CompositeSource"];
-		CompositeSource && !TryReadTextureSourceFromCompactBinary(CompositeSource, Context,BuildSettings, AssociatedNormalSourceMips))
+		CompositeSource && !TryReadTextureSourceFromCompactBinary(CompositeSource, Context, BuildSettings.bVolume, AssociatedNormalSourceMips))
 	{
 		return;
 	}
diff --git a/Engine/Source/Developer/TextureBuildUtilities/Private/TextureBuildUtilities.cpp b/Engine/Source/Developer/TextureBuildUtilities/Private/TextureBuildUtilities.cpp
--- a/Engine/Source/Developer/TextureBuildUtilities/Private/TextureBuildUtilities.cpp
+++ b/Engine/Source/Developer/TextureBuildUtilities/Private/TextureBuildUtilities.cpp
@@ -2,6 +2,7 @@
 
 #include "TextureBuildUtilities.h"
 
+#include "Async/TaskGraphInterfaces.h"
 #include "ImageCoreUtils.h"
 #include "TextureCompressorModule.h" // for FTextureBuildSettings
 #include "Misc/DataDrivenPlatformInfoRegistry.h"
@@ -9,6 +10,7 @@
 #include "Serialization/CompactBinaryWriter.h"
 //#include "EngineLogs.h" // can't use from SCW
 
+DEFINE_LOG_CATEGORY_STATIC(LogTextureBuildUtilities, Log, All);
 
 namespace UE
 {
@@ -369,5 +371,444 @@ TEXTUREBUILDUTILITIES_API bool TextureNeedsDecodeForPC(EPixelFormat InPixelForma
 	return !EnumHasAnyFlags(GPixelFormats[InPixelFormat].Capabilities, EPixelFormatCapabilities::Texture2D);
 }
 
+
+static int GetWithinSliceRDOMemoryUsePerPixel(EPixelFormat PixelFormat)
+{
+	// Memory use of RDO data structures, per pixel, within each slice
+	// not counting per-image memory use
+	const int MemUse_BC1 = 57;
+	const int MemUse_BC4 = 90;
+	const int MemUse_BC5 = 2 * MemUse_BC4;
+	const int MemUse_BC6 = 8;
+	const int MemUse_BC7 = 30;
+	const int MemUse_BC3 = MemUse_BC4; // max of BC1,BC4
+
+	switch (PixelFormat)
+	{
+	case PF_DXT1:
+		return MemUse_BC1;
+	case PF_DXT3:
+	case PF_DXT5:
+		return MemUse_BC3;
+	case PF_BC4:
+		return MemUse_BC4;
+	case PF_BC5:
+		return MemUse_BC5;
+	case PF_BC6H:
+		return MemUse_BC6;
+	case PF_BC7:
+		return MemUse_BC7;
+	default:
+		// is this possible?
+		UE_CALL_ONCE([&]() {
+			UE_LOG(LogTextureBuildUtilities, Display, TEXT("Unexpected non-BC PixelFormat: %d."), (int)PixelFormat);
+		});
+
+		return 100;
+	}
+}
+
+static int64 CalculateTextureSourceBytesFromImageInfo(FImageInfo InImageInfo, int32 InMipCount, bool bInVolume)
+{
+	if (InMipCount > 0)
+	{
+		int64 SourceBytes = 0;
+		while (InMipCount)
+		{
+			SourceBytes += InImageInfo.GetImageSizeBytes();
+
+			InImageInfo.SizeX = FEncodedTextureDescription::GetMipWidth(InImageInfo.SizeX, 1);
+			InImageInfo.SizeY = FEncodedTextureDescription::GetMipHeight(InImageInfo.SizeY, 1);
+			if (bInVolume)
+			{
+				InImageInfo.NumSlices = FEncodedTextureDescription::GetMipDepth(InImageInfo.NumSlices, 1, true);
+			}
+			InMipCount--;
+		}
+
+		return SourceBytes;
+	}
+
+	return -1;
+}
+
+
+TEXTUREBUILDUTILITIES_API EPixelFormat GetOutputPixelFormatWithFallback(const FTextureBuildSettings& InBuildSettings, bool bInKnownAlphaFallback)
+{
+	bool bHasAlpha = false;
+	InBuildSettings.GetOutputAlphaFromKnownAlphaOrFallback(&bHasAlpha, bInKnownAlphaFallback);
+
+	// If we get called and we have not stripped any platform prefix, this will crash if it's not TextureFormatOodle because
+	// the others still look at TextureFormatName instead of BaseTextureFormatName.
+	bool bNeedsBaseCopy = false;
+	if (InBuildSettings.TextureFormatName != InBuildSettings.BaseTextureFormatName)
+	{
+		// Welp, we are different. If it's TextureFormatOodle then we are actually OK.
+		// TextureFormatNames are pretty short so we can just copy out:
+		TStringBuilder<32> FormatName;
+		FormatName << InBuildSettings.BaseTextureFormatName;
+		if (FormatName.Len() < 4 ||
+			FCString::Strncmp(FormatName.GetData(), TEXT("TFO_"), 4))
+		{
+			// We aren't TFO and we are different so make a copy of us.
+			bNeedsBaseCopy = true;
+		}
+	}
+
+	EPixelFormat PixelFormat = PF_Unknown;
+	if (bNeedsBaseCopy)
+	{
+		// Base texture formats expect to get TextureFormatName without the platform prefix.
+		// We could call through the non-base TextureFormat but all it's doing is this:
+		// eventually we'll migrate all texture formats to reference BaseTextureFormatName.
+		FTextureBuildSettings BaseTextureBuildSettings = InBuildSettings;
+		BaseTextureBuildSettings.TextureFormatName = BaseTextureBuildSettings.BaseTextureFormatName;
+		PixelFormat = InBuildSettings.BaseTextureFormat->GetEncodedPixelFormat(BaseTextureBuildSettings, bHasAlpha);
+	}
+	else
+	{
+		PixelFormat = InBuildSettings.BaseTextureFormat->GetEncodedPixelFormat(InBuildSettings, bHasAlpha);
+	}
+	check(PixelFormat != PF_Unknown);
+
+	return PixelFormat;
+}
+
+
+TEXTUREBUILDUTILITIES_API int64 GetVirtualTextureRequiredMemoryEstimate(const FTextureBuildSettings* InBuildSettingsPerLayer,
+	TConstArrayView<ERawImageFormat::Type> InLayerFormats,
+	TConstArrayView<UE::TextureBuildUtilities::FVirtualTextureSourceBlockInfo> InSourceBlocks)
+{
+	const bool bRDO = true;
+	// @todo Oodle : be careful about using BuildSettings for this as there are two buildsettingses, just assume its on for now
+	//   <- FIX ME, allow lower mem estimates for non-RDO
+
+	// over-estimate is okay
+	// try not to over-estimate by too much (reduces parallelism of cook)
+
+	int64 MaxNumberOfWorkers = FMath::Max(1, FTaskGraphInterface::Get().GetNumWorkerThreads());
+
+	// VT build does :
+	// load all source images
+	// for each layer/block :
+	//    generate mips (requires F32 copy)
+	//    output to intermediate format
+	//    intermediate format copy is then used to make tiles
+	//    for each tile :
+	//       make padded tile in intermediate format
+	//       encode to output format
+	//       discard padded tile in intermediate format
+	// all output tiles are then aggregated
+
+	// Compute the memory it should take to uncompress the bulkdata in memory
+	int64 TotalSourceBytes = 0;
+	int64 TotalTopMipNumPixelsPerLayer = 0;
+	int64 LargestBlockTopMipNumPixels = 0;
+
+	// All layers in a VT must have the same layout for each block. Layers only can change source pixels + format, not dims.
+	for (int32 BlockIndex = 0; BlockIndex < InSourceBlocks.Num(); ++BlockIndex)
+	{
+		const UE::TextureBuildUtilities::FVirtualTextureSourceBlockInfo& SourceBlock = InSourceBlocks[BlockIndex];
+
+		for (int32 LayerIndex = 0; LayerIndex < InLayerFormats.Num(); ++LayerIndex)
+		{
+			// Create an FImageInfo so we can calcualte size off of it.
+			FImageInfo LayerBlock;
+			LayerBlock.GammaSpace = EGammaSpace::Linear; // doesn't matter for size
+			LayerBlock.Format = InLayerFormats[LayerIndex];
+			LayerBlock.SizeX = SourceBlock.SizeX;
+			LayerBlock.SizeY = SourceBlock.SizeY;
+			LayerBlock.NumSlices = SourceBlock.NumSlices;
+
+			TotalSourceBytes += CalculateTextureSourceBytesFromImageInfo(LayerBlock, SourceBlock.NumMips, false);
+		}
+
+		// assume pow2 options are the same for all layers, just use layer 0 here :
+		const FTextureBuildSettings& LayerBuildSettings = InBuildSettingsPerLayer[0];
+
+		int32 TargetSizeX, TargetSizeY, TargetSizeZ;
+		UE::TextureBuildUtilities::GetPowerOfTwoTargetTextureSize(SourceBlock.SizeX, SourceBlock.SizeY, SourceBlock.NumSlices,
+			LayerBuildSettings.bVolume, (ETexturePowerOfTwoSetting::Type)LayerBuildSettings.PowerOfTwoMode,
+			LayerBuildSettings.ResizeDuringBuildX, LayerBuildSettings.ResizeDuringBuildY,
+			TargetSizeX, TargetSizeY, TargetSizeZ);
+
+		int64 CurrentBlockTopMipNumPixels = (int64)TargetSizeX * TargetSizeY * TargetSizeZ;
+
+		TotalTopMipNumPixelsPerLayer += CurrentBlockTopMipNumPixels;
+
+		LargestBlockTopMipNumPixels = FMath::Max(CurrentBlockTopMipNumPixels, LargestBlockTopMipNumPixels);
+	}
+
+	if (TotalSourceBytes <= 0)
+	{
+		return -1; /* Unknown */
+	}
+
+	// assume full mip chain :
+	int64 TotalPixelsPerLayer = (TotalTopMipNumPixelsPerLayer * 4) / 3;
+
+	int64 TotalNumPixels = TotalPixelsPerLayer * InLayerFormats.Num();
+
+	// only one block of one layer does the float image mip build at a time :
+	int64 IntermediateFloatColorBytes = (LargestBlockTopMipNumPixels * sizeof(FLinearColor) * 4) / 3;
+
+	int64 TileSize = InBuildSettingsPerLayer[0].VirtualTextureTileSize;
+	int64 BorderSize = InBuildSettingsPerLayer[0].VirtualTextureBorderSize;
+
+	int64 NumTilesPerLayer = FMath::DivideAndRoundUp<int64>(TotalPixelsPerLayer, TileSize * TileSize);
+	int64 NumTiles = NumTilesPerLayer * InLayerFormats.Num();
+	int64 TilePixels = (TileSize + 2 * BorderSize) * (TileSize + 2 * BorderSize);
+
+	int64 NumOutputPixelsPerLayer = NumTilesPerLayer * TilePixels;
+
+	// intermediate is created just once per block, use max size estimate
+	int64 VTIntermediateSizeBytes = IntermediateFloatColorBytes;
+	int64 OutputSizeBytes = 0;
+
+	int64 MaxPerPixelEncoderMemUse = 0;
+
+	for (int32 LayerIndex = 0; LayerIndex < InLayerFormats.Num(); ++LayerIndex)
+	{
+		const FTextureBuildSettings& LayerBuildSettings = InBuildSettingsPerLayer[LayerIndex];
+
+		// VT builds to an intermediate format.
+
+		ERawImageFormat::Type IntermediateImageFormat = UE::TextureBuildUtilities::GetVirtualTextureBuildIntermediateFormat(LayerBuildSettings);
+
+		int64 IntermediateBytesPerPixel = ERawImageFormat::GetBytesPerPixel(IntermediateImageFormat);
+
+		// + output bytes? (but can overlap with IntermediateFloatColorBytes)
+		//	almost always less than IntermediateFloatColorBytes
+		//  exception would be lots of udim blocks + lots of layers
+		//  because IntermediateFloatColorBytes is per block/layer but output is held for all
+
+		EPixelFormat PixelFormat = UE::TextureBuildUtilities::GetOutputPixelFormatWithFallback(LayerBuildSettings, true);
+
+		if (PixelFormat == PF_Unknown)
+		{
+			return -1; /* Unknown */
+		}
+
+		const FPixelFormatInfo& PFI = GPixelFormats[PixelFormat];
+
+		OutputSizeBytes += (NumOutputPixelsPerLayer * PFI.BlockBytes) / (PFI.BlockSizeX * PFI.BlockSizeY);
+
+		// is it a blocked format :
+		if (PFI.BlockSizeX > 1)
+		{
+			// another copy of Intermediate in BlockSurf swizzle :
+			int CurPerPixelEncoderMemUse = IntermediateBytesPerPixel;
+
+			if (bRDO)
+			{
+				int RDOMemUse = GetWithinSliceRDOMemoryUsePerPixel(PixelFormat);
+				CurPerPixelEncoderMemUse += 4; // activity
+				CurPerPixelEncoderMemUse += RDOMemUse;
+				CurPerPixelEncoderMemUse += 1; // output again
+			}
+
+			// max over any layer :
+			MaxPerPixelEncoderMemUse = FMath::Max(MaxPerPixelEncoderMemUse, CurPerPixelEncoderMemUse);
+		}
+	}
+
+	// after we make the Intermediate layer, it is cut into tiles
+	// we then need mem for the intermediate format padded up to tiles
+	// and then working encoder mem & compressed output space for each tile
+	//	(tiles are made one by one in the ParallelFor to make the compressed output)
+	// but at that point the FloatColorBytes is freed
+
+	int64 NumberOfWorkingTiles = FMath::Min(NumTiles, MaxNumberOfWorkers);
+
+	// VT tile encode mem :  
+	int64 MemoryUsePerTile = MaxPerPixelEncoderMemUse * TilePixels; // around 1.8 MB
+	{
+		 // MemoryUsePerTile
+		 // makes tile in IntermediateBytesPerPixel
+		 // encodes out to OutputSizeBytes
+		 // encoder (Oodle) temp mem
+		 // TilePixels * IntermediateBytesPerPixel (twice: surf+blocksurf)
+		 // TilePixels * Output bytes (twice: baseline+rdo output) (output already counted)
+		 // TilePixels * activity mask
+		 // MaxPerPixelEncoderMemUse is around 100
+	}
+
+	int64 TileCompressionBytes = NumberOfWorkingTiles * MemoryUsePerTile;
+
+	int64 MemoryEstimate = TotalSourceBytes + VTIntermediateSizeBytes;
+	// @todo Oodle : After we make the VT Intermediate, is the source BulkData freed?
+	//   -> it seems no at the moment, but it could be
+
+	// take larger of mem use during float image filter phase or tile compression phase
+	MemoryEstimate += FMath::Max(IntermediateFloatColorBytes, TileCompressionBytes + OutputSizeBytes);
+
+	MemoryEstimate += 64 * 1024; // overhead room
+
+	//UE_LOG(LogTexture,Display,TEXT("GetBuildRequiredMemoryEstimate VT : %.3f MB"),MemoryEstimate/(1024*1024.f));
+
+	return MemoryEstimate;
+}
+
+
+// Returns the estimated memory cost of building the given texture. Only valid for physical (i.e. non-virtual) textures.
+TEXTUREBUILDUTILITIES_API int64 GetPhysicalTextureBuildMemoryEstimate(const FTextureBuildSettings* InSettingsPerLayerFetchFirst, const FImageInfo& InSourceImageInfo, int32 InMipCount)
+{
+	const bool bRDO = true;
+	// @todo Oodle : be careful about using BuildSettings for this as there are two buildsettingses, just assume its on for now
+	//   <- FIX ME, allow lower mem estimates for non-RDO
+
+	// over-estimate is okay
+	// try not to over-estimate by too much (reduces parallelism of cook)
+
+	int64 MaxNumberOfWorkers = FMath::Max(1, FTaskGraphInterface::Get().GetNumWorkerThreads());
+
+	{
+		const FTextureBuildSettings& BuildSettings = InSettingsPerLayerFetchFirst[0];
+		// non VT
+
+		// Compute the memory it should take to uncompress the bulkdata in memory
+		int64 TotalSourceBytes = CalculateTextureSourceBytesFromImageInfo(InSourceImageInfo, InMipCount, BuildSettings.bVolume);
+		if (TotalSourceBytes <= 0)
+		{
+			return -1; /* Unknown */
+		}
+
+
+
+		int32 TargetSizeX, TargetSizeY, TargetSizeZ;
+		UE::TextureBuildUtilities::GetPowerOfTwoTargetTextureSize(InSourceImageInfo.SizeX, InSourceImageInfo.SizeY, InSourceImageInfo.NumSlices,
+			BuildSettings.bVolume, (ETexturePowerOfTwoSetting::Type)BuildSettings.PowerOfTwoMode,
+			BuildSettings.ResizeDuringBuildX, BuildSettings.ResizeDuringBuildY,
+			TargetSizeX, TargetSizeY, TargetSizeZ);
+
+		int64 TotalTopMipNumPixels = (int64)TargetSizeX * TargetSizeY * TargetSizeZ;
+
+		// assume full mip chain :
+		int64 TotalNumPixels = (TotalTopMipNumPixels * 4) / 3;
+
+		// actually we have each mip twice for the float image filter phase so this is under-counting
+		//	but that isn't held allocated while the output is made, so it can overlap with that mem
+		int64 IntermediateFloatColorBytes = TotalNumPixels * sizeof(FLinearColor);
+
+		int64 MemoryEstimate = TotalSourceBytes + IntermediateFloatColorBytes;
+
+		// Assume alpha exists if we don't know for worst-case handling.
+		EPixelFormat PixelFormat = GetOutputPixelFormatWithFallback(BuildSettings, true);
+
+		if (PixelFormat == PF_Unknown)
+		{
+			return -1; /* Unknown */
+		}
+
+		const FPixelFormatInfo& PFI = GPixelFormats[PixelFormat];
+
+		const int64 OutputSizeBytes = (TotalNumPixels * PFI.BlockBytes) / (PFI.BlockSizeX * PFI.BlockSizeY);
+
+		MemoryEstimate += OutputSizeBytes;
+
+		// check to see if it's uncompressed or a BCN format :
+		if (IsDXTCBlockCompressedTextureFormat(PixelFormat))
+		{
+			// block-compressed format ; assume it's using Oodle Texture
+
+			if (bRDO)
+			{
+				// two more copies in outputsize
+				// baseline encode + UT or Layout
+				MemoryEstimate += OutputSizeBytes * 2;
+			}
+
+			// you also have to convert the float surface to an input format for Oodle
+			//	this copy is done in TFO
+			//  Oodle then allocs another copy to swizzle into blocks before encoding
+
+			int IntermediateBytesPerPixel;
+			bool bNeedsIntermediateCopy = true;
+
+			// this matches the logic in TextureFormatOodle :
+			if (PixelFormat == PF_BC6H)
+			{
+				IntermediateBytesPerPixel = 16; //RGBAF32
+				bNeedsIntermediateCopy = false; // no intermediate used in TFO (float source kept), 1 blocksurf
+			}
+			else if (PixelFormat == PF_BC4 || PixelFormat == PF_BC5)
+			{
+				// changed: TFO uses 2_U16 now (4 byte intermediate)
+				IntermediateBytesPerPixel = 8; // RGBA16
+			}
+			else
+			{
+				IntermediateBytesPerPixel = 4; // RGBA8
+			}
+
+			int NumIntermediateCopies = 1; // BlockSurf
+			if (bNeedsIntermediateCopy) NumIntermediateCopies++;
+
+			MemoryEstimate += NumIntermediateCopies * IntermediateBytesPerPixel * TotalNumPixels;
+
+			if (bRDO)
+			{
+				// activity map for whole image :
+				// (this has changed in newer versions of Oodle Texture)
+
+				// Phase1 = computing activity map
+				int ActivityBytesPerPixel;
+
+				if (PixelFormat == PF_BC4) ActivityBytesPerPixel = 12;
+				else if (PixelFormat == PF_BC5) ActivityBytesPerPixel = 16;
+				else ActivityBytesPerPixel = 24;
+
+				int64 RDOPhase1MemUse = ActivityBytesPerPixel * TotalNumPixels;
+
+				// Phase2 = cut into slices, encode each slice
+				// per-slice data structure memory use
+				// non-RDO is all on stack so zero
+
+				// fewer workers for small images ; roughly one slice per 64 KB of output
+				//int64 NumberofSlices = FMath::DivideAndRoundUp<int64>(OutputSizeBytes,64*1024);
+				int64 PixelsPerSlice = (64 * 1024 * TotalNumPixels) / OutputSizeBytes;
+				int64 NumberofSlices = FMath::DivideAndRoundUp<int64>(TotalNumPixels, PixelsPerSlice);
+				if (NumberofSlices <= 4)
+				{
+					PixelsPerSlice = TotalNumPixels / NumberofSlices;
+				}
+
+				int64 MemoryUsePerWorker = PixelsPerSlice * GetWithinSliceRDOMemoryUsePerPixel(PixelFormat);
+					// MemoryUsePerWorker is around 10 MB
+				int64 NumberOfWorkers = FMath::Min(NumberofSlices, MaxNumberOfWorkers);
+
+				int64 RDOPhase2MemUse = 4 * TotalNumPixels; // activity map held on whole image
+				RDOPhase2MemUse += NumberOfWorkers * MemoryUsePerWorker;
+
+				// usually phase2 is higher
+				// but on large BC6 images on machines with low core counts, phase1 can be higher
+
+				MemoryEstimate += FMath::Max(RDOPhase1MemUse, RDOPhase2MemUse);
+			}
+		}
+		else if (IsASTCBlockCompressedTextureFormat(PixelFormat))
+		{
+			// ASTCenc does an entermediate copy to RGBA16F for HDR formats and RGBA8 for LDR
+			MemoryEstimate += (IsHDR(PixelFormat) ? 8 : 4) * TotalNumPixels;
+			// internal memory use of ASTCenc is not estimated
+			// @todo : fix me
+		}
+		else
+		{
+			// note: memory ues of non-Oodle encoders is not estimated
+			// @todo : fix me
+		}
+
+		MemoryEstimate += 64 * 1024; // overhead room
+
+		//UE_LOG(LogTexture,Display,TEXT("GetBuildRequiredMemoryEstimate non-VT : %.3f MB"),MemoryEstimate/(1024*1024.f));
+
+		return MemoryEstimate;
+
+		// @todo Oodle : not right for volumes & latlong cubes
+		// @todo Oodle : not right with Composite , CPU textures
+	}
+}
+
 } // namespace
 }
\ No newline at end of file
diff --git a/Engine/Source/Developer/TextureBuildUtilities/Public/TextureBuildUtilities.h b/Engine/Source/Developer/TextureBuildUtilities/Public/TextureBuildUtilities.h
--- a/Engine/Source/Developer/TextureBuildUtilities/Public/TextureBuildUtilities.h
+++ b/Engine/Source/Developer/TextureBuildUtilities/Public/TextureBuildUtilities.h
@@ -91,5 +91,38 @@ TEXTUREBUILDUTILITIES_API void GetPlaceholderTextureImage(FImage* OutImage);
 //	if InPow2Setting == None, the Out sizes match the In sizes, and false is returned
 TEXTUREBUILDUTILITIES_API bool GetPowerOfTwoTargetTextureSize(int32 InMip0SizeX, int32 InMip0SizeY, int32 InMip0NumSlices, bool bInIsVolume, ETexturePowerOfTwoSetting::Type InPow2Setting, int32 InResizeDuringBuildX, int32 InResizeDuringBuildY, int32& OutTargetSizeX, int32& OutTargetSizeY, int32& OutTargetSizeZ);
 
+// Return the final output pixel format. When we don't know the source's alpha channel information this might not be knowable. In those cases
+// treat such sources as having an alpha channel if bKnownAlphaFallback is true.
+TEXTUREBUILDUTILITIES_API EPixelFormat GetOutputPixelFormatWithFallback(const FTextureBuildSettings& InBuildSettings, bool bInKnownAlphaFallback);
+
+// Return the estimate for how much memory a physical texture will take to build, for managing memory resources during build dispatch.
+TEXTUREBUILDUTILITIES_API int64 GetPhysicalTextureBuildMemoryEstimate(const FTextureBuildSettings* InBuildSettings, const FImageInfo& InSourceImageInfo, int32 InMipCount);
+
+// Mirrors Texture.h FTextureSourceBlock for use without Texture.h
+//
+// This provides the layout representation of a block of source pixels for a virtual texture. The format is separate
+// as VTs can have multiple layers, each with the same source pixel layout but with a different format (and thus byte size).
+struct TEXTUREBUILDUTILITIES_API FVirtualTextureSourceBlockInfo
+{
+	// @@ I don't actually know what coordinate space blocks are in and none of the code I'm looking at seems very clear on it.
+	// afaict it's just the index of the source image and where it exists. I don't know what happens if you have a hole in your source blocks.
+	int32 BlockX=0;
+	int32 BlockY=0;
+
+	// Pixel dims.
+	int32 SizeX=0;
+	int32 SizeY=0;
+
+	// afaict this is ignored: see "BlockData.NumSlices = 1; // TODO?" in VirtualTextureDataBuilder
+	int32 NumSlices=0;
+
+	int32 NumMips=0;
+};
+
+// Return the estimate for how much memory a virtual texture will take to build, for managing memory resources during build dispatch.
+TEXTUREBUILDUTILITIES_API int64 GetVirtualTextureRequiredMemoryEstimate(const FTextureBuildSettings* InBuildSettingsPerLayer,
+	TConstArrayView<ERawImageFormat::Type> InLayerFormats,
+	TConstArrayView<UE::TextureBuildUtilities::FVirtualTextureSourceBlockInfo> InSourceBlocks);
+
 } // namespace TextureBuildUtilities
 } // namespace UE
diff --git a/Engine/Source/Developer/TextureCompressor/Public/TextureCompressorModule.h b/Engine/Source/Developer/TextureCompressor/Public/TextureCompressorModule.h
--- a/Engine/Source/Developer/TextureCompressor/Public/TextureCompressorModule.h
+++ b/Engine/Source/Developer/TextureCompressor/Public/TextureCompressorModule.h
@@ -274,7 +274,7 @@ struct FTextureBuildSettings
 	// to TextureFormatName.
 	FName BaseTextureFormatName;
 
-	// Cached pointer to BaseTextureFormatName
+	// Cached pointer to BaseTextureFormatName. Resolved in FinalizeBuildSettingsForLayer and ResolveBuildSettings
 	const ITextureFormat* BaseTextureFormat = nullptr;
 	
 	// Whether bHasTransparentAlpha is valid.
@@ -395,18 +395,17 @@ struct FTextureBuildSettings
 		return bInSourceMipsAlphaDetected;
 	}
 
-	// the return value is whether the *output* has alpha
-	// the argument is whether we should treat the source as having alpha, in cases where the source alpha is unknown
-	bool GetOutputAlphaFromKnownAlphaOrFallback(bool bUnknownSourceAlphaFallback = true) const
+	// bUnknownSourceAlphaFallback is whether we should treat the source as having alpha, in cases where the source alpha is unknown
+	void GetOutputAlphaFromKnownAlphaOrFallback(bool* bOutOutputHasAlpha, bool bUnknownSourceAlphaFallback = true) const
 	{
 		if ( bKnowAlphaTransparency )
 		{
-			return bHasTransparentAlpha;
+			*bOutOutputHasAlpha = bHasTransparentAlpha;
 		}
 		else
 		{
 			// fallback without known alpha
-			return GetTextureExpectsAlphaInPixelFormat(bUnknownSourceAlphaFallback);
+			*bOutOutputHasAlpha = GetTextureExpectsAlphaInPixelFormat(bUnknownSourceAlphaFallback);
 		}
 	}
 
diff --git a/Engine/Source/Runtime/Engine/Private/TextureDerivedData.cpp b/Engine/Source/Runtime/Engine/Private/TextureDerivedData.cpp
--- a/Engine/Source/Runtime/Engine/Private/TextureDerivedData.cpp
+++ b/Engine/Source/Runtime/Engine/Private/TextureDerivedData.cpp
@@ -651,7 +651,8 @@ static void GetBuiltTextureSizeBytesEstimate(
 	int64 NumMips = FImageCoreUtils::GetMipCountFromDimensions(TopMipSizeX,TopMipSizeY,TopMipSizeZ,bIsVolume);
 	check( NumMips > 0 );
 
-	bool bHasAlpha = BuildSettings.GetOutputAlphaFromKnownAlphaOrFallback(true);
+	bool bHasAlpha;
+	BuildSettings.GetOutputAlphaFromKnownAlphaOrFallback(&bHasAlpha, true);
 
 	FEncodedTextureDescription TextureDescription;
 	BuildSettings.GetEncodedTextureDescription(&TextureDescription, TextureFormat, TopMipSizeX, TopMipSizeY, TopMipSizeZ, NumMips, bHasAlpha);
@@ -775,11 +776,11 @@ static void ModifyMaxTextureResolutionBuildSettingsForPlatformLimit(
 	uint64 MaxSurfaceBytes,MaxPackageBytes;
 	TargetPlatform->GetTextureSizeLimits(MaxSurfaceBytes,MaxPackageBytes);
 	
-	EPixelFormat PixelFormat = GetOutputPixelFormat(OutSettings);
+	EPixelFormat PixelFormat = UE::TextureBuildUtilities::GetOutputPixelFormatWithFallback(OutSettings, true);
 
 	if ( PixelFormat == PF_Unknown )
 	{
-		UE_LOG(LogTexture, Error, TEXT("Texture %s failed GetOutputPixelFormat (format=%s)"), 
+		UE_LOG(LogTexture, Error, TEXT("Texture %s failed GetOutputPixelFormatWithFallback (format=%s)"), 
 			*Texture.GetPathName(),
 			*OutSettings.TextureFormatName.ToString());
 			
diff --git a/Engine/Source/Runtime/Engine/Private/TextureDerivedDataBuildUtils.cpp b/Engine/Source/Runtime/Engine/Private/TextureDerivedDataBuildUtils.cpp
--- a/Engine/Source/Runtime/Engine/Private/TextureDerivedDataBuildUtils.cpp
+++ b/Engine/Source/Runtime/Engine/Private/TextureDerivedDataBuildUtils.cpp
@@ -334,7 +334,7 @@ UE::DerivedData::FUtf8SharedString FindTextureBuildFunction(const FName TextureF
 	return Function;
 }
 
-FCbObject SaveTextureBuildSettings(const UTexture& Texture, const FTextureBuildSettings& BuildSettings, int32 LayerIndex, bool bUseCompositeTexture, int64 RequiredMemoryEstimate)
+FCbObject SaveTextureBuildSettings(const UTexture& Texture, const FTextureBuildSettings& BuildSettings, int32 LayerIndex, bool bUseCompositeTexture)
 {
 	const ITextureFormat* TextureFormat = nullptr;
 	if (ITextureFormatManagerModule* TFM = GetTextureFormatManager())
@@ -364,8 +364,6 @@ FCbObject SaveTextureBuildSettings(const UTexture& Texture, const FTextureBuildS
 		Writer.AddUuid("CompressionCacheId", Texture.CompressionCacheId);
 	}
 
-	Writer.AddInteger("RequiredMemoryEstimate", RequiredMemoryEstimate);
-
 	if (uint16 TextureFormatVersion = TextureFormat->GetVersion(BuildSettings.TextureFormatName, &BuildSettings))
 	{
 		Writer.AddInteger("FormatVersion", TextureFormatVersion);
diff --git a/Engine/Source/Runtime/Engine/Private/TextureDerivedDataBuildUtils.h b/Engine/Source/Runtime/Engine/Private/TextureDerivedDataBuildUtils.h
--- a/Engine/Source/Runtime/Engine/Private/TextureDerivedDataBuildUtils.h
+++ b/Engine/Source/Runtime/Engine/Private/TextureDerivedDataBuildUtils.h
@@ -13,6 +13,6 @@ class UTexture;
 struct FTextureBuildSettings;
 
 UE::DerivedData::FUtf8SharedString FindTextureBuildFunction(FName TextureFormatName);
-FCbObject SaveTextureBuildSettings(const UTexture& Texture, const FTextureBuildSettings& BuildSettings, int32 LayerIndex, bool bUseCompositeTexture, int64 RequiredMemoryEstimate);
+FCbObject SaveTextureBuildSettings(const UTexture& Texture, const FTextureBuildSettings& BuildSettings, int32 LayerIndex, bool bUseCompositeTexture);
 
 #endif // WITH_EDITOR
diff --git a/Engine/Source/Runtime/Engine/Private/TextureDerivedDataTask.cpp b/Engine/Source/Runtime/Engine/Private/TextureDerivedDataTask.cpp
--- a/Engine/Source/Runtime/Engine/Private/TextureDerivedDataTask.cpp
+++ b/Engine/Source/Runtime/Engine/Private/TextureDerivedDataTask.cpp
@@ -1174,403 +1174,44 @@ static void DDC1_BuildTexture(
 	}
 }
 
-
-EPixelFormat GetOutputPixelFormat(const FTextureBuildSettings & BuildSettings)
-{
-	// get the TextureFormat so we can get the output pixel format :
-		
-	const ITextureFormat* TextureFormat = nullptr;
-
-	ITextureFormatManagerModule* TFM = GetTextureFormatManager();
-	if (TFM)
-	{
-		TextureFormat = TFM->FindTextureFormat(BuildSettings.TextureFormatName);
-	}
-	if (TextureFormat == nullptr)
-	{
-		UE_LOG(LogTexture, Warning,
-			TEXT("Failed to find compressor for texture format '%s'."),
-			*BuildSettings.TextureFormatName.ToString()
-		);
-			
-		return PF_Unknown; /* Unknown */
-	}
-	
-	bool bHasAlpha = BuildSettings.GetOutputAlphaFromKnownAlphaOrFallback(true);
-
-	EPixelFormat PixelFormat = TextureFormat->GetEncodedPixelFormat(BuildSettings,bHasAlpha);
-	check( PixelFormat != PF_Unknown );
-
-	return PixelFormat;
-}
-
-static int GetWithinSliceRDOMemoryUsePerPixel(EPixelFormat PixelFormat)
+static int64 GetBuildRequiredMemoryEstimate(UTexture* InTexture, const FTextureBuildSettings* InBuildSettingsPerLayer)
 {
-	// Memory use of RDO data structures, per pixel, within each slice
-	// not counting per-image memory use
-	const int MemUse_BC1 = 57;
-	const int MemUse_BC4 = 90;
-	const int MemUse_BC5 = 2*MemUse_BC4;
-	const int MemUse_BC6 = 8;
-	const int MemUse_BC7 = 30;
-	const int MemUse_BC3 = MemUse_BC4; // max of BC1,BC4
-	
-	switch(PixelFormat)
-	{
-	case PF_DXT1:
-		return MemUse_BC1;
-	case PF_DXT3:
-	case PF_DXT5:
-		return MemUse_BC3;
-	case PF_BC4:
-		return MemUse_BC4;
-	case PF_BC5:
-		return MemUse_BC5;
-	case PF_BC6H:
-		return MemUse_BC6;
-	case PF_BC7:
-		return MemUse_BC7;
-	default:
-		// is this possible?
-		UE_CALL_ONCE( [&](){
-			UE_LOG(LogTexture, Display, TEXT("Unexpected non-BC PixelFormat: %d."), (int)PixelFormat);
-		} );
-
-		return 100;
-	}
-}
+	// Thunk to our computation functions that don't rely on Texture.h.
 
-static int64 GetBuildRequiredMemoryEstimate(UTexture* InTexture,
-	const FTextureBuildSettings* InSettingsPerLayerFetchFirst)
-{
 	const FTextureSource & Source = InTexture->Source;
 
-	const bool bIsVT = InSettingsPerLayerFetchFirst[0].bVirtualStreamable;
-	
-	const bool bRDO = true;
-	// @todo Oodle : be careful about using BuildSettings for this as there are two buildsettingses, just assume its on for now
-	//   <- FIX ME, allow lower mem estimates for non-RDO
-
-	// over-estimate is okay
-	// try not to over-estimate by too much (reduces parallelism of cook)
-	
-	int64 MaxNumberOfWorkers = FMath::Max(1, FTaskGraphInterface::Get().GetNumWorkerThreads());
-
-	if ( bIsVT )
+	if (InBuildSettingsPerLayer[0].bVirtualStreamable)
 	{
-		// VT build does :
-		// load all source images
-		// for each layer/block :
-		//    generate mips (requires F32 copy)
-		//    output to intermediate format
-		//    intermediate format copy is then used to make tiles
-		//    for each tile :
-		//       make padded tile in intermediate format
-		//       encode to output format
-		//       discard padded tile in intermediate format
-		// all output tiles are then aggregated
-
-		// Compute the memory it should take to uncompress the bulkdata in memory
-		int64 TotalSourceBytes = 0;
-		int64 TotalTopMipNumPixelsPerLayer = 0;
-		int64 LargestBlockTopMipNumPixels = 0;
-
-		for (int32 BlockIndex = 0; BlockIndex < Source.GetNumBlocks(); ++BlockIndex)
+		TArray<ERawImageFormat::Type, TInlineAllocator<1>> LayerFormats;
+		LayerFormats.AddZeroed(Source.GetNumLayers());
+		for (int32 LayerIndex = 0; LayerIndex < Source.GetNumLayers(); LayerIndex++)
 		{
-			FTextureSourceBlock SourceBlock;
-			Source.GetBlock(BlockIndex, SourceBlock);
-
-			for (int32 LayerIndex = 0; LayerIndex < Source.GetNumLayers(); ++LayerIndex)
-			{
-				for (int32 MipIndex = 0; MipIndex < SourceBlock.NumMips; ++MipIndex)
-				{
-					TotalSourceBytes += Source.CalcMipSize(BlockIndex, LayerIndex, MipIndex);
-				}
-			}
-		
-			// assume pow2 options are the same for all layers, just use layer 0 here :
-			const FTextureBuildSettings & BuildSettings = InSettingsPerLayerFetchFirst[0];
-
-			int32 TargetSizeX, TargetSizeY, TargetSizeZ;
-			UE::TextureBuildUtilities::GetPowerOfTwoTargetTextureSize(SourceBlock.SizeX,SourceBlock.SizeY,SourceBlock.NumSlices,
-				BuildSettings.bVolume, (ETexturePowerOfTwoSetting::Type)BuildSettings.PowerOfTwoMode, 
-				BuildSettings.ResizeDuringBuildX, BuildSettings.ResizeDuringBuildY, 
-				TargetSizeX, TargetSizeY, TargetSizeZ);
-
-			int64 CurrentBlockTopMipNumPixels = (int64)TargetSizeX * TargetSizeY * TargetSizeZ;
-
-			TotalTopMipNumPixelsPerLayer += CurrentBlockTopMipNumPixels;
-
-			LargestBlockTopMipNumPixels = FMath::Max( CurrentBlockTopMipNumPixels , LargestBlockTopMipNumPixels );
+			LayerFormats[LayerIndex] = FImageCoreUtils::ConvertToRawImageFormat(Source.GetFormat(LayerIndex));
 		}
-		
-		if ( TotalSourceBytes <= 0 )
-		{
-			return -1; /* Unknown */
-		}
-		
-		// assume full mip chain :
-		int64 TotalPixelsPerLayer = (TotalTopMipNumPixelsPerLayer * 4) / 3;
-
-		int64 TotalNumPixels = TotalPixelsPerLayer * Source.GetNumLayers();
-
-		// only one block of one layer does the float image mip build at a time :
-		int64 IntermediateFloatColorBytes = (LargestBlockTopMipNumPixels * sizeof(FLinearColor) * 4) / 3;
-		
-		int64 TileSize = InSettingsPerLayerFetchFirst[0].VirtualTextureTileSize;
-		int64 BorderSize = InSettingsPerLayerFetchFirst[0].VirtualTextureBorderSize;
-
-		int64 NumTilesPerLayer = FMath::DivideAndRoundUp<int64>(TotalPixelsPerLayer,TileSize*TileSize);
-		int64 NumTiles = NumTilesPerLayer * Source.GetNumLayers();
-		int64 TilePixels = (TileSize + 2*BorderSize)*(TileSize + 2*BorderSize);
-
-		int64 NumOutputPixelsPerLayer = NumTilesPerLayer * TilePixels;
 
-		// intermediate is created just once per block, use max size estimate
-		int64 VTIntermediateSizeBytes = IntermediateFloatColorBytes;
-		int64 OutputSizeBytes = 0;
-	
-		int64 MaxPerPixelEncoderMemUse = 0;
-
-		for (int32 LayerIndex = 0; LayerIndex < Source.GetNumLayers(); ++LayerIndex)
+		TArray<UE::TextureBuildUtilities::FVirtualTextureSourceBlockInfo, TInlineAllocator<4>> Blocks;
+		Blocks.AddDefaulted(Source.GetNumBlocks());
+		for (int32 BlockIndex = 0; BlockIndex < Source.GetNumBlocks(); BlockIndex++)
 		{
-			const FTextureBuildSettings & BuildSettings = InSettingsPerLayerFetchFirst[LayerIndex];
-		
-			// VT builds to an intermediate format.
-			
-			ERawImageFormat::Type IntermediateImageFormat = UE::TextureBuildUtilities::GetVirtualTextureBuildIntermediateFormat(BuildSettings);
-
-			int64 IntermediateBytesPerPixel = ERawImageFormat::GetBytesPerPixel(IntermediateImageFormat);
+			FTextureSourceBlock Block;
+			Source.GetBlock(BlockIndex, Block);
 
-			// + output bytes? (but can overlap with IntermediateFloatColorBytes)
-			//	almost always less than IntermediateFloatColorBytes
-			//  exception would be lots of udim blocks + lots of layers
-			//  because IntermediateFloatColorBytes is per block/layer but output is held for all
-			
-			EPixelFormat PixelFormat = GetOutputPixelFormat(BuildSettings);
-
-			if ( PixelFormat == PF_Unknown )
-			{
-				return -1; /* Unknown */
-			}
-
-			const FPixelFormatInfo & PFI = GPixelFormats[PixelFormat];
-
-			OutputSizeBytes += ( NumOutputPixelsPerLayer * PFI.BlockBytes ) / ( PFI.BlockSizeX * PFI.BlockSizeY );
-
-			// is it a blocked format :
-			if ( PFI.BlockSizeX > 1 )
-			{
-				// another copy of Intermediate in BlockSurf swizzle :
-				int CurPerPixelEncoderMemUse = IntermediateBytesPerPixel;
-
-				if ( bRDO )
-				{
-					int RDOMemUse = GetWithinSliceRDOMemoryUsePerPixel(PixelFormat);
-					CurPerPixelEncoderMemUse += 4; // activity
-					CurPerPixelEncoderMemUse += RDOMemUse;
-					CurPerPixelEncoderMemUse += 1; // output again
-				}
-
-				// max over any layer :
-				MaxPerPixelEncoderMemUse = FMath::Max(MaxPerPixelEncoderMemUse,CurPerPixelEncoderMemUse);
-			}
+			Blocks[BlockIndex].BlockX = Block.BlockX;
+			Blocks[BlockIndex].BlockY = Block.BlockY;
+			Blocks[BlockIndex].SizeX = Block.SizeX;
+			Blocks[BlockIndex].SizeY = Block.SizeY;
+			Blocks[BlockIndex].NumSlices = Block.NumSlices;
+			Blocks[BlockIndex].NumMips = Block.NumMips;
 		}
-				
-		// after we make the Intermediate layer, it is cut into tiles
-		// we then need mem for the intermediate format padded up to tiles
-		// and then working encoder mem & compressed output space for each tile
-		//	(tiles are made one by one in the ParallelFor to make the compressed output)
-		// but at that point the FloatColorBytes is freed
-		
-		int64 NumberOfWorkingTiles = FMath::Min(NumTiles,MaxNumberOfWorkers);
-		
-		// VT tile encode mem :  
-		int64 MemoryUsePerTile = MaxPerPixelEncoderMemUse * TilePixels; // around 1.8 MB
-		{
-			 // MemoryUsePerTile
-			 // makes tile in IntermediateBytesPerPixel
-			 // encodes out to OutputSizeBytes
-			 // encoder (Oodle) temp mem
-			 // TilePixels * IntermediateBytesPerPixel (twice: surf+blocksurf)
-			 // TilePixels * Output bytes (twice: baseline+rdo output) (output already counted)
-			 // TilePixels * activity mask
-			 // MaxPerPixelEncoderMemUse is around 100
-		}
-
-		int64 TileCompressionBytes = NumberOfWorkingTiles * MemoryUsePerTile;
-
-		int64 MemoryEstimate = TotalSourceBytes + VTIntermediateSizeBytes;
-		// @todo Oodle : After we make the VT Intermediate, is the source BulkData freed?
-		//   -> it seems no at the moment, but it could be
-		
-		// take larger of mem use during float image filter phase or tile compression phase
-		MemoryEstimate += FMath::Max( IntermediateFloatColorBytes , TileCompressionBytes + OutputSizeBytes );
 
-		MemoryEstimate += 64 * 1024; // overhead room
-
-		//UE_LOG(LogTexture,Display,TEXT("GetBuildRequiredMemoryEstimate VT : %.3f MB"),MemoryEstimate/(1024*1024.f));
-
-		return MemoryEstimate;
+		return UE::TextureBuildUtilities::GetVirtualTextureRequiredMemoryEstimate(InBuildSettingsPerLayer, LayerFormats, Blocks);
 	}
 	else
 	{
 		// non VT
-		
-		if ( Source.GetNumBlocks() != 1 || Source.GetNumLayers() != 1 )
-		{
-			return -1; /* Requires VT enabled. */
-		}
-
-		// Compute the memory it should take to uncompress the bulkdata in memory
-		int64 TotalSourceBytes = 0;
-
-		FTextureSourceBlock SourceBlock;
-		Source.GetBlock(0, SourceBlock);
-
-		for (int32 MipIndex = 0; MipIndex < SourceBlock.NumMips; ++MipIndex)
-		{
-			TotalSourceBytes += Source.CalcMipSize(0, 0, MipIndex);
-		}
-		
-		if ( TotalSourceBytes <= 0 )
-		{
-			return -1; /* Unknown */
-		}
-		
-		const FTextureBuildSettings & BuildSettings = InSettingsPerLayerFetchFirst[0];
-		
-		int32 TargetSizeX, TargetSizeY, TargetSizeZ;
-		UE::TextureBuildUtilities::GetPowerOfTwoTargetTextureSize(SourceBlock.SizeX,SourceBlock.SizeY,SourceBlock.NumSlices,
-			BuildSettings.bVolume, (ETexturePowerOfTwoSetting::Type)BuildSettings.PowerOfTwoMode, 
-			BuildSettings.ResizeDuringBuildX, BuildSettings.ResizeDuringBuildY, 
-			TargetSizeX, TargetSizeY, TargetSizeZ);
-
-		int64 TotalTopMipNumPixels = (int64)TargetSizeX * TargetSizeY * TargetSizeZ;
-
-		// assume full mip chain :
-		int64 TotalNumPixels = (TotalTopMipNumPixels * 4)/3;
-
-		// actually we have each mip twice for the float image filter phase so this is under-counting
-		//	but that isn't held allocated while the output is made, so it can overlap with that mem
-		int64 IntermediateFloatColorBytes = TotalNumPixels * sizeof(FLinearColor);
-		
-		int64 MemoryEstimate = TotalSourceBytes + IntermediateFloatColorBytes;
-	
-		EPixelFormat PixelFormat = GetOutputPixelFormat(BuildSettings);
-
-		if ( PixelFormat == PF_Unknown )
-		{
-			return -1; /* Unknown */
-		}
-
-		const FPixelFormatInfo & PFI = GPixelFormats[PixelFormat];
-
-		const int64 OutputSizeBytes = ( TotalNumPixels * PFI.BlockBytes ) / ( PFI.BlockSizeX * PFI.BlockSizeY );
-
-		MemoryEstimate += OutputSizeBytes;
-
-		// check to see if it's uncompressed or a BCN format :
-		if ( IsDXTCBlockCompressedTextureFormat(PixelFormat) )
-		{
-			// block-compressed format ; assume it's using Oodle Texture
-			
-			if ( bRDO )
-			{
-				// two more copies in outputsize
-				// baseline encode + UT or Layout
-				MemoryEstimate += OutputSizeBytes*2;
-			}
-
-			// you also have to convert the float surface to an input format for Oodle
-			//	this copy is done in TFO
-			//  Oodle then allocs another copy to swizzle into blocks before encoding
-
-			int IntermediateBytesPerPixel;
-			bool bNeedsIntermediateCopy = true;
-
-			// this matches the logic in TextureFormatOodle :
-			if ( PixelFormat == PF_BC6H )
-			{
-				IntermediateBytesPerPixel = 16; //RGBAF32
-				bNeedsIntermediateCopy = false; // no intermediate used in TFO (float source kept), 1 blocksurf
-			}
-			else if ( PixelFormat == PF_BC4 || PixelFormat == PF_BC5 )
-			{
-				// changed: TFO uses 2_U16 now (4 byte intermediate)
-				IntermediateBytesPerPixel = 8; // RGBA16
-			}
-			else
-			{
-				IntermediateBytesPerPixel = 4; // RGBA8
-			}
-
-			int NumIntermediateCopies = 1; // BlockSurf
-			if ( bNeedsIntermediateCopy ) NumIntermediateCopies ++;
-
-			MemoryEstimate += NumIntermediateCopies * IntermediateBytesPerPixel * TotalNumPixels;
-			
-			if ( bRDO )
-			{
-				// activity map for whole image :
-				// (this has changed in newer versions of Oodle Texture)
-
-				// Phase1 = computing activity map
-				int ActivityBytesPerPixel;
-
-				if ( PixelFormat == PF_BC4 ) ActivityBytesPerPixel = 12;
-				else if ( PixelFormat == PF_BC5 ) ActivityBytesPerPixel = 16;
-				else ActivityBytesPerPixel = 24;
-
-				int64 RDOPhase1MemUse = ActivityBytesPerPixel * TotalNumPixels;
-
-				// Phase2 = cut into slices, encode each slice
-				// per-slice data structure memory use
-				// non-RDO is all on stack so zero
-
-				// fewer workers for small images ; roughly one slice per 64 KB of output
-				//int64 NumberofSlices = FMath::DivideAndRoundUp<int64>(OutputSizeBytes,64*1024);
-				int64 PixelsPerSlice = (64*1024*TotalNumPixels)/OutputSizeBytes;
-				int64 NumberofSlices = FMath::DivideAndRoundUp<int64>(TotalNumPixels,PixelsPerSlice);
-				if ( NumberofSlices <= 4 )
-				{
-					PixelsPerSlice = TotalNumPixels / NumberofSlices;
-				}
-			
-				int64 MemoryUsePerWorker = PixelsPerSlice * GetWithinSliceRDOMemoryUsePerPixel(PixelFormat);
-					// MemoryUsePerWorker is around 10 MB
-				int64 NumberOfWorkers = FMath::Min(NumberofSlices,MaxNumberOfWorkers);
-			
-				int64 RDOPhase2MemUse = 4 * TotalNumPixels; // activity map held on whole image
-				RDOPhase2MemUse += NumberOfWorkers * MemoryUsePerWorker;
-
-				// usually phase2 is higher
-				// but on large BC6 images on machines with low core counts, phase1 can be higher
-
-				MemoryEstimate += FMath::Max(RDOPhase1MemUse,RDOPhase2MemUse);
-			}
-		}
-		else if (IsASTCBlockCompressedTextureFormat(PixelFormat))
-		{
-			// ASTCenc does an entermediate copy to RGBA16F for HDR formats and RGBA8 for LDR
-			MemoryEstimate += (IsHDR(PixelFormat) ? 8 : 4) * TotalNumPixels;
-			// internal memory use of ASTCenc is not estimated
-			// @todo : fix me
-		}
-		else
-		{
-			// note: memory ues of non-Oodle encoders is not estimated
-			// @todo : fix me
-		}
-		
-		MemoryEstimate += 64 * 1024; // overhead room
-		
-		//UE_LOG(LogTexture,Display,TEXT("GetBuildRequiredMemoryEstimate non-VT : %.3f MB"),MemoryEstimate/(1024*1024.f));
-
-		return MemoryEstimate;
-
-		// @todo Oodle : not right for volumes & latlong cubes
-		// @todo Oodle : not right with Composite , CPU textures
+		FImageInfo Mip0Info;
+		Source.GetMipImageInfo(Mip0Info, 0, 0, 0);
+		return UE::TextureBuildUtilities::GetPhysicalTextureBuildMemoryEstimate(InBuildSettingsPerLayer, Mip0Info, Source.GetNumMips());
 	}
 }
 
@@ -2853,7 +2494,6 @@ public:
 		UE::DerivedData::FSharedString& TexturePath,
 		UTexture* InTexture,
 		bool bUseCompositeTexture,
-		int64 RequiredMemoryEstimate,
 		const UE::DerivedData::FUtf8SharedString& FunctionName,
 		const UE::DerivedData::FUtf8SharedString& TilingFunctionName,
 		const FTextureBuildSettings& InBuildSettings,
@@ -2864,7 +2504,7 @@ public:
 	{
 		using namespace UE::DerivedData;
 
-		FBuildDefinition BaseDefinition = CreateDefinition(InBuild, *InTexture, TexturePath, FunctionName, InBuildSettings, bUseCompositeTexture, RequiredMemoryEstimate);
+		FBuildDefinition BaseDefinition = CreateDefinition(InBuild, *InTexture, TexturePath, FunctionName, InBuildSettings, bUseCompositeTexture);
 		FBuildDefinition* RunDefinition = &BaseDefinition;
 
 		// If we have a build chain, then the next build determines what the output is as the data they
@@ -2977,7 +2617,6 @@ public:
 		
 		// we don't support VT layers here (no SettingsPerLayer)
 		check( Texture.Source.GetNumLayers() == 1 );
-		int64 RequiredMemoryEstimate = GetBuildRequiredMemoryEstimate(&Texture,&InSettingsFetchOrBuild);
 
 		// Debug string.
 		FSharedString TexturePath;
@@ -2999,7 +2638,7 @@ public:
 				EncodeSpeed = ETextureEncodeSpeed::FinalIfAvailable;
 			}
 
-			StatusMessage.Emplace(ComposeTextureBuildText(Texture, InSettingsFetchOrBuild, EncodeSpeed, RequiredMemoryEstimate, EnumHasAnyFlags(Flags, ETextureCacheFlags::ForVirtualTextureStreamingBuild)));
+			StatusMessage.Emplace(ComposeTextureBuildText(Texture, InSettingsFetchOrBuild, EncodeSpeed, GetBuildRequiredMemoryEstimate(&Texture, &InSettingsFetchOrBuild), EnumHasAnyFlags(Flags, ETextureCacheFlags::ForVirtualTextureStreamingBuild)));
 		}
 		
 
@@ -3036,7 +2675,6 @@ public:
 			TexturePath,
 			&Texture,
 			bUseCompositeTexture,
-			RequiredMemoryEstimate,
 			FunctionName,
 			TilingFunctionName,
 			InSettingsFetchOrBuild,
@@ -3055,7 +2693,6 @@ public:
 				TexturePath,
 				&Texture,
 				bUseCompositeTexture,
-				RequiredMemoryEstimate,
 				FunctionName,
 				TilingFunctionName,
 				*InSettingsFetchFirst,
@@ -3092,12 +2729,11 @@ public:
 		const UE::DerivedData::FSharedString& TexturePath,
 		const UE::DerivedData::FUtf8SharedString& FunctionName,
 		const FTextureBuildSettings& Settings,
-		const bool bUseCompositeTexture,
-		const int64 RequiredMemoryEstimate)
+		const bool bUseCompositeTexture)
 	{
 		UE::DerivedData::FBuildDefinitionBuilder DefinitionBuilder = Build.CreateDefinition(TexturePath, FunctionName);
 		DefinitionBuilder.AddConstant(UTF8TEXTVIEW("EngineParameters"), UE::TextureBuildUtilities::TextureEngineParameters::ToCompactBinaryWithDefaults(GenerateTextureEngineParameters()));
-		DefinitionBuilder.AddConstant(UTF8TEXTVIEW("Settings"), SaveTextureBuildSettings(Texture, Settings, 0, bUseCompositeTexture, RequiredMemoryEstimate));
+		DefinitionBuilder.AddConstant(UTF8TEXTVIEW("Settings"), SaveTextureBuildSettings(Texture, Settings, 0, bUseCompositeTexture));
 
 		// Texture.Source must be uncompressed for TextureBuildFunction
 		Texture.Source.RemoveCompression();
@@ -3603,16 +3239,8 @@ FTexturePlatformData::FStructuredDerivedDataKey CreateTextureDerivedDataKey(
 		bool bUseCompositeTexture = false;
 		if (FTextureBuildTask::IsTextureValidForBuilding(Texture, CacheFlags, Settings.bCPUAccessible, bUseCompositeTexture))
 		{
-			// this is just to make DDC Key so I don't need RequiredMemoryEstimate
-			// but it goes in the the DDC Key, so I have to compute it
-			// how do I pass something to TBF without it going in the DDC Key ? -> currently you can't
-			// @todo Oodle : RequiredMemoryEstimate goes in the key for DDC2, not ideal
-			// @@ remove this and compute it in the worker instead so it doesn't affect the key, currently depends on the worker
-			// count so we get a diff key per machine.
 			check( Texture.Source.GetNumLayers() == 1 ); // no SettingsPerLayer here
-			int64 RequiredMemoryEstimate = GetBuildRequiredMemoryEstimate(&Texture,&Settings);
-
-			FBuildDefinition Definition = FTextureBuildTask::CreateDefinition(Build, Texture, TexturePath.ToView(), FunctionName, *UseSettings, bUseCompositeTexture, RequiredMemoryEstimate);
+			FBuildDefinition Definition = FTextureBuildTask::CreateDefinition(Build, Texture, TexturePath.ToView(), FunctionName, *UseSettings, bUseCompositeTexture);
 			FBuildDefinition* ParentDefinition = &Definition;
 			TOptional<FBuildDefinition> TilingDefinition;
 			if (TilingFunctionName.IsEmpty() == false)
diff --git a/Engine/Source/Runtime/Engine/Private/TextureDerivedDataTask.h b/Engine/Source/Runtime/Engine/Private/TextureDerivedDataTask.h
--- a/Engine/Source/Runtime/Engine/Private/TextureDerivedDataTask.h
+++ b/Engine/Source/Runtime/Engine/Private/TextureDerivedDataTask.h
@@ -313,6 +313,4 @@ FTexturePlatformData::FStructuredDerivedDataKey CreateTextureDerivedDataKey(
 	ETextureCacheFlags CacheFlags,
 	const FTextureBuildSettings& Settings);
 
-EPixelFormat GetOutputPixelFormat(const FTextureBuildSettings & BuildSettings);
-
 #endif // WITH_EDITOR
