Prevent mobile-specific SVT textures to be displayed if Virtual Textures on mobile is disabled in the project settings. If mobile support for virtual textures is to be used on certain mobile platforms only, the project settings should still be enabled (so that the editor can know that at least one mobile platform supports it) and those platforms should have r.Mobile.VirtualTextures set to 0 in the .ini files.

Note : this is effectively a backout of 32885946, which was itself a partial backout of 32792673.

diff --git a/Engine/Source/Runtime/Engine/Classes/VT/VirtualTextureBuilder.h b/Engine/Source/Runtime/Engine/Classes/VT/VirtualTextureBuilder.h
--- a/Engine/Source/Runtime/Engine/Classes/VT/VirtualTextureBuilder.h
+++ b/Engine/Source/Runtime/Engine/Classes/VT/VirtualTextureBuilder.h
@@ -77,5 +77,8 @@ protected:
 	//~ Begin UObject Interface
 	virtual void Serialize(FArchive& Ar) override;
 	virtual void PostLoad() override;
+#if WITH_EDITOR
+	virtual bool CanEditChange(const FProperty* InProperty) const override;
+#endif // WITH_EDITOR
 	//~ End UObject Interface
 };
diff --git a/Engine/Source/Runtime/Engine/Private/VT/VirtualTextureBuilder.cpp b/Engine/Source/Runtime/Engine/Private/VT/VirtualTextureBuilder.cpp
--- a/Engine/Source/Runtime/Engine/Private/VT/VirtualTextureBuilder.cpp
+++ b/Engine/Source/Runtime/Engine/Private/VT/VirtualTextureBuilder.cpp
@@ -91,6 +91,25 @@ UVirtualTexture2D* UVirtualTextureBuilder::GetVirtualTexture(EShadingPath Shadin
 
 #if WITH_EDITOR
 
+bool UVirtualTextureBuilder::CanEditChange(const FProperty* InProperty) const
+{
+	bool bCanEdit = Super::CanEditChange(InProperty);
+
+	static const TConsoleVariableData<int32>* CVarVirtualTextureOnMobile = IConsoleManager::Get().FindTConsoleVariableDataInt(TEXT("r.Mobile.VirtualTextures"));
+	check(CVarVirtualTextureOnMobile);
+	const bool bEnableVirtualTextureOnMobile = (CVarVirtualTextureOnMobile->GetValueOnGameThread() != 0);
+
+	if (InProperty->GetFName() == GET_MEMBER_NAME_CHECKED(UVirtualTextureBuilder, bSeparateTextureForMobile))
+	{
+		bCanEdit &= bEnableVirtualTextureOnMobile;
+	}
+	else if (InProperty->GetFName() == GET_MEMBER_NAME_CHECKED(UVirtualTextureBuilder, TextureMobile))
+	{
+		bCanEdit &= bEnableVirtualTextureOnMobile && bSeparateTextureForMobile;
+	}
+	return bCanEdit;
+}
+
 static void BuildVirtualTexture2D(UVirtualTexture2D* Texture, FVirtualTextureBuildDesc const& BuildDesc, bool bWaitForCompilation)
 {
 	Texture->VirtualTextureStreaming = true;
