Fix comment on debug display for stream cache

diff --git a/Engine/Source/Runtime/Engine/Private/AudioStreamingCache.cpp b/Engine/Source/Runtime/Engine/Private/AudioStreamingCache.cpp
--- a/Engine/Source/Runtime/Engine/Private/AudioStreamingCache.cpp
+++ b/Engine/Source/Runtime/Engine/Private/AudioStreamingCache.cpp
@@ -122,7 +122,7 @@ static int32 DebugViewCVar = 2;
 FAutoConsoleVariableRef CVarDebugView(
 	TEXT("au.streamcaching.DebugView"),
 	DebugViewCVar,
-	TEXT("Enables the comparison of FObjectKeys when comparing Stream Cache Chunk Keys.  Without this FName collisions could occur if 2 SoundWaves have the same name.\n")
+	TEXT("Controls Drawing the Debug Display for the Stream Cache.\n")
 	TEXT("0: Legacy, 1: Default, 2: Averaged View, 3: High Detail View"),
 	ECVF_Default);
 
