#ue5 -CharacterMovement: Add option to disable advancing CurrentClientTimeStamp when server performs a ForcePositionUpdate(), now set to false.

Advancing the client's time can cause the server to outrun the client and prevent the client from catching up until the updates are stopped on the server. However in some cases the server can be prevented from ever stopping the forced updates (or not within a relatively long period of time), which causes long client hangs. It's preferable to catch the client up as quickly as possible once the client connection is restored (which was already the intent in VerifyClientTimeStamp(), but server time advancement could still outrun the client catchup).

New CVars:
p.NetServerForcePositionUpdateAdvanceClientTimestamp (true/false) - controls new behavior

#tests PIE with heavy simulated bursty client packet loss

diff --git a/Engine/Source/Runtime/Engine/Private/Components/CharacterMovementComponent.cpp b/Engine/Source/Runtime/Engine/Private/Components/CharacterMovementComponent.cpp
--- a/Engine/Source/Runtime/Engine/Private/Components/CharacterMovementComponent.cpp
+++ b/Engine/Source/Runtime/Engine/Private/Components/CharacterMovementComponent.cpp
@@ -207,6 +207,13 @@ namespace CharacterMovementCVars
 		TEXT("0: Disable, 1: Enable"),
 		ECVF_Default);
 
+	static bool bNetServerForcePositionUpdateAdvanceClientTimestamp = false;
+	FAutoConsoleVariableRef CVarNetServerForcePositionUpdateAdvanceClientTimestamp(
+		TEXT("p.NetServerForcePositionUpdateAdvanceClientTimestamp"),
+		bNetServerForcePositionUpdateAdvanceClientTimestamp,
+		TEXT("If enabled, the server will advance client timestamps when running ForcePositionUpdate() on client timeout.\n"),
+		ECVF_Default);
+
 
 	static int32 ReplayLerpAcceleration = 0;
 	FAutoConsoleVariableRef CVarReplayLerpAcceleration(
@@ -8428,8 +8435,12 @@ bool UCharacterMovementComponent::ForcePositionUpdate(float DeltaTime)
 
 	FNetworkPredictionData_Server_Character* ServerData = GetPredictionData_Server_Character();
 
-	// Increment client timestamp so we reject client moves after this new simulated time position.
-	ServerData->CurrentClientTimeStamp += DeltaTime;
+	if (CharacterMovementCVars::bNetServerForcePositionUpdateAdvanceClientTimestamp)
+	{
+		// Increment client timestamp so we reject client moves after this new simulated time position.
+		// See handling of CurrentClientTimeStamp in VerifyClientTimeStamp().
+		ServerData->CurrentClientTimeStamp += DeltaTime;
+	}
 
 	// Increment server timestamp so ServerLastTransformUpdateTimeStamp gets changed if there is an actual movement.
 	const double SavedServerTimestamp = ServerData->ServerAccumulatedClientTimeStamp;
@@ -9070,12 +9081,7 @@ bool UCharacterMovementComponent::VerifyClientTimeStamp(float TimeStamp, FNetwor
 			// Also apply the reset to any active root motions.
 			CurrentRootMotion.ApplyTimeStampReset(MinTimeBetweenTimeStampResets);
 		}
-		else
-		{
-			UE_LOG(LogNetPlayerMovement, VeryVerbose, TEXT("TimeStamp %f Accepted! CurrentTimeStamp: %f"), TimeStamp, ServerData.CurrentClientTimeStamp);
-			ProcessClientTimeStampForTimeDiscrepancy(TimeStamp, ServerData);
-		}
-
+		
 		if (bFirstMoveAfterForcedUpdates)
 		{
 			// We have been performing ForcedUpdates because we hadn't received any moves from this connection in a while but we've now received a new move!
@@ -9089,6 +9095,11 @@ bool UCharacterMovementComponent::VerifyClientTimeStamp(float TimeStamp, FNetwor
 				ServerData.ServerTimeStamp = World->GetTimeSeconds();
 			}
 		}
+		else if (!bTimeStampResetDetected)
+		{
+			UE_LOG(LogNetPlayerMovement, VeryVerbose, TEXT("TimeStamp %f Accepted! CurrentTimeStamp: %f"), TimeStamp, ServerData.CurrentClientTimeStamp);
+			ProcessClientTimeStampForTimeDiscrepancy(TimeStamp, ServerData);
+		}
 	}
 	else
 	{
