Fix runtime DumpMaterialStats. Added info in the dump if the shadermaps where actually used for rendering (meaning they created at least one RHIShader).

diff --git a/Engine/Source/Runtime/Engine/Private/Materials/MaterialShader.cpp b/Engine/Source/Runtime/Engine/Private/Materials/MaterialShader.cpp
--- a/Engine/Source/Runtime/Engine/Private/Materials/MaterialShader.cpp
+++ b/Engine/Source/Runtime/Engine/Private/Materials/MaterialShader.cpp
@@ -3786,8 +3786,14 @@ void DumpMaterialStats(EShaderPlatform Platform)
 	TMultiMap<FString, TShaderRef<FShader>> MaterialToShaderMap;
 	TMultiMap<FString, FShaderPipeline*> MaterialToShaderPipelineMap;
 
+	struct FMaterialInfo
+	{
+		uint32 CodeSize;
+		bool bUseForRendering;
+	};
+
 	// Set of material names.
-	TSet<FString> MaterialNames;
+	TMap<FString, FMaterialInfo> MaterialStats;
 
 	// Look at in-memory shader use.
 	FScopeLock AllMatSMAccess(&FMaterialShaderMap::AllMaterialShaderMapsGuard);
@@ -3798,10 +3804,24 @@ void DumpMaterialStats(EShaderPlatform Platform)
 		TArray<FShaderPipelineRef> ShaderPipelines;
 		MaterialShaderMap->GetShaderList(Shaders);
 		MaterialShaderMap->GetShaderPipelineList(ShaderPipelines);
+		FShaderMapResource* Resource = MaterialShaderMap->GetResource();
+
+		FMaterialInfo MI;
+		MI.CodeSize = 0;
 
 		// Add friendly name to list of materials.
 		FString FriendlyName = MaterialShaderMap->GetFriendlyName();
-		MaterialNames.Add(FriendlyName);
+
+#if !WITH_EDITORONLY_DATA
+		if (FriendlyName.IsEmpty() && Resource != nullptr)
+		{
+			FriendlyName = Resource->GetOwnerName().ToString();
+			MI.CodeSize = MaterialShaderMap->GetFrozenContentSize() + Resource->GetSizeBytes();
+			MI.bUseForRendering = Resource->ContainsAtLeastOneRHIShaderCreated();
+		}
+#endif
+
+		MaterialStats.FindOrAdd(FriendlyName, MI);
 
 		// Add shaders to mapping per friendly name as there might be multiple
 		for (auto& KeyValue : Shaders)
@@ -3823,45 +3843,58 @@ void DumpMaterialStats(EShaderPlatform Platform)
 	MaterialViewer.AddColumn(TEXT("Name"));
 	MaterialViewer.AddColumn(TEXT("Shaders"));
 	MaterialViewer.AddColumn(TEXT("Code Size"));
+	MaterialViewer.AddColumn(TEXT("Used For Rendering"));
 	MaterialViewer.AddColumn(TEXT("Pipelines"));
 	MaterialViewer.CycleRow();
 
 	// Iterate over all materials, gathering shader stats.
-	int32 TotalCodeSize		= 0;
-	int32 TotalShaderCount	= 0;
+	int32 TotalCodeSize = 0;
+	int32 TotalShaderCount = 0;
 	int32 TotalShaderPipelineCount = 0;
-	for( TSet<FString>::TConstIterator It(MaterialNames); It; ++It )
+	int32 TotalUsedShaderMapCount = 0;
+
+	for (TMap<FString, FMaterialInfo>::TIterator It(MaterialStats); It; ++It)
 	{
 		// Retrieve list of shaders in map.
 		TArray<TShaderRef<FShader>> Shaders;
-		MaterialToShaderMap.MultiFind( *It, Shaders );
+		MaterialToShaderMap.MultiFind(It.Key(), Shaders);
 		TArray<FShaderPipeline*> ShaderPipelines;
-		MaterialToShaderPipelineMap.MultiFind(*It, ShaderPipelines);
-		
+		MaterialToShaderPipelineMap.MultiFind(It.Key(), ShaderPipelines);
+
 		// Iterate over shaders and gather stats.
-		int32 CodeSize = 0;
-		for( int32 ShaderIndex=0; ShaderIndex<Shaders.Num(); ShaderIndex++ )
+		FMaterialInfo& MI = It.Value();
+
+#if WITH_EDITORONLY_DATA
+		if (MI.CodeSize == 0)
 		{
-			const TShaderRef<FShader>& Shader = Shaders[ShaderIndex];
-			CodeSize += Shader->GetCodeSize();
+			for (int32 ShaderIndex = 0; ShaderIndex < Shaders.Num(); ShaderIndex++)
+			{
+				const TShaderRef<FShader>& Shader = Shaders[ShaderIndex];
+				MI.CodeSize += Shader->GetCodeSize();
+			}
 		}
+#endif
 
-		TotalCodeSize += CodeSize;
+		TotalCodeSize += MI.CodeSize;
 		TotalShaderCount += Shaders.Num();
 		TotalShaderPipelineCount += ShaderPipelines.Num();
+		TotalUsedShaderMapCount += MI.bUseForRendering ? 1 : 0;
 
 		// Dump stats
-		MaterialViewer.AddColumn(**It);
-		MaterialViewer.AddColumn(TEXT("%u"),Shaders.Num());
-		MaterialViewer.AddColumn(TEXT("%u"),CodeSize);
+		MaterialViewer.AddColumn(*It.Key());
+		MaterialViewer.AddColumn(TEXT("%u"), Shaders.Num());
+		MaterialViewer.AddColumn(TEXT("%u"), MI.CodeSize);
+		MaterialViewer.AddColumn(TEXT("%u"), MI.bUseForRendering);
 		MaterialViewer.AddColumn(TEXT("%u"), ShaderPipelines.Num());
 		MaterialViewer.CycleRow();
 	}
 
 	// Add a total row.
+	MaterialViewer.CycleRow();
 	MaterialViewer.AddColumn(TEXT("Total"));
-	MaterialViewer.AddColumn(TEXT("%u"),TotalShaderCount);
-	MaterialViewer.AddColumn(TEXT("%u"),TotalCodeSize);
+	MaterialViewer.AddColumn(TEXT("%u"), TotalShaderCount);
+	MaterialViewer.AddColumn(TEXT("%u"), TotalCodeSize);
+	MaterialViewer.AddColumn(TEXT("%u"), TotalUsedShaderMapCount);
 	MaterialViewer.AddColumn(TEXT("%u"), TotalShaderPipelineCount);
 	MaterialViewer.CycleRow();
 #endif // ALLOW_DEBUG_FILES && ALLOW_SHADERMAP_DEBUG_DATA
diff --git a/Engine/Source/Runtime/RenderCore/Public/Shader.h b/Engine/Source/Runtime/RenderCore/Public/Shader.h
--- a/Engine/Source/Runtime/RenderCore/Public/Shader.h
+++ b/Engine/Source/Runtime/RenderCore/Public/Shader.h
@@ -360,6 +360,11 @@ public:
 		return RHIShaders[ShaderIndex].load(std::memory_order_acquire) != nullptr;
 	}
 
+	inline bool ContainsAtLeastOneRHIShaderCreated() const
+	{
+		return bAtLeastOneRHIShaderCreated;
+	}
+
 	inline FRHIShader* GetShader(int32 ShaderIndex, bool bRequired = true)
 	{
 		// This is a double checked locking. This trickery arises from the fact that we're
