[TSAN]
 - Fix benign warning by using relaxed atomic

#rnx

diff --git a/Engine/Source/Runtime/Engine/Private/DistanceFieldAtlas.cpp b/Engine/Source/Runtime/Engine/Private/DistanceFieldAtlas.cpp
--- a/Engine/Source/Runtime/Engine/Private/DistanceFieldAtlas.cpp
+++ b/Engine/Source/Runtime/Engine/Private/DistanceFieldAtlas.cpp
@@ -520,9 +520,10 @@ void FDistanceFieldAsyncQueue::AddTask(FAsyncDistanceFieldTask* Task)
 		return;
 	}
 
-	if (!MeshUtilities)
+	// This is ok to only use relaxed atomics here since if threads are racing they should all get the same module pointer to store.
+	if (!MeshUtilities.load(std::memory_order_relaxed))
 	{
-		MeshUtilities = &FModuleManager::Get().LoadModuleChecked<IMeshUtilities>(TEXT("MeshUtilities"));
+		MeshUtilities.store(&FModuleManager::Get().LoadModuleChecked<IMeshUtilities>(TEXT("MeshUtilities")), std::memory_order_relaxed);
 	}
 	
 	const bool bUseAsyncBuild = GUseAsyncDistanceFieldBuildQueue || !IsInGameThread();
@@ -723,7 +724,7 @@ void FDistanceFieldAsyncQueue::Build(FAsyncDistanceFieldTask* Task, FQueuedThrea
 		TRACE_CPUPROFILER_EVENT_SCOPE(FDistanceFieldAsyncQueue::Build);
 
 		const FStaticMeshLODResources& LODModel = Task->GenerateSource->GetRenderData()->LODResources[0];
-		MeshUtilities->GenerateSignedDistanceFieldVolumeData(
+		MeshUtilities.load(std::memory_order_relaxed)->GenerateSignedDistanceFieldVolumeData(
 			Task->StaticMesh->GetName(),
 			Task->SourceMeshData,
 			LODModel,
diff --git a/Engine/Source/Runtime/Engine/Public/DistanceFieldAtlas.h b/Engine/Source/Runtime/Engine/Public/DistanceFieldAtlas.h
--- a/Engine/Source/Runtime/Engine/Public/DistanceFieldAtlas.h
+++ b/Engine/Source/Runtime/Engine/Public/DistanceFieldAtlas.h
@@ -430,7 +430,7 @@ private:
 
 	FDelegateHandle PostReachabilityAnalysisHandle;
 
-	class IMeshUtilities* MeshUtilities;
+	std::atomic<class IMeshUtilities*> MeshUtilities;
 
 	mutable FCriticalSection CriticalSection;
 
