Only destroy the state after we commit

diff --git a/Engine/Source/Runtime/Engine/Private/Components/ActorComponent.cpp b/Engine/Source/Runtime/Engine/Private/Components/ActorComponent.cpp
--- a/Engine/Source/Runtime/Engine/Private/Components/ActorComponent.cpp
+++ b/Engine/Source/Runtime/Engine/Private/Components/ActorComponent.cpp
@@ -40,6 +40,7 @@
 #include "PSOPrecacheMaterial.h"
 #include "Materials/MaterialInterface.h"
 #include "ObjectCacheContext.h"
+#include "AutoRTFM/AutoRTFM.h"
 
 #if WITH_EDITOR
 #include "Kismet2/ComponentEditorUtils.h"
@@ -1817,22 +1818,26 @@ void UActorComponent::ExecuteRegisterEvents(FRegisterComponentContext* Context)
 
 void UActorComponent::ExecuteUnregisterEvents()
 {
-	DestroyPhysicsState();
-
-	if(bRenderStateCreated)
+	// Delay the destroying the physics, as well as the rendering state until after we commit
+	UE_AUTORTFM_ONCOMMIT(
 	{
-		SCOPE_CYCLE_COUNTER(STAT_ComponentDestroyRenderState);
-		checkf(bRegistered, TEXT("Component has render state when not registered (%s)"), *GetFullName());
-		DestroyRenderState_Concurrent();
-		checkf(!bRenderStateCreated, TEXT("Failed to route DestroyRenderState_Concurrent (%s)"), *GetFullName());
-	}
+		DestroyPhysicsState();
 
-	if(bRegistered)
-	{
-		SCOPE_CYCLE_COUNTER(STAT_ComponentOnUnregister);
-		OnUnregister();
-		checkf(!bRegistered, TEXT("Failed to route OnUnregister (%s)"), *GetFullName());
-	}
+		if (bRenderStateCreated)
+		{
+			SCOPE_CYCLE_COUNTER(STAT_ComponentDestroyRenderState);
+			checkf(bRegistered, TEXT("Component has render state when not registered (%s)"), *GetFullName());
+			DestroyRenderState_Concurrent();
+			checkf(!bRenderStateCreated, TEXT("Failed to route DestroyRenderState_Concurrent (%s)"), *GetFullName());
+		}
+
+		if (bRegistered)
+		{
+			SCOPE_CYCLE_COUNTER(STAT_ComponentOnUnregister);
+			OnUnregister();
+			checkf(!bRegistered, TEXT("Failed to route OnUnregister (%s)"), *GetFullName());
+		}
+	});
 }
 
 void UActorComponent::ReregisterComponent()
