[NetPhysics] Per actor hardsnap distance setting for default physics replication.

diff --git a/Engine/Source/Runtime/Engine/Private/PhysicsEngine/PhysicsReplication.cpp b/Engine/Source/Runtime/Engine/Private/PhysicsEngine/PhysicsReplication.cpp
--- a/Engine/Source/Runtime/Engine/Private/PhysicsEngine/PhysicsReplication.cpp
+++ b/Engine/Source/Runtime/Engine/Private/PhysicsEngine/PhysicsReplication.cpp
@@ -588,7 +588,20 @@ bool FPhysicsReplication::ApplyRigidBodyState(float DeltaSeconds, FBodyInstance*
 	const float LinearVelocityCoefficient = CharacterMovementCVars::LinearVelocityCoefficient >= 0.0f ? CharacterMovementCVars::LinearVelocityCoefficient : ErrorCorrection.LinearVelocityCoefficient;
 	const float AngleLerp = CharacterMovementCVars::AngleLerp >= 0.0f ? CharacterMovementCVars::AngleLerp : ErrorCorrection.AngleLerp;
 	const float AngularVelocityCoefficient = CharacterMovementCVars::AngularVelocityCoefficient >= 0.0f ? CharacterMovementCVars::AngularVelocityCoefficient : ErrorCorrection.AngularVelocityCoefficient;
-	const float MaxLinearHardSnapDistance = CharacterMovementCVars::MaxLinearHardSnapDistance >= 0.f ? CharacterMovementCVars::MaxLinearHardSnapDistance : ErrorCorrection.MaxLinearHardSnapDistance;
+	float MaxLinearHardSnapDistance = CharacterMovementCVars::MaxLinearHardSnapDistance >= 0.f ? CharacterMovementCVars::MaxLinearHardSnapDistance : ErrorCorrection.MaxLinearHardSnapDistance;
+
+	bool bHardsnapLegacyInPT = PhysicsReplicationCVars::DefaultReplicationCVars::bHardsnapLegacyInPT;
+	bool bCorrectConnectedBodies = PhysicsReplicationCVars::DefaultReplicationCVars::bCorrectConnectedBodies;
+	bool bCorrectConnectedBodiesFriction = PhysicsReplicationCVars::DefaultReplicationCVars::bCorrectConnectedBodiesFriction;
+
+	// Assign per-actor settings from NetworkPhysicSettingsComponent if this actor has one
+	if (SettingsCurrent.Get())
+	{
+		MaxLinearHardSnapDistance = SettingsCurrent.Get()->DefaultReplicationSettings.GetMaxLinearHardSnapDistance(MaxLinearHardSnapDistance);
+		bHardsnapLegacyInPT = SettingsCurrent.Get()->DefaultReplicationSettings.GetHardsnapDefaultLegacyInPT();
+		bCorrectConnectedBodies = SettingsCurrent.Get()->DefaultReplicationSettings.GetCorrectConnectedBodies();
+		bCorrectConnectedBodiesFriction = SettingsCurrent.Get()->DefaultReplicationSettings.GetCorrectConnectedBodiesFriction();
+	}
 
 	// Get Current state
 	FRigidBodyState CurrentState;
@@ -714,17 +727,6 @@ bool FPhysicsReplication::ApplyRigidBodyState(float DeltaSeconds, FBodyInstance*
 			bool bPTHardSnapSuccess = false;
 			if (PhysicsReplicationAsync != nullptr)
 			{
-				bool bHardsnapLegacyInPT = PhysicsReplicationCVars::DefaultReplicationCVars::bHardsnapLegacyInPT;
-				bool bCorrectConnectedBodies = PhysicsReplicationCVars::DefaultReplicationCVars::bCorrectConnectedBodies;
-				bool bCorrectConnectedBodiesFriction = PhysicsReplicationCVars::DefaultReplicationCVars::bCorrectConnectedBodiesFriction;
-				
-				if (SettingsCurrent.Get())
-				{
-					bHardsnapLegacyInPT = SettingsCurrent.Get()->DefaultReplicationSettings.GetHardsnapDefaultLegacyInPT();
-					bCorrectConnectedBodies = SettingsCurrent.Get()->DefaultReplicationSettings.GetCorrectConnectedBodies();
-					bCorrectConnectedBodiesFriction = SettingsCurrent.Get()->DefaultReplicationSettings.GetCorrectConnectedBodiesFriction();
-				}
-
 				if (bHardsnapLegacyInPT)
 				{
 					if (Chaos::FSingleParticlePhysicsProxy* Proxy = static_cast<Chaos::FSingleParticlePhysicsProxy*>(BI->GetPhysicsActorHandle()))
@@ -1392,7 +1394,8 @@ bool FPhysicsReplicationAsync::DefaultReplication(Chaos::FPBDRigidParticleHandle
 	const float LinearVelocityCoefficient = CharacterMovementCVars::LinearVelocityCoefficient >= 0.0f ? CharacterMovementCVars::LinearVelocityCoefficient : ErrorCorrectionDefault.LinearVelocityCoefficient;
 	const float AngleLerp = CharacterMovementCVars::AngleLerp >= 0.0f ? CharacterMovementCVars::AngleLerp : ErrorCorrectionDefault.AngleLerp;
 	const float AngularVelocityCoefficient = CharacterMovementCVars::AngularVelocityCoefficient >= 0.0f ? CharacterMovementCVars::AngularVelocityCoefficient : ErrorCorrectionDefault.AngularVelocityCoefficient;
-	const float MaxLinearHardSnapDistance = CharacterMovementCVars::MaxLinearHardSnapDistance >= 0.f ? CharacterMovementCVars::MaxLinearHardSnapDistance : ErrorCorrectionDefault.MaxLinearHardSnapDistance;
+	float MaxLinearHardSnapDistance = CharacterMovementCVars::MaxLinearHardSnapDistance >= 0.f ? CharacterMovementCVars::MaxLinearHardSnapDistance : ErrorCorrectionDefault.MaxLinearHardSnapDistance;
+	MaxLinearHardSnapDistance = SettingsCurrent.DefaultReplicationSettings.GetMaxLinearHardSnapDistance(MaxLinearHardSnapDistance);
 
 
 	// Get Current state
diff --git a/Engine/Source/Runtime/Engine/Public/Physics/NetworkPhysicsSettingsComponent.h b/Engine/Source/Runtime/Engine/Public/Physics/NetworkPhysicsSettingsComponent.h
--- a/Engine/Source/Runtime/Engine/Public/Physics/NetworkPhysicsSettingsComponent.h
+++ b/Engine/Source/Runtime/Engine/Public/Physics/NetworkPhysicsSettingsComponent.h
@@ -83,20 +83,27 @@ struct FNetworkPhysicsSettingsDefaultReplication
 {
 	GENERATED_BODY()
 
+	UPROPERTY(config, EditDefaultsOnly, Category = "Overrides", Meta = (InlineEditConditionToggle))
+	uint32 bOverrideMaxLinearHardSnapDistance : 1;
+	// Overrides CVar: p.MaxLinearHardSnapDistance -- Hardsnap if distance between current position and extrapolated target position is larger than this value.
+	UPROPERTY(config, EditDefaultsOnly, Category = "Overrides", Meta = (EditCondition = "bOverrideMaxLinearHardSnapDistance"))
+	float MaxLinearHardSnapDistance = 400.0f;
+	float GetMaxLinearHardSnapDistance(float DefaultValue) { return bOverrideMaxLinearHardSnapDistance ? MaxLinearHardSnapDistance : DefaultValue; }
+
 	UPROPERTY(config, EditDefaultsOnly, Category = "Overrides", Meta = (InlineEditConditionToggle))
 	uint32 bOverrideDefaultLegacyHardsnapInPT : 1;
 	// Overrides CVar: p.DefaultReplication.Legacy.HardsnapInPT -- If default replication is used and it's running the legacy flow through Game Thread, allow hardsnapping to be performed on Physics Thread if async physics is enabled.
 	UPROPERTY(config, EditDefaultsOnly, Category = "Overrides", Meta = (EditCondition = "bOverrideDefaultLegacyHardsnapInPT"))
 	bool bHardsnapInPhysicsThread = PhysicsReplicationCVars::DefaultReplicationCVars::bHardsnapLegacyInPT;
 	bool GetHardsnapDefaultLegacyInPT() { return bOverrideDefaultLegacyHardsnapInPT ? bHardsnapInPhysicsThread : PhysicsReplicationCVars::DefaultReplicationCVars::bHardsnapLegacyInPT; }
-	
+
 	UPROPERTY(config, EditDefaultsOnly, Category = "Overrides", Meta = (InlineEditConditionToggle))
 	uint32 bOverrideCorrectConnectedBodies : 1;
 	// Overrides CVar: p.DefaultReplication.CorrectConnectedBodies -- When true, transform corrections will also apply to any connected physics object.
 	UPROPERTY(config, EditDefaultsOnly, Category = "Overrides", Meta = (EditCondition = "bOverrideCorrectConnectedBodies"))
 	bool bCorrectConnectedBodies = PhysicsReplicationCVars::DefaultReplicationCVars::bCorrectConnectedBodies;
 	bool GetCorrectConnectedBodies() { return bOverrideCorrectConnectedBodies ? bCorrectConnectedBodies : PhysicsReplicationCVars::DefaultReplicationCVars::bCorrectConnectedBodies; }
-	
+
 	UPROPERTY(config, EditDefaultsOnly, Category = "Overrides", Meta = (InlineEditConditionToggle))
 	uint32 bOverrideCorrectConnectedBodiesFriction : 1;
 	// Overrides CVar: p.DefaultReplication.CorrectConnectedBodiesFriction -- When true, transform correction on any connected physics object will also recalculate their friction.
