Netdriver
* Fix wrong assignment for MaxReplicationWriterObjectCount

#rnx

diff --git a/Engine/Source/Runtime/Engine/Private/NetDriver.cpp b/Engine/Source/Runtime/Engine/Private/NetDriver.cpp
--- a/Engine/Source/Runtime/Engine/Private/NetDriver.cpp
+++ b/Engine/Source/Runtime/Engine/Private/NetDriver.cpp
@@ -256,7 +256,7 @@ namespace UE::Net::Private
 			OutParams.MaxReplicatedObjectCount = ReplicationSystemConfig.MaxReplicatedObjectCount;
 		}
 
-		if (ReplicationSystemConfig.InitialNetObjectListCount != 0 )
+		if (ReplicationSystemConfig.InitialNetObjectListCount != 0)
 		{
 			OutParams.InitialNetObjectListCount = ReplicationSystemConfig.InitialNetObjectListCount;
 		}
@@ -273,7 +273,7 @@ namespace UE::Net::Private
 		
 		if (ReplicationSystemConfig.MaxReplicationWriterObjectCount!= 0)
 		{
-			OutParams.MaxReplicationWriterObjectCount = OutParams.MaxReplicationWriterObjectCount;
+			OutParams.MaxReplicationWriterObjectCount = ReplicationSystemConfig.MaxReplicationWriterObjectCount;
 		}
 
 		if (ReplicationSystemConfig.MaxDeltaCompressedObjectCount != 0)
