New material translator

diff --git a/Engine/Source/Runtime/Engine/Private/Materials/MaterialIR.cpp b/Engine/Source/Runtime/Engine/Private/Materials/MaterialIR.cpp
new file mode 100644
--- /dev/null
+++ b/Engine/Source/Runtime/Engine/Private/Materials/MaterialIR.cpp
@@ -0,0 +1,31 @@
+// Copyright Epic Games, Inc. All Rights Reserved.
+
+#include "Materials/MaterialIR.h"
+#include "Materials/MaterialIRTypes.h"
+
+#if WITH_EDITOR
+
+namespace MaterialIR
+{
+
+struct FValue;
+using FValuePtr = const FValue*;
+
+TArrayView<const FValuePtr> FVectorValue::GetComponents() const
+{
+	TArrayView<FValuePtr> Components = const_cast<FVectorValue*>(this)->GetMutableComponents();
+	return { Components.GetData(), Components.Num() };
+}
+
+TArrayView<FValuePtr> FVectorValue::GetMutableComponents()
+{
+	FArithmeticTypePtr ArithmeticType = Type->ToArithmetic();
+	check(ArithmeticType);
+
+	FValuePtr* Ptr = (FValuePtr*)static_cast<TVectorValue<1>*>(this)->Components;
+	return { Ptr, ArithmeticType->NumRows };
+}
+
+} // namespace MaterialIR
+
+#endif // #if WITH_EDITOR
diff --git a/Engine/Source/Runtime/Engine/Private/Materials/MaterialIRBuilder.cpp b/Engine/Source/Runtime/Engine/Private/Materials/MaterialIRBuilder.cpp
new file mode 100644
--- /dev/null
+++ b/Engine/Source/Runtime/Engine/Private/Materials/MaterialIRBuilder.cpp
@@ -0,0 +1,179 @@
+// Copyright Epic Games, Inc. All Rights Reserved.
+
+#include "Materials/MaterialIRBuilder.h"
+#include "Materials/MaterialIRModule.h"
+#include "Materials/MaterialIR.h"
+#include "Materials/MaterialIRTypes.h"
+#include "Shader/ShaderTypes.h"
+
+#if WITH_EDITOR
+
+namespace MaterialIR
+{
+
+struct FBuilder::FPrivate : FBuilder
+{
+	template <typename T>
+	T* New(FTypePtr InType)
+	{
+		T* Val = new T;
+		Val->Kind = T::TypeKind;
+		Val->Type = InType;
+		Module->Values.Add(Val);
+		return Val;
+	}
+};
+
+FBuilder::FBuilder()
+{
+}
+
+FBuilder::FBuilder(FMaterialIRModule* InModule)
+{
+	SetTargetModule(InModule);
+}
+
+void FBuilder::SetTargetModule(FMaterialIRModule* InModule)
+{
+	Module = InModule;
+}
+
+FValuePtr FBuilder::NewConstantFromShaderValue(const UE::Shader::FValue& InValue)
+{
+	using namespace UE::Shader;
+
+	switch (InValue.Type.ValueType)
+	{
+	case UE::Shader::EValueType::Float1: return NewConstantFloat1(InValue.AsFloatScalar());
+	case UE::Shader::EValueType::Float2: return NewConstantFloat2(FVector2f{ InValue.Component[0].Float, InValue.Component[1].Float });
+	case UE::Shader::EValueType::Float3: return NewConstantFloat3(FVector3f{ InValue.Component[0].Float, InValue.Component[1].Float, InValue.Component[2].Float });
+	case UE::Shader::EValueType::Float4: return NewConstantFloat4(FVector4f{ InValue.Component[0].Float, InValue.Component[1].Float, InValue.Component[2].Float, InValue.Component[3].Float });
+
+	case UE::Shader::EValueType::Int1: return NewConstantInt1(InValue.AsFloatScalar());
+	case UE::Shader::EValueType::Int2: return NewConstantInt2(FIntVector2{ InValue.Component[0].Int, InValue.Component[1].Int });
+	case UE::Shader::EValueType::Int3: return NewConstantInt3(FIntVector3{ InValue.Component[0].Int, InValue.Component[1].Int, InValue.Component[2].Int });
+	case UE::Shader::EValueType::Int4: return NewConstantInt4(FIntVector4{ InValue.Component[0].Int, InValue.Component[1].Int, InValue.Component[2].Int, InValue.Component[3].Int });
+	}
+
+	UE_MIR_UNREACHABLE();
+}
+
+FValuePtr FBuilder::NewConstantFloat1(float InX)
+{
+	FScalarValue* S = AsPrivate()->New<FScalarValue>(FArithmeticType::GetScalar(SK_Float));
+	S->Float = InX;
+	return S;
+}
+
+FValuePtr FBuilder::NewConstantFloat2(const FVector2f& InValue)
+{
+	FValuePtr X = NewConstantFloat1(InValue.X);
+	FValuePtr Y = NewConstantFloat1(InValue.Y);
+	return NewVector2(X, Y);
+}
+
+FValuePtr FBuilder::NewConstantFloat3(const FVector3f& InValue)
+{
+	FValuePtr X = NewConstantFloat1(InValue.X);
+	FValuePtr Y = NewConstantFloat1(InValue.Y);
+	FValuePtr Z = NewConstantFloat1(InValue.Z);
+	return NewVector3(X, Y, Z);
+}
+
+FValuePtr FBuilder::NewConstantFloat4(const FVector4f& InValue)
+{
+	FValuePtr X = NewConstantFloat1(InValue.X);
+	FValuePtr Y = NewConstantFloat1(InValue.Y);
+	FValuePtr Z = NewConstantFloat1(InValue.Z);
+	FValuePtr W = NewConstantFloat1(InValue.W);
+	return NewVector4(X, Y, Z, W);
+}
+
+FValuePtr FBuilder::NewConstantInt1(int InX)
+{
+	FScalarValue* S = AsPrivate()->New<FScalarValue>(FArithmeticType::GetScalar(SK_Int));
+	S->Integer = InX;
+	return S;
+}
+
+FValuePtr FBuilder::NewConstantInt2(const FIntVector2& InValue)
+{
+	FValuePtr X = NewConstantInt1(InValue.X);
+	FValuePtr Y = NewConstantInt1(InValue.Y);
+	return NewVector2(X, Y);
+}
+
+FValuePtr FBuilder::NewConstantInt3(const FIntVector3& InValue)
+{
+	FValuePtr X = NewConstantInt1(InValue.X);
+	FValuePtr Y = NewConstantInt1(InValue.Y);
+	FValuePtr Z = NewConstantInt1(InValue.Z);
+	return NewVector3(X, Y, Z);
+}
+
+FValuePtr FBuilder::NewConstantInt4(const FIntVector4& InValue)
+{
+	FValuePtr X = NewConstantInt1(InValue.X);
+	FValuePtr Y = NewConstantInt1(InValue.Y);
+	FValuePtr Z = NewConstantInt1(InValue.Z);
+	FValuePtr W = NewConstantInt1(InValue.W);
+	return NewVector4(X, Y, Z, W);
+}
+
+FValuePtr FBuilder::NewVector2(FValuePtr InX, FValuePtr InY)
+{
+	check(InX->Type->IsScalar());
+	check(InX->Type == InY->Type);
+
+	FVectorValue* V = AsPrivate()->New<TVectorValue<2>>(FArithmeticType::GetVector(InX->Type->ToArithmetic()->ScalarKind, 2));
+	TArrayView<FValuePtr> Components = V->GetMutableComponents();
+	Components[0] = InX;
+	Components[1] = InY;
+
+	return V;
+}
+
+FValuePtr FBuilder::NewVector3(FValuePtr InX, FValuePtr InY, FValuePtr InZ)
+{
+	check(InX->Type->IsScalar());
+	check(InX->Type == InY->Type);
+	check(InY->Type == InZ->Type);
+
+	FVectorValue* V = AsPrivate()->New<TVectorValue<3>>(FArithmeticType::GetVector(InX->Type->ToArithmetic()->ScalarKind, 3));
+	TArrayView<FValuePtr> Components = V->GetMutableComponents();
+	Components[0] = InX;
+	Components[1] = InY;
+	Components[2] = InZ;
+
+	return V;
+}
+
+FValuePtr FBuilder::NewVector4(FValuePtr InX, FValuePtr InY, FValuePtr InZ, FValuePtr InW)
+{
+	check(InX->Type->IsScalar());
+	check(InX->Type == InY->Type);
+	check(InY->Type == InZ->Type);
+	check(InZ->Type == InW->Type);
+
+	FVectorValue* V = AsPrivate()->New<TVectorValue<4>>(FArithmeticType::GetVector(InX->Type->ToArithmetic()->ScalarKind, 4));
+	TArrayView<FValuePtr> Components = V->GetMutableComponents();
+	Components[0] = InX;
+	Components[1] = InY;
+	Components[2] = InZ;
+	Components[3] = InW;
+
+	return V;
+}
+
+FSetMaterialOutputInstr* FBuilder::EmitSetMaterialOutput(EMaterialProperty InProperty, FValue* InArgValue)
+{
+	FSetMaterialOutputInstr* Instr = AsPrivate()->New<FSetMaterialOutputInstr>(nullptr);
+	Instr->Property = InProperty;
+	Instr->ArgValue = InArgValue;
+	Module->Outputs.Push(Instr);
+	return Instr;
+}
+
+} // namespace MaterialIR
+
+#endif // #if WITH_EDITOR
diff --git a/Engine/Source/Runtime/Engine/Private/Materials/MaterialIRModule.cpp b/Engine/Source/Runtime/Engine/Private/Materials/MaterialIRModule.cpp
new file mode 100644
--- /dev/null
+++ b/Engine/Source/Runtime/Engine/Private/Materials/MaterialIRModule.cpp
@@ -0,0 +1,34 @@
+// Copyright Epic Games, Inc. All Rights Reserved.
+
+#include "Materials/MaterialIRModule.h"
+#include "Materials/MaterialIR.h"
+
+#if WITH_EDITOR
+
+namespace MIR = MaterialIR;
+
+FMaterialIRModule::~FMaterialIRModule()
+{
+	Empty();
+}
+
+void FMaterialIRModule::Empty()
+{
+	for (MIR::FValue* Value : Values)
+	{
+		delete Value;
+	}
+
+	Values.Empty();
+	Outputs.Empty();
+}
+
+void FMaterialIRModule::Reset(EShaderPlatform InShaderPlatform, ERHIFeatureLevel::Type InFeatureLevel, const ITargetPlatform* InTargetPlatform)
+{
+	Empty();
+
+	ShaderPlatform = InShaderPlatform;
+	TargetPlatform = InTargetPlatform;
+}
+
+#endif
diff --git a/Engine/Source/Runtime/Engine/Private/Materials/MaterialIRModuleBuilder.cpp b/Engine/Source/Runtime/Engine/Private/Materials/MaterialIRModuleBuilder.cpp
new file mode 100644
--- /dev/null
+++ b/Engine/Source/Runtime/Engine/Private/Materials/MaterialIRModuleBuilder.cpp
@@ -0,0 +1,76 @@
+// Copyright Epic Games, Inc. All Rights Reserved.
+
+#include "Materials/MaterialIRModuleBuilder.h"
+
+#if WITH_EDITOR
+
+#include "Materials/MaterialIRModule.h"
+#include "Materials/MaterialIR.h"
+#include "Materials/MaterialIRTypes.h"
+#include "Materials/MaterialIRBuilder.h"
+#include "MaterialShared.h"
+#include "Materials/MaterialAttributeDefinitionMap.h"
+#include "Materials/Material.h"
+#include "MaterialExpressionIO.h"
+
+struct FMaterialIRModuleBuilder::FHelper
+{
+};
+
+void FMaterialIRModuleBuilder::SetSource(FMaterial* InMaterial, const FStaticParameterSet* InStaticParameters)
+{
+	Material = InMaterial;
+	StaticParameters = InStaticParameters;
+}
+
+void FMaterialIRModuleBuilder::SetPlatform(EShaderPlatform InShaderPlatform, ERHIFeatureLevel::Type InFeatureLevel, const ITargetPlatform* InTargetPlatform)
+{
+	ShaderPlatform	= InShaderPlatform;
+	FeatureLevel		= InFeatureLevel;
+	TargetPlatform	= InTargetPlatform;
+}
+
+void FMaterialIRModuleBuilder::SetTarget(FMaterialIRModule* InModule)
+{
+	Module = InModule;
+}
+
+bool FMaterialIRModuleBuilder::Build()
+{
+	Module->Reset(ShaderPlatform, FeatureLevel, TargetPlatform);
+
+	MaterialIR::FBuilder Builder{ Module };
+
+	UMaterial* BaseMaterial = Material->GetMaterialInterface()->GetMaterial();
+
+	for (int32 PropertyIndex = 0; PropertyIndex < MP_MAX; ++PropertyIndex)
+	{
+		EMaterialProperty Property = (EMaterialProperty)PropertyIndex;
+		FMaterialInputDescription Input;
+
+		if (Property == MP_CustomOutput || Property == MP_MaterialAttributes || !BaseMaterial->GetExpressionInputDescription(Property, Input))
+		{
+			continue;
+		}
+
+		MaterialIR::FSetMaterialOutputInstr* Output = Builder.EmitSetMaterialOutput(Property, nullptr);
+	
+		if (Input.bUseConstant)
+		{
+			Output->ArgValue = Builder.NewConstantFromShaderValue(Input.ConstantValue);
+		}
+		else if (!Input.Input->IsConnected() && Input.Type != UE::Shader::EValueType::Void)
+		{
+			UE::Shader::FValue Zero{ Input.Type };
+			Output->ArgValue = Builder.NewConstantFromShaderValue(Zero);
+		}
+		else if (Input.Input->IsConnected())
+		{
+			UE_MIR_UNREACHABLE();
+		}
+	}
+
+	return true;
+}
+
+#endif
diff --git a/Engine/Source/Runtime/Engine/Private/Materials/MaterialIRToHLSLTranslator.cpp b/Engine/Source/Runtime/Engine/Private/Materials/MaterialIRToHLSLTranslator.cpp
new file mode 100644
--- /dev/null
+++ b/Engine/Source/Runtime/Engine/Private/Materials/MaterialIRToHLSLTranslator.cpp
@@ -0,0 +1,385 @@
+// Copyright Epic Games, Inc. All Rights Reserved.
+
+#include "Materials/MaterialIRToHLSLTranslator.h"
+#include "Materials/MaterialIRModule.h"
+#include "Materials/MaterialIRTypes.h"
+#include "Materials/MaterialIR.h"
+
+#include "ShaderCore.h"
+#include "MaterialShared.h"
+#include "Materials/MaterialAttributeDefinitionMap.h"
+#include "Materials/Material.h"
+#include "MaterialDomain.h"
+#include "DataDrivenShaderPlatformInfo.h"
+#include "Materials/MaterialExpressionVolumetricAdvancedMaterialOutput.h"
+#include "RenderUtils.h"
+
+#if WITH_EDITOR
+
+namespace IR = MaterialIR;
+
+struct FHLSLPrinter
+{
+	FString& Buffer;
+
+	template <typename... Types>
+	void Print(const TCHAR* Format, Types... Args)
+	{
+		Buffer.Printf(Format, Args...);
+	}
+
+	FHLSLPrinter& operator<<(const TCHAR* Text)
+	{
+ 		Buffer.Append(Text);
+		return *this;
+	}
+
+	FHLSLPrinter& operator<<(int32 Value)
+	{
+ 		Buffer.Appendf(TEXT("%d"), Value);
+		return *this;
+	}
+
+	FHLSLPrinter& operator<<(uint32 Value)
+	{
+ 		Buffer.Appendf(TEXT("%u"), Value);
+		return *this;
+	}
+
+	FHLSLPrinter& operator<<(float Value)
+	{
+ 		Buffer.Appendf(TEXT("%.5f"), Value);
+		return *this;
+	}
+};
+
+static const TCHAR* GetHLSLTypeString(EMaterialValueType Type)
+{
+	switch (Type)
+	{
+	case MCT_Float1: return TEXT("MaterialFloat");
+	case MCT_Float2: return TEXT("MaterialFloat2");
+	case MCT_Float3: return TEXT("MaterialFloat3");
+	case MCT_Float4: return TEXT("MaterialFloat4");
+	case MCT_Float: return TEXT("MaterialFloat");
+	case MCT_Texture2D: return TEXT("texture2D");
+	case MCT_TextureCube: return TEXT("textureCube");
+	case MCT_Texture2DArray: return TEXT("texture2DArray");
+	case MCT_VolumeTexture: return TEXT("volumeTexture");
+	case MCT_StaticBool: return TEXT("static bool");
+	case MCT_Bool:  return TEXT("bool");
+	case MCT_MaterialAttributes: return TEXT("FMaterialAttributes");
+	case MCT_TextureExternal: return TEXT("TextureExternal");
+	case MCT_TextureVirtual: return TEXT("TextureVirtual");
+	case MCT_VTPageTableResult: return TEXT("VTPageTableResult");
+	case MCT_ShadingModel: return TEXT("uint");
+	case MCT_UInt: return TEXT("uint");
+	case MCT_UInt1: return TEXT("uint");
+	case MCT_UInt2: return TEXT("uint2");
+	case MCT_UInt3: return TEXT("uint3");
+	case MCT_UInt4: return TEXT("uint4");
+	case MCT_Substrate: return TEXT("FSubstrateData");
+	case MCT_TextureCollection: return TEXT("FResourceCollection");
+	default: return TEXT("unknown");
+	};
+}
+
+static bool IsMaterialPropertyShared(EMaterialProperty InProperty)
+{
+	switch (InProperty)
+	{
+		case MP_Normal:
+		case MP_Tangent:
+		case MP_EmissiveColor:
+		case MP_Opacity:
+		case MP_OpacityMask:
+		case MP_BaseColor:
+		case MP_Metallic:
+		case MP_Specular:
+		case MP_Roughness:
+		case MP_Anisotropy:
+		case MP_AmbientOcclusion:
+		case MP_Refraction:
+		case MP_PixelDepthOffset:
+		case MP_SubsurfaceColor:
+		case MP_ShadingModel:
+		case MP_SurfaceThickness:
+		case MP_FrontMaterial:
+		case MP_Displacement:
+			return true;
+		default:
+			return false;
+	}
+};
+
+static const TCHAR* GetShadingModelParameterName(EMaterialShadingModel InModel)
+{
+	switch (InModel)
+	{
+		case MSM_Unlit: return TEXT("MATERIAL_SHADINGMODEL_UNLIT");
+		case MSM_DefaultLit: return TEXT("MATERIAL_SHADINGMODEL_DEFAULT_LIT");
+		case MSM_Subsurface: return TEXT("MATERIAL_SHADINGMODEL_SUBSURFACE");
+		case MSM_PreintegratedSkin: return TEXT("MATERIAL_SHADINGMODEL_PREINTEGRATED_SKIN");
+		case MSM_ClearCoat: return TEXT("MATERIAL_SHADINGMODEL_CLEAR_COAT");
+		case MSM_SubsurfaceProfile: return TEXT("MATERIAL_SHADINGMODEL_SUBSURFACE_PROFILE");
+		case MSM_TwoSidedFoliage: return TEXT("MATERIAL_SHADINGMODEL_TWOSIDED_FOLIAGE");
+		case MSM_Hair: return TEXT("MATERIAL_SHADINGMODEL_HAIR");
+		case MSM_Cloth: return TEXT("MATERIAL_SHADINGMODEL_CLOTH");
+		case MSM_Eye: return TEXT("MATERIAL_SHADINGMODEL_EYE");
+		case MSM_SingleLayerWater: return TEXT("MATERIAL_SHADINGMODEL_SINGLELAYERWATER");
+		case MSM_ThinTranslucent: return TEXT("MATERIAL_SHADINGMODEL_THIN_TRANSLUCENT");
+		default: UE_MIR_UNREACHABLE();
+	}
+}
+
+UE_MIR_BEGIN_PRIVATE(FMaterialIRToHLSLTranslator);
+
+void GenerateHLSL(const FMaterial& InMaterial, const FMaterialIRModule& InModule, FParametersMap& InParams)
+{
+	FString PixelAttributes;
+	FString EvaluateOtherMaterialAttributesHLSL;
+
+	FHLSLPrinter Printer{ EvaluateOtherMaterialAttributesHLSL };
+
+ 	for (const IR::FSetMaterialOutputInstr* Output : InModule.GetOutputs())
+	{
+		// Special case MP_SubsurfaceColor as the actual property is a combination of the color and the profile but we don't want to expose the profile
+		FString PropertyName = (Output->Property == MP_SubsurfaceColor) ? "Subsurface" : FMaterialAttributeDefinitionMap::GetAttributeName(Output->Property);
+		EMaterialValueType Type = (Output->Property == MP_SubsurfaceColor) ? MCT_Float4 : FMaterialAttributeDefinitionMap::GetValueType(Output->Property);
+		check(PropertyName.Len() > 0);
+
+		EvaluateOtherMaterialAttributesHLSL.Appendf(TEXT("\tPixelMaterialInputs.%s = "), *PropertyName);
+
+		LowerValue(Printer, Output);
+
+		EvaluateOtherMaterialAttributesHLSL.Append(TEXT("\n"));
+	}
+
+	for (int32 PropertyIndex = 0; PropertyIndex < MP_MAX; ++PropertyIndex)
+	{
+		EMaterialProperty Property = (EMaterialProperty)PropertyIndex;
+		if (!IsMaterialPropertyShared(Property))
+		{
+			continue;
+		}
+		
+		check(FMaterialAttributeDefinitionMap::GetShaderFrequency(Property) == SF_Pixel);
+			
+		// Special case MP_SubsurfaceColor as the actual property is a combination of the color and the profile but we don't want to expose the profile
+		FString PropertyName = (Property == MP_SubsurfaceColor) ? "Subsurface" : FMaterialAttributeDefinitionMap::GetAttributeName(Property);
+		EMaterialValueType Type = (Property == MP_SubsurfaceColor) ? MCT_Float4 : FMaterialAttributeDefinitionMap::GetValueType(Property);
+		check(PropertyName.Len() > 0);
+
+		PixelAttributes.Appendf(TEXT("\t%s %s;\n"), GetHLSLTypeString(Type), *PropertyName);
+	}
+	
+	InParams.Add(TEXT("pixel_material_inputs"), MoveTemp(PixelAttributes));
+	InParams.Add(TEXT("calc_pixel_material_inputs_normal"), TEXT("PixelMater ialInputs.Normal = MaterialFloat3(0.00000000, 0.00000000, 1.00000000);"));
+	InParams.Add(TEXT("calc_pixel_material_inputs_other_inputs"), MoveTemp(EvaluateOtherMaterialAttributesHLSL));
+	InParams.Add(TEXT("material_declarations"), TEXT("struct FMaterialAttributes {};"));
+	
+	auto SetParamInt = [&] (const TCHAR* InParamName, int InValue)
+	{
+		InParams.Add(InParamName, FString::Printf(TEXT("%d"), InValue));
+	};
+
+	SetParamInt(TEXT("num_material_texcoords_vertex"), 0);
+	SetParamInt(TEXT("num_material_texcoords"), 0);
+	SetParamInt(TEXT("num_custom_vertex_interpolators"), 0);
+	SetParamInt(TEXT("num_tex_coord_interpolators"), 0);
+}
+
+void SetMaterialParameters(const FMaterial& InMaterial, FParametersMap& InParams)
+{
+	auto SetParamReturnFloat = [&] (const TCHAR* InParamName, float InValue)
+	{
+		InParams.Add(InParamName, FString::Printf(TEXT("\treturn %.5f"), InValue));
+	};
+
+	SetParamReturnFloat(TEXT("get_material_emissive_for_cs"), 0.f);
+	SetParamReturnFloat(TEXT("get_material_translucency_directional_lighting_intensity"), InMaterial.GetTranslucencyDirectionalLightingIntensity());
+	SetParamReturnFloat(TEXT("get_material_translucent_shadow_density_scale"), InMaterial.GetTranslucentShadowDensityScale());
+	SetParamReturnFloat(TEXT("get_material_translucent_self_shadow_density_scale"), InMaterial.GetTranslucentSelfShadowDensityScale());
+	SetParamReturnFloat(TEXT("get_material_translucent_self_shadow_second_density_scale"), InMaterial.GetTranslucentSelfShadowSecondDensityScale());
+	SetParamReturnFloat(TEXT("get_material_translucent_self_shadow_second_opacity"), InMaterial.GetTranslucentSelfShadowSecondOpacity());
+	SetParamReturnFloat(TEXT("get_material_translucent_backscattering_exponent"), InMaterial.GetTranslucentBackscatteringExponent());
+
+	FLinearColor Extinction = InMaterial.GetTranslucentMultipleScatteringExtinction();
+	InParams.Add(TEXT("get_material_translucent_multiple_scattering_extinction"), FString::Printf(TEXT("\treturn MaterialFloat3(%.5f, %.5f, %.5f)"), Extinction.R, Extinction.G, Extinction.B));
+
+	SetParamReturnFloat(TEXT("get_material_opacity_mask_clip_value"), InMaterial.GetOpacityMaskClipValue());
+	InParams.Add(TEXT("get_material_world_position_offset_raw"), TEXT("\treturn 0; // todo"));
+	InParams.Add(TEXT("get_material_previous_world_position_offset_raw"), TEXT("\treturn 0; // todo"));
+	
+	// CustomData0/1 are named ClearCoat/ClearCoatRoughness
+	InParams.Add(TEXT("get_material_custom_data0"), TEXT("\treturn 0; // todo"));
+	InParams.Add(TEXT("get_material_custom_data1"), TEXT("\treturn 0; // todo"));
+
+	FString EvaluateMaterialDeclaration;
+	EvaluateMaterialDeclaration.Append(TEXT("void EvaluateVertexMaterialAttributes(in out FMaterialVertexParameters Parameters)\n{\n"));
+	EvaluateMaterialDeclaration.Append(TEXT("\n}\n"));
+	InParams.Add(TEXT("evaluate_material_attributes"), EvaluateMaterialDeclaration);
+}
+
+void GetShaderCompilerEnvironment(const FMaterial& InMaterial, const FMaterialIRModule& InModule, FShaderCompilerEnvironment& OutEnvironment)
+{
+	EShaderPlatform ShaderPlatform = InModule.GetShaderPlatform();
+	const FMaterialCompilationOutput& CompilationOutput = InModule.GetCompilationOutput();
+
+	OutEnvironment.TargetPlatform = InModule.GetTargetPlatform();
+	OutEnvironment.SetDefine(TEXT("ENABLE_NEW_HLSL_GENERATOR"), 1);
+	OutEnvironment.SetDefine(TEXT("MATERIAL_ATMOSPHERIC_FOG"), false);
+	OutEnvironment.SetDefine(TEXT("MATERIAL_SKY_ATMOSPHERE"), false);
+	OutEnvironment.SetDefine(TEXT("INTERPOLATE_VERTEX_COLOR"), false);
+	OutEnvironment.SetDefine(TEXT("NEEDS_PARTICLE_COLOR"), false);
+	OutEnvironment.SetDefine(TEXT("NEEDS_PARTICLE_LOCAL_TO_WORLD"), false);
+	OutEnvironment.SetDefine(TEXT("NEEDS_PARTICLE_WORLD_TO_LOCAL"), false);
+	OutEnvironment.SetDefine(TEXT("NEEDS_PER_INSTANCE_RANDOM_PS"), false);
+	OutEnvironment.SetDefine(TEXT("USES_TRANSFORM_VECTOR"), false);
+	OutEnvironment.SetDefine(TEXT("WANT_PIXEL_DEPTH_OFFSET"), CompilationOutput.bUsesPixelDepthOffset);
+	OutEnvironment.SetDefineAndCompileArgument(TEXT("USES_WORLD_POSITION_OFFSET"), (bool)CompilationOutput.bUsesWorldPositionOffset);
+	OutEnvironment.SetDefineAndCompileArgument(TEXT("USES_DISPLACEMENT"), false);
+	OutEnvironment.SetDefine(TEXT("USES_EMISSIVE_COLOR"), false);
+	OutEnvironment.SetDefine(TEXT("USES_DISTORTION"), InMaterial.IsDistorted());
+	OutEnvironment.SetDefine(TEXT("MATERIAL_ENABLE_TRANSLUCENCY_FOGGING"), InMaterial.ShouldApplyFogging());
+	OutEnvironment.SetDefine(TEXT("MATERIAL_ENABLE_TRANSLUCENCY_CLOUD_FOGGING"), InMaterial.ShouldApplyCloudFogging());
+	OutEnvironment.SetDefine(TEXT("MATERIAL_IS_SKY"), InMaterial.IsSky());
+	OutEnvironment.SetDefine(TEXT("MATERIAL_COMPUTE_FOG_PER_PIXEL"), InMaterial.ComputeFogPerPixel());
+	OutEnvironment.SetDefine(TEXT("MATERIAL_FULLY_ROUGH"), false);
+	OutEnvironment.SetDefine(TEXT("MATERIAL_USES_ANISOTROPY"), false);
+	OutEnvironment.SetDefine(TEXT("MATERIAL_NEURAL_POST_PROCESS"), (CompilationOutput.bUsedWithNeuralNetworks || InMaterial.IsUsedWithNeuralNetworks()) && InMaterial.IsPostProcessMaterial());
+	OutEnvironment.SetDefine(TEXT("NUM_VIRTUALTEXTURE_SAMPLES"), 0);
+	OutEnvironment.SetDefine(TEXT("MATERIAL_VIRTUALTEXTURE_FEEDBACK"), false);
+	OutEnvironment.SetDefine(TEXT("IS_MATERIAL_SHADER"), true);
+
+	FMaterialShadingModelField ShadingModels = InMaterial.GetShadingModels();
+	ensure(ShadingModels.IsValid());
+
+	int32 NumActiveShadingModels = 0;
+	if (ShadingModels.IsLit())
+	{
+		// This is to have platforms use the simple single layer water shading similar to mobile: no dynamic lights, only sun and sky, no distortion, no colored transmittance on background, no custom depth read.
+		const bool bSingleLayerWaterUsesSimpleShading = FDataDrivenShaderPlatformInfo::GetWaterUsesSimpleForwardShading(ShaderPlatform) && IsForwardShadingEnabled(ShaderPlatform);
+
+		for (int i = 0; i < MSM_NUM; ++i)
+		{
+			EMaterialShadingModel Model = (EMaterialShadingModel)i;
+			if (Model == MSM_Strata || !ShadingModels.HasShadingModel(Model))
+			{
+				continue;
+			}
+
+			if (Model == MSM_SingleLayerWater && !FDataDrivenShaderPlatformInfo::GetRequiresDisableForwardLocalLights(ShaderPlatform))
+			{
+				continue;
+			}
+
+			if (Model == MSM_SingleLayerWater && bSingleLayerWaterUsesSimpleShading)
+			{
+				// Value must match SINGLE_LAYER_WATER_SHADING_QUALITY_MOBILE_WITH_DEPTH_TEXTURE in SingleLayerWaterCommon.ush!
+				OutEnvironment.SetDefine(TEXT("SINGLE_LAYER_WATER_SHADING_QUALITY"), true);
+			}
+
+			OutEnvironment.SetDefine(GetShadingModelParameterName(Model), true);
+			NumActiveShadingModels += 1;
+		}
+	}
+	else
+	{
+		// Unlit shading model can only exist by itself
+		OutEnvironment.SetDefine(TEXT("MATERIAL_SINGLE_SHADINGMODEL"), true);
+		OutEnvironment.SetDefine(GetShadingModelParameterName(MSM_Unlit), true);
+	}
+
+	if (NumActiveShadingModels == 1)
+	{
+		OutEnvironment.SetDefine(TEXT("MATERIAL_SINGLE_SHADINGMODEL"), true);
+	}
+	else if (!ensure(NumActiveShadingModels > 0))
+	{
+		UE_LOG(LogMaterial, Warning, TEXT("Unknown material shading model(s). Setting to MSM_DefaultLit"));
+		OutEnvironment.SetDefine(GetShadingModelParameterName(MSM_DefaultLit), true);
+	}
+
+	static IConsoleVariable* CVarLWCIsEnabled = IConsoleManager::Get().FindConsoleVariable(TEXT("r.MaterialEditor.LWCEnabled"));
+	OutEnvironment.SetDefine(TEXT("MATERIAL_LWC_ENABLED"), CVarLWCIsEnabled->GetInt());
+	OutEnvironment.SetDefine(TEXT("WSVECTOR_IS_TILEOFFSET"), true);
+	OutEnvironment.SetDefine(TEXT("WSVECTOR_IS_DOUBLEFLOAT"), false);
+
+	if (InMaterial.GetMaterialDomain() == MD_Volume)
+	{
+		TArray<const UMaterialExpressionVolumetricAdvancedMaterialOutput*> VolumetricAdvancedExpressions;
+		InMaterial.GetMaterialInterface()->GetMaterial()->GetAllExpressionsOfType(VolumetricAdvancedExpressions);
+		if (VolumetricAdvancedExpressions.Num() > 0)
+		{
+			if (VolumetricAdvancedExpressions.Num() > 1)
+			{
+				UE_LOG(LogMaterial, Fatal, TEXT("Only a single UMaterialExpressionVolumetricAdvancedMaterialOutput node is supported."));
+			}
+
+			const UMaterialExpressionVolumetricAdvancedMaterialOutput* VolumetricAdvancedNode = VolumetricAdvancedExpressions[0];
+			const TCHAR* Param = VolumetricAdvancedNode->GetEvaluatePhaseOncePerSample() ? TEXT("MATERIAL_VOLUMETRIC_ADVANCED_PHASE_PERSAMPLE") : TEXT("MATERIAL_VOLUMETRIC_ADVANCED_PHASE_PERPIXEL");
+			OutEnvironment.SetDefine(Param, true);
+
+			OutEnvironment.SetDefine(TEXT("MATERIAL_VOLUMETRIC_ADVANCED"), true);
+			OutEnvironment.SetDefine(TEXT("MATERIAL_VOLUMETRIC_ADVANCED_GRAYSCALE_MATERIAL"), VolumetricAdvancedNode->bGrayScaleMaterial);
+			OutEnvironment.SetDefine(TEXT("MATERIAL_VOLUMETRIC_ADVANCED_RAYMARCH_VOLUME_SHADOW"), VolumetricAdvancedNode->bRayMarchVolumeShadow);
+			OutEnvironment.SetDefine(TEXT("MATERIAL_VOLUMETRIC_ADVANCED_CLAMP_MULTISCATTERING_CONTRIBUTION"), VolumetricAdvancedNode->bClampMultiScatteringContribution);
+			OutEnvironment.SetDefine(TEXT("MATERIAL_VOLUMETRIC_ADVANCED_MULTISCATTERING_OCTAVE_COUNT"), VolumetricAdvancedNode->GetMultiScatteringApproximationOctaveCount());
+			OutEnvironment.SetDefine(TEXT("MATERIAL_VOLUMETRIC_ADVANCED_CONSERVATIVE_DENSITY"), VolumetricAdvancedNode->ConservativeDensity.IsConnected());
+			OutEnvironment.SetDefine(TEXT("MATERIAL_VOLUMETRIC_ADVANCED_OVERRIDE_AMBIENT_OCCLUSION"), InMaterial.HasAmbientOcclusionConnected());
+			OutEnvironment.SetDefine(TEXT("MATERIAL_VOLUMETRIC_ADVANCED_GROUND_CONTRIBUTION"), VolumetricAdvancedNode->bGroundContribution);
+		}
+	}
+
+	OutEnvironment.SetDefine(TEXT("MATERIAL_IS_SUBSTRATE"), false);
+	OutEnvironment.SetDefine(TEXT("DUAL_SOURCE_COLOR_BLENDING_ENABLED"), false);
+	OutEnvironment.SetDefine(TEXT("TEXTURE_SAMPLE_DEBUG"), false);
+}
+
+void LowerValue(FHLSLPrinter& Printer, const IR::FValue* InValue)
+{
+	if (const IR::FScalarValue* S = InValue->Cast<IR::FScalarValue>())
+	{
+		IR::FArithmeticTypePtr ArithType = S->Type->ToArithmetic();
+		check(ArithType && ArithType->IsScalar());
+
+		switch (ArithType->ScalarKind)
+		{
+			case IR::SK_Bool: Printer << S->Boolean; break;
+			case IR::SK_Int: Printer << S->Integer; break;
+			case IR::SK_Float: Printer << S->Float; break;
+		}
+	}
+	else if (const IR::FVectorValue* V = InValue->Cast<IR::FVectorValue>())
+	{
+		IR::FArithmeticTypePtr ArithType = V->Type->ToArithmetic();
+		check(ArithType && ArithType->IsVector());
+
+		Printer << ScalarKindToString(ArithType->ScalarKind) << ArithType->NumColumns << TEXT("(");
+
+		for (IR::FValuePtr ComponentVal : V->GetComponents())
+		{
+			LowerValue(Printer, ComponentVal);
+		}
+
+		Printer << TEXT(")");
+	}
+}
+
+UE_MIR_END_PRIVATE(); //FMaterialIRToHLSLTranslator
+
+FMaterialIRToHLSLTranslator::FMaterialIRToHLSLTranslator()
+{
+}
+
+void FMaterialIRToHLSLTranslator::Translate(const FMaterial& InMaterial, const FMaterialIRModule& InModule, FParametersMap& OutParametersMap, FShaderCompilerEnvironment& OutEnvironment)
+{
+	OutParametersMap.Empty();
+
+	AsPrivate()->GenerateHLSL(InMaterial, InModule, OutParametersMap);
+	AsPrivate()->SetMaterialParameters(InMaterial, OutParametersMap);
+	AsPrivate()->GetShaderCompilerEnvironment(InMaterial, InModule, OutEnvironment);
+}
+
+#endif // #if WITH_EDITOR
diff --git a/Engine/Source/Runtime/Engine/Private/Materials/MaterialIRTypes.cpp b/Engine/Source/Runtime/Engine/Private/Materials/MaterialIRTypes.cpp
new file mode 100644
--- /dev/null
+++ b/Engine/Source/Runtime/Engine/Private/Materials/MaterialIRTypes.cpp
@@ -0,0 +1,154 @@
+// Copyright Epic Games, Inc. All Rights Reserved.
+
+#include "Materials/MaterialIRTypes.h"
+
+#if WITH_EDITOR
+
+namespace MaterialIR
+{
+
+FTypePtr FType::FromShaderType(const UE::Shader::FType& InShaderType)
+{
+	check(!InShaderType.IsStruct());
+	check(!InShaderType.IsObject());
+
+	switch (InShaderType.ValueType)
+	{
+		case UE::Shader::EValueType::Void:
+			return FType::GetVoid();
+
+		case UE::Shader::EValueType::Float1:
+		case UE::Shader::EValueType::Float2:
+		case UE::Shader::EValueType::Float3:
+		case UE::Shader::EValueType::Float4:
+			return FArithmeticType::GetVector(SK_Float, (int)InShaderType.ValueType - (int)UE::Shader::EValueType::Float1 + 1);
+
+		case UE::Shader::EValueType::Int1:
+		case UE::Shader::EValueType::Int2:
+		case UE::Shader::EValueType::Int3:
+		case UE::Shader::EValueType::Int4:
+			return FArithmeticType::GetVector(SK_Int, (int)InShaderType.ValueType - (int)UE::Shader::EValueType::Int1 + 1);
+
+		case UE::Shader::EValueType::Bool1:
+		case UE::Shader::EValueType::Bool2:
+		case UE::Shader::EValueType::Bool3:
+		case UE::Shader::EValueType::Bool4:
+			return FArithmeticType::GetVector(SK_Bool, (int)InShaderType.ValueType - (int)UE::Shader::EValueType::Int1 + 1);
+
+		default:
+			UE_MIR_UNREACHABLE();
+	}
+}
+
+FTypePtr FType::GetVoid()
+{
+	static FType Type{ TK_Void };
+	return &Type;
+}
+
+const FArithmeticType* FType::ToArithmetic() const
+{
+	return Kind == TK_Arithmetic ? static_cast<const FArithmeticType*>(this) : nullptr; 
+}
+
+bool FType::IsScalar() const
+{
+	if (FArithmeticTypePtr Arith = ToArithmetic())
+	{
+		return Arith->IsScalar();
+	}
+	return false;
+}
+
+const TCHAR* ScalarKindToString(EScalarKind Kind)
+{
+	switch (Kind)
+	{
+		case SK_Bool: return TEXT("bool"); break;
+		case SK_Int: return TEXT("int"); break;
+		case SK_Float: return TEXT("float"); break;
+		default: UE_MIR_UNREACHABLE();
+	}
+}
+
+static const FArithmeticType* GetNumericalType(EScalarKind InScalarKind, int NumRows, int NumColumns)
+{
+	check(InScalarKind >= 0 && InScalarKind <= SK_Float);
+	
+	static const FArithmeticType Types[] {
+		{ { TK_Arithmetic }, SK_Bool, 1, 1 },
+		{ { TK_Arithmetic }, SK_Bool, 1, 2 }, 
+		{ { TK_Arithmetic }, SK_Bool, 1, 3 },
+		{ { TK_Arithmetic }, SK_Bool, 1, 4 },
+		{ { TK_Arithmetic }, SK_Bool, 2, 1 },
+		{ { TK_Arithmetic }, SK_Bool, 2, 2 },
+		{ { TK_Arithmetic }, SK_Bool, 2, 3 },
+		{ { TK_Arithmetic }, SK_Bool, 2, 4 },
+		{ { TK_Arithmetic }, SK_Bool, 3, 1 },
+		{ { TK_Arithmetic }, SK_Bool, 3, 2 },
+		{ { TK_Arithmetic }, SK_Bool, 3, 3 },
+		{ { TK_Arithmetic }, SK_Bool, 3, 4 },
+		{ { TK_Arithmetic }, SK_Bool, 4, 1 },
+		{ { TK_Arithmetic }, SK_Bool, 4, 2 },
+		{ { TK_Arithmetic }, SK_Bool, 4, 3 },
+		{ { TK_Arithmetic }, SK_Bool, 4, 4 },
+		{ { TK_Arithmetic }, SK_Int, 1, 1 },
+		{ { TK_Arithmetic }, SK_Int, 1, 2 },
+		{ { TK_Arithmetic }, SK_Int, 1, 3 },
+		{ { TK_Arithmetic }, SK_Int, 1, 4 },
+		{ { TK_Arithmetic }, SK_Int, 2, 1 },
+		{ { TK_Arithmetic }, SK_Int, 2, 2 },
+		{ { TK_Arithmetic }, SK_Int, 2, 3 },
+		{ { TK_Arithmetic }, SK_Int, 2, 4 },
+		{ { TK_Arithmetic }, SK_Int, 3, 1 },
+		{ { TK_Arithmetic }, SK_Int, 3, 2 },
+		{ { TK_Arithmetic }, SK_Int, 3, 3 },
+		{ { TK_Arithmetic }, SK_Int, 3, 4 },
+		{ { TK_Arithmetic }, SK_Int, 4, 1 },
+		{ { TK_Arithmetic }, SK_Int, 4, 2 },
+		{ { TK_Arithmetic }, SK_Int, 4, 3 },
+		{ { TK_Arithmetic }, SK_Int, 4, 4 },
+		{ { TK_Arithmetic }, SK_Float, 1, 1 },
+		{ { TK_Arithmetic }, SK_Float, 1, 2 },
+		{ { TK_Arithmetic }, SK_Float, 1, 3 },
+		{ { TK_Arithmetic }, SK_Float, 1, 4 },
+		{ { TK_Arithmetic }, SK_Float, 2, 1 },
+		{ { TK_Arithmetic }, SK_Float, 2, 2 },
+		{ { TK_Arithmetic }, SK_Float, 2, 3 },
+		{ { TK_Arithmetic }, SK_Float, 2, 4 },
+		{ { TK_Arithmetic }, SK_Float, 3, 1 },
+		{ { TK_Arithmetic }, SK_Float, 3, 2 },
+		{ { TK_Arithmetic }, SK_Float, 3, 3 },
+		{ { TK_Arithmetic }, SK_Float, 3, 4 },
+		{ { TK_Arithmetic }, SK_Float, 4, 1 },
+		{ { TK_Arithmetic }, SK_Float, 4, 2 },
+		{ { TK_Arithmetic }, SK_Float, 4, 3 },
+		{ { TK_Arithmetic }, SK_Float, 4, 4 },
+	};
+
+	int Index = InScalarKind * 4 * 4 + (NumRows - 1) * 4 + (NumColumns - 1);
+	check(Index < UE_ARRAY_COUNT(Types));
+	return &Types[Index];
+}
+
+const FArithmeticType* FArithmeticType::GetScalar(EScalarKind InScalarKind)
+{
+	return GetNumericalType(InScalarKind, 1, 1);
+}
+
+const FArithmeticType* FArithmeticType::GetVector(EScalarKind InScalarKind, int NumComponents)
+{
+	check(NumComponents >= 1 && NumComponents <= 4);
+	return GetNumericalType(InScalarKind, NumComponents, 1);
+}
+
+const FArithmeticType* FArithmeticType::GetMatrix(EScalarKind InScalarKind, int NumRows, int NumColumns)
+{
+	check(NumColumns > 1 && NumColumns <= 4);
+	check(NumRows > 1 && NumRows <= 4);
+	return GetNumericalType(InScalarKind, NumRows, NumColumns);
+}
+
+} // namespace MaterialIR
+
+#endif // #if WITH_EDITOR
diff --git a/Engine/Source/Runtime/Engine/Private/Materials/MaterialShared.cpp b/Engine/Source/Runtime/Engine/Private/Materials/MaterialShared.cpp
--- a/Engine/Source/Runtime/Engine/Private/Materials/MaterialShared.cpp
+++ b/Engine/Source/Runtime/Engine/Private/Materials/MaterialShared.cpp
@@ -62,6 +62,11 @@
 #include "ProfilingDebugging/CookStats.h"
 #include "Engine/NeuralProfile.h"
 
+#include "Materials/MaterialIRModule.h"
+#include "Materials/MaterialIRModuleBuilder.h"
+#include "Materials/MaterialIRToHLSLTranslator.h"
+#include "Materials/MaterialSourceTemplate.h"
+
 #define LOCTEXT_NAMESPACE "MaterialShared"
 
 DEFINE_LOG_CATEGORY(LogMaterial);
@@ -81,11 +86,19 @@ namespace MaterialSharedCookStats
 #endif
 
 #if WITH_EDITOR
+
 static TAutoConsoleVariable<bool> CVarMaterialEdPreshaderDumpToHLSL(
 	TEXT("r.MaterialEditor.PreshaderDumpToHLSL"),
 	true,
 	TEXT("Controls whether to append preshader expressions and parameter reference counts to the HLSL source window (as comments at the end of the code)."),
 	ECVF_RenderThreadSafe);
+
+static TAutoConsoleVariable<bool> CVarMaterialEdUseNewTranslatorPrototype(
+	TEXT("r.MaterialEditor.UseNewTranslatorPrototype"),
+	false,
+	TEXT("Controls whether to enable the new material translator prototype (WIP) ."),
+	ECVF_RenderThreadSafe);
+
 #endif
 
 IMPLEMENT_TYPE_LAYOUT(FHashedMaterialParameterInfo);
@@ -3341,15 +3354,54 @@ bool FMaterial::Translate_Legacy(const FMaterialShaderMapId& ShaderMapId,
 	return bSuccess;
 }
 
-bool FMaterial::Translate_New(const FMaterialShaderMapId& ShaderMapId,
+bool FMaterial::Translate_New(const FMaterialShaderMapId& InShaderMapId,
 	const FStaticParameterSet& InStaticParameters,
 	EShaderPlatform InPlatform,
 	const ITargetPlatform* InTargetPlatform,
 	FMaterialCompilationOutput& OutCompilationOutput,
 	TRefCountPtr<FSharedShaderCompilerEnvironment>& OutMaterialEnvironment)
 {
-	const FMaterialCompileTargetParameters TargetParams(InPlatform, ShaderMapId.FeatureLevel, InTargetPlatform);
-	return MaterialEmitHLSL(TargetParams, InStaticParameters, *this, OutCompilationOutput, OutMaterialEnvironment);
+	const bool bUseNewTranslatorPrototype = CVarMaterialEdUseNewTranslatorPrototype->GetBool();
+	if (bUseNewTranslatorPrototype)
+	{
+		FMaterialIRModule Module;
+
+		/* Setup the builder and bild the material */
+		FMaterialIRModuleBuilder Builder;
+		Builder.SetSource(this, &InStaticParameters);
+		Builder.SetPlatform(InPlatform, InShaderMapId.FeatureLevel, InTargetPlatform);
+		Builder.SetTarget(&Module);
+		if (!Builder.Build())
+		{
+			return false;
+		}
+
+		/* Copy over the compilation output */
+		OutCompilationOutput = Module.GetCompilationOutput();
+		OutMaterialEnvironment = new FSharedShaderCompilerEnvironment();
+
+		/* Translate the material IR module to HLSL template string parameters and material environment */
+		FMaterialIRToHLSLTranslator::FParametersMap ShaderStringParameters;
+		FMaterialIRToHLSLTranslator Translator;
+		Translator.Translate(*this, Module, ShaderStringParameters, *OutMaterialEnvironment);
+
+		/* Interpolate HLSL parameters with the material shader template to produce the final shader source */
+		int32 LineNumber;
+		FStringTemplateResolver Resolver = FMaterialSourceTemplate::Get().BeginResolve(InPlatform, &LineNumber);
+		ShaderStringParameters.Add({TEXT("line_number"), FString::Printf(TEXT("%u"), LineNumber)});
+		Resolver.SetParameterMap(&ShaderStringParameters);
+
+		// Interpolate the string
+		FString MaterialShaderCode = Resolver.Finalize();
+		OutMaterialEnvironment->IncludeVirtualPathToContentsMap.Add(TEXT("/Engine/Generated/Material.ush"), MoveTemp(MaterialShaderCode));
+
+		return true;
+	}
+	else
+	{
+		const FMaterialCompileTargetParameters TargetParams(InPlatform, InShaderMapId.FeatureLevel, InTargetPlatform);
+		return MaterialEmitHLSL(TargetParams, InStaticParameters, *this, OutCompilationOutput, OutMaterialEnvironment);
+	}
 }
 
 bool FMaterial::Translate(const FMaterialShaderMapId& InShaderMapId,
diff --git a/Engine/Source/Runtime/Engine/Public/Materials/MaterialIR.h b/Engine/Source/Runtime/Engine/Public/Materials/MaterialIR.h
new file mode 100644
--- /dev/null
+++ b/Engine/Source/Runtime/Engine/Public/Materials/MaterialIR.h
@@ -0,0 +1,84 @@
+// Copyright Epic Games, Inc. All Rights Reserved.
+#pragma once
+
+#include "Materials/MaterialIRCommon.h"
+
+enum EMaterialProperty : int;
+
+#if WITH_EDITOR
+
+namespace MaterialIR
+{
+
+enum EValueKind
+{
+	/* Values */
+	VK_Value = 0,
+	VK_Scalar = 1 << 1,
+	VK_Vector = 2 << 1,
+	
+	/* Instructions */
+	VK_Instruction = 1,
+ 	VK_SetMaterialOutputInstr = VK_Instruction | 1 << 1,
+};
+
+/* Values */
+
+struct FValue : FIRNode<EValueKind>
+{
+	EValueKind Kind = VK_Value;
+	FTypePtr  Type{};
+
+	void Destroy();
+};
+
+	template <EValueKind TTypeKind>
+	struct TValue : FValue
+	{
+		static constexpr EValueKind TypeKind = TTypeKind;
+	};
+
+struct FScalarValue : TValue<VK_Scalar>
+{
+	union
+	{
+		bool  Boolean;
+		int   Integer;
+		float Float;
+	};
+};
+
+struct FVectorValue : TValue<VK_Vector>
+{
+	TArrayView<const FValuePtr> GetComponents() const;
+	TArrayView<FValuePtr> GetMutableComponents();
+};
+
+template <int TDimension>
+struct TVectorValue : FVectorValue
+{
+	FValuePtr Components[TDimension];
+};
+
+/* Instructions */
+
+struct FInstruction : TValue<VK_Instruction>
+{
+	FInstruction* Next{};
+};
+
+	template <EValueKind TTypeKind>
+	struct TInstruction : FInstruction
+	{
+		static constexpr EValueKind TypeKind = TTypeKind;
+	};
+
+struct FSetMaterialOutputInstr : TInstruction<VK_SetMaterialOutputInstr>
+{
+	EMaterialProperty Property;
+	FValuePtr ArgValue;
+};
+
+} // namespace MaterialIR
+
+#endif
diff --git a/Engine/Source/Runtime/Engine/Public/Materials/MaterialIRBuilder.h b/Engine/Source/Runtime/Engine/Public/Materials/MaterialIRBuilder.h
new file mode 100644
--- /dev/null
+++ b/Engine/Source/Runtime/Engine/Public/Materials/MaterialIRBuilder.h
@@ -0,0 +1,48 @@
+// Copyright Epic Games, Inc. All Rights Reserved.
+#pragma once
+
+#include "Materials/MaterialIRCommon.h"
+
+#if WITH_EDITOR
+
+namespace MaterialIR
+{
+
+//
+class FBuilder
+{
+public:
+
+	FBuilder();
+	FBuilder(FMaterialIRModule* InModule);
+
+	void SetTargetModule(FMaterialIRModule* InModule);
+
+	/* Constants */
+
+	FValuePtr NewConstantFromShaderValue(const UE::Shader::FValue& InValue);
+	FValuePtr NewConstantFloat1(float InX);
+	FValuePtr NewConstantFloat2(const FVector2f& InValue);
+	FValuePtr NewConstantFloat3(const FVector3f& InValue);
+	FValuePtr NewConstantFloat4(const FVector4f& InValue);
+	FValuePtr NewConstantInt1(int InX);
+	FValuePtr NewConstantInt2(const FIntVector2& InValue);
+	FValuePtr NewConstantInt3(const FIntVector3& InValue);
+	FValuePtr NewConstantInt4(const FIntVector4& InValue);
+	FValuePtr NewVector2(FValuePtr InX, FValuePtr InY);
+	FValuePtr NewVector3(FValuePtr InX, FValuePtr InY, FValuePtr InZ);
+	FValuePtr NewVector4(FValuePtr InX, FValuePtr InY, FValuePtr InZ, FValuePtr InW);
+
+	/* Instructions */
+
+	FSetMaterialOutputInstr* EmitSetMaterialOutput(EMaterialProperty InProperty, FValue* InArgValue);
+
+private:
+	FMaterialIRModule* Module{};
+
+	UE_MIR_PRIVATE();
+};
+
+} // namespace MaterialIR
+
+#endif
diff --git a/Engine/Source/Runtime/Engine/Public/Materials/MaterialIRCommon.h b/Engine/Source/Runtime/Engine/Public/Materials/MaterialIRCommon.h
new file mode 100644
--- /dev/null
+++ b/Engine/Source/Runtime/Engine/Public/Materials/MaterialIRCommon.h
@@ -0,0 +1,83 @@
+// Copyright Epic Games, Inc. All Rights Reserved.
+#pragma once
+
+#include "SceneTypes.h"
+#include "RHIShaderPlatform.h"
+#include "Containers/Array.h"
+#include "Math/MathFwd.h"
+
+#if WITH_EDITOR
+
+/* Forward declarations */
+
+namespace UE::Shader
+{
+	struct FValue;
+}
+
+namespace ERHIFeatureLevel { enum Type : int; }
+
+class ITargetPlatform;
+class FMaterialIRModule;
+class FMaterial;
+struct FStaticParameterSet;
+struct FShaderCompilerEnvironment;
+class FMaterialIRModule;
+class FMaterialIRModuleBuilder;
+struct FExpressionInput;
+
+namespace MaterialIR
+{
+
+/* Types*/
+struct FType;
+struct FArithmeticType;
+
+using FTypePtr = const FType*;
+using FArithmeticTypePtr = const FArithmeticType*;
+
+/* IR */
+struct FValue;
+using FValuePtr = const FValue*;
+struct FScalarValue;
+struct FVector;
+struct FSetMaterialOutputInstr;
+
+/* Others */
+class FBuilder;
+
+//
+template <typename TKindType>
+struct FIRNode
+{
+	TKindType Kind;
+
+	bool IsA(TKindType InKind) const
+	{
+		return (static_cast<uint64>(Kind) & static_cast<uint64>(InKind)) == static_cast<uint64>(InKind);
+	}
+
+	template <typename T>
+	const T* Cast() const
+	{
+ 		return this && IsA(T::TypeKind) ? static_cast<const T*>(this) : nullptr;
+	}
+};
+
+}
+
+#define UE_MIR_UNREACHABLE() { check(!"Unreachable"); UE_ASSUME(false); }
+
+#define UE_MIR_PRIVATE() \
+		struct FPrivate; \
+		friend FPrivate; \
+		FPrivate* AsPrivate() { return reinterpret_cast<FPrivate*>(this); }\
+		const FPrivate* AsPrivate() const { return reinterpret_cast<const FPrivate*>(this); }
+
+#define UE_MIR_BEGIN_PRIVATE(TypeName)\
+		struct TypeName::FPrivate : TypeName {
+
+#define UE_MIR_END_PRIVATE()\
+		}
+
+#endif // #if WITH_EDITOR
diff --git a/Engine/Source/Runtime/Engine/Public/Materials/MaterialIRModule.h b/Engine/Source/Runtime/Engine/Public/Materials/MaterialIRModule.h
new file mode 100644
--- /dev/null
+++ b/Engine/Source/Runtime/Engine/Public/Materials/MaterialIRModule.h
@@ -0,0 +1,30 @@
+// Copyright Epic Games, Inc. All Rights Reserved.
+#pragma once
+
+#include "Materials/MaterialIRCommon.h"
+#include "MaterialShared.h"
+
+#if WITH_EDITOR
+
+class FMaterialIRModule
+{
+public:
+	~FMaterialIRModule();
+	void Empty();
+	void Reset(EShaderPlatform InShaderPlatform, ERHIFeatureLevel::Type InFeatureLevel, const ITargetPlatform* InTargetPlatform);
+	EShaderPlatform GetShaderPlatform() const { return ShaderPlatform; }
+	const ITargetPlatform* GetTargetPlatform() const { return TargetPlatform; }
+	const FMaterialCompilationOutput& GetCompilationOutput() const { return CompilationOutput; }
+	TArrayView<const MaterialIR::FSetMaterialOutputInstr* const> GetOutputs() const { return Outputs; }
+
+private:
+	EShaderPlatform ShaderPlatform;
+	const ITargetPlatform* TargetPlatform;
+	FMaterialCompilationOutput CompilationOutput;
+	TArray<MaterialIR::FValue*> Values;
+	TArray<MaterialIR::FSetMaterialOutputInstr*> Outputs;
+
+	friend MaterialIR::FBuilder;
+};
+
+#endif
diff --git a/Engine/Source/Runtime/Engine/Public/Materials/MaterialIRModuleBuilder.h b/Engine/Source/Runtime/Engine/Public/Materials/MaterialIRModuleBuilder.h
new file mode 100644
--- /dev/null
+++ b/Engine/Source/Runtime/Engine/Public/Materials/MaterialIRModuleBuilder.h
@@ -0,0 +1,38 @@
+// Copyright Epic Games, Inc. All Rights Reserved.
+#pragma once
+
+#include "Materials/MaterialIRCommon.h"
+
+#if WITH_EDITOR
+
+//
+class FMaterialIRModuleBuilder
+{
+public:
+	//
+	void SetSource(FMaterial* InMaterial, const FStaticParameterSet* InStaticParameters);
+
+	//
+	void SetPlatform(EShaderPlatform InShaderPlatform, ERHIFeatureLevel::Type InFeatureLevel, const ITargetPlatform* InTargetPlatform);
+
+	//
+	void SetTarget(FMaterialIRModule* InModule);
+
+	//
+	bool Build();
+
+private:
+	FMaterial* Material;
+	const FStaticParameterSet* StaticParameters;
+	EShaderPlatform ShaderPlatform{};
+	ERHIFeatureLevel::Type FeatureLevel{};
+	const ITargetPlatform* TargetPlatform{};
+	FMaterialIRModule* Module;
+
+	TArray<FExpressionInput*> ExpressionAnalysisStack;
+
+	struct FHelper;
+	friend FHelper;
+};
+
+#endif
diff --git a/Engine/Source/Runtime/Engine/Public/Materials/MaterialIRToHLSLTranslator.h b/Engine/Source/Runtime/Engine/Public/Materials/MaterialIRToHLSLTranslator.h
new file mode 100644
--- /dev/null
+++ b/Engine/Source/Runtime/Engine/Public/Materials/MaterialIRToHLSLTranslator.h
@@ -0,0 +1,23 @@
+// Copyright Epic Games, Inc. All Rights Reserved.
+
+#pragma once
+
+#include "Materials/MaterialIRCommon.h"
+
+#if WITH_EDITOR
+
+class FMaterialIRToHLSLTranslator
+{
+public:
+	using FParametersMap = TMap<FString, FString>;
+
+public:
+	FMaterialIRToHLSLTranslator();
+	void Translate(const FMaterial& InMaterial, const FMaterialIRModule& InModule, FParametersMap& OutParameters, FShaderCompilerEnvironment& OutEnvironment);
+
+private:
+ 	
+	UE_MIR_PRIVATE();
+};
+
+#endif // #if WITH_EDITOR
diff --git a/Engine/Source/Runtime/Engine/Public/Materials/MaterialIRTypes.h b/Engine/Source/Runtime/Engine/Public/Materials/MaterialIRTypes.h
new file mode 100644
--- /dev/null
+++ b/Engine/Source/Runtime/Engine/Public/Materials/MaterialIRTypes.h
@@ -0,0 +1,54 @@
+// Copyright Epic Games, Inc. All Rights Reserved.
+#pragma once
+
+#include "Materials/MaterialIRCommon.h"
+#include "Shader/ShaderTypes.h"
+
+#if WITH_EDITOR
+
+namespace MaterialIR
+{
+
+enum ETypeKind
+{
+	TK_Void = 1,
+	TK_Arithmetic
+};
+
+struct FType
+{
+	ETypeKind Kind;
+
+	static FTypePtr FromShaderType(const UE::Shader::FType& InShaderType);
+	static FTypePtr GetVoid();
+
+	const FArithmeticType* ToArithmetic() const;
+	bool IsScalar() const;
+
+};
+
+enum EScalarKind
+{
+	SK_Bool, SK_Int, SK_Float,
+};
+
+const TCHAR* ScalarKindToString(EScalarKind Kind);
+
+struct FArithmeticType : FType
+{
+	EScalarKind ScalarKind;
+	int NumRows;
+	int NumColumns;
+
+	static const FArithmeticType* GetScalar(EScalarKind InScalarKind);
+	static const FArithmeticType* GetVector(EScalarKind InScalarKind, int NumRows);
+	static const FArithmeticType* GetMatrix(EScalarKind InScalarKind, int NumColumns, int NumRows);
+
+	bool IsScalar() const { return NumRows == 1 && NumColumns == 1; }
+	bool IsVector() const { return NumRows > 1 && NumColumns == 1; }
+	bool IsMatrix() const { return NumRows > 1 && NumColumns > 1; }
+};
+
+} // namespace MaterialIR
+
+#endif
