[Experimental] World Partition Streaming Cells Transformers
- Added streaming cells transformers experimental feature (works in PIE and cooked builds), which provides a framework for Licenses to implement their own transformers to transform cells.
- Provided implementation Log: will dump stats before and after applying level transformers, with option to fully log or just log deltas (default).
- Provided implementation ISM: the streaming cell actors will be packed into a single actor at the center of the cell with instanced Static Mesh components for packed instances
  - No special treatment for static meshes with a material not supporting instances, will report a warning, same as HLOD instancing.
  - Also works with partition actors such as Foliage, so cells with several compatible Partition actors will be merged into a single one.
  - Provides Allowed and Disallowed classes for content specific rules, by default non-static actors, replicated and HLOD actors.
  - Supports actor coloration to show merged actors in green at runtime using console command `Show ActorColoration CellTransformerISM`.
- Potential other implementations: MASS conversion, PCG, Physics split, etc.

diff --git a/Engine/Source/Runtime/Engine/Private/WorldPartition/WorldPartition.cpp b/Engine/Source/Runtime/Engine/Private/WorldPartition/WorldPartition.cpp
--- a/Engine/Source/Runtime/Engine/Private/WorldPartition/WorldPartition.cpp
+++ b/Engine/Source/Runtime/Engine/Private/WorldPartition/WorldPartition.cpp
@@ -4,6 +4,7 @@
 	WorldPartition.cpp: UWorldPartition implementation
 =============================================================================*/
 #include "WorldPartition/WorldPartition.h"
+#include "Engine/Engine.h"
 #include "Misc/PackageName.h"
 #include "UObject/AssetRegistryTagsContext.h"
 #include "UObject/UObjectIterator.h"
@@ -57,6 +58,8 @@
 #include "WorldPartition/ErrorHandling/WorldPartitionStreamingGenerationMapCheckErrorHandler.h"
 #include "Modules/ModuleManager.h"
 #include "GameDelegates.h"
+
+#include "WorldPartition/WorldPartitionRuntimeCellTransformerISM.h"
 #else
 #include "Engine/Level.h"
 #endif //WITH_EDITOR
@@ -484,6 +487,32 @@ bool UWorldPartition::CanEditChange(const FProperty* InProperty) const
 	return true;
 }
 
+void UWorldPartition::PostEditChangeChainProperty(FPropertyChangedChainEvent& PropertyChangedEvent)
+{
+	Super::PostEditChangeChainProperty(PropertyChangedEvent);
+
+	static FName NAME_RuntimeCellsTransformerStack(TEXT("RuntimeCellsTransformerStack"));
+	FName PropertyName = PropertyChangedEvent.Property ? PropertyChangedEvent.Property->GetFName() : NAME_None;
+
+	if (PropertyName == GET_MEMBER_NAME_CHECKED(FRuntimeCellTransformerInstance, Class))
+	{
+		int32 TransformerIndex = PropertyChangedEvent.GetArrayIndex(NAME_RuntimeCellsTransformerStack.ToString());
+		check(RuntimeCellsTransformerStack.IsValidIndex(TransformerIndex));
+
+		FRuntimeCellTransformerInstance& TransformerInstance = RuntimeCellsTransformerStack[TransformerIndex];
+		
+		UWorldPartitionRuntimeCellTransformer* OldTransformerInstance = TransformerInstance.Instance;
+
+		TransformerInstance.Instance = TransformerInstance.Class ? NewObject<UWorldPartitionRuntimeCellTransformer>(this, TransformerInstance.Class, NAME_None) : nullptr;
+
+		if (OldTransformerInstance && TransformerInstance.Instance)
+		{
+			UEngine::FCopyPropertiesForUnrelatedObjectsParams Params;
+			UEngine::CopyPropertiesForUnrelatedObjects(OldTransformerInstance, TransformerInstance.Instance, Params);
+		}
+	}
+}
+
 FName UWorldPartition::GetWorldPartitionEditorName() const
 {
 	if (SupportsStreaming())
@@ -1667,6 +1696,25 @@ void UWorldPartition::Tick(float DeltaSeconds)
 			}
 		}
 	}
+
+	if (RuntimeCellsTransformerStackTimes.Num())
+	{
+		RuntimeCellsTransformerStackDumpTime += DeltaSeconds;
+
+		if (RuntimeCellsTransformerStackDumpTime > 10.0f)
+		{
+			RuntimeCellsTransformerStackTimes.ValueSort([](const TPair<double, int32>& A, const TPair<double, int32>& B) { return B.Key < A.Key; } );
+
+			UE_LOG(LogWorldPartition, Log, TEXT("Runtime cells transformer stack per-cell stats:"));
+			for (TPair<UClass*, TPair<double, int32>>& StatsPair : RuntimeCellsTransformerStackTimes)
+			{
+				UE_LOG(LogWorldPartition, Log, TEXT("\t%s: %s"), *StatsPair.Key->GetName(), *FPlatformTime::PrettyTime(StatsPair.Value.Key / (double)StatsPair.Value.Value));
+			}
+
+			RuntimeCellsTransformerStackTimes.Empty();
+			RuntimeCellsTransformerStackDumpTime = 0.0f;
+		}
+	}
 #endif
 }
 
@@ -2209,6 +2257,27 @@ FBox UWorldPartition::GetRuntimeWorldBounds() const
 
 	return EditorHash->GetNonSpatialBounds();
 }
+
+void UWorldPartition::ApplyRuntimeCellsTransformerStack(ULevel* InLevel)
+{
+	auto ApplyTransformPhase = [this, InLevel](TFunction<void(const FRuntimeCellTransformerInstance& TransformerInstance, ULevel* Level)> Func)
+	{
+		for (const FRuntimeCellTransformerInstance& TransformerInstance : RuntimeCellsTransformerStack)
+		{
+			TPair<double, int32>& TotalStats = RuntimeCellsTransformerStackTimes.FindOrAdd(TransformerInstance.Class);
+			TotalStats.Key -= FPlatformTime::Seconds();
+		
+			Func(TransformerInstance, InLevel);
+		
+			TotalStats.Key += FPlatformTime::Seconds();
+			TotalStats.Value++;
+		}
+	};
+
+	ApplyTransformPhase([](const FRuntimeCellTransformerInstance& TransformerInstance, ULevel* Level) { TransformerInstance.PreTransform(Level); });
+	ApplyTransformPhase([](const FRuntimeCellTransformerInstance& TransformerInstance, ULevel* Level) { TransformerInstance.Transform(Level); });
+	ApplyTransformPhase([](const FRuntimeCellTransformerInstance& TransformerInstance, ULevel* Level) { TransformerInstance.PostTransform(Level); });
+}
 #endif
 
 bool UWorldPartition::SupportsWorldAssetStreaming(const FName& InTargetGrid)
@@ -2232,4 +2301,3 @@ TArray<UWorldPartitionRuntimeCell*> UWorldPartition::GetWorldAssetStreamingCells
 }
 
 #undef LOCTEXT_NAMESPACE
-
diff --git a/Engine/Source/Runtime/Engine/Private/WorldPartition/WorldPartitionLevelStreamingDynamic.cpp b/Engine/Source/Runtime/Engine/Private/WorldPartition/WorldPartitionLevelStreamingDynamic.cpp
--- a/Engine/Source/Runtime/Engine/Private/WorldPartition/WorldPartitionLevelStreamingDynamic.cpp
+++ b/Engine/Source/Runtime/Engine/Private/WorldPartition/WorldPartitionLevelStreamingDynamic.cpp
@@ -537,6 +537,8 @@ void UWorldPartitionLevelStreamingDynamic::FinalizeRuntimeLevel()
 			// Remap Runtime Level's SoftObjectPaths
 			FWorldPartitionLevelHelper::RemapLevelSoftObjectPaths(RuntimeLevel, OuterWorldPartition.Get());
 		}
+
+		OuterWorldPartition->ApplyRuntimeCellsTransformerStack(RuntimeLevel);
 	}
 
 	SetLoadedLevel(RuntimeLevel);
diff --git a/Engine/Source/Runtime/Engine/Private/WorldPartition/WorldPartitionRuntimeCellTransformer.cpp b/Engine/Source/Runtime/Engine/Private/WorldPartition/WorldPartitionRuntimeCellTransformer.cpp
new file mode 100644
--- /dev/null
+++ b/Engine/Source/Runtime/Engine/Private/WorldPartition/WorldPartitionRuntimeCellTransformer.cpp
@@ -0,0 +1,7 @@
+// Copyright Epic Games, Inc. All Rights Reserved.
+
+#include "WorldPartition/WorldPartitionRuntimeCellTransformer.h"
+
+UWorldPartitionRuntimeCellTransformer::UWorldPartitionRuntimeCellTransformer(const FObjectInitializer& ObjectInitializer)
+	: Super(ObjectInitializer)
+{}
\ No newline at end of file
diff --git a/Engine/Source/Runtime/Engine/Private/WorldPartition/WorldPartitionRuntimeCellTransformerISM.cpp b/Engine/Source/Runtime/Engine/Private/WorldPartition/WorldPartitionRuntimeCellTransformerISM.cpp
new file mode 100644
--- /dev/null
+++ b/Engine/Source/Runtime/Engine/Private/WorldPartition/WorldPartitionRuntimeCellTransformerISM.cpp
@@ -0,0 +1,215 @@
+// Copyright Epic Games, Inc. All Rights Reserved.
+
+#include "WorldPartition/WorldPartitionRuntimeCellTransformerISM.h"
+
+#if WITH_EDITOR
+#include "CoreMinimal.h"
+#include "GameFramework/Actor.h"
+#include "ISMPartition/ISMComponentBatcher.h"
+#include "ISMPartition/ISMComponentDescriptor.h"
+#include "Engine/Level.h"
+#include "UObject/Package.h"
+#endif
+
+#if WITH_EDITORONLY_DATA
+#include "ActorPartition/PartitionActor.h"
+#include "Engine/StaticMeshActor.h"
+#endif
+
+#include "GameFramework/ActorPrimitiveColorHandler.h"
+
+#define LOCTEXT_NAMESPACE "WorldPartition"
+
+static const FName NAME_AutoInstancingActor(TEXT("AutoInstancingActor"));
+
+UWorldPartitionRuntimeCellTransformerISM::UWorldPartitionRuntimeCellTransformerISM(const FObjectInitializer& ObjectInitializer)
+	: Super(ObjectInitializer)
+{
+#if WITH_EDITORONLY_DATA
+	AllowedClasses.Add(APartitionActor::StaticClass());
+	AllowedClasses.Add(AStaticMeshActor::StaticClass());
+	MinNumInstances = 2;
+#endif
+
+#if ENABLE_ACTOR_PRIMITIVE_COLOR_HANDLER
+	if (HasAnyFlags(RF_ClassDefaultObject) && ExactCast<UWorldPartitionRuntimeCellTransformerISM>(this))
+	{
+		FActorPrimitiveColorHandler::Get().RegisterPrimitiveColorHandler(TEXT("CellTransformerISM"), LOCTEXT("CellTransformerISM", "Cell Transformer ISM"), false, [](const UPrimitiveComponent* InPrimitiveComponent)
+		{
+			if (AActor* Actor = InPrimitiveComponent ? InPrimitiveComponent->GetOwner() : nullptr)
+			{
+				if (Actor->GetFName() == NAME_AutoInstancingActor)
+				{
+					return FLinearColor::MakeRandomSeededColor(GetTypeHash(InPrimitiveComponent->GetFName()));
+				}
+			}
+			return FLinearColor::White;
+		});
+	}
+#endif
+}
+
+#if WITH_EDITOR
+void UWorldPartitionRuntimeCellTransformerISM::Transform(ULevel* InLevel)
+{
+	check(InLevel);
+
+	struct FActorComponentBatcherDescriptor
+	{
+		TMap<TObjectPtr<AActor>*, TArray<UStaticMeshComponent*>> ActorComponents;
+		FISMComponentBatcher ISMComponentBatcher;
+	};
+
+	TMap<FISMComponentDescriptor, FActorComponentBatcherDescriptor> ISMComponentBatchers;
+
+	for (TObjectPtr<AActor>& Actor : InLevel->Actors)
+	{
+		if (IsValid(Actor) && CanAutoInstanceActor(Actor))
+		{
+			// Gather potential components that can be merged
+			Actor->ForEachComponent<UStaticMeshComponent>(true, [&ISMComponentBatchers, &Actor](UStaticMeshComponent* StaticMeshComponent)
+			{
+				if (!StaticMeshComponent->IsEditorOnly() && StaticMeshComponent->IsVisible() && (StaticMeshComponent->Mobility == EComponentMobility::Static))
+				{
+					FISMComponentDescriptor ISMComponentDescriptor;
+					ISMComponentDescriptor.InitFrom(StaticMeshComponent, false);
+					ISMComponentBatchers.FindOrAdd(ISMComponentDescriptor).ActorComponents.FindOrAdd(&Actor).Add(StaticMeshComponent);
+				}
+			});
+		}
+	}
+
+	int32 NumInstancedComponents = 0;
+	for (auto& [ISMComponentDescriptor, ActorComponentBatcherDescriptor] : ISMComponentBatchers)
+	{
+		if ((uint32)ActorComponentBatcherDescriptor.ActorComponents.Num() >= MinNumInstances)
+		{
+			for (auto& [Actor, Components] : ActorComponentBatcherDescriptor.ActorComponents)
+			{
+				for (UStaticMeshComponent* StaticMeshComponent : Components)
+				{
+					// Register the component into the batcher
+					StaticMeshComponent->UpdateComponentToWorld();
+					ActorComponentBatcherDescriptor.ISMComponentBatcher.Add(StaticMeshComponent);
+
+					// Remove the component from the actor
+					(*Actor)->RemoveOwnedComponent(StaticMeshComponent);
+					StaticMeshComponent->MarkAsGarbage();
+
+					NumInstancedComponents++;
+				}
+
+				if (CanRemoveActor(*Actor))
+				{
+					*Actor = nullptr;
+				}
+			}
+		}
+	}
+
+	InLevel->Actors.Remove(nullptr);
+
+	if (NumInstancedComponents)
+	{
+		AActor* PackedActor = NewObject<AActor>(InLevel, NAME_AutoInstancingActor);
+
+		for (auto& [ISMComponentDescriptor, ActorComponentBatcherDescriptor] : ISMComponentBatchers)
+		{
+			if (ActorComponentBatcherDescriptor.ISMComponentBatcher.GetNumInstances())
+			{
+				UInstancedStaticMeshComponent* ISMComponent = NewObject<UInstancedStaticMeshComponent>(PackedActor);
+				ISMComponentDescriptor.InitComponent(ISMComponent);
+				ActorComponentBatcherDescriptor.ISMComponentBatcher.InitComponent(ISMComponent);
+
+				if (!PackedActor->GetRootComponent())
+				{
+					PackedActor->SetRootComponent(ISMComponent);
+				}
+
+				ISMComponent->SetMobility(EComponentMobility::Static);
+				ISMComponent->SetWorldTransform(FTransform::Identity);
+
+				PackedActor->AddInstanceComponent(ISMComponent);
+			}
+		}
+
+		InLevel->Actors.Add(PackedActor);
+	}
+}
+
+bool UWorldPartitionRuntimeCellTransformerISM::CanAutoInstanceActor(AActor* InActor) const
+{
+	if (InActor->GetIsReplicated())
+	{
+		return false;
+	}
+
+	if (!InActor->IsRootComponentStatic())
+	{
+		return false;
+	}
+
+	if (InActor->IsHidden())
+	{
+		return false;
+	}
+
+	if (InActor->IsEditorOnly())
+	{
+		return false;
+	}
+
+	if (InActor->Children.Num() || InActor->IsChildActor())
+	{
+		return false;
+	}
+
+	UClass* ActorClass = InActor->GetClass();
+	for (TSubclassOf<AActor> DisallowedClass : DisallowedClasses)
+	{
+		if (ActorClass == DisallowedClass)
+		{
+			return false;
+		}
+	}
+
+	for (TSubclassOf<AActor> AllowedClass : AllowedClasses)
+	{
+		if (ActorClass->IsChildOf(AllowedClass))
+		{
+			return true;
+		}
+	}
+
+	return false;
+}
+
+bool UWorldPartitionRuntimeCellTransformerISM::CanRemoveActor(AActor* InActor) const
+{
+	uint32 NumRuntimeComponents = 0;
+	for (UActorComponent* Component : InActor->GetComponents())
+	{
+		if (!Component->IsEditorOnly())
+		{
+			NumRuntimeComponents++;
+		}
+	}
+		
+	if (!NumRuntimeComponents)
+	{
+		return true;
+	}
+
+	if (NumRuntimeComponents == 1)
+	{
+		if (InActor->GetRootComponent() && ExactCast<USceneComponent>(InActor->GetRootComponent()))
+		{
+			return true;
+		}
+	}
+
+	return false;
+}
+#endif
+
+#undef LOCTEXT_NAMESPACE
\ No newline at end of file
diff --git a/Engine/Source/Runtime/Engine/Private/WorldPartition/WorldPartitionRuntimeCellTransformerLog.cpp b/Engine/Source/Runtime/Engine/Private/WorldPartition/WorldPartitionRuntimeCellTransformerLog.cpp
new file mode 100644
--- /dev/null
+++ b/Engine/Source/Runtime/Engine/Private/WorldPartition/WorldPartitionRuntimeCellTransformerLog.cpp
@@ -0,0 +1,100 @@
+// Copyright Epic Games, Inc. All Rights Reserved.
+
+#include "WorldPartition/WorldPartitionRuntimeCellTransformerLog.h"
+#include "WorldPartition/WorldPartitionLog.h"
+#include "CoreMinimal.h"
+#include "Engine/Level.h"
+#include "GameFramework/Actor.h"
+#include "Templates/Greater.h"
+#include "UObject/ObjectMacros.h"
+#include "UObject/Package.h"
+
+UWorldPartitionRuntimeCellTransformerLog::UWorldPartitionRuntimeCellTransformerLog(const FObjectInitializer& ObjectInitializer)
+	: Super(ObjectInitializer)
+{
+#if WITH_EDITORONLY_DATA
+	bOnlyLogDifferences = true;
+#endif
+}
+
+#if WITH_EDITOR
+void UWorldPartitionRuntimeCellTransformerLog::PreTransform(ULevel* InLevel)
+{
+	check(ClassNumInstancesBefore.IsEmpty());
+	GatherLevelContentStats(InLevel, ClassNumInstancesBefore);
+
+	if (!bOnlyLogDifferences)
+	{
+		UE_LOG(LogWorldPartition, Log, TEXT("Level %s PreTransform Content:"), *InLevel->GetPackage()->GetName());
+		DumpLevelContentStats(InLevel, ClassNumInstancesBefore);
+	}
+}
+
+void UWorldPartitionRuntimeCellTransformerLog::PostTransform(ULevel* InLevel)
+{
+	TMap<UClass*, int32> ClassNumInstancesAfter;
+	GatherLevelContentStats(InLevel, ClassNumInstancesAfter);
+
+	if (bOnlyLogDifferences)
+	{
+		if (!ClassNumInstancesBefore.OrderIndependentCompareEqual(ClassNumInstancesAfter))
+		{
+			UE_LOG(LogWorldPartition, Log, TEXT("Level %s PostTransform Delta:"), *InLevel->GetPackage()->GetName());
+
+			TMap<UClass*, int32> ClassNumInstancesDelta;
+
+			for (auto& [Class, NumInstances] : ClassNumInstancesBefore)
+			{
+				const int32* NumInstancesAfter = ClassNumInstancesAfter.Find(Class);
+				ClassNumInstancesDelta.FindOrAdd(Class) += (NumInstancesAfter ? *NumInstancesAfter : 0) - NumInstances;
+			}
+
+			for (auto& [Class, NumInstances] : ClassNumInstancesAfter)
+			{
+				if (!ClassNumInstancesBefore.Contains(Class))
+				{
+					ClassNumInstancesDelta.FindOrAdd(Class) += NumInstances;
+				}
+			}
+
+			DumpLevelContentStats(InLevel, ClassNumInstancesDelta);
+		}
+	}
+	else
+	{
+		UE_LOG(LogWorldPartition, Log, TEXT("Level %s PostTransform Content:"), *InLevel->GetPackage()->GetName());
+		DumpLevelContentStats(InLevel, ClassNumInstancesAfter);
+	}
+
+	ClassNumInstancesBefore.Empty();
+}
+
+void UWorldPartitionRuntimeCellTransformerLog::GatherLevelContentStats(ULevel* InLevel, TMap<UClass*, int32>& OutClassNumInstances)
+{
+	for (AActor* Actor : InLevel->Actors)
+	{
+		if (IsValid(Actor))
+		{
+			OutClassNumInstances.FindOrAdd(Actor->GetClass())++;
+		}
+	}
+
+	OutClassNumInstances.ValueSort(TGreater<>());
+}
+
+void UWorldPartitionRuntimeCellTransformerLog::DumpLevelContentStats(ULevel* InLevel, const TMap<UClass*, int32>& InClassNumInstances)
+{
+	if (InClassNumInstances.Num())
+	{
+		UE_LOG(LogWorldPartition, Log, TEXT("\tNum Actors: %d"), InLevel->Actors.Num());
+		UE_LOG(LogWorldPartition, Log, TEXT("\tActor Class Breakdown:"));
+		for (auto& [Class, NumInstances] : InClassNumInstances)
+		{
+			if (NumInstances)
+			{
+				UE_LOG(LogWorldPartition, Log, TEXT("\t\t%s: %d"), *Class->GetName(), NumInstances);
+			}
+		}
+	}
+}
+#endif
\ No newline at end of file
diff --git a/Engine/Source/Runtime/Engine/Private/WorldPartition/WorldPartitionRuntimeLevelStreamingCell.cpp b/Engine/Source/Runtime/Engine/Private/WorldPartition/WorldPartitionRuntimeLevelStreamingCell.cpp
--- a/Engine/Source/Runtime/Engine/Private/WorldPartition/WorldPartitionRuntimeLevelStreamingCell.cpp
+++ b/Engine/Source/Runtime/Engine/Private/WorldPartition/WorldPartitionRuntimeLevelStreamingCell.cpp
@@ -517,6 +517,8 @@ bool UWorldPartitionRuntimeLevelStreamingCell::OnPopulateGeneratedPackageForCook
 		check(NewLevel->GetPackage() == InPackage);
 		FWorldPartitionLevelHelper::MoveExternalActorsToLevel(Packages, NewLevel, OutModifiedPackages);
 
+		WorldPartition->ApplyRuntimeCellsTransformerStack(NewLevel);
+
 		// Push temporarily the cooking ExternalStreamingObject in the policy for RemapLevelSoftObjectPaths to use it to resolve softobjectpaths
 		// Do this only if the ExternalStreamingObject has a valid root external data layer asset, as Content Bundle soft object remapping is not supported at cook time (there is no world package remapping)
 		const URuntimeHashExternalStreamingObjectBase* ExternalStreamingObject = GetTypedOuter<URuntimeHashExternalStreamingObjectBase>();
diff --git a/Engine/Source/Runtime/Engine/Public/WorldPartition/WorldPartition.h b/Engine/Source/Runtime/Engine/Public/WorldPartition/WorldPartition.h
--- a/Engine/Source/Runtime/Engine/Public/WorldPartition/WorldPartition.h
+++ b/Engine/Source/Runtime/Engine/Public/WorldPartition/WorldPartition.h
@@ -22,6 +22,7 @@
 #include "WorldPartition/WorldPartitionStreamingGeneration.h"
 #include "WorldPartition/WorldPartitionActorLoaderInterface.h"
 #include "WorldPartition/WorldPartitionEditorLoaderAdapter.h"
+#include "WorldPartition/WorldPartitionRuntimeCellTransformer.h"
 #include "ExternalDirtyActorsTracker.h"
 #include "PackageSourceControlHelper.h"
 #include "CookPackageSplitter.h"
@@ -113,6 +114,27 @@ public:
 };
 #endif
 
+/** Holds an instance of a runtime cell transformer. */
+USTRUCT()
+struct FRuntimeCellTransformerInstance
+{
+	GENERATED_USTRUCT_BODY()
+
+#if WITH_EDITORONLY_DATA
+	inline void PreTransform(ULevel* InLevel) const { if (Instance) { Instance->PreTransform(InLevel); } }
+	inline void Transform(ULevel* InLevel) const { if (Instance) { Instance->Transform(InLevel); } }
+	inline void PostTransform(ULevel* InLevel) const { if (Instance) { Instance->PostTransform(InLevel); } }
+
+	/** Runtime cell transformer class */
+	UPROPERTY(EditAnywhere, Category = WorldPartitionSetup, AdvancedDisplay)
+	TSubclassOf<UWorldPartitionRuntimeCellTransformer> Class;
+
+	/** Transformer object instance */
+	UPROPERTY(VisibleAnywhere, Category = WorldPartitionSetup, Instanced, Meta = (EditCondition = "Class != nullptr", HideEditConditionToggle, NoResetToDefault))
+	TObjectPtr<UWorldPartitionRuntimeCellTransformer> Instance;
+#endif
+};
+
 UCLASS(AutoExpandCategories=(WorldPartition), MinimalAPI)
 class UWorldPartition final : public UObject, public FActorDescContainerInstanceCollection, public IWorldPartitionCookPackageGenerator
 {
@@ -186,6 +208,7 @@ public:
 	//~ Begin UObject Interface
 #if WITH_EDITOR
 	ENGINE_API virtual bool CanEditChange(const FProperty* InProperty) const override;
+	ENGINE_API virtual void PostEditChangeChainProperty(FPropertyChangedChainEvent& PropertyChangedEvent) override;
 #endif //WITH_EDITOR
 	ENGINE_API virtual void Serialize(FArchive& Ar) override;
 	ENGINE_API virtual UWorld* GetWorld() const override;
@@ -310,6 +333,8 @@ public:
 
 	ENGINE_API FBox GetEditorWorldBounds() const;
 	ENGINE_API FBox GetRuntimeWorldBounds() const;
+
+	ENGINE_API void ApplyRuntimeCellsTransformerStack(ULevel* InLevel);
 	
 	UHLODLayer* GetDefaultHLODLayer() const { return DefaultHLODLayer; }
 	void SetDefaultHLODLayer(UHLODLayer* InDefaultHLODLayer) { DefaultHLODLayer = InDefaultHLODLayer; }
@@ -589,6 +614,14 @@ private:
 #endif
 
 #if WITH_EDITORONLY_DATA
+	/** Runtime cells transform stack objects */
+	UPROPERTY(EditAnywhere, Category = WorldPartitionSetup)
+	TArray<FRuntimeCellTransformerInstance> RuntimeCellsTransformerStack;
+
+	/** Runtime cells transform stack objects execution stats */
+	float RuntimeCellsTransformerStackDumpTime = 0.0f;
+	TMap<UClass*, TPair<double, int32>> RuntimeCellsTransformerStackTimes;
+
 	// Default HLOD layer
 	UPROPERTY(EditAnywhere, Category = WorldPartitionSetup, meta = (DisplayName = "Default HLOD Layer", EditCondition="bEnableStreaming", EditConditionHides, HideEditConditionToggle))
 	TObjectPtr<class UHLODLayer> DefaultHLODLayer;
diff --git a/Engine/Source/Runtime/Engine/Public/WorldPartition/WorldPartitionLevelHelper.h b/Engine/Source/Runtime/Engine/Public/WorldPartition/WorldPartitionLevelHelper.h
--- a/Engine/Source/Runtime/Engine/Public/WorldPartition/WorldPartitionLevelHelper.h
+++ b/Engine/Source/Runtime/Engine/Public/WorldPartition/WorldPartitionLevelHelper.h
@@ -44,7 +44,7 @@ public:
 	UE_DEPRECATED(5.4, "LoadActors is deprecated, LoadActors with FLoadActorsParams should be used instead.")
 	static bool LoadActors(UWorld* InOuterWorld, ULevel* InDestLevel, TArrayView<FWorldPartitionRuntimeCellObjectMapping> InActorPackages, FPackageReferencer& InPackageReferencer, TFunction<void(bool)> InCompletionCallback, bool bInLoadAsync, FLinkerInstancingContext InInstancingContext);
 
-	/* Struct of optional parameters passed to foreach actordesc functions. */
+	/* Struct of optional parameters passed to LoadActors function. */
 	struct FLoadActorsParams
 	{
 		FLoadActorsParams()
diff --git a/Engine/Source/Runtime/Engine/Public/WorldPartition/WorldPartitionRuntimeCellTransformer.h b/Engine/Source/Runtime/Engine/Public/WorldPartition/WorldPartitionRuntimeCellTransformer.h
new file mode 100644
--- /dev/null
+++ b/Engine/Source/Runtime/Engine/Public/WorldPartition/WorldPartitionRuntimeCellTransformer.h
@@ -0,0 +1,18 @@
+// Copyright Epic Games, Inc. All Rights Reserved.
+#pragma once
+
+#include "CoreMinimal.h"
+#include "WorldPartitionRuntimeCellTransformer.generated.h"
+
+UCLASS(Abstract)
+class UWorldPartitionRuntimeCellTransformer : public UObject
+{
+	GENERATED_UCLASS_BODY()
+
+#if WITH_EDITOR
+public:
+	virtual void PreTransform(ULevel* InLevel) {}
+	virtual void Transform(ULevel* InLevel) {}
+	virtual void PostTransform(ULevel* InLevel) {}
+#endif
+};
\ No newline at end of file
diff --git a/Engine/Source/Runtime/Engine/Public/WorldPartition/WorldPartitionRuntimeCellTransformerISM.h b/Engine/Source/Runtime/Engine/Public/WorldPartition/WorldPartitionRuntimeCellTransformerISM.h
new file mode 100644
--- /dev/null
+++ b/Engine/Source/Runtime/Engine/Public/WorldPartition/WorldPartitionRuntimeCellTransformerISM.h
@@ -0,0 +1,35 @@
+// Copyright Epic Games, Inc. All Rights Reserved.
+#pragma once
+
+#include "WorldPartition/WorldPartitionRuntimeCellTransformer.h"
+#include "WorldPartitionRuntimeCellTransformerISM.generated.h"
+
+UCLASS()
+class UWorldPartitionRuntimeCellTransformerISM : public UWorldPartitionRuntimeCellTransformer
+{
+	GENERATED_UCLASS_BODY()
+
+#if WITH_EDITOR
+public:
+	virtual void Transform(ULevel* InLevel) override;
+
+private:
+	bool CanAutoInstanceActor(AActor* InActor) const;
+	bool CanRemoveActor(AActor* InActor) const;
+#endif
+
+#if WITH_EDITORONLY_DATA
+public:
+	/** Allowed classes (recursive) to convert to instances */
+	UPROPERTY(EditAnywhere, Category = ISM)
+	TArray<TSubclassOf<AActor>> AllowedClasses;
+
+	/** Disallowed classes (non-recursive) to convert to instances */
+	UPROPERTY(EditAnywhere, Category = ISM)
+	TArray<TSubclassOf<AActor>> DisallowedClasses;
+
+	/** Minimum number of instances required to allow converting actors to ISM */
+	UPROPERTY(EditAnywhere, Category = ISM)
+	uint32 MinNumInstances;
+#endif
+};
\ No newline at end of file
diff --git a/Engine/Source/Runtime/Engine/Public/WorldPartition/WorldPartitionRuntimeCellTransformerLog.h b/Engine/Source/Runtime/Engine/Public/WorldPartition/WorldPartitionRuntimeCellTransformerLog.h
new file mode 100644
--- /dev/null
+++ b/Engine/Source/Runtime/Engine/Public/WorldPartition/WorldPartitionRuntimeCellTransformerLog.h
@@ -0,0 +1,31 @@
+// Copyright Epic Games, Inc. All Rights Reserved.
+#pragma once
+
+#include "CoreMinimal.h"
+#include "WorldPartitionRuntimeCellTransformer.h"
+#include "WorldPartitionRuntimeCellTransformerLog.generated.h"
+
+UCLASS()
+class UWorldPartitionRuntimeCellTransformerLog : public UWorldPartitionRuntimeCellTransformer
+{
+	GENERATED_UCLASS_BODY()
+
+#if WITH_EDITOR
+public:
+	virtual void PreTransform(ULevel* InLevel) override;
+	virtual void PostTransform(ULevel* InLevel) override;
+
+private:
+	void GatherLevelContentStats(ULevel* InLevel, TMap<UClass*, int32>& OutClassNumInstances);
+	void DumpLevelContentStats(ULevel* InLevel, const TMap<UClass*, int32>& InClassNumInstances);
+#endif
+
+#if WITH_EDITORONLY_DATA
+public:
+	/** Only log when the level was actually transformed */
+	UPROPERTY(EditAnywhere, Category = ISM)
+	bool bOnlyLogDifferences;
+
+	TMap<UClass*, int32> ClassNumInstancesBefore;
+#endif
+};
\ No newline at end of file
