Add missing include.

diff --git a/Engine/Source/Runtime/Engine/Public/WorldPartition/WorldPartitionRuntimeCellTransformerISM.h b/Engine/Source/Runtime/Engine/Public/WorldPartition/WorldPartitionRuntimeCellTransformerISM.h
--- a/Engine/Source/Runtime/Engine/Public/WorldPartition/WorldPartitionRuntimeCellTransformerISM.h
+++ b/Engine/Source/Runtime/Engine/Public/WorldPartition/WorldPartitionRuntimeCellTransformerISM.h
@@ -2,6 +2,8 @@
 #pragma once
 
 #include "WorldPartition/WorldPartitionRuntimeCellTransformer.h"
+#include "Templates/SubclassOf.h"
+
 #include "WorldPartitionRuntimeCellTransformerISM.generated.h"
 
 UCLASS()
