Fix TextureSource.GetId() changing when you call Compress()
the main problem was actually the CompressionFormat being put into the Id
but go ahead and make it more robust so that we clarify that Id is the BulkData Payload Id *before* compression

diff --git a/Engine/Source/Runtime/Engine/Private/Texture.cpp b/Engine/Source/Runtime/Engine/Private/Texture.cpp
--- a/Engine/Source/Runtime/Engine/Private/Texture.cpp
+++ b/Engine/Source/Runtime/Engine/Private/Texture.cpp
@@ -1363,7 +1363,9 @@ void UTexture::PreSave(FObjectPreSaveContext ObjectSaveContext)
 	// because it invalidates the texture build due to source hash change
 	// and could cause another build to be triggered during PostCompilation
 	// causing reentrancy problems.
-	BlockOnAnyAsyncBuild();
+	//BlockOnAnyAsyncBuild();
+	// use Modify(false) so that we also block on other textures that use us as a composite
+	Modify(false);
 
 	if (!GEngine->IsAutosaving() && !ObjectSaveContext.IsProceduralSave())
 	{
@@ -2138,10 +2140,9 @@ void FTextureSource::InitLayered(
 	}
 	else
 	{
-		BulkData.UpdatePayload(FUniqueBuffer::Alloc(TotalBytes).MoveToShared(), Owner);
-		// @todo Oodle : FIX ME ?? uninitialized memory ??
-		// ?? with no incoming data this is hashing garbage ???
-		// I think this also does FEditorBulkData::Register !?
+		// make sure data is initialized to zero:
+		FUniqueBuffer Buffer = FUniqueBuffer::AllocZeroed(TotalBytes);
+		BulkData.UpdatePayload(Buffer.MoveToShared(), Owner);
 	}
 	
 	// don't compress BulkData yet, it will be done by Compress() from PreSave()
@@ -2377,7 +2378,9 @@ void FTextureSource::RemoveCompression()
 	}
 
 	// BulkData LZ options not changed here
-	// also no UseHashAsGuid
+	
+	// update the Id from the decompressed data :
+	UseHashAsGuid();
 }
 
 void FTextureSource::Compress()
@@ -2408,6 +2411,17 @@ void FTextureSource::Compress()
 			// RemoveCompression();
 			
 			FSharedBuffer Buffer = Decompress();
+
+			if ( CompressionFormat != TSCF_None )
+			{
+				// set BulkData to the uncompressed data so we can get the hash before compression :
+				// BulkData.UpdatePayload does a slow hash update
+				BulkData.UpdatePayload(Buffer, Owner);
+				CompressionFormat = TSCF_None;
+			}
+
+			// update the Id from the decompressed data :
+			UseHashAsGuid();
 		
 			if ( ! HasLayerColorInfo() )
 			{
@@ -2417,14 +2431,18 @@ void FTextureSource::Compress()
 
 			FSharedBuffer DeltaBuffer = DoUEDeltaTransform(Buffer,true);
 			
-			// at this moment it would be easy to try LZ compression on the delta and non-delta data and choose the best
+			// note: at this moment it would be easy to try LZ compression on the delta and non-delta data and choose the best
 			// if you care about small uasset size and don't mind a slightly slower encode
 			//  (90% of uasset save time is not in this function)
 
 			// BulkData.UpdatePayload does a slow hash update
 			BulkData.UpdatePayload(DeltaBuffer, Owner);
-
 			CompressionFormat = TSCF_UEDELTA;
+			
+			//	we try to keep "Id" == to the hash of the BulkData when it was the raw data
+			//	the invariant
+			//	( Id == UE::Serialization::IoHashToGuid(BulkData.GetPayloadId()) )
+			//	is no longer true after this			
 		}
 	}
 	else // not ShouldUseUEDeltaForFormat
@@ -2433,13 +2451,6 @@ void FTextureSource::Compress()
 	}
 		
 	BulkData.SetCompressionOptions(ECompressedBufferCompressor::Kraken,ECompressedBufferCompressionLevel::Fast);
-
-	// note: we changed BulkData payload here
-	//	but we do NOT call UseHashAsGuid
-	//	we try to keep "Id" == to the hash of the BulkData when it was the raw data
-	//	the invariant
-	//	( Id == UE::Serialization::IoHashToGuid(BulkData.GetPayloadId()) )
-	//	is no longer true after this
 }
 
 FSharedBuffer FTextureSource::Decompress(IImageWrapperModule* ) const
@@ -3478,6 +3489,15 @@ void FTextureSource::UseHashAsGuid()
 	FScopeLock BulkDataExclusiveScope(&const_cast<FCriticalSection &>(BulkDataLock.Get()));
 #endif
 
+	if ( bGuidIsHash && CompressionFormat == TSCF_UEDELTA )
+	{
+		// we try to keep Id == the hash of the TSCF_None data before Compress()
+		// when the data is changed to UEDELTA , the hash is captured at that point
+		// if you call UseHashAsGuid again after that, we do not change Id
+		return;
+	}
+
+	// HasPayloadData is the same as Payload Size != 0
 	if (HasPayloadData())
 	{
 		CheckTextureIsUnlocked(TEXT("UseHashAsGuid"));
@@ -3487,6 +3507,9 @@ void FTextureSource::UseHashAsGuid()
 	}
 	else
 	{
+		// or ForceGenerateGuid() here?
+
+		bGuidIsHash = true;
 		Id.Invalidate();
 	}
 }
@@ -3513,7 +3536,16 @@ FGuid FTextureSource::GetId() const
 	IdBuilder << NumMips;
 	IdBuilder << NumLayers;	
 	IdBuilder << bLongLatCubemap;
-	IdBuilder << CompressionFormat;
+
+	// GetId() result should not change when CompressionFormat changes
+	//	so that before and after calling Compress() (save) , GetId() doesn't change
+	TEnumAsByte<enum ETextureSourceCompressionFormat> CompressionFormatForIdBuilder = CompressionFormat;
+	if ( CompressionFormat == TSCF_UEDELTA )
+	{
+		CompressionFormatForIdBuilder = TSCF_None;
+	}
+
+	IdBuilder << CompressionFormatForIdBuilder;
 	IdBuilder << bGuidIsHash; // always true here
 	IdBuilder << static_cast<uint8>(Format.GetValue());
 	
@@ -3543,6 +3575,8 @@ FGuid FTextureSource::GetId() const
 		}
 	}
 
+	// bUseHashAsGuid is true , so Id == UE::Serialization::IoHashToGuid(BulkData.GetPayloadId())
+	//	however, "Id" is kept as the hash of the data before Compress
 	IdBuilder << const_cast<FGuid&>(Id);
 
 	return IdBuilder.Build();
