Heterogeneous Volume Component: Fix crash when modifying bound material.

#tests editor tests modifying HV materials

diff --git a/Engine/Source/Runtime/Engine/Private/Components/HeterogeneousVolumeComponent.cpp b/Engine/Source/Runtime/Engine/Private/Components/HeterogeneousVolumeComponent.cpp
--- a/Engine/Source/Runtime/Engine/Private/Components/HeterogeneousVolumeComponent.cpp
+++ b/Engine/Source/Runtime/Engine/Private/Components/HeterogeneousVolumeComponent.cpp
@@ -56,8 +56,8 @@ private:
 	FHeterogeneousVolumeData HeterogeneousVolumeData;
 
 	// Cache UObject values
+	UMaterialInterface* MaterialInterface;
 	FMaterialRelevance MaterialRelevance;
-	FMaterialRenderProxy* MaterialRenderProxy;
 };
 
 /*=============================================================================
@@ -72,6 +72,7 @@ FHeterogeneousVolumeSceneProxy::FHeterogeneousVolumeSceneProxy(UHeterogeneousVol
 #else
 	, HeterogeneousVolumeData(this)
 #endif
+	, MaterialInterface(InComponent->GetMaterial(0))
 {
 	bIsHeterogeneousVolume = true;
 	bCastDynamicShadow = InComponent->CastShadow;
@@ -90,17 +91,14 @@ FHeterogeneousVolumeSceneProxy::FHeterogeneousVolumeSceneProxy(UHeterogeneousVol
 	VoxelSize.Z /= InComponent->VolumeResolution.Z;
 	HeterogeneousVolumeData.MinimumVoxelSize = FMath::Max(VoxelSize.GetMin(), 0.001);
 
-	UMaterialInterface* MaterialInterface = InComponent->GetMaterial(0);
 	if (InComponent->MaterialInstanceDynamic)
 	{
 		MaterialInterface = InComponent->MaterialInstanceDynamic;
 	}
 
-	MaterialRenderProxy = nullptr;
 	if (MaterialInterface)
 	{
 		MaterialRelevance = MaterialInterface->GetRelevance_Concurrent(GetScene().GetFeatureLevel());
-		MaterialRenderProxy = MaterialInterface->GetRenderProxy();
 	}
 
 	HeterogeneousVolumeData.StepFactor = InComponent->StepFactor;
@@ -192,13 +190,13 @@ void FHeterogeneousVolumeSceneProxy::GetDynamicMeshElements(
 		{
 			if (VisibilityMap & (1 << ViewIndex))
 			{
-				if (MaterialRenderProxy)
+				if (MaterialInterface)
 				{
 					// Set up MeshBatch
 					FMeshBatch& Mesh = Collector.AllocateMesh();
 
 					Mesh.VertexFactory = &VertexFactory;
-					Mesh.MaterialRenderProxy = MaterialRenderProxy;
+					Mesh.MaterialRenderProxy = MaterialInterface->GetRenderProxy();
 					Mesh.LCI = NULL;
 					Mesh.ReverseCulling = IsLocalToWorldDeterminantNegative() ? true : false;
 					Mesh.CastShadow = CastsDynamicShadow();
@@ -501,6 +499,7 @@ void UHeterogeneousVolumeComponent::PostEditChangeProperty(FPropertyChangedEvent
 			MaterialInstanceDynamic = CreateOrCastToMID(MaterialInterface);
 		}
 		OnSparseVolumeTextureChanged(GetSparseVolumeTexture(MaterialInterface, SVTParameterIndex));
+		MarkRenderStateDirty();
 	}
 
 	if (PropertyName == GET_MEMBER_NAME_CHECKED(UHeterogeneousVolumeComponent, VolumeResolution))
