Fix missing calls to RecacheUniformExpressions for material instances that had already cached their uniform expressions, but were still not considered visible by TObjectIterator because they still have the EInternalObjectFlags::AsyncLoading flag
This fixes the check UniformExpressionCache.CachedUniformExpressionShaderMap == Material.GetRenderingThreadShaderMap() in FMaterialShader::GetShaderBindings

diff --git a/Engine/Source/Runtime/Engine/Private/Materials/MaterialInstance.cpp b/Engine/Source/Runtime/Engine/Private/Materials/MaterialInstance.cpp
--- a/Engine/Source/Runtime/Engine/Private/Materials/MaterialInstance.cpp
+++ b/Engine/Source/Runtime/Engine/Private/Materials/MaterialInstance.cpp
@@ -66,6 +66,10 @@
 #include "Cooker/CookDependency.h"
 #endif
 
+#if WITH_ODSC
+#include "ODSC/ODSCManager.h"
+#endif
+
 #include UE_INLINE_GENERATED_CPP_BY_NAME(MaterialInstance)
 
 DECLARE_CYCLE_STAT(TEXT("MaterialInstance CopyMatInstParams"), STAT_MaterialInstance_CopyMatInstParams, STATGROUP_Shaders);
@@ -81,6 +85,9 @@ const FMaterialInstanceCachedData FMaterialInstanceCachedData::EmptyData{};
 void UMaterialInstance::StartCacheUniformExpressions() const
 {
 	UsedByRT |= (uint32)EMaterialInstanceUsedByRTFlag::CacheUniformExpressions;
+#if WITH_ODSC
+	FODSCManager::RegisterMaterialInstance(this);
+#endif
 }
 
 void UMaterialInstance::FinishCacheUniformExpressions() const
@@ -3350,6 +3357,10 @@ void UMaterialInstance::DeclareConstructClasses(TArray<FTopLevelAssetPath>& OutC
 
 void UMaterialInstance::BeginDestroy()
 {
+#if WITH_ODSC
+	FODSCManager::UnregisterMaterialInstance(this);
+#endif
+
 	TArray<TRefCountPtr<FMaterialResource>> ResourcesToDestroy;
 	for (FMaterialResource* CurrentResource : StaticPermutationMaterialResources)
 	{
diff --git a/Engine/Source/Runtime/Engine/Private/ODSC/ODSCManager.cpp b/Engine/Source/Runtime/Engine/Private/ODSC/ODSCManager.cpp
--- a/Engine/Source/Runtime/Engine/Private/ODSC/ODSCManager.cpp
+++ b/Engine/Source/Runtime/Engine/Private/ODSC/ODSCManager.cpp
@@ -6,6 +6,7 @@
 #include "ODSCLog.h"
 #include "ODSCThread.h"
 #include "Containers/BackgroundableTicker.h"
+#include "Materials/MaterialInstance.h"
 
 DEFINE_LOG_CATEGORY(LogODSC);
 
@@ -72,6 +73,13 @@ bool FODSCManager::Tick(float DeltaSeconds)
 		TArray<FODSCMessageHandler*> CompletedThreadedRequests;
 		Thread->GetCompletedRequests(CompletedThreadedRequests);
 
+		bool bFlushAsyncLoading = HasAsyncLoadingInstances();
+
+		if (CompletedThreadedRequests.Num() && bFlushAsyncLoading)
+		{
+			FlushAsyncLoading();
+		}
+
 		// Finish and remove any completed requests
 		for (FODSCMessageHandler* CompletedRequest : CompletedThreadedRequests)
 		{
@@ -117,3 +125,48 @@ void FODSCManager::AddThreadedShaderPipelineRequest(
 		Thread->AddShaderPipelineRequest(ShaderPlatform, FeatureLevel, QualityLevel, MaterialName, VertexFactoryName, PipelineName, ShaderTypeNames, PermutationId);
 	}
 }
+
+static inline bool IsODSCActive()
+{
+	return GODSCManager && GODSCManager->IsHandlingRequests();
+}
+
+void FODSCManager::RegisterMaterialInstance(const UMaterialInstance* MaterialInstance)
+{
+	if (IsODSCActive() && MaterialInstance->HasAnyInternalFlags(EInternalObjectFlags::AsyncLoading))
+	{
+		FScopeLock Lock(&GODSCManager->MaterialInstancesCachedUniformExpressionsCS);
+		TWeakObjectPtr<const UMaterialInstance>& MaterialInstanceSoftPtr = GODSCManager->MaterialInstancesCachedUniformExpressions.FindOrAdd(MaterialInstance);
+		MaterialInstanceSoftPtr = TWeakObjectPtr<const UMaterialInstance>(MaterialInstance);
+	}
+} 
+
+void FODSCManager::UnregisterMaterialInstance(const UMaterialInstance* MaterialInstance)
+{
+	if (GODSCManager != nullptr)
+	{
+		FScopeLock Lock(&GODSCManager->MaterialInstancesCachedUniformExpressionsCS);
+		GODSCManager->MaterialInstancesCachedUniformExpressions.Remove(MaterialInstance);
+	}
+}
+
+bool FODSCManager::HasAsyncLoadingInstances()
+{
+	FScopeLock Lock(&MaterialInstancesCachedUniformExpressionsCS);
+
+	bool bHasAsyncLoadingInstances = false;
+	for (auto Iter = MaterialInstancesCachedUniformExpressions.CreateIterator(); Iter; ++Iter)
+	{
+		const UMaterialInstance* MI = Iter.Value().Get();
+
+		if (MI == nullptr || !MI->HasAnyInternalFlags(EInternalObjectFlags::AsyncLoading))
+		{
+			Iter.RemoveCurrent();
+			continue;
+		}
+
+		bHasAsyncLoadingInstances = true;
+	}
+
+	return bHasAsyncLoadingInstances;
+}
diff --git a/Engine/Source/Runtime/Engine/Public/ODSC/ODSCManager.h b/Engine/Source/Runtime/Engine/Public/ODSC/ODSCManager.h
--- a/Engine/Source/Runtime/Engine/Public/ODSC/ODSCManager.h
+++ b/Engine/Source/Runtime/Engine/Public/ODSC/ODSCManager.h
@@ -9,6 +9,7 @@
 #include "ShaderCompiler.h"
 
 class FODSCThread;
+class UMaterialInstance;
 
 /**
  * Responsible for processing shader compile responses from the ODSC Thread.
@@ -81,15 +82,22 @@ public:
 	/** Returns true if we would actually add a request when calling AddThreadedShaderPipelineRequest. */
 	inline bool IsHandlingRequests() const { return Thread != nullptr; }
 
+	static void RegisterMaterialInstance(const UMaterialInstance* MI);
+	static void UnregisterMaterialInstance(const UMaterialInstance* MI);
+
 private:
 
 	ENGINE_API void OnEnginePreExit();
 	ENGINE_API void StopThread();
 
+	bool HasAsyncLoadingInstances();
+
 	/** Handles communicating directly with the cook on the fly server. */
 	FODSCThread* Thread = nullptr;
 
 	FDelegateHandle OnScreenMessagesHandle;
+	FCriticalSection MaterialInstancesCachedUniformExpressionsCS;
+	TMap<const void*, TWeakObjectPtr<const UMaterialInstance> > MaterialInstancesCachedUniformExpressions;
 
 };
 
