Add IsObjectOfSoftClass helper function to check if an object could cast to a soft class. To use as an alternative to a Cast node without creating a hard reference to the class.

diff --git a/Engine/Source/Runtime/Engine/Classes/Kismet/KismetSystemLibrary.h b/Engine/Source/Runtime/Engine/Classes/Kismet/KismetSystemLibrary.h
--- a/Engine/Source/Runtime/Engine/Classes/Kismet/KismetSystemLibrary.h
+++ b/Engine/Source/Runtime/Engine/Classes/Kismet/KismetSystemLibrary.h
@@ -373,6 +373,10 @@ class UKismetSystemLibrary : public UBlueprintFunctionLibrary
 	UFUNCTION(BlueprintCallable, Category = "Utilities", meta = (DeterminesOutputType = "AssetClass"))
 	static ENGINE_API UClass* LoadClassAsset_Blocking(TSoftClassPtr<UObject> AssetClass);
 
+	/** Returns true if Object is of type SoftClass - either an instance of the class or child class, or implements the interface. Alternative to Cast - slower but without adding a hard reference. */
+	UFUNCTION(BlueprintCallable, Category = "Utilities", meta = (ExpandEnumAsExecs = ReturnValue, DisplayName = "IsA ( soft )"))
+	static bool IsObjectOfSoftClass(const UObject* Object, TSoftClassPtr<UObject> SoftClass);
+
 	// Internal functions used by K2Node_LoadAsset and K2Node_ConvertAsset
 
 	UFUNCTION(BlueprintPure, meta = (BlueprintInternalUseOnly = "true"), Category = "Utilities")
diff --git a/Engine/Source/Runtime/Engine/Private/KismetSystemLibrary.cpp b/Engine/Source/Runtime/Engine/Private/KismetSystemLibrary.cpp
--- a/Engine/Source/Runtime/Engine/Private/KismetSystemLibrary.cpp
+++ b/Engine/Source/Runtime/Engine/Private/KismetSystemLibrary.cpp
@@ -1441,6 +1441,29 @@ UClass* UKismetSystemLibrary::LoadClassAsset_Blocking(TSoftClassPtr<UObject> Ass
 	return AssetClass.LoadSynchronous();
 }
 
+bool UKismetSystemLibrary::IsObjectOfSoftClass(const UObject* Object, TSoftClassPtr<UObject> SoftClass)
+{
+	if (!Object)
+	{
+		return false;
+	}
+
+	TSubclassOf<UObject> ObjectClass = SoftClass.Get();
+	if (!ObjectClass)
+	{
+		return false;
+	}
+
+	TSubclassOf<UInterface> InterfaceClass = ObjectClass.Get();
+	if (InterfaceClass)
+	{
+		check(Object->GetClass());
+		return Object->GetClass()->ImplementsInterface(InterfaceClass);
+	}
+
+	return Object->IsA(ObjectClass);
+}
+
 UObject* UKismetSystemLibrary::Conv_SoftObjectReferenceToObject(const TSoftObjectPtr<UObject>& SoftObject)
 {
 	return SoftObject.Get();
