[Material Editor] Fixed a few issues to the new Material Translator. It is now capable of translating a basic material with constant attributes.

diff --git a/Engine/Source/Runtime/Engine/Private/Materials/MaterialIRModuleBuilder.cpp b/Engine/Source/Runtime/Engine/Private/Materials/MaterialIRModuleBuilder.cpp
--- a/Engine/Source/Runtime/Engine/Private/Materials/MaterialIRModuleBuilder.cpp
+++ b/Engine/Source/Runtime/Engine/Private/Materials/MaterialIRModuleBuilder.cpp
@@ -8,10 +8,14 @@
 #include "Materials/MaterialIR.h"
 #include "Materials/MaterialIRTypes.h"
 #include "Materials/MaterialIRBuilder.h"
-#include "MaterialShared.h"
+#include "MaterialIRUtility.h"
+
 #include "Materials/MaterialAttributeDefinitionMap.h"
 #include "Materials/Material.h"
 #include "MaterialExpressionIO.h"
+#include "MaterialShared.h"
+
+namespace IR = MaterialIR;
 
 struct FMaterialIRModuleBuilder::FHelper
 {
@@ -39,7 +43,7 @@ bool FMaterialIRModuleBuilder::Build()
 {
 	Module->Reset(ShaderPlatform, FeatureLevel, TargetPlatform);
 
-	MaterialIR::FBuilder Builder{ Module };
+	IR::FBuilder Builder{ Module };
 
 	UMaterial* BaseMaterial = Material->GetMaterialInterface()->GetMaterial();
 
@@ -48,21 +52,23 @@ bool FMaterialIRModuleBuilder::Build()
 		EMaterialProperty Property = (EMaterialProperty)PropertyIndex;
 		FMaterialInputDescription Input;
 
-		if (Property == MP_CustomOutput || Property == MP_MaterialAttributes || !BaseMaterial->GetExpressionInputDescription(Property, Input))
+		if (!Utility::IsMaterialPropertyShared(Property) 
+			|| Property == MP_SubsurfaceColor
+			|| Property == MP_FrontMaterial
+			|| !BaseMaterial->GetExpressionInputDescription(Property, Input))
 		{
 			continue;
 		}
 
-		MaterialIR::FSetMaterialOutputInstr* Output = Builder.EmitSetMaterialOutput(Property, nullptr);
+		IR::FSetMaterialOutputInstr* Output = Builder.EmitSetMaterialOutput(Property, nullptr);
 	
 		if (Input.bUseConstant)
 		{
 			Output->ArgValue = Builder.NewConstantFromShaderValue(Input.ConstantValue);
 		}
-		else if (!Input.Input->IsConnected() && Input.Type != UE::Shader::EValueType::Void)
+		else if (!Input.Input->IsConnected())
 		{
-			UE::Shader::FValue Zero{ Input.Type };
-			Output->ArgValue = Builder.NewConstantFromShaderValue(Zero);
+			Output->ArgValue = Utility::CreateMaterialAttributeDefaultValue(Builder, Material, Property);
 		}
 		else if (Input.Input->IsConnected())
 		{
diff --git a/Engine/Source/Runtime/Engine/Private/Materials/MaterialIRToHLSLTranslator.cpp b/Engine/Source/Runtime/Engine/Private/Materials/MaterialIRToHLSLTranslator.cpp
--- a/Engine/Source/Runtime/Engine/Private/Materials/MaterialIRToHLSLTranslator.cpp
+++ b/Engine/Source/Runtime/Engine/Private/Materials/MaterialIRToHLSLTranslator.cpp
@@ -4,6 +4,7 @@
 #include "Materials/MaterialIRModule.h"
 #include "Materials/MaterialIRTypes.h"
 #include "Materials/MaterialIR.h"
+#include "MaterialIRUtility.h"
 
 #include "ShaderCore.h"
 #include "MaterialShared.h"
@@ -21,11 +22,12 @@ namespace IR = MaterialIR;
 struct FHLSLPrinter
 {
 	FString& Buffer;
+	bool bFirstListItem = false;
 
-	template <typename... Types>
-	void Print(const TCHAR* Format, Types... Args)
+	template <int N, typename... Types>
+	void Printf(const TCHAR (&Format)[N], Types... Args)
 	{
-		Buffer.Printf(Format, Args...);
+		Buffer.Appendf(Format, Args...);
 	}
 
 	FHLSLPrinter& operator<<(const TCHAR* Text)
@@ -48,9 +50,23 @@ struct FHLSLPrinter
 
 	FHLSLPrinter& operator<<(float Value)
 	{
- 		Buffer.Appendf(TEXT("%.5f"), Value);
+ 		Buffer.Appendf(TEXT("%.5ff"), Value);
 		return *this;
 	}
+
+	void BeginList()
+	{
+		bFirstListItem = true;
+	}
+
+	void PrintListSeparator()
+	{
+		if (!bFirstListItem)
+		{
+			Buffer.Append(TEXT(", "));
+		}
+		bFirstListItem = false;
+	}
 };
 
 static const TCHAR* GetHLSLTypeString(EMaterialValueType Type)
@@ -84,34 +100,6 @@ static const TCHAR* GetHLSLTypeString(EMaterialValueType Type)
 	};
 }
 
-static bool IsMaterialPropertyShared(EMaterialProperty InProperty)
-{
-	switch (InProperty)
-	{
-		case MP_Normal:
-		case MP_Tangent:
-		case MP_EmissiveColor:
-		case MP_Opacity:
-		case MP_OpacityMask:
-		case MP_BaseColor:
-		case MP_Metallic:
-		case MP_Specular:
-		case MP_Roughness:
-		case MP_Anisotropy:
-		case MP_AmbientOcclusion:
-		case MP_Refraction:
-		case MP_PixelDepthOffset:
-		case MP_SubsurfaceColor:
-		case MP_ShadingModel:
-		case MP_SurfaceThickness:
-		case MP_FrontMaterial:
-		case MP_Displacement:
-			return true;
-		default:
-			return false;
-	}
-};
-
 static const TCHAR* GetShadingModelParameterName(EMaterialShadingModel InModel)
 {
 	switch (InModel)
@@ -138,27 +126,20 @@ void GenerateHLSL(const FMaterial& InMaterial, const FMaterialIRModule& InModule
 {
 	FString PixelAttributes;
 	FString EvaluateOtherMaterialAttributesHLSL;
-
 	FHLSLPrinter Printer{ EvaluateOtherMaterialAttributesHLSL };
 
  	for (const IR::FSetMaterialOutputInstr* Output : InModule.GetOutputs())
 	{
-		// Special case MP_SubsurfaceColor as the actual property is a combination of the color and the profile but we don't want to expose the profile
-		FString PropertyName = (Output->Property == MP_SubsurfaceColor) ? "Subsurface" : FMaterialAttributeDefinitionMap::GetAttributeName(Output->Property);
-		EMaterialValueType Type = (Output->Property == MP_SubsurfaceColor) ? MCT_Float4 : FMaterialAttributeDefinitionMap::GetValueType(Output->Property);
-		check(PropertyName.Len() > 0);
-
-		EvaluateOtherMaterialAttributesHLSL.Appendf(TEXT("\tPixelMaterialInputs.%s = "), *PropertyName);
-
 		LowerValue(Printer, Output);
-
-		EvaluateOtherMaterialAttributesHLSL.Append(TEXT("\n"));
 	}
 
+	EvaluateOtherMaterialAttributesHLSL.Append(TEXT("\tPixelMaterialInputs.FrontMaterial = GetInitialisedSubstrateData();\n"));
+	EvaluateOtherMaterialAttributesHLSL.Append(TEXT("\tPixelMaterialInputs.Subsurface = 0;\n"));
+
 	for (int32 PropertyIndex = 0; PropertyIndex < MP_MAX; ++PropertyIndex)
 	{
 		EMaterialProperty Property = (EMaterialProperty)PropertyIndex;
-		if (!IsMaterialPropertyShared(Property))
+		if (!Utility::IsMaterialPropertyShared(Property))
 		{
 			continue;
 		}
@@ -174,8 +155,13 @@ void GenerateHLSL(const FMaterial& InMaterial, const FMaterialIRModule& InModule
 	}
 	
 	InParams.Add(TEXT("pixel_material_inputs"), MoveTemp(PixelAttributes));
-	InParams.Add(TEXT("calc_pixel_material_inputs_normal"), TEXT("PixelMater ialInputs.Normal = MaterialFloat3(0.00000000, 0.00000000, 1.00000000);"));
-	InParams.Add(TEXT("calc_pixel_material_inputs_other_inputs"), MoveTemp(EvaluateOtherMaterialAttributesHLSL));
+
+	InParams.Add(TEXT("calc_pixel_material_inputs_normal"), TEXT("\tPixelMaterialInputs.Normal = MaterialFloat3(0.00000000, 0.00000000, 1.00000000);"));
+	InParams.Add(TEXT("calc_pixel_material_inputs_other_inputs"), EvaluateOtherMaterialAttributesHLSL);
+
+	InParams.Add(TEXT("calc_pixel_material_inputs_analytic_derivatives_normal"), TEXT("\tPixelMaterialInputs.Normal = MaterialFloat3(0.00000000, 0.00000000, 1.00000000);"));
+	InParams.Add(TEXT("calc_pixel_material_inputs_analytic_derivatives_other_inputs"), MoveTemp(EvaluateOtherMaterialAttributesHLSL));
+	
 	InParams.Add(TEXT("material_declarations"), TEXT("struct FMaterialAttributes {};"));
 	
 	auto SetParamInt = [&] (const TCHAR* InParamName, int InValue)
@@ -212,8 +198,8 @@ void SetMaterialParameters(const FMaterial& InMaterial, FParametersMap& InParams
 	InParams.Add(TEXT("get_material_previous_world_position_offset_raw"), TEXT("\treturn 0; // todo"));
 	
 	// CustomData0/1 are named ClearCoat/ClearCoatRoughness
-	InParams.Add(TEXT("get_material_custom_data0"), TEXT("\treturn 0; // todo"));
-	InParams.Add(TEXT("get_material_custom_data1"), TEXT("\treturn 0; // todo"));
+	InParams.Add(TEXT("get_material_custom_data0"), TEXT("\treturn 1.0f; // todo"));
+	InParams.Add(TEXT("get_material_custom_data1"), TEXT("\treturn 0.1f; // todo"));
 
 	FString EvaluateMaterialDeclaration;
 	EvaluateMaterialDeclaration.Append(TEXT("void EvaluateVertexMaterialAttributes(in out FMaterialVertexParameters Parameters)\n{\n"));
@@ -339,32 +325,49 @@ void GetShaderCompilerEnvironment(const FMaterial& InMaterial, const FMaterialIR
 
 void LowerValue(FHLSLPrinter& Printer, const IR::FValue* InValue)
 {
-	if (const IR::FScalarValue* S = InValue->Cast<IR::FScalarValue>())
+	if (const IR::FScalarValue* Scalar = InValue->Cast<IR::FScalarValue>())
 	{
-		IR::FArithmeticTypePtr ArithType = S->Type->ToArithmetic();
+		IR::FArithmeticTypePtr ArithType = Scalar->Type->ToArithmetic();
 		check(ArithType && ArithType->IsScalar());
 
 		switch (ArithType->ScalarKind)
 		{
-			case IR::SK_Bool: Printer << S->Boolean; break;
-			case IR::SK_Int: Printer << S->Integer; break;
-			case IR::SK_Float: Printer << S->Float; break;
+			case IR::SK_Bool: Printer << Scalar->Boolean; break;
+			case IR::SK_Int: Printer << Scalar->Integer; break;
+			case IR::SK_Float: Printer << Scalar->Float; break;
 		}
 	}
-	else if (const IR::FVectorValue* V = InValue->Cast<IR::FVectorValue>())
+	else if (const IR::FVectorValue* Vector = InValue->Cast<IR::FVectorValue>())
 	{
-		IR::FArithmeticTypePtr ArithType = V->Type->ToArithmetic();
+		IR::FArithmeticTypePtr ArithType = Vector->Type->ToArithmetic();
 		check(ArithType && ArithType->IsVector());
 
-		Printer << ScalarKindToString(ArithType->ScalarKind) << ArithType->NumColumns << TEXT("(");
+		Printer << ScalarKindToString(ArithType->ScalarKind) << ArithType->NumRows << TEXT("(");
 
-		for (IR::FValuePtr ComponentVal : V->GetComponents())
+		Printer.BeginList();
+		for (IR::FValuePtr Component : Vector->GetComponents())
 		{
-			LowerValue(Printer, ComponentVal);
+			Printer.PrintListSeparator();
+			LowerValue(Printer, Component);
 		}
 
 		Printer << TEXT(")");
 	}
+	else if (const IR::FSetMaterialOutputInstr* Output = InValue->Cast<IR::FSetMaterialOutputInstr>())
+	{
+		// Special case MP_SubsurfaceColor as the actual property is a combination of the color and the profile but we don't want to expose the profile
+		const FString& PropertyName = (Output->Property == MP_SubsurfaceColor) ? "Subsurface" : FMaterialAttributeDefinitionMap::GetAttributeName(Output->Property);
+
+		Printer.Printf(TEXT("\tPixelMaterialInputs.%s = "), *PropertyName);
+		
+		LowerValue(Printer, Output->ArgValue);
+
+		Printer << TEXT(";\n");
+	}
+	else
+	{
+		UE_MIR_UNREACHABLE();
+	}
 }
 
 UE_MIR_END_PRIVATE(); //FMaterialIRToHLSLTranslator
diff --git a/Engine/Source/Runtime/Engine/Private/Materials/MaterialIRUtility.cpp b/Engine/Source/Runtime/Engine/Private/Materials/MaterialIRUtility.cpp
new file mode 100644
--- /dev/null
+++ b/Engine/Source/Runtime/Engine/Private/Materials/MaterialIRUtility.cpp
@@ -0,0 +1,66 @@
+// Copyright Epic Games, Inc. All Rights Reserved.
+
+#include "Materials/MaterialIRUtility.h"
+#include "Materials/MaterialAttributeDefinitionMap.h"
+#include "Materials/MaterialIRBuilder.h"
+#include "MaterialShared.h"
+
+#if WITH_EDITOR
+
+namespace Utility {
+
+namespace MIR = MaterialIR;
+
+bool IsMaterialPropertyShared(EMaterialProperty InProperty)
+{
+	switch (InProperty)
+	{
+	case MP_Normal:
+	case MP_Tangent:
+	case MP_EmissiveColor:
+	case MP_Opacity:
+	case MP_OpacityMask:
+	case MP_BaseColor:
+	case MP_Metallic:
+	case MP_Specular:
+	case MP_Roughness:
+	case MP_Anisotropy:
+	case MP_AmbientOcclusion:
+	case MP_Refraction:
+	case MP_PixelDepthOffset:
+	case MP_SubsurfaceColor:
+	case MP_ShadingModel:
+	case MP_SurfaceThickness:
+	case MP_FrontMaterial:
+	case MP_Displacement:
+		return true;
+	default:
+		return false;
+	}
+}
+
+MIR::FValuePtr CreateMaterialAttributeDefaultValue(MIR::FBuilder& Builder, const FMaterial* Material, EMaterialProperty Property)
+{
+	EMaterialValueType Type = FMaterialAttributeDefinitionMap::GetValueType(Property);
+	FVector4f DefaultValue = FMaterialAttributeDefinitionMap::GetDefaultValue(Property);
+
+	switch (Type)
+	{
+		case MCT_ShadingModel: return Builder.NewConstantInt1(Material->GetShadingModels().GetFirstShadingModel());
+
+		case MCT_Float1: return Builder.NewConstantFloat1(  DefaultValue.X );
+		case MCT_Float2: return Builder.NewConstantFloat2({ DefaultValue.X, DefaultValue.Y });
+		case MCT_Float3: return Builder.NewConstantFloat3({ DefaultValue.X, DefaultValue.Y, DefaultValue.Z });
+		case MCT_Float4: return Builder.NewConstantFloat4(DefaultValue);
+
+		case MCT_UInt1: return Builder.NewConstantInt1(  (int32)DefaultValue.X );
+		case MCT_UInt2: return Builder.NewConstantInt2({ (int32)DefaultValue.X, (int32)DefaultValue.Y });
+		case MCT_UInt3: return Builder.NewConstantInt3({ (int32)DefaultValue.X, (int32)DefaultValue.Y, (int32)DefaultValue.Z });
+		case MCT_UInt4: return Builder.NewConstantInt4({ (int32)DefaultValue.X, (int32)DefaultValue.Y, (int32)DefaultValue.Z, (int32)DefaultValue.W });
+
+		default: UE_MIR_UNREACHABLE();
+	}
+}
+
+} // namespace Utility
+#endif
diff --git a/Engine/Source/Runtime/Engine/Private/Materials/MaterialIRUtility.h b/Engine/Source/Runtime/Engine/Private/Materials/MaterialIRUtility.h
new file mode 100644
--- /dev/null
+++ b/Engine/Source/Runtime/Engine/Private/Materials/MaterialIRUtility.h
@@ -0,0 +1,17 @@
+// Copyright Epic Games, Inc. All Rights Reserved.
+#pragma once
+
+#include "Materials/MaterialIRCommon.h"
+
+#if WITH_EDITOR
+
+namespace Utility {
+
+bool IsMaterialPropertyShared(EMaterialProperty InProperty);
+
+//
+MaterialIR::FValuePtr CreateMaterialAttributeDefaultValue(MaterialIR::FBuilder& Builder, const FMaterial* Material, EMaterialProperty Property);
+
+} // namespace Utility
+
+#endif
diff --git a/Engine/Source/Runtime/Engine/Public/Materials/MaterialIR.h b/Engine/Source/Runtime/Engine/Public/Materials/MaterialIR.h
--- a/Engine/Source/Runtime/Engine/Public/Materials/MaterialIR.h
+++ b/Engine/Source/Runtime/Engine/Public/Materials/MaterialIR.h
@@ -13,22 +13,33 @@ namespace MaterialIR
 enum EValueKind
 {
 	/* Values */
-	VK_Value = 0,
-	VK_Scalar = 1 << 1,
-	VK_Vector = 2 << 1,
+	VK_Value = 1 << 0,
+	VK_Scalar = VK_Value | 1 << 2,
+	VK_Vector = VK_Value | 2 << 2,
 	
 	/* Instructions */
-	VK_Instruction = 1,
- 	VK_SetMaterialOutputInstr = VK_Instruction | 1 << 1,
+	VK_Instruction = 1 << 1,
+ 	VK_SetMaterialOutputInstr = VK_Instruction | 1 << 2,
 };
 
 /* Values */
 
-struct FValue : FIRNode<EValueKind>
+struct FValue
 {
-	EValueKind Kind = VK_Value;
+	EValueKind Kind{};
 	FTypePtr  Type{};
 
+	bool IsA(EValueKind InKind) const
+	{
+		return (Kind & InKind) == InKind;
+	}
+
+	template <typename T>
+	const T* Cast() const
+	{
+		return this && IsA(T::TypeKind) ? static_cast<const T*>(this) : nullptr;
+	}
+
 	void Destroy();
 };
 
diff --git a/Engine/Source/Runtime/Engine/Public/Materials/MaterialIRCommon.h b/Engine/Source/Runtime/Engine/Public/Materials/MaterialIRCommon.h
--- a/Engine/Source/Runtime/Engine/Public/Materials/MaterialIRCommon.h
+++ b/Engine/Source/Runtime/Engine/Public/Materials/MaterialIRCommon.h
@@ -46,24 +46,6 @@ struct FSetMaterialOutputInstr;
 /* Others */
 class FBuilder;
 
-//
-template <typename TKindType>
-struct FIRNode
-{
-	TKindType Kind;
-
-	bool IsA(TKindType InKind) const
-	{
-		return (static_cast<uint64>(Kind) & static_cast<uint64>(InKind)) == static_cast<uint64>(InKind);
-	}
-
-	template <typename T>
-	const T* Cast() const
-	{
- 		return this && IsA(T::TypeKind) ? static_cast<const T*>(this) : nullptr;
-	}
-};
-
 }
 
 #define UE_MIR_UNREACHABLE() { check(!"Unreachable"); UE_ASSUME(false); }
diff --git a/Engine/Source/Runtime/Engine/Public/Materials/MaterialIRTypes.h b/Engine/Source/Runtime/Engine/Public/Materials/MaterialIRTypes.h
--- a/Engine/Source/Runtime/Engine/Public/Materials/MaterialIRTypes.h
+++ b/Engine/Source/Runtime/Engine/Public/Materials/MaterialIRTypes.h
@@ -11,7 +11,7 @@ namespace MaterialIR
 
 enum ETypeKind
 {
-	TK_Void = 1,
+	TK_Void,
 	TK_Arithmetic
 };
 
