[World Partition Class Descriptors] Don't try to prefetch transient classes from the asset registry.

diff --git a/Engine/Source/Runtime/Engine/Private/WorldPartition/WorldPartitionClassDescRegistry.cpp b/Engine/Source/Runtime/Engine/Private/WorldPartition/WorldPartitionClassDescRegistry.cpp
--- a/Engine/Source/Runtime/Engine/Private/WorldPartition/WorldPartitionClassDescRegistry.cpp
+++ b/Engine/Source/Runtime/Engine/Private/WorldPartition/WorldPartitionClassDescRegistry.cpp
@@ -448,7 +448,8 @@ void FWorldPartitionClassDescRegistry::PrefetchClassDesc(UClass* InClass)
 		if (!ClassByPath.Contains(ParentClassPath))
 		{
 			// Only prefetch classes that exists on disk
-			if (!ParentClass->GetPackage()->HasAnyPackageFlags(PKG_NewlyCreated))
+			const UPackage* ParentClassPackage = ParentClass->GetPackage();
+			if ((ParentClassPackage != GetTransientPackage()) && !ParentClassPackage->HasAnyPackageFlags(PKG_NewlyCreated))
 			{
 				PrefetchClassDescs({ ParentClassPath });
 			}
