[NetPhysics] Reliable flow in NetworkPhysicsComponent
- Feature to mark input/state data as important and have those sent reliably, works alongside the unreliable / unimportant data flow
- If an important data arrive too late, call MergeData to merge that and all subsequent data into the next data that is being used.
- Option to have ApplyData called instead of MergeData

diff --git a/Engine/Source/Runtime/Engine/Private/PhysicsEngine/NetworkPhysicsComponent.cpp b/Engine/Source/Runtime/Engine/Private/PhysicsEngine/NetworkPhysicsComponent.cpp
--- a/Engine/Source/Runtime/Engine/Private/PhysicsEngine/NetworkPhysicsComponent.cpp
+++ b/Engine/Source/Runtime/Engine/Private/PhysicsEngine/NetworkPhysicsComponent.cpp
@@ -29,6 +29,14 @@ namespace PhysicsReplicationCVars
 		static FAutoConsoleVariableRef CVarResimCompareStateToTriggerRewind(TEXT("np2.Resim.CompareStateToTriggerRewind"), bCompareStateToTriggerRewind, TEXT("When true, cache local players custom state struct in rewind history and compare the predicted state with incoming server state to trigger resimulations if they differ, comparison done through FNetworkPhysicsData::CompareData"));
 		bool bCompareInputToTriggerRewind = false;
 		static FAutoConsoleVariableRef CVarResimCompareInputToTriggerRewind(TEXT("np2.Resim.CompareInputToTriggerRewind"), bCompareInputToTriggerRewind, TEXT("When true, compare local players predicted inputs with incoming server inputs to trigger resimulations if they differ, comparison done through FNetworkPhysicsData::CompareData."));
+		bool bEnableUnreliableFlow = true;
+		static FAutoConsoleVariableRef CVarResimEnableUnreliableFlow(TEXT("np2.Resim.EnableUnreliableFlow"), bEnableUnreliableFlow, TEXT("When true, allow data to be sent unreliably. Also sends FNetworkPhysicsData not marked with FNetworkPhysicsData::bimportant unreliably over the network."));
+		bool bEnableReliableFlow = false;
+		static FAutoConsoleVariableRef CVarResimEnableReliableFlow(TEXT("np2.Resim.EnableReliableFlow"), bEnableReliableFlow, TEXT("EXPERIMENTAL -- When true, allow data to be sent reliably. Also send FNetworkPhysicsData marked with FNetworkPhysicsData::bimportant reliably over the network."));
+		bool bApplyDataInsteadOfMergeData = false;
+		static FAutoConsoleVariableRef CVarResimApplyDataInsteadOfMergeData(TEXT("np2.Resim.ApplyDataInsteadOfMergeData"), bApplyDataInsteadOfMergeData, TEXT("When true, call ApplyData for each data instead of MergeData when having to use multiple data entries in one frame."));
+		bool bAllowInputExtrapolation = true;
+		static FAutoConsoleVariableRef CVarResimAllowInputExtrapolation(TEXT("np2.Resim.AllowInputExtrapolation"), bAllowInputExtrapolation, TEXT("When true and not locally controlled, allow inputs to be extrapolated from last known and if there is a gap allow interpolation between two known inputs."));
 	}
 }
 
@@ -136,6 +144,8 @@ FNetworkPhysicsRewindDataProxy& FNetworkPhysicsRewindDataProxy::operator=(const
 #if UE_WITH_IRIS
 UE_NET_IMPLEMENT_NAMED_STRUCT_LASTRESORT_NETSERIALIZER_AND_REGISTRY_DELEGATES(NetworkPhysicsRewindDataInputProxy);
 UE_NET_IMPLEMENT_NAMED_STRUCT_LASTRESORT_NETSERIALIZER_AND_REGISTRY_DELEGATES(NetworkPhysicsRewindDataStateProxy);
+UE_NET_IMPLEMENT_NAMED_STRUCT_LASTRESORT_NETSERIALIZER_AND_REGISTRY_DELEGATES(NetworkPhysicsRewindDataImportantInputProxy);
+UE_NET_IMPLEMENT_NAMED_STRUCT_LASTRESORT_NETSERIALIZER_AND_REGISTRY_DELEGATES(NetworkPhysicsRewindDataImportantStateProxy);
 #endif // UE_WITH_IRIS
 
 bool FNetworkPhysicsRewindDataInputProxy::NetSerialize(FArchive& Ar, class UPackageMap* Map, bool& bOutSuccess)
@@ -148,6 +158,16 @@ bool FNetworkPhysicsRewindDataStateProxy::NetSerialize(FArchive& Ar, class UPack
 	return NetSerializeBase(Ar, Map, bOutSuccess, [this]() { return Owner->ReplicatedStates.History->CreateNew(); });
 }
 
+bool FNetworkPhysicsRewindDataImportantInputProxy::NetSerialize(FArchive& Ar, class UPackageMap* Map, bool& bOutSuccess)
+{
+	return NetSerializeBase(Ar, Map, bOutSuccess, [this]() { return Owner->ReplicatedImportantInput.History->CreateNew(); });
+}
+
+bool FNetworkPhysicsRewindDataImportantStateProxy::NetSerialize(FArchive& Ar, class UPackageMap* Map, bool& bOutSuccess)
+{
+	return NetSerializeBase(Ar, Map, bOutSuccess, [this]() { return Owner->ReplicatedImportantState.History->CreateNew(); });
+}
+
 // after presimulate internal (asyncinput internal simulation done and the output created)
 void FNetworkPhysicsCallback::ApplyCallbacks_Internal(int32 PhysicsStep, const TArray<Chaos::ISimCallbackObject*>& SimCallbackObjects)
 {
@@ -471,14 +491,25 @@ void UNetworkPhysicsComponent::InitializeComponent()
 {
 	Super::InitializeComponent();
 
+	// Cache CVar values
+	bEnableUnreliableFlow = PhysicsReplicationCVars::ResimulationCVars::bEnableUnreliableFlow;
+	bEnableReliableFlow = PhysicsReplicationCVars::ResimulationCVars::bEnableReliableFlow;
+	bApplyDataInsteadOfMergeData = PhysicsReplicationCVars::ResimulationCVars::bApplyDataInsteadOfMergeData;
+	bAllowInputExtrapolation = PhysicsReplicationCVars::ResimulationCVars::bAllowInputExtrapolation;
+
 	if (AActor* Owner = GetOwner())
 	{
+		// Get settings from NetworkPhysicsSettingsComponent, if there is one
 		if (UNetworkPhysicsSettingsComponent* PhysicsSettings = Owner->FindComponentByClass<UNetworkPhysicsSettingsComponent>())
 		{
-			InputRedundancy = PhysicsSettings->ResimulationSettings.bOverrideRedundantInputs ? PhysicsSettings->ResimulationSettings.RedundantInputs : InputRedundancy;
-			StateRedundancy = PhysicsSettings->ResimulationSettings.bOverrideRedundantStates ? PhysicsSettings->ResimulationSettings.RedundantStates : StateRedundancy;
-			bCompareInputToTriggerRewind = PhysicsSettings->ResimulationSettings.GetCompareInputToTriggerRewind(bCompareInputToTriggerRewind);
-			bCompareStateToTriggerRewind = PhysicsSettings->ResimulationSettings.GetCompareStateToTriggerRewind(bCompareStateToTriggerRewind);
+			InputRedundancy = PhysicsSettings->NetworkPhysicsComponentSettings.GetRedundantInputs(InputRedundancy);
+			StateRedundancy = PhysicsSettings->NetworkPhysicsComponentSettings.GetRedundantStates(StateRedundancy);
+			bCompareInputToTriggerRewind = PhysicsSettings->NetworkPhysicsComponentSettings.GetCompareInputToTriggerRewind(bCompareInputToTriggerRewind);
+			bCompareStateToTriggerRewind = PhysicsSettings->NetworkPhysicsComponentSettings.GetCompareStateToTriggerRewind(bCompareStateToTriggerRewind);
+			bEnableUnreliableFlow = PhysicsSettings->NetworkPhysicsComponentSettings.GetEnableUnreliableFlow();
+			bEnableReliableFlow = PhysicsSettings->NetworkPhysicsComponentSettings.GetEnableReliableFlow();
+			bApplyDataInsteadOfMergeData = PhysicsSettings->NetworkPhysicsComponentSettings.GetApplyDataInsteadOfMergeData();
+			bAllowInputExtrapolation = PhysicsSettings->NetworkPhysicsComponentSettings.GetAllowInputExtrapolation();
 
 			if (ReplicatedInputs.History)
 			{
@@ -491,6 +522,13 @@ void UNetworkPhysicsComponent::InitializeComponent()
 		}
 	}
 
+	// TEMP: Stops ensure that happens if first ever network broadcast for the actor is a multicast from PT (with GT frozen), this instead makes the first multicast happen in GT.
+	// ToDo, remove this when NetworkPhysicComponent has a proper Async ISimCallback flow.
+	if (bEnableReliableFlow && ReplicatedImportantInput.Owner != nullptr)
+	{
+		MulticastReceiveImportantInputData(ReplicatedImportantInput);
+	}
+
 	if (UWorld* World = GetWorld())
 	{
 		if (UNetworkPhysicsSystem* NetworkManager = World->GetSubsystem<UNetworkPhysicsSystem>())
@@ -518,7 +556,7 @@ void UNetworkPhysicsComponent::GetLifetimeReplicatedProps(TArray< FLifetimePrope
 
 	FDoRepLifetimeParams ReplicatedParams;
 	ReplicatedParams.Condition = COND_None;
-	ReplicatedParams.RepNotifyCondition = REPNOTIFY_OnChanged;
+	ReplicatedParams.RepNotifyCondition = REPNOTIFY_Always; //REPNOTIFY_OnChanged;
 	ReplicatedParams.bIsPushBased = true;
 
 	DOREPLIFETIME_WITH_PARAMS_FAST(UNetworkPhysicsComponent, ReplicatedInputs, ReplicatedParams);
@@ -550,8 +588,9 @@ void UNetworkPhysicsComponent::AsyncPhysicsTickComponent(float DeltaTime, float
 
 void UNetworkPhysicsComponent::SendInputData()
 {
+	const bool bIsServer = HasServerWorld();
 	// Inputs are sent from the server or locally controlled actors/pawns
-	if (InputHistory && (IsLocallyControlled() || HasServerWorld()))
+	if (InputHistory && (IsLocallyControlled() || bIsServer))
 	{
 		const APlayerController* PlayerController = GetPlayerController();
 		if (!PlayerController)
@@ -564,34 +603,74 @@ void UNetworkPhysicsComponent::SendInputData()
 		{
 			// Send latest N frames from history
 			const int32 ToFrame = InputHistory->GetLatestFrame();
-			const int32 FromFrame = ToFrame - InputRedundancy;
 
-			InputHistory->CopyData(*ReplicatedInputs.History, FromFrame, ToFrame);
-
-			if (HasServerWorld())
+			// -- Default / Unreliable Flow --
+			if (bEnableUnreliableFlow)
 			{
-				// Server sends inputs through property replication
-				MARK_PROPERTY_DIRTY_FROM_NAME(UNetworkPhysicsComponent, ReplicatedInputs, this);
-			}
-			else
-			{
-				// Clients send inputs through an RPC to the server
-				ServerReceiveInputData(ReplicatedInputs);
-			}
+				const int32 FromFrame = ToFrame - InputRedundancy;
+				if (InputHistory->CopyData(*ReplicatedInputs.History, FromFrame, ToFrame, /*bIncludeUnimportant*/ true, /*bIncludeImportant*/ bEnableReliableFlow == false))
+				{
+					if (bIsServer)
+					{
+						// Server sends inputs through property replication
+						MARK_PROPERTY_DIRTY_FROM_NAME(UNetworkPhysicsComponent, ReplicatedInputs, this);
+					}
+					else
+					{
+						// Clients send inputs through an RPC to the server
+						ServerReceiveInputData(ReplicatedInputs);
+					}
 
 #if DEBUG_NETWORK_PHYSICS
-			FAsyncPhysicsTimestamp Timestamp = const_cast<APlayerController*>(PlayerController)->GetPhysicsTimestamp();
-			if (HasServerWorld())
-			{
-				ReplicatedInputs.History->DebugData(FString::Printf(TEXT("SERVER | PT | SendInputData | CurrentLocalFrame = %d | CurrentServerFrame = %d | Component = %s"), Timestamp.LocalFrame, Timestamp.ServerFrame, *GetFullName()));
+					{
+						FAsyncPhysicsTimestamp Timestamp = const_cast<APlayerController*>(PlayerController)->GetPhysicsTimestamp();
+						FString NetRole = HasServerWorld() ? FString("SERVER    ") : (IsLocallyControlled() ? FString("AUTONOMOUS") : FString("PROXY     "));
+						ReplicatedInputs.History->DebugData(FString::Printf(TEXT("%s | PT | SendInputData | CurrentLocalFrame = %d | CurrentServerFrame = %d | Component = %s"), *NetRole, Timestamp.LocalFrame, Timestamp.ServerFrame, *GetFullName()));
+					}
+#endif
+				}
 			}
-			else
+
+			// -- Important / Reliable flow --
+			if (bEnableReliableFlow)
 			{
-				ReplicatedInputs.History->DebugData(FString::Printf(TEXT("CLIENT | PT | SendInputData | CurrentLocalFrame = %d | CurrentServerFrame = %d | Component = %s"), Timestamp.LocalFrame, Timestamp.ServerFrame, *GetFullName()));
-			}
+				/* Get the latest valid frame that can hold new important data:
+				* 1. Frame after last time we called SendInputData
+				* 2. Frame after last entry in replicated data
+				* 3. Earliest possible frame in history */
+				const int32 FromFrame = FMath::Max(FMath::Max(LastInputSendFrame + 1, ReplicatedImportantInput.History->GetLatestFrame() + 1), ToFrame - InputHistory->GetHistorySize());
+
+				for (int32 CopyFrame = FromFrame; CopyFrame <= ToFrame; CopyFrame++)
+				{
+					if (InputHistory->CopyData(*ReplicatedImportantInput.History, CopyFrame, CopyFrame, /*bIncludeUnimportant*/ false, /*bIncludeImportant*/ true))
+					{
+						if (bIsServer)
+						{
+							MulticastReceiveImportantInputData(ReplicatedImportantInput);
+						}
+						else
+						{
+							ServerReceiveImportantInputData(ReplicatedImportantInput);
+						}
+#if DEBUG_NETWORK_PHYSICS
+						{
+							FAsyncPhysicsTimestamp Timestamp = const_cast<APlayerController*>(PlayerController)->GetPhysicsTimestamp();
+							FString NetRole = HasServerWorld() ? FString("SERVER    ") : (IsLocallyControlled() ? FString("AUTONOMOUS") : FString("PROXY     "));
+							ReplicatedImportantInput.History->DebugData(FString::Printf(TEXT("%s | PT | SendInputData | RELIABLE | CurrentLocalFrame = %d | CurrentServerFrame = %d | Component = %s"), *NetRole, Timestamp.LocalFrame, Timestamp.ServerFrame, *GetFullName()));
+						}
 #endif
+					}
+				}
+			}
+			LastInputSendFrame = InputHistory->GetLatestFrame();
 		}
 	}
+
+	if (bStopRelayingLocalInputsDeferred)
+	{
+		bIsRelayingLocalInputs = false;
+		bStopRelayingLocalInputsDeferred = false;
+	}
 }
 
 void UNetworkPhysicsComponent::SendStateData()
@@ -600,15 +679,43 @@ void UNetworkPhysicsComponent::SendStateData()
 	{
 		// Send latest N frames from history
 		const int32 ToFrame = StateHistory->GetLatestFrame();
-		const int32 FromFrame = ToFrame - StateRedundancy;
 
-		// If on server we should send the states onto all the clients through repnotify
-		StateHistory->CopyData(*ReplicatedStates.History, FromFrame, ToFrame);
-		MARK_PROPERTY_DIRTY_FROM_NAME(UNetworkPhysicsComponent, ReplicatedStates, this);
+		// -- Default / Unreliable Flow --
+		if (bEnableUnreliableFlow)
+		{
+			const int32 FromFrame = ToFrame - StateRedundancy;
+			if (StateHistory->CopyData(*ReplicatedStates.History, FromFrame, ToFrame, /*bIncludeUnimportant*/ true, /*bIncludeImportant*/ bEnableReliableFlow == false))
+			{
+				// If on server we should send the states onto all the clients through repnotify
+				MARK_PROPERTY_DIRTY_FROM_NAME(UNetworkPhysicsComponent, ReplicatedStates, this);
+
+#if DEBUG_NETWORK_PHYSICS
+				ReplicatedStates.History->DebugData(FString::Printf(TEXT("SERVER | PT | SendStateData | Component = %s"), *GetFullName()));
+#endif
+			}
+		}
 		
+		// -- Important / Reliable flow --
+		if (bEnableReliableFlow)
+		{
+			/* Get the latest valid frame that can hold new important data: 
+			* 1. Frame after last time we called SendStateData
+			* 2. Frame after last entry in replicated data
+			* 3. Earliest possible frame in history */
+			const int32 FromFrame = FMath::Max(FMath::Max(LastStateSendFrame + 1, ReplicatedImportantState.History->GetLatestFrame() + 1), ToFrame - StateHistory->GetHistorySize());
+
+			for (int32 CopyFrame = FromFrame; CopyFrame <= ToFrame; CopyFrame++)
+			{
+				if (StateHistory->CopyData(*ReplicatedImportantState.History, CopyFrame, CopyFrame, /*bIncludeUnimportant*/ false, /*bIncludeImportant*/ true))
+				{
+					MulticastReceiveImportantStateData(ReplicatedImportantState);
 #if DEBUG_NETWORK_PHYSICS
-		ReplicatedStates.History->DebugData(FString::Printf(TEXT("SERVER | PT | SendStateData | Component = %s"), *GetFullName()));
+					ReplicatedImportantState.History->DebugData(FString::Printf(TEXT("SERVER | PT | SendStateData | RELIABLE | Component = %s"), *GetFullName()));
 #endif
+				}
+			}
+		}
+		LastStateSendFrame = StateHistory->GetLatestFrame();
 	}
 }
 
@@ -672,29 +779,15 @@ void UNetworkPhysicsComponent::OnRep_SetReplicatedStates()
 				{
 					if (bCompareStateToTriggerRewind && IsLocallyControlled())
 					{
-						int32 ResimFrame = StateHistory->ReceiveNewData(*ReceivedStates, LocalOffset, /*CompareDataForRewind*/ bCompareStateToTriggerRewind);
-						if (ResimFrame != INDEX_NONE)
-						{
-							if (Chaos::FPhysicsSolver* Solver = PhysScene->GetSolver())
-							{
-								if (Chaos::FRewindData* RewindData = Solver->GetRewindData())
-								{
-									// Mark particle/island as resim
-									Chaos::FReadPhysicsObjectInterface_Internal Interface = Chaos::FPhysicsObjectInternalInterface::GetRead();
-									if (Chaos::FPBDRigidParticleHandle* POHandle = Interface.GetRigidParticle(RootPhysicsObject))
-									{
-										Solver->GetEvolution()->GetIslandManager().SetParticleResimFrame(POHandle, ResimFrame);
-									}
-
-									// Set resim frame in rewind data
-									ResimFrame = (RewindData->GetResimFrame() == INDEX_NONE) ? ResimFrame : FMath::Min(ResimFrame, RewindData->GetResimFrame());
-									RewindData->SetResimFrame(ResimFrame);
-								}
-							}
-						}
+						// Record received data into history
+						const int32 ResimFrame = StateHistory->ReceiveNewData(*ReceivedStates, LocalOffset, /*CompareDataForRewind*/ bCompareStateToTriggerRewind);
+					
+						// Trigger resimulation if ReceiveNewData returned a valid frame to rewind to
+						TriggerResimulation(ResimFrame);
 					}
 					else 
 					{
+						// Record received data into history
 						PRAGMA_DISABLE_DEPRECATION_WARNINGS // TODO: Change to ReceiveNewData() in UE 5.6 and remove deprecation pragma
 						StateHistory->ReceiveNewDatas(*ReceivedStates, LocalOffset);
 						PRAGMA_ENABLE_DEPRECATION_WARNINGS
@@ -735,29 +828,15 @@ void UNetworkPhysicsComponent::OnRep_SetReplicatedInputs()
 				{
 					if (bCompareInputToTriggerRewind && IsLocallyControlled())
 					{
-						int32 ResimFrame = InputHistory->ReceiveNewData(*ReceivedInputs, LocalOffset, /*CompareDataForRewind*/ bCompareInputToTriggerRewind);
-						if (ResimFrame != INDEX_NONE)
-						{
-							if (Chaos::FPhysicsSolver* Solver = PhysScene->GetSolver())
-							{
-								if (Chaos::FRewindData* RewindData = Solver->GetRewindData())
-								{
-									// Mark particle/island as resim
-									Chaos::FReadPhysicsObjectInterface_Internal Interface = Chaos::FPhysicsObjectInternalInterface::GetRead();
-									if (Chaos::FPBDRigidParticleHandle* POHandle = Interface.GetRigidParticle(RootPhysicsObject))
-									{
-										Solver->GetEvolution()->GetIslandManager().SetParticleResimFrame(POHandle, ResimFrame);
-									}
-
-									// Set resim frame in rewind data
-									ResimFrame = (RewindData->GetResimFrame() == INDEX_NONE) ? ResimFrame : FMath::Min(ResimFrame, RewindData->GetResimFrame());
-									RewindData->SetResimFrame(ResimFrame);
-								}
-							}
-						}
+						// Record received data into history
+						const int32 ResimFrame = InputHistory->ReceiveNewData(*ReceivedInputs, LocalOffset, /*CompareDataForRewind*/ bCompareInputToTriggerRewind, /*Important*/false);
+						
+						// Trigger resimulation if ReceiveNewData returned a valid frame to rewind to
+						TriggerResimulation(ResimFrame);
 					}
 					else
 					{
+						// Record received data into history
 						PRAGMA_DISABLE_DEPRECATION_WARNINGS // TODO: Change to ReceiveNewData() in UE 5.6 and remove deprecation pragma
 						InputHistory->ReceiveNewDatas(*ReceivedInputs, LocalOffset);
 						PRAGMA_ENABLE_DEPRECATION_WARNINGS
@@ -785,9 +864,6 @@ void UNetworkPhysicsComponent::ServerReceiveInputData_Implementation(const FNetw
 {
 	if (InputHistory)
 	{ 
-		// We could probably skip that test since the server RPC is on server
-		ensure(HasServerWorld());
-		
 		if (UWorld* World = GetWorld())
 		{
 			if (FPhysScene* PhysScene = World->GetPhysicsScene())
@@ -799,6 +875,7 @@ void UNetworkPhysicsComponent::ServerReceiveInputData_Implementation(const FNetw
 					// Validate data in the received inputs
 					ReceivedInputs->ValidateDataInHistory(ActorComponent);
 
+					// Record received data into history
 					PRAGMA_DISABLE_DEPRECATION_WARNINGS // TODO: Change to ReceiveNewData() in UE 5.6 and remove deprecation pragma
 					InputHistory->ReceiveNewDatas(*ReceivedInputs, 0);
 					PRAGMA_ENABLE_DEPRECATION_WARNINGS
@@ -815,6 +892,133 @@ void UNetworkPhysicsComponent::ServerReceiveInputData_Implementation(const FNetw
 	}
 }
 
+void UNetworkPhysicsComponent::ServerReceiveImportantInputData_Implementation(const FNetworkPhysicsRewindDataImportantInputProxy& ClientInputs)
+{
+	if (InputHistory)
+	{
+		if (UWorld* World = GetWorld())
+		{
+			if (FPhysScene* PhysScene = World->GetPhysicsScene())
+			{
+				// Make another copy of the client inputs for the physics thread to consume
+				TSharedPtr<Chaos::FBaseRewindHistory> ReceivedInputs = MakeShareable(ClientInputs.History->Clone().Release());
+				PhysScene->EnqueueAsyncPhysicsCommand(0, this, [this, ReceivedInputs, PhysScene]()
+					{
+						// Validate data in the received inputs
+						ReceivedInputs->ValidateDataInHistory(ActorComponent);
+
+						// Record received data into history
+						InputHistory->ReceiveNewData(*ReceivedInputs, 0, /*CompareDataForRewind*/false, /*Important*/true);
+
+						// Cache earliest important data frame since last ApplyData was called
+						NewImportantInputFrame = FMath::Min(NewImportantInputFrame, ReceivedInputs->GetEarliestFrame());
+
+#if DEBUG_NETWORK_PHYSICS
+						{
+							const int32 CurrentFrame = PhysScene->GetSolver()->GetCurrentFrame();
+							ReceivedInputs->DebugData(FString::Printf(TEXT("SERVER | PT | ServerReceiveImportantInputData | RELIABLE | CurrentLocalFrame = %d | NewImportantInputFrame = %d | Component = %s"), CurrentFrame, NewImportantInputFrame, *GetFullName()));
+						}
+#endif
+					}, false);
+			}
+		}
+	}
+}
+
+void UNetworkPhysicsComponent::MulticastReceiveImportantInputData_Implementation(const FNetworkPhysicsRewindDataImportantInputProxy& ServerInputs)
+{
+	// Multicast RPC, ignore on server
+	if (HasServerWorld())
+	{
+		return;
+	}
+
+	APlayerController* PlayerController = GetPlayerController();
+	if (!PlayerController)
+	{
+		PlayerController = GetWorld()->GetFirstPlayerController();
+	}
+
+	if (PlayerController && InputHistory)
+	{
+		const int32 LocalOffset = PlayerController->GetNetworkPhysicsTickOffset();
+		
+		if (UWorld* World = GetWorld())
+		{
+			if (FPhysScene* PhysScene = World->GetPhysicsScene())
+			{
+				// Make another copy of the client inputs for the physics thread to consume
+				TSharedPtr<Chaos::FBaseRewindHistory> ReceivedInputs = MakeShareable(ServerInputs.History->Clone().Release());
+				PhysScene->EnqueueAsyncPhysicsCommand(0, this, [this, ReceivedInputs, PhysScene, LocalOffset]()
+					{
+						// Record received data into history
+						const bool bCompareDataForRewind = (bCompareInputToTriggerRewind && IsLocallyControlled());
+						const int32 ResimFrame = InputHistory->ReceiveNewData(*ReceivedInputs, LocalOffset, /*CompareDataForRewind*/bCompareDataForRewind, /*Important*/true);
+						
+						// Trigger resimulation if ReceiveNewData returned a valid frame to rewind to
+						TriggerResimulation(ResimFrame);
+
+						// Cache earliest important data frame since last ApplyData was called
+						NewImportantInputFrame = FMath::Min(NewImportantInputFrame, ReceivedInputs->GetEarliestFrame());
+
+#if DEBUG_NETWORK_PHYSICS
+						{
+							const int32 CurrentFrame = PhysScene->GetSolver()->GetCurrentFrame();
+							FString NetRole = HasServerWorld() ? FString("SERVER    ") : (IsLocallyControlled() ? FString("AUTONOMOUS") : FString("PROXY     "));
+							ReceivedInputs->DebugData(FString::Printf(TEXT("%s | PT | MulticastReceiveImportantInputData | RELIABLE | CurrentLocalFrame = %d | NewImportantInputFrame = %d | Component = %s"), *NetRole, CurrentFrame, NewImportantInputFrame, *GetFullName()));
+						}
+#endif
+					}, false);
+			}
+		}
+
+	}
+}
+
+void UNetworkPhysicsComponent::MulticastReceiveImportantStateData_Implementation(const FNetworkPhysicsRewindDataImportantStateProxy& ServerStates)
+{
+	// Multicast RPC, ignore on server
+	if (HasServerWorld())
+	{
+		return;
+	}
+
+	APlayerController* PlayerController = GetPlayerController();
+	if (!PlayerController)
+	{
+		PlayerController = GetWorld()->GetFirstPlayerController();
+	}
+
+	if (PlayerController && StateHistory)
+	{
+		const int32 LocalOffset = PlayerController->GetNetworkPhysicsTickOffset();
+
+		if (UWorld* World = GetWorld())
+		{
+			if (FPhysScene* PhysScene = World->GetPhysicsScene())
+			{
+				// Make another copy of the client inputs for the physics thread to consume
+				TSharedPtr<Chaos::FBaseRewindHistory> ReceivedStates = MakeShareable(ServerStates.History->Clone().Release());
+				PhysScene->EnqueueAsyncPhysicsCommand(0, this, [this, ReceivedStates, PhysScene, LocalOffset]()
+					{
+						const bool bCompareDataForRewind = (bCompareStateToTriggerRewind && IsLocallyControlled());
+						const int32 ResimFrame = StateHistory->ReceiveNewData(*ReceivedStates, LocalOffset, /*CompareDataForRewind*/bCompareDataForRewind, /*bImportant*/true);
+						TriggerResimulation(ResimFrame);
+
+#if DEBUG_NETWORK_PHYSICS
+						{
+							const int32 CurrentFrame = PhysScene->GetSolver()->GetCurrentFrame();
+							FString NetRole = HasServerWorld() ? FString("SERVER    ") : (IsLocallyControlled() ? FString("AUTONOMOUS") : FString("PROXY     "));
+							ReceivedStates->DebugData(FString::Printf(TEXT("%s | PT | MulticastReceiveImportantStateData | RELIABLE | CurrentLocalFrame = %d | Component = %s"), *NetRole, CurrentFrame, *GetFullName()));
+						}
+#endif
+					}, false);
+			}
+		}
+
+	}
+}
+
 void UNetworkPhysicsComponent::OnPreProcessInputsInternal(const int32 PhysicsStep)
 {
 	const bool bIsServer = HasServerWorld();
@@ -856,8 +1060,8 @@ void UNetworkPhysicsComponent::OnPreProcessInputsInternal(const int32 PhysicsSte
 			{
 				PhysicsData->ApplyData(ActorComponent);
 #if DEBUG_NETWORK_PHYSICS
-				UE_LOG(LogChaos, Log, TEXT("		Applying extracted state from history | bExactFrame = %d | LocalFrame = %d | ServerFrame = %d | InputFrame = %d | ")
-					, bExactFrame, PhysicsData->LocalFrame, PhysicsData->ServerFrame, PhysicsData->InputFrame);
+				UE_LOG(LogChaos, Log, TEXT("		Applying extracted state from history | bExactFrame = %d | LocalFrame = %d | ServerFrame = %d | InputFrame = %d | Data: %s")
+					, bExactFrame, PhysicsData->LocalFrame, PhysicsData->ServerFrame, PhysicsData->InputFrame, *PhysicsData->DebugData());
 #endif
 			}
 #if DEBUG_NETWORK_PHYSICS
@@ -873,13 +1077,34 @@ void UNetworkPhysicsComponent::OnPreProcessInputsInternal(const int32 PhysicsSte
 		if (InputHistory && (!IsLocallyControlled() || bIsSolverResim))
 		{
 			FNetworkPhysicsData* PhysicsData = InputData.Get();
-			const int32 NextExpectedLocalFrame = PhysicsData->LocalFrame + 1;
+			int32 NextExpectedLocalFrame = PhysicsData->LocalFrame + 1;
 			PhysicsData->LocalFrame = PhysicsStep;
 
+			// There are important inputs earlier than upcoming input to apply
+			if (NewImportantInputFrame < NextExpectedLocalFrame && !bIsSolverResim)
+			{
+				if (bApplyDataInsteadOfMergeData)
+				{
+#if DEBUG_NETWORK_PHYSICS
+					UE_LOG(LogChaos, Log, TEXT("	Reapplying multiple data due to receiving an important data that was previously missed. FromFrame: %d | ToFrame: %d | IsLocallyControlled = %d"), NewImportantInputFrame, (NextExpectedLocalFrame - 1), IsLocallyControlled());
+#endif
+					// Apply all inputs in range
+					InputHistory->ApplyDataRange(NewImportantInputFrame, NextExpectedLocalFrame - 1, ActorComponent, /*bOnlyImportant*/false);
+				}
+				else
+				{
+					// Merge all inputs from earliest new important
+					NextExpectedLocalFrame = NewImportantInputFrame;
+#if DEBUG_NETWORK_PHYSICS
+					UE_LOG(LogChaos, Log, TEXT("	Prepare to reapply multiple data through MergeData due to receiving an important data that was previously missed. FromFrame: %d | ToFrame: %d | IsLocallyControlled = %d"), NewImportantInputFrame, (NextExpectedLocalFrame - 1), IsLocallyControlled());
+#endif
+				}
+			}
+
 			PRAGMA_DISABLE_DEPRECATION_WARNINGS // TODO: Change to ExtractData() in UE 5.6 and remove deprecation pragma
-			if (InputHistory->ExtractDatas(PhysicsStep, bIsSolverReset, PhysicsData))
+			if (InputHistory->ExtractDatas(PhysicsStep, bIsSolverReset, PhysicsData, /*bExactFrame*/(bAllowInputExtrapolation == false)))
 			PRAGMA_ENABLE_DEPRECATION_WARNINGS
-			{ 
+			{
 				// Calculate input decay if we are resimulating and we don't have up to date inputs
 				if (bIsSolverResim)
 				{
@@ -889,25 +1114,52 @@ void UNetworkPhysicsComponent::OnPreProcessInputsInternal(const int32 PhysicsSte
 						PhysicsData->DecayData(InputDecay);
 					}
 				}
-				// Merge all inputs since last used input, if not resimulating
-				else if (PhysicsData->LocalFrame > NextExpectedLocalFrame)
+				// Apply input if data is unimportant or if data is important and unaltered (extrapolated, interpolated)
+				else if (!PhysicsData->bImportant || PhysicsData->InputFrame != INDEX_NONE)
 				{
-					InputHistory->MergeData(NextExpectedLocalFrame, PhysicsData);
-				}
+					// Check if we have a gap between last used input and current input
+					if (PhysicsData->LocalFrame > NextExpectedLocalFrame)
+					{
+						if (bApplyDataInsteadOfMergeData)
+						{
+#if DEBUG_NETWORK_PHYSICS
+							UE_LOG(LogChaos, Log, TEXT("	Applying multiple data instead of merging, from LocalFrame %d into LocalFrame %d | IsLocallyControlled = %d"), NextExpectedLocalFrame, PhysicsData->LocalFrame, IsLocallyControlled());
+#endif
+							// Iterate over each input and call ApplyData, except on the last, it will get handled by the normal ApplyData call further down
+							const int32 LastFrame = PhysicsData->LocalFrame;
+							for (; NextExpectedLocalFrame <= LastFrame; NextExpectedLocalFrame++)
+							{
+								if (InputHistory->ExtractData(NextExpectedLocalFrame, bIsSolverReset, PhysicsData, true) && NextExpectedLocalFrame < LastFrame)
+								{
+									PhysicsData->ApplyData(ActorComponent);
+								}
+							}
+						}
+						else
+						{
+#if DEBUG_NETWORK_PHYSICS
+							UE_LOG(LogChaos, Log, TEXT("	Merging inputs from LocalFrame %d into LocalFrame %d | IsLocallyControlled = %d"), NextExpectedLocalFrame, PhysicsData->LocalFrame, IsLocallyControlled());
+#endif
+							// Merge all inputs since last used input
+							InputHistory->MergeData(NextExpectedLocalFrame, PhysicsData);
+						}
+					}
 
-				// If the extracted input data was altered (interpolated, merged, extrapolated/predicted) on the server, record it into the history for it to get replicated to clients
-				if (bIsServer && !bIsSolverResim && PhysicsData->InputFrame == INDEX_NONE)
-				{
-					PhysicsData->bReceivedData = true; // Mark the input data as received so that it doesn't get overwritten by incoming client inputs
-					PhysicsData->LocalFrame = PhysicsStep;
-					InputHistory->RecordData(PhysicsStep, PhysicsData);
+					// If the extracted input data was altered (extrapolated, interpolated) on the server, record it into the history for it to get replicated to clients
+					if (bIsServer && !bIsSolverResim && PhysicsData->InputFrame == INDEX_NONE)
+					{
+						PhysicsData->bReceivedData = true; // Mark the input data as received so that it doesn't get overwritten by incoming client inputs
+						PhysicsData->bImportant = false;
+						PhysicsData->LocalFrame = PhysicsStep;
+						InputHistory->RecordData(PhysicsStep, PhysicsData);
+					}
 				}
 
 				PhysicsData->ApplyData(ActorComponent);
 
 #if DEBUG_NETWORK_PHYSICS
-				UE_LOG(LogChaos, Log, TEXT("		Applying extracted input from history | LocalFrame = %d | ServerFrame = %d | InputFrame = %d | IsResim = %d | IsLocallyControlled = %d | InputDecay = %f")
-					, PhysicsData->LocalFrame, PhysicsData->ServerFrame, PhysicsData->InputFrame, bIsSolverResim, IsLocallyControlled(), GetCurrentInputDecay(PhysicsData));
+				UE_LOG(LogChaos, Log, TEXT("		Applying extracted input from history | LocalFrame = %d | ServerFrame = %d | InputFrame = %d | IsResim = %d | IsLocallyControlled = %d | InputDecay = %f | Data: %s")
+					, PhysicsData->LocalFrame, PhysicsData->ServerFrame, PhysicsData->InputFrame, bIsSolverResim, IsLocallyControlled(), GetCurrentInputDecay(PhysicsData), *PhysicsData->DebugData());
 #endif
 			}
 #if DEBUG_NETWORK_PHYSICS
@@ -919,6 +1171,7 @@ void UNetworkPhysicsComponent::OnPreProcessInputsInternal(const int32 PhysicsSte
 #endif
 		}
 	}
+	NewImportantInputFrame = INT_MAX;
 }
 
 void UNetworkPhysicsComponent::OnPostProcessInputsInternal(const int32 PhysicsStep)
@@ -955,15 +1208,12 @@ void UNetworkPhysicsComponent::OnPostProcessInputsInternal(const int32 PhysicsSt
 			PlayerController = GetWorld()->GetFirstPlayerController();
 		}
 
-		const bool bShouldCacheInputHistory = PlayerController && IsLocallyControlled() && !bIsSolverResim;
+		const bool bShouldCacheInputHistory = IsLocallyControlled() && !bIsSolverResim;
 		// for the inputs client local ones are ground truth otherwise use the replicated ones coming from the server
-		if (bShouldCacheInputHistory && (InputData != nullptr))
+		if (PlayerController && bShouldCacheInputHistory && (InputData != nullptr))
 		{
 			FNetworkPhysicsData* PhysicsData = InputData.Get();
-			PhysicsData->LocalFrame = PhysicsStep;
-			PhysicsData->ServerFrame = bIsServer ? PhysicsStep : PhysicsStep + PlayerController->GetNetworkPhysicsTickOffset();
-			PhysicsData->InputFrame = PhysicsStep;
-			PhysicsData->bReceivedData = false;
+			PhysicsData->PrepareFrame(PhysicsStep, bIsServer, PlayerController->GetNetworkPhysicsTickOffset());
 
 			PhysicsData->BuildData(ActorComponent);
 
@@ -972,13 +1222,15 @@ void UNetworkPhysicsComponent::OnPostProcessInputsInternal(const int32 PhysicsSt
 			PRAGMA_ENABLE_DEPRECATION_WARNINGS
 
 #if DEBUG_NETWORK_PHYSICS
+			{
 				UE_LOG(LogChaos, Log, TEXT("		Recording input into history | LocalFrame = %d | ServerFrame = %d | InputFrame = %d | Input: %s ")
 					, PhysicsData->LocalFrame, PhysicsData->ServerFrame, PhysicsData->InputFrame, *PhysicsData->DebugData());
+			}
 #endif
 		}
 
 		const bool bShouldCacheStateHistory = bIsServer || (bCompareStateToTriggerRewind && bShouldCacheInputHistory);
-		if (StateHistory && bShouldCacheStateHistory)
+		if (PlayerController && StateHistory && bShouldCacheStateHistory)
 		{
 			// Compute of the local frame coming from the client that was used to generate this state
 			int32 InputFrame = INDEX_NONE;
@@ -993,10 +1245,8 @@ void UNetworkPhysicsComponent::OnPostProcessInputsInternal(const int32 PhysicsSt
 			}
 
 			FNetworkPhysicsData* PhysicsData = StateData.Get();
-			PhysicsData->LocalFrame = PhysicsStep;
-			PhysicsData->ServerFrame = PhysicsStep;
+			PhysicsData->PrepareFrame(PhysicsStep, bIsServer, PlayerController->GetNetworkPhysicsTickOffset());
 			PhysicsData->InputFrame = InputFrame;
-			PhysicsData->bReceivedData = false;
 
 			PhysicsData->BuildData(ActorComponent);
 
@@ -1005,8 +1255,10 @@ void UNetworkPhysicsComponent::OnPostProcessInputsInternal(const int32 PhysicsSt
 			PRAGMA_ENABLE_DEPRECATION_WARNINGS
 
 #if DEBUG_NETWORK_PHYSICS
+			{
 				UE_LOG(LogChaos, Log, TEXT("		Recording state into history | LocalFrame = %d | ServerFrame = %d | InputFrame = %d | State: %s ")
 					, PhysicsData->LocalFrame, PhysicsData->ServerFrame, PhysicsData->InputFrame, *PhysicsData->DebugData());
+			}
 #endif
 		}
 	}
@@ -1164,6 +1416,31 @@ int32 UNetworkPhysicsComponent::SetupRewindData()
 	return NumFrames;
 }
 
+void UNetworkPhysicsComponent::TriggerResimulation(int32 ResimFrame)
+{
+	if (ResimFrame != INDEX_NONE)
+	{
+		if (FPhysScene* PhysScene = GetWorld()->GetPhysicsScene())
+		{
+			if (Chaos::FPhysicsSolver* Solver = PhysScene->GetSolver())
+			{
+				if (Chaos::FRewindData* RewindData = Solver->GetRewindData())
+				{
+					// Mark particle/island as resim
+					Chaos::FReadPhysicsObjectInterface_Internal Interface = Chaos::FPhysicsObjectInternalInterface::GetRead();
+					if (Chaos::FPBDRigidParticleHandle* POHandle = Interface.GetRigidParticle(RootPhysicsObject))
+					{
+						Solver->GetEvolution()->GetIslandManager().SetParticleResimFrame(POHandle, ResimFrame);
+					}
+
+					// Set resim frame in rewind data
+					ResimFrame = (RewindData->GetResimFrame() == INDEX_NONE) ? ResimFrame : FMath::Min(ResimFrame, RewindData->GetResimFrame());
+					RewindData->SetResimFrame(ResimFrame);
+				}
+			}
+		}
+	}
+}
 
 
 
diff --git a/Engine/Source/Runtime/Engine/Private/PhysicsEngine/NetworkPhysicsSettingsComponent.cpp b/Engine/Source/Runtime/Engine/Private/PhysicsEngine/NetworkPhysicsSettingsComponent.cpp
--- a/Engine/Source/Runtime/Engine/Private/PhysicsEngine/NetworkPhysicsSettingsComponent.cpp
+++ b/Engine/Source/Runtime/Engine/Private/PhysicsEngine/NetworkPhysicsSettingsComponent.cpp
@@ -60,6 +60,7 @@ void UNetworkPhysicsSettingsComponent::InitializeComponent()
 								AsyncInput->Settings.DefaultReplicationSettings = DefaultReplicationSettings;
 								AsyncInput->Settings.ResimulationSettings = ResimulationSettings;
 								AsyncInput->Settings.PredictiveInterpolationSettings = PredictiveInterpolationSettings;
+								AsyncInput->Settings.NetworkPhysicsComponentSettings = NetworkPhysicsComponentSettings;
 							}
 						}
 					}
diff --git a/Engine/Source/Runtime/Engine/Public/Physics/NetworkPhysicsComponent.h b/Engine/Source/Runtime/Engine/Public/Physics/NetworkPhysicsComponent.h
--- a/Engine/Source/Runtime/Engine/Public/Physics/NetworkPhysicsComponent.h
+++ b/Engine/Source/Runtime/Engine/Public/Physics/NetworkPhysicsComponent.h
@@ -58,9 +58,25 @@ struct TNetRewindHistory : public Chaos::TDataRewindHistory<DataType>
 		}
 	}
 
-	virtual void CopyData(Chaos::FBaseRewindHistory& OutHistory, const uint32 StartFrame, const uint32 EndFrame) override
+	virtual void ApplyDataRange(const int32 FromFrame, const int32 ToFrame, void* ActorComponent, const bool bOnlyImportant = false) override
+	{
+		UActorComponent* NetworkComponent = static_cast<UActorComponent*>(ActorComponent);
+
+		for (int32 ApplyFrame = FromFrame; ApplyFrame <= ToFrame; ++ApplyFrame)
+		{
+			const int32 ApplyIndex = Super::GetFrameIndex(ApplyFrame);
+			DataType& FrameData = Super::DataHistory[ApplyIndex];
+			if (ApplyFrame == FrameData.LocalFrame && (!bOnlyImportant || FrameData.bImportant))
+			{
+				FrameData.ApplyData(NetworkComponent);
+			}
+		}
+	}
+
+	virtual bool CopyData(Chaos::FBaseRewindHistory& OutHistory, const uint32 StartFrame, const uint32 EndFrame, bool bIncludeUnimportant = true, bool bIncludeImportant = false) override
 	{
 		TNetRewindHistory& OutNetHistory = static_cast<TNetRewindHistory&>(OutHistory);
+		bool bHasCopiedData = false;
 
 		DataType FrameData;
 		for (uint32 CopyFrame = StartFrame; CopyFrame <= EndFrame; ++CopyFrame)
@@ -69,9 +85,16 @@ struct TNetRewindHistory : public Chaos::TDataRewindHistory<DataType>
 			if (CopyFrame == Super::DataHistory[CopyIndex].LocalFrame)
 			{
 				FrameData = Super::DataHistory[CopyIndex];
-				OutNetHistory.RecordData(CopyFrame, &FrameData);
+
+				// Check if we should include unimportant and/or important data
+				if ((!FrameData.bImportant && bIncludeUnimportant) || (FrameData.bImportant && bIncludeImportant))
+				{
+					OutNetHistory.RecordData(CopyFrame, &FrameData);
+					bHasCopiedData = true;
+				}
 			}
 		}
+		return bHasCopiedData;
 	}
 
 	virtual TUniquePtr<Chaos::FBaseRewindHistory> CopyFramesWithOffset(const uint32 StartFrame, const uint32 EndFrame, const int32 FrameOffset) override
@@ -97,7 +120,7 @@ struct TNetRewindHistory : public Chaos::TDataRewindHistory<DataType>
 		return Copy;
 	}
 
-	virtual int32 ReceiveNewData(Chaos::FBaseRewindHistory& NewData, const int32 FrameOffset, bool CompareDataForRewind = false) override
+	virtual int32 ReceiveNewData(Chaos::FBaseRewindHistory& NewData, const int32 FrameOffset, bool CompareDataForRewind = false, const bool bImportant = false) override
 	{
 		TNetRewindHistory& NetNewData = static_cast<TNetRewindHistory&>(NewData);
 
@@ -107,14 +130,14 @@ struct TNetRewindHistory : public Chaos::TDataRewindHistory<DataType>
 			for (int32 FrameIndex = 0; FrameIndex < NetNewData.NumFrames; ++FrameIndex)
 			{
 				DataType& FrameData = NetNewData.DataHistory[FrameIndex];
+				FrameData.bImportant = bImportant;
+				FrameData.bReceivedData = true; // Received data is marked to differentiate from locally predicted data
 
 				FrameData.LocalFrame = FrameData.ServerFrame - FrameOffset;
 
-				if (Super::ShouldRecordReceivedDataOnFrame(FrameData.LocalFrame))
+				if (ShouldRecordReceivedDataOnFrame(FrameData))
 				{
-					FrameData.bReceivedData = true; // Received data is marked to differentiate from locally predicted data
-
-					if (CompareDataForRewind && FrameData.LocalFrame > RewindFrame && Super::TriggerRewindFromNewData(&FrameData))
+					if (CompareDataForRewind && FrameData.LocalFrame > RewindFrame && TriggerRewindFromNewData(FrameData))
 					{
 						RewindFrame = FrameData.LocalFrame;
 					}
@@ -128,6 +151,31 @@ struct TNetRewindHistory : public Chaos::TDataRewindHistory<DataType>
 		return RewindFrame;
 	}
 
+	/** Check if we should record received data into history.
+	* Can for example block received data from client from overriding server authoritative data */
+	virtual bool ShouldRecordReceivedDataOnFrame(const DataType& ReceivedData)
+	{
+		if (ReceivedData.LocalFrame < 0)
+		{
+			return false;
+		}
+
+		// Allow received data to get recorded if it's for a newer frame than already stored at the history index or if the data stored is not marked as received (i.e. it's locally predicted)
+		Super::LoadData(ReceivedData.LocalFrame);
+		return (Super::DataHistory[Super::CurrentIndex].LocalFrame < ReceivedData.LocalFrame || (!Super::DataHistory[Super::CurrentIndex].bReceivedData && Super::DataHistory[Super::CurrentIndex].LocalFrame == ReceivedData.LocalFrame));
+	}
+
+	/** Compares new received data with local predicted data and returns true if they differ enough to trigger a resimulation  */
+	virtual bool TriggerRewindFromNewData(DataType& NewData)
+	{
+		if (Super::EvalData(NewData.LocalFrame) && !Super::DataHistory[Super::CurrentIndex].bReceivedData)
+		{
+			return !NewData.CompareData(Super::DataHistory[Super::CurrentIndex]);
+		}
+
+		return false;
+	}
+
 	virtual void NetSerialize(FArchive& Ar, UPackageMap* InPackageMap) override
 	{
 		Ar << Super::NumFrames;
@@ -182,11 +230,13 @@ struct TNetRewindHistory : public Chaos::TDataRewindHistory<DataType>
 		{
 			for (int32 FrameIndex = 0; FrameIndex < Super::NumFrames; ++FrameIndex)
 			{
-				UE_LOG(LogChaos, Log, TEXT("		Index: %d || LocalFrame = %d || ServerFrame = %d || InputFrame = %d  ||  Data: %s")
+				UE_LOG(LogChaos, Log, TEXT("		Index: %d || LocalFrame = %d || ServerFrame = %d || InputFrame = %d || bReceivedData = %d || bImportant = %d  ||  Data: %s")
 				, FrameIndex
 				, Super::DataHistory[FrameIndex].LocalFrame
 				, Super::DataHistory[FrameIndex].ServerFrame
 				, Super::DataHistory[FrameIndex].InputFrame
+				, Super::DataHistory[FrameIndex].bReceivedData
+				, Super::DataHistory[FrameIndex].bImportant
 				, *Super::DataHistory[FrameIndex].DebugData());
 			}
 		}
@@ -297,6 +347,47 @@ struct TStructOpsTypeTraits<FNetworkPhysicsRewindDataStateProxy> : public TStruc
 		WithIdenticalViaEquality = true
 	};
 };
+/**
+ * Struct suitable for use as a replicated property to replicate input rewind history
+ */
+USTRUCT()
+struct FNetworkPhysicsRewindDataImportantInputProxy : public FNetworkPhysicsRewindDataProxy
+{
+	GENERATED_BODY()
+		
+	bool NetSerialize(FArchive& Ar, class UPackageMap* Map, bool& bOutSuccess);
+};
+
+template<>
+struct TStructOpsTypeTraits<FNetworkPhysicsRewindDataImportantInputProxy> : public TStructOpsTypeTraitsBase2<FNetworkPhysicsRewindDataImportantInputProxy>
+{
+	enum
+	{
+		WithNetSerializer = true,
+		WithIdenticalViaEquality = true
+	};
+};
+
+/**
+ * Struct suitable for use as a replicated property to replicate state rewind history
+ */
+USTRUCT()
+struct FNetworkPhysicsRewindDataImportantStateProxy : public FNetworkPhysicsRewindDataProxy
+{
+	GENERATED_BODY()
+
+	bool NetSerialize(FArchive& Ar, class UPackageMap* Map, bool& bOutSuccess);
+};
+
+template<>
+struct TStructOpsTypeTraits<FNetworkPhysicsRewindDataImportantStateProxy> : public TStructOpsTypeTraitsBase2<FNetworkPhysicsRewindDataImportantStateProxy>
+{
+	enum
+	{
+		WithNetSerializer = true,
+		WithIdenticalViaEquality = true
+	};
+};
 
 /**
  * Network physics rewind callback to manage all the sim callbacks rewind functionalities
@@ -422,6 +513,9 @@ struct FNetworkPhysicsData
 	// If this data was received over the network or locally predicted
 	bool bReceivedData = false;
 
+	// If this data is marked as important (replicated reliably)
+	bool bImportant = false;
+
 	// Serialize the data into/from the archive
 	void SerializeFrames(FArchive& Ar)
 	{
@@ -430,6 +524,22 @@ struct FNetworkPhysicsData
 		Ar << InputFrame;
 	}
 
+	void PrepareFrame(int32 CurrentFrame, bool bIsServer, int32 ClientFrameOffset)
+	{
+		LocalFrame = CurrentFrame;
+		ServerFrame = bIsServer ? CurrentFrame : CurrentFrame + ClientFrameOffset;
+		InputFrame = CurrentFrame;
+		bReceivedData = false;
+		bImportant = false;
+	}
+
+	/** Set if this data is important(replicated reliably) or unimportant(replicated unreliably)
+	* NOTE: Default is to handle all inputs as unimportant, while one time events can be marked as important. */
+	void SetImportant(bool bIsImportant)
+	{
+		bImportant = bIsImportant;
+	}
+
 	// Apply the data onto the network physics component
 	virtual void ApplyData(UActorComponent* NetworkComponent) const
 	{
@@ -527,14 +637,9 @@ public:
 	// Init the network physics component 
 	ENGINE_API void InitPhysics();
 
-	// Server RPC to receive inputs from client
-	UFUNCTION(Server, unreliable)
-	ENGINE_API void ServerReceiveInputData(const FNetworkPhysicsRewindDataInputProxy& ClientInputs);
-
 	UE_DEPRECATED(5.4, "Deprecated, use SendInputData() instead")
 	UFUNCTION(Server, unreliable, meta = (DeprecatedFunction, DeprecationMessage = "ServerReceiveInputsDatas has been deprecated. Use ServerReceiveInputData instead."))
 	ENGINE_API void ServerReceiveInputsDatas(const FNetworkPhysicsRewindDataInputProxy& ClientInputs);
-
 	// Async physics tick component function per frame from the solver
 	ENGINE_API virtual void AsyncPhysicsTickComponent(float DeltaTime, float SimTime) override;
 
@@ -621,12 +726,24 @@ public:
 	ENGINE_API bool IsLocallyControlled() const;
 
 	/** Mark this as controlled through locally relayed inputs rather than controlled as a pawn through a player controller.
-	* Set if NetworkPhysicsComponent is implemented on an AActor instead of APawn and it's currently being fed inputs from the local player / autonomous proxy */
+	* Set if NetworkPhysicsComponent is implemented on an AActor instead of APawn and it's currently being fed inputs, or if this is controlled by the server. 
+	* NOTE: The actor for this NetworkPhysicsComponent also needs to be owned by the local client if this is used client-side. */
 	ENGINE_API void SetIsRelayingLocalInputs(bool bInRelayingLocalInputs)
 	{
 		bIsRelayingLocalInputs = bInRelayingLocalInputs;
 	}
 
+	/** Stop relaying local inputs after next network send.
+	* Deferred version of SetIsRelayingLocalInputs(false) to ensure that the last replicated data gets sent.
+	* This does not work on locally controlled APawns, see SetIsRelayingLocalInputs() for description. */
+	ENGINE_API void StopRelayingLocalInputsDeferred()
+	{
+		if (bIsRelayingLocalInputs)
+		{
+			bStopRelayingLocalInputsDeferred = true;
+		}
+	}
+
 	/** Check if this is controlled locally through relayed inputs from autonomous proxy. It's recommended to use IsLocallyControlled() when checking if this is locally controlled. */
 	ENGINE_API const bool GetIsRelayingLocalInputs() const { return bIsRelayingLocalInputs; }
 
@@ -635,6 +752,30 @@ public:
 
 protected : 
 
+	// Server RPC to receive inputs from client
+	UFUNCTION(Server, unreliable)
+	ENGINE_API void ServerReceiveInputData(const FNetworkPhysicsRewindDataInputProxy& ClientInputs);
+
+	// Server RPC to receive important inputs from client
+	UFUNCTION(Server, reliable)
+	ENGINE_API void ServerReceiveImportantInputData(const FNetworkPhysicsRewindDataImportantInputProxy& ClientInputs);
+
+	// Client RPC to receive important inputs from server
+	UFUNCTION(NetMulticast, reliable)
+	ENGINE_API void MulticastReceiveImportantInputData(const FNetworkPhysicsRewindDataImportantInputProxy& ServerInputs);
+
+	// Client RPC to receive important states from server
+	UFUNCTION(NetMulticast, reliable)
+	ENGINE_API void MulticastReceiveImportantStateData(const FNetworkPhysicsRewindDataImportantStateProxy& ServerStates);
+
+	// replicated important physics input
+	UPROPERTY(Transient)
+	FNetworkPhysicsRewindDataImportantInputProxy ReplicatedImportantInput;
+
+	// replicated important physics state
+	UPROPERTY(Transient)
+	FNetworkPhysicsRewindDataImportantStateProxy ReplicatedImportantState;
+
 	// repnotify for the inputs on the client
 	UFUNCTION()
 	ENGINE_API void OnRep_SetReplicatedInputs();
@@ -651,14 +792,17 @@ protected :
 	UPROPERTY(Transient, ReplicatedUsing = OnRep_SetReplicatedStates)
 	FNetworkPhysicsRewindDataStateProxy ReplicatedStates;
 
-	// Frame counter to compute the local to server offset
-	int32 FrameCounter = 0;
+private:
+	/** Trigger a resimulation on frame */
+	void TriggerResimulation(int32 ResimFrame);
 
 private:
 
 	friend FNetworkPhysicsCallback;
 	friend struct FNetworkPhysicsRewindDataInputProxy;
 	friend struct FNetworkPhysicsRewindDataStateProxy;
+	friend struct FNetworkPhysicsRewindDataImportantInputProxy;
+	friend struct FNetworkPhysicsRewindDataImportantStateProxy;
 
 	// States history uses to rewind simulation 
 	TSharedPtr<Chaos::FBaseRewindHistory> StateHistory;
@@ -686,12 +830,21 @@ private:
 
 	// Locally relayed inputs makes this component act as if it's a locally controlled pawn.
 	bool bIsRelayingLocalInputs = false;
+	
+	// If we are currently relaying inputs and will stop after next network send.
+	bool bStopRelayingLocalInputsDeferred = false;
 
-	// Cache locally predicted states and then compare then via FNetworkPhysicsData::CompareData to trigger rewind if comparison differ
-	bool bCompareStateToTriggerRewind = false;
+	int32 LastInputSendFrame = INDEX_NONE;
+	int32 LastStateSendFrame = INDEX_NONE;
+	int32 NewImportantInputFrame = INT_MAX;
 
-	// Compare locally predicted inputs via FNetworkPhysicsData::CompareData to trigger rewind if comparison differ
+	// Temporary settings variables, ToDo, Remove when async ISimCallback flow is implemented and these can be retrieved from NetworkPhysicsSettingsComponent on PT
+	bool bCompareStateToTriggerRewind = false;
 	bool bCompareInputToTriggerRewind = false;
+	bool bEnableUnreliableFlow = true;
+	bool bEnableReliableFlow = false;
+	bool bApplyDataInsteadOfMergeData = false;
+	bool bAllowInputExtrapolation = true;
 };
 
 /** DEPRECATED UE 5.4 */
@@ -718,6 +871,12 @@ FORCEINLINE void UNetworkPhysicsComponent::CreateDataHistory(UActorComponent* Hi
 	ReplicatedStates.History = MakeUnique<TNetRewindHistory<typename PhysicsTraits::StatesType>>(StateRedundancy);
 	ReplicatedStates.Owner = this;
 
+	ReplicatedImportantInput.History = MakeUnique<TNetRewindHistory<typename PhysicsTraits::InputsType>>(1);
+	ReplicatedImportantInput.Owner = this;
+
+	ReplicatedImportantState.History = MakeUnique<TNetRewindHistory<typename PhysicsTraits::StatesType>>(1);
+	ReplicatedImportantState.Owner = this;
+
 	ActorComponent = HistoryComponent;
 	
 	AddDataHistory();
@@ -735,6 +894,9 @@ FORCEINLINE void UNetworkPhysicsComponent::CreateInputHistory(UActorComponent* H
 	ReplicatedInputs.History = MakeUnique<TNetRewindHistory<InputsType>>(InputRedundancy);
 	ReplicatedInputs.Owner = this;
 
+	ReplicatedImportantInput.History = MakeUnique<TNetRewindHistory<InputsType>>(1);
+	ReplicatedImportantInput.Owner = this;
+
 	ActorComponent = HistoryComponent;
 
 	AddDataHistory();
diff --git a/Engine/Source/Runtime/Engine/Public/Physics/NetworkPhysicsSettingsComponent.h b/Engine/Source/Runtime/Engine/Public/Physics/NetworkPhysicsSettingsComponent.h
--- a/Engine/Source/Runtime/Engine/Public/Physics/NetworkPhysicsSettingsComponent.h
+++ b/Engine/Source/Runtime/Engine/Public/Physics/NetworkPhysicsSettingsComponent.h
@@ -62,6 +62,10 @@ namespace PhysicsReplicationCVars
 		extern float VelStabilityMultiplier;
 		extern float AngVelStabilityMultiplier;
 		extern bool bRuntimeCorrectConnectedBodies;
+		extern bool bEnableUnreliableFlow;
+		extern bool bEnableReliableFlow;
+		extern bool bApplyDataInsteadOfMergeData;
+		extern bool bAllowInputExtrapolation;
 	}
 }
 
@@ -243,18 +247,6 @@ struct FNetworkPhysicsSettingsResimulation
 	uint32 ResimulationErrorThreshold = 10;
 	uint32 GetResimulationErrorThreshold(uint32 DefaultValue) { return bOverrideRotStabilityMultiplier ? ResimulationErrorThreshold : DefaultValue; }
 
-	UPROPERTY(config, EditDefaultsOnly, Category = "Overrides", Meta = (InlineEditConditionToggle))
-	uint32 bOverrideRedundantInputs : 1;
-	// Override how many inputs to synchronize each sync to cover packet loss.
-	UPROPERTY(config, EditDefaultsOnly, Category = "Overrides", Meta = (EditCondition = "bOverrideRedundantInputs"))
-	uint8 RedundantInputs = 3;
-
-	UPROPERTY(config, EditDefaultsOnly, Category = "Overrides", Meta = (InlineEditConditionToggle))
-	uint32 bOverrideRedundantStates : 1;
-	// Override how many states to synchronize each sync to cover packet loss.
-	UPROPERTY(config, EditDefaultsOnly, Category = "Overrides", Meta = (EditCondition = "bOverrideRedundantStates"))
-	uint8 RedundantStates = 1;
-
 	UPROPERTY(config, EditDefaultsOnly, Category = "Overrides", Meta = (InlineEditConditionToggle))
 	uint32 bOverrideRuntimeCorrectionEnabled : 1;
 	// Overrides CVar: np2.Resim.RuntimeCorrectionEnabled -- Apply positional and rotational runtime corrections while within resim trigger distance.
@@ -303,6 +295,26 @@ struct FNetworkPhysicsSettingsResimulation
 	UPROPERTY(config, EditDefaultsOnly, Category = "Overrides", Meta = (EditCondition = "bOverrideAngVelStabilityMultiplier"))
 	float AngVelStabilityMultiplier = PhysicsReplicationCVars::ResimulationCVars::AngVelStabilityMultiplier;
 	float GetAngVelStabilityMultiplier() { return bOverrideAngVelStabilityMultiplier ? AngVelStabilityMultiplier : PhysicsReplicationCVars::ResimulationCVars::AngVelStabilityMultiplier; }
+};
+
+USTRUCT()
+struct FNetworkPhysicsSettingsNetworkPhysicsComponent
+{
+	GENERATED_BODY()
+
+	UPROPERTY(config, EditDefaultsOnly, Category = "Overrides", Meta = (InlineEditConditionToggle))
+	uint32 bOverrideRedundantInputs : 1;
+	// Override how many inputs to synchronize each sync to cover packet loss.
+	UPROPERTY(config, EditDefaultsOnly, Category = "Overrides", Meta = (EditCondition = "bOverrideRedundantInputs"))
+	uint8 RedundantInputs = 3;
+	uint8 GetRedundantInputs(uint8 DefaultValue) { return bOverrideRedundantInputs ? RedundantInputs : DefaultValue; }
+
+	UPROPERTY(config, EditDefaultsOnly, Category = "Overrides", Meta = (InlineEditConditionToggle))
+	uint32 bOverrideRedundantStates : 1;
+	// Override how many states to synchronize each sync to cover packet loss.
+	UPROPERTY(config, EditDefaultsOnly, Category = "Overrides", Meta = (EditCondition = "bOverrideRedundantStates"))
+	uint8 RedundantStates = 1;
+	uint8 GetRedundantStates(uint8 DefaultValue) { return bOverrideRedundantStates ? RedundantStates : DefaultValue; }
 
 	UPROPERTY(config, EditDefaultsOnly, Category = "Overrides", Meta = (InlineEditConditionToggle))
 	uint32 bOverrideCompareStateToTriggerRewind : 1;
@@ -317,6 +329,34 @@ struct FNetworkPhysicsSettingsResimulation
 	UPROPERTY(config, EditDefaultsOnly, Category = "Overrides", Meta = (EditCondition = "bOverrideCompareInputToTriggerRewind"))
 	bool bCompareInputToTriggerRewind = false;
 	bool GetCompareInputToTriggerRewind(bool DefaultValue) { return bOverrideCompareInputToTriggerRewind ? bCompareInputToTriggerRewind : DefaultValue; }
+
+	UPROPERTY(config, EditDefaultsOnly, Category = "Overrides", Meta = (InlineEditConditionToggle))
+	uint32 bOverrideEnableUnreliableFlow : 1;
+	// Overrides CVar: np2.Resim.EnableUnreliableFlow -- When true, allow data to be sent unreliably. Also sends FNetworkPhysicsData not marked with FNetworkPhysicsData::bimportant unreliably over the network.
+	UPROPERTY(config, EditDefaultsOnly, Category = "Overrides", Meta = (EditCondition = "bOverrideEnableUnreliableFlow"))
+	bool bEnableUnreliableFlow = PhysicsReplicationCVars::ResimulationCVars::bEnableUnreliableFlow;
+	bool GetEnableUnreliableFlow() { return bOverrideEnableUnreliableFlow ? bEnableUnreliableFlow : PhysicsReplicationCVars::ResimulationCVars::bEnableUnreliableFlow; }
+
+	UPROPERTY(config, EditDefaultsOnly, Category = "Overrides", Meta = (InlineEditConditionToggle))
+	uint32 bOverrideEnableReliableFlow : 1;
+	// Overrides CVar: np2.Resim.EnableReliableFlow -- EXPERIMENTAL -- When true, allow data to be sent reliably. Also send FNetworkPhysicsData marked with FNetworkPhysicsData::bimportant reliably over the network.
+	UPROPERTY(config, EditDefaultsOnly, Category = "Overrides", Meta = (EditCondition = "bOverrideEnableReliableFlow"))
+	bool bEnableReliableFlow = PhysicsReplicationCVars::ResimulationCVars::bEnableReliableFlow;
+	bool GetEnableReliableFlow() { return bOverrideEnableReliableFlow ? bEnableReliableFlow : PhysicsReplicationCVars::ResimulationCVars::bEnableReliableFlow; }
+
+	UPROPERTY(config, EditDefaultsOnly, Category = "Overrides", Meta = (InlineEditConditionToggle))
+	uint32 bOverrideApplyDataInsteadOfMergeData : 1;
+	// Overrides CVar: np2.Resim.ApplyDataInsteadOfMergeData -- When true, call ApplyData for each data instead of MergeData when having to use multiple data entries in one frame.
+	UPROPERTY(config, EditDefaultsOnly, Category = "Overrides", Meta = (EditCondition = "bOverrideApplyDataInsteadOfMergeData"))
+	bool bApplyDataInsteadOfMergeData = PhysicsReplicationCVars::ResimulationCVars::bApplyDataInsteadOfMergeData;
+	bool GetApplyDataInsteadOfMergeData() { return bOverrideApplyDataInsteadOfMergeData ? bApplyDataInsteadOfMergeData : PhysicsReplicationCVars::ResimulationCVars::bApplyDataInsteadOfMergeData; }
+
+	UPROPERTY(config, EditDefaultsOnly, Category = "Overrides", Meta = (InlineEditConditionToggle))
+	uint32 bOverrideAllowInputExtrapolation : 1;
+	// Overrides CVar: np2.Resim.AllowInputExtrapolation -- When true and not locally controlled, allow inputs to be extrapolated from last known and if there is a gap allow interpolation between two known inputs.
+	UPROPERTY(config, EditDefaultsOnly, Category = "Overrides", Meta = (EditCondition = "bOverrideAllowInputExtrapolation"))
+	bool bAllowInputExtrapolation = PhysicsReplicationCVars::ResimulationCVars::bAllowInputExtrapolation;
+	bool GetAllowInputExtrapolation() { return bOverrideAllowInputExtrapolation ? bAllowInputExtrapolation : PhysicsReplicationCVars::ResimulationCVars::bAllowInputExtrapolation; }
 };
 
 /*
@@ -361,6 +401,9 @@ public:
 	UPROPERTY(EditDefaultsOnly, Category = "Networked Physics Settings")
 	FNetworkPhysicsSettingsResimulation ResimulationSettings;
 	
+	UPROPERTY(EditDefaultsOnly, Category = "Networked Physics Settings")
+	FNetworkPhysicsSettingsNetworkPhysicsComponent NetworkPhysicsComponentSettings;
+	
 	/*
 	UPROPERTY(EditDefaultsOnly, Category = "Networked Physics Settings")
 	FNetworkPhysicsSettingsRewindData RewindSettings;
@@ -390,8 +433,9 @@ struct FNetworkPhysicsSettingsAsync
 {
 	FNetworkPhysicsSettings GeneralSettings;
 	FNetworkPhysicsSettingsDefaultReplication DefaultReplicationSettings;
-	FNetworkPhysicsSettingsResimulation ResimulationSettings;
 	FNetworkPhysicsSettingsPredictiveInterpolation PredictiveInterpolationSettings;
+	FNetworkPhysicsSettingsResimulation ResimulationSettings;
+	FNetworkPhysicsSettingsNetworkPhysicsComponent NetworkPhysicsComponentSettings;
 };
 
 struct FNetworkPhysicsSettingsAsyncInput : public Chaos::FSimCallbackInput
diff --git a/Engine/Source/Runtime/Experimental/Chaos/Public/RewindData.h b/Engine/Source/Runtime/Experimental/Chaos/Public/RewindData.h
--- a/Engine/Source/Runtime/Experimental/Chaos/Public/RewindData.h
+++ b/Engine/Source/Runtime/Experimental/Chaos/Public/RewindData.h
@@ -48,6 +48,9 @@ struct FBaseRewindHistory
 	UE_DEPRECATED(5.4, "Deprecated, use ExtractData() instead")
 	FORCEINLINE virtual bool ExtractDatas(const int32 ExtractFrame, const bool bResetSolver, void* HistoryDatas, const bool bExactFrame = false) { return ExtractData(ExtractFrame, bResetSolver, HistoryDatas, bExactFrame); }
 	
+	/** Call ApplyData on each frame data within range */
+	FORCEINLINE virtual void ApplyDataRange(const int32 FromFrame, const int32 ToFrame, void* ActorComponent, const bool bOnlyImportant = false) {}
+
 	/** Iterate over and merge data */
 	FORCEINLINE virtual void MergeData(const int32 FromFrame, void* ToData) {}
 
@@ -56,14 +59,12 @@ struct FBaseRewindHistory
 	UE_DEPRECATED(5.4, "Deprecated, use RecordData() instead")
 	FORCEINLINE virtual bool RecordDatas(const int32 RecordFrame, const void* HistoryDatas) { return RecordData(RecordFrame, HistoryDatas); }
 	
-	/** Check if we should record received data into history on @param RecordFrame.
-	* Can for example block received data from client from overriding server authoritative data */
-	FORCEINLINE virtual bool ShouldRecordReceivedDataOnFrame(const int32 RecordFrame) { return true; }
-
 	/** Copy data from local history into @param OutHistory
 	* @param StartFrame = Included
-	* @param EndFrame = Included */
-	virtual void CopyData(Chaos::FBaseRewindHistory& OutHistory, const uint32 StartFrame, const uint32 EndFrame) {};
+	* @param EndFrame = Included 
+	* @param bIncludeUnimportant = If to copy unimportant data entries
+	* @param bIncludeImportant = If to copy important data entries*/
+	virtual bool CopyData(Chaos::FBaseRewindHistory& OutHistory, const uint32 StartFrame, const uint32 EndFrame, bool bIncludeUnimportant = true, bool bIncludeImportant = false) { return false; }
 
 	/** Create a polymorphic copy of only a range of frames, applying the frame offset to the copies
 	* @param StartFrame = Included
@@ -71,13 +72,10 @@ struct FBaseRewindHistory
 	virtual TUniquePtr<FBaseRewindHistory> CopyFramesWithOffset(const uint32 StartFrame, const uint32 EndFrame, const int32 FrameOffset) = 0;
 
 	/** Copy new data (received from the network) into this history, returns frame to resimulate from if @param CompareDataForRewind is set to true and compared data differ enough */
-	virtual int32 ReceiveNewData(FBaseRewindHistory& NewData, const int32 FrameOffset, bool CompareDataForRewind = false) { return INDEX_NONE; }
+	virtual int32 ReceiveNewData(FBaseRewindHistory& NewData, const int32 FrameOffset, const bool CompareDataForRewind = false, const bool bImportant = false) { return INDEX_NONE; }
 	UE_DEPRECATED(5.4, "Deprecated, use ReceiveNewData() instead")
 	virtual void ReceiveNewDatas(FBaseRewindHistory& NewDatas, const int32 FrameOffset) { ReceiveNewData(NewDatas, FrameOffset); }
 
-	/** Compares new received data with local predicted data and returns true if they differ enough to trigger a resimulation  */
-	FORCEINLINE virtual bool TriggerRewindFromNewData(void* NewData) { return false; }
-
 	/** Serialize the data to or from a network archive */
 	virtual void NetSerialize(FArchive& Ar, UPackageMap* PackageMap) {}
 	
@@ -98,9 +96,15 @@ struct FBaseRewindHistory
 	/** Legacy interface to apply inputs */
 	FORCEINLINE virtual bool ApplyInputs(const int32 ApplyFrame, const bool bResetSolver) { return false; }
 
-	/** Return the most up to date frame entry in history */
+	/** Return the most up to date frame entry in history, returns INDEX_NONE if no frame was found */
 	virtual const int32 GetLatestFrame() const { return INDEX_NONE; }
 
+	/** Return the least up to date frame entry in history, returns INT_MAX if no frame was found */
+	virtual const int32 GetEarliestFrame() const { return INT_MAX; }
+
+	/** Return the max size of the history */
+	virtual const int32 GetHistorySize() const { return 0; }
+
 	/** Resize the history */
 	virtual void ResizeDataHistory(const int32 FrameCount, const EAllowShrinking AllowShrinking = EAllowShrinking::Yes) {}
 };
@@ -188,6 +192,7 @@ public :
 			if (MinFrameIndex != INDEX_NONE && MaxFrameIndex != INDEX_NONE)
 			{
 				DataType& ExtractedData = *static_cast<DataType*>(HistoryData);
+				ExtractedData = DataHistory[MinFrameIndex];
 
 				PRAGMA_DISABLE_DEPRECATION_WARNINGS
 				// TODO: Change to InterpolateData() in UE 5.6 and remove deprecation pragma
@@ -247,16 +252,6 @@ public :
 		}
 	}
 
-	FORCEINLINE virtual bool TriggerRewindFromNewData(void* NewData) override
-	{
-		if (EvalData(static_cast<DataType*>(NewData)->LocalFrame) && !DataHistory[CurrentIndex].bReceivedData)
-		{
-			return !static_cast<DataType*>(NewData)->CompareData(DataHistory[CurrentIndex]);
-		}
-
-		return false;
-	}
-
 	/** Load the data from the buffer at a specific frame */
 	FORCEINLINE bool LoadData(const int32 LoadFrame)
 	{
@@ -295,18 +290,6 @@ public :
 		return true;
 	}
 
-	FORCEINLINE virtual bool ShouldRecordReceivedDataOnFrame(const int32 RecordFrame) override
-	{
-		if (RecordFrame < 0)
-		{
-			return false;
-		}
-
-		// Allow received data to get recorded if it's for a newer frame than already stored at the history index or if the data stored is not marked as received (i.e. it's locally predicted)
-		LoadData(RecordFrame);
-		return (!DataHistory[CurrentIndex].bReceivedData || DataHistory[CurrentIndex].LocalFrame < RecordFrame);
-	}
-
 	/** Current data that is being loaded/recorded*/
 	DataType& GetCurrentData() { return DataHistory[CurrentIndex]; }
 	UE_DEPRECATED(5.4, "Deprecated, use GetCurrentData() instead")
@@ -343,6 +326,28 @@ public :
 		return LatestFrame;
 	}
 
+	/** Return the least up to date frame entry in history */
+	virtual const int32 GetEarliestFrame() const override
+	{
+		int32 EarliestFrame = INT_MAX;
+		
+		for (int32 FrameIndex = 0; FrameIndex < NumFrames; ++FrameIndex)
+		{
+			if (DataHistory[FrameIndex].LocalFrame > INDEX_NONE)
+			{
+				EarliestFrame = FMath::Min(EarliestFrame, DataHistory[FrameIndex].LocalFrame);
+			}
+		}
+
+		return EarliestFrame;
+	}
+
+	/** Return the max size of the history */
+	virtual const int32 GetHistorySize() const
+	{
+		return NumFrames;
+	}
+
 	/** Resize the history */
 	FORCEINLINE void ResizeDataHistory(const int32 FrameCount, const EAllowShrinking AllowShrinking = EAllowShrinking::Yes) override
 	{
