- Removes unused [meta] texturepath names from DDC1_FetchAndFillDerivedData
- Corrects debug name used for DDC1 texture requests in DDC1_FetchAndFillDerivedData to be derived from the encodespeed. If the speed is Fast, only then will we use the name with the appended "[Fast]" tag

#rnx

diff --git a/Engine/Source/Runtime/Engine/Private/TextureDerivedDataTask.cpp b/Engine/Source/Runtime/Engine/Private/TextureDerivedDataTask.cpp
--- a/Engine/Source/Runtime/Engine/Private/TextureDerivedDataTask.cpp
+++ b/Engine/Source/Runtime/Engine/Private/TextureDerivedDataTask.cpp
@@ -1441,14 +1441,11 @@ static void DDC1_FetchAndFillDerivedData(
 
 	FSharedBuffer RawDerivedData;
 	const FSharedString SharedTexturePathName(TexturePathName);
-	const FSharedString SharedTextureMetaPathName(WriteToString<256>(TexturePathName, TEXTVIEW(" [Meta]")));
 	const FSharedString SharedTextureFastPathName(WriteToString<256>(TexturePathName, TEXTVIEW(" [Fast]")));
-	const FSharedString SharedTextureFastMetaPathName(WriteToString<256>(TexturePathName, TEXTVIEW(" [Fast][Meta]")));
 
 	FString LocalDerivedDataKeySuffix;
 	FString LocalDerivedDataKey;
 
-
 	bool bGotDDCData = false;
 	bool bUsedFetchFirst = false;
 	if (BuildSettingsPerLayerFetchFirst.Num() && !bForceRebuild)
@@ -1463,7 +1460,8 @@ static void DDC1_FetchAndFillDerivedData(
 			GetTextureDerivedDataKeyFromSuffix(FetchFirstKeySuffix, FetchFirstKey);
 
 			TArray<FCacheGetValueRequest, TInlineAllocator<1>> Requests;
-			Requests.Add({ SharedTexturePathName, ConvertLegacyCacheKey(FetchFirstKey), ECachePolicy::Default, 0 /* UserData */});
+			const FSharedString& TexturePathRequestName = (FetchFirstMetadata.EncodeSpeed == (uint8) ETextureEncodeSpeed::Fast) ? SharedTextureFastPathName : SharedTexturePathName;
+			Requests.Add({ TexturePathRequestName, ConvertLegacyCacheKey(FetchFirstKey), ECachePolicy::Default, 0 /* UserData */});
 
 			FRequestOwner BlockingOwner(EPriority::Blocking);
 
@@ -1493,7 +1491,8 @@ static void DDC1_FetchAndFillDerivedData(
 		GetTextureDerivedDataKeyFromSuffix(LocalDerivedDataKeySuffix, LocalDerivedDataKey);
 
 		TArray<FCacheGetValueRequest, TInlineAllocator<1>> Requests;
-		Requests.Add({ SharedTextureFastPathName, ConvertLegacyCacheKey(LocalDerivedDataKey), ECachePolicy::Default, 0 /* UserData */ });
+		const FSharedString& TexturePathRequestName = (FetchOrBuildMetadata.EncodeSpeed == (uint8) ETextureEncodeSpeed::Fast) ? SharedTextureFastPathName : SharedTexturePathName;
+		Requests.Add({ TexturePathRequestName, ConvertLegacyCacheKey(LocalDerivedDataKey), ECachePolicy::Default, 0 /* UserData */ });
 
 		FRequestOwner BlockingOwner(EPriority::Blocking);
 
