Fixed crash in UCharacterMovementComponent::GetActorFeetLocation() when the CharacterOwner is null.

#tests Tested with 100% repro using EQSTestPawn
#rnx

diff --git a/Engine/Source/Runtime/Engine/Private/Components/CharacterMovementComponent.cpp b/Engine/Source/Runtime/Engine/Private/Components/CharacterMovementComponent.cpp
--- a/Engine/Source/Runtime/Engine/Private/Components/CharacterMovementComponent.cpp
+++ b/Engine/Source/Runtime/Engine/Private/Components/CharacterMovementComponent.cpp
@@ -3814,8 +3814,14 @@ bool UCharacterMovementComponent::ApplyRequestedMove(float DeltaTime, float MaxA
 
 FVector UCharacterMovementComponent::GetActorFeetLocation() const
 {
-	const float HalfHeight = CharacterOwner->GetCapsuleComponent()->GetScaledCapsuleHalfHeight();
-	return UpdatedComponent->GetComponentLocation() + HalfHeight * GetGravityDirection();
+	const UCapsuleComponent* const CapsuleComponent = CharacterOwner ? CharacterOwner->GetCapsuleComponent() : Cast<UCapsuleComponent>(UpdatedComponent);
+	if (CapsuleComponent)
+	{
+		const float HalfHeight = CapsuleComponent->GetScaledCapsuleHalfHeight();
+		return UpdatedComponent->GetComponentLocation() + HalfHeight * GetGravityDirection();
+	}
+
+	return Super::GetActorFeetLocation();
 }
 
 void UCharacterMovementComponent::RequestDirectMove(const FVector& MoveVelocity, bool bForceMaxSpeed)
