Texture2D : avoid stall in GetCPUCopy if it will be null due to Availability

diff --git a/Engine/Source/Runtime/Engine/Private/Texture2D.cpp b/Engine/Source/Runtime/Engine/Private/Texture2D.cpp
--- a/Engine/Source/Runtime/Engine/Private/Texture2D.cpp
+++ b/Engine/Source/Runtime/Engine/Private/Texture2D.cpp
@@ -1315,6 +1315,12 @@ UTexture2D* UTexture2D::GetCPUCopyTexture()
 
 FSharedImageConstRef UTexture2D::GetCPUCopy() const
 {
+	// early out to avoid stalling when not a CPU texture
+	if ( Availability != ETextureAvailability::CPU)
+	{
+		return FSharedImageConstRef();
+	}
+
 	// could stall if texture isn't built!
 	const FTexturePlatformData* LocalPlatformData = GetPlatformData();
 
