Cook Indeterminism: Blueprints with ACameraActor can have an empty InstanceGuid when saved in the editor, and will create a new guid during cook. If the guid is indeterministic it will create cook indeterminism. Create the guid deterministically instead.
#rnx

diff --git a/Engine/Source/Runtime/Engine/Private/Actor.cpp b/Engine/Source/Runtime/Engine/Private/Actor.cpp
--- a/Engine/Source/Runtime/Engine/Private/Actor.cpp
+++ b/Engine/Source/Runtime/Engine/Private/Actor.cpp
@@ -893,7 +893,7 @@ void AActor::Serialize(FArchive& Ar)
 		}
 		else if (Ar.IsPersistent() && !ActorGuid.IsValid())
 		{
-			ActorGuid = FGuid::NewGuid();
+			ActorGuid = FGuid::NewDeterministicGuid(GetPathName());
 		}
 		else if ((Ar.GetPortFlags() & (PPF_Duplicate | PPF_DuplicateForPIE)) == PPF_Duplicate)
 		{
