Resample audio files that have an unsupported sample rate when we do other resampling.

diff --git a/Engine/Source/Developer/AudioFormatRad/Private/AudioFormatRad.cpp b/Engine/Source/Developer/AudioFormatRad/Private/AudioFormatRad.cpp
--- a/Engine/Source/Developer/AudioFormatRad/Private/AudioFormatRad.cpp
+++ b/Engine/Source/Developer/AudioFormatRad/Private/AudioFormatRad.cpp
@@ -50,6 +50,12 @@ public:
 		return true;
 	}
 
+	virtual TConstArrayView<float> GetSortedAllowedSampleRates() const
+	{
+		static const float ValidRates [] = { 24000, 32000, 44100, 48000} ;
+		return MakeArrayView<const float>(ValidRates, UE_ARRAY_COUNT(ValidRates));
+	}
+
 	virtual uint16 GetVersion(FName Format) const override
 	{
 		check(Format == NAME_RADA);
diff --git a/Engine/Source/Developer/TargetPlatform/Public/Interfaces/IAudioFormat.h b/Engine/Source/Developer/TargetPlatform/Public/Interfaces/IAudioFormat.h
--- a/Engine/Source/Developer/TargetPlatform/Public/Interfaces/IAudioFormat.h
+++ b/Engine/Source/Developer/TargetPlatform/Public/Interfaces/IAudioFormat.h
@@ -52,6 +52,15 @@ public:
 		return false;
 	}
 
+	/**
+	* If the format only allows a subset of sample rates, return them here. They must be in sorted order
+	* lowest to highest.
+	*/
+	virtual TConstArrayView<float> GetSortedAllowedSampleRates() const
+	{
+		return MakeArrayView<float>(nullptr, 0);
+	}
+
 	/**
 	 * Cooks the source data for the platform and stores the cooked data internally.
 	 *
diff --git a/Engine/Source/Runtime/Engine/Private/AudioDerivedData.cpp b/Engine/Source/Runtime/Engine/Private/AudioDerivedData.cpp
--- a/Engine/Source/Runtime/Engine/Private/AudioDerivedData.cpp
+++ b/Engine/Source/Runtime/Engine/Private/AudioDerivedData.cpp
@@ -1689,6 +1689,43 @@ static int32 CalculateModifiedCompressionQuality(int32 InQuality, float InQualit
 	return FMath::Clamp<int32>(FMath::FloorToInt(ModifiedCompressionQuality), 1, 100);
 }
 
+// Apply any overrides and force the final rate to be a valid one for the codec, if necessary.
+static float GetFinalSampleRate(int32 InWaveSampleRate, const FAudioCookInputs& InInputs)
+{
+	float FinalSampleRate = InWaveSampleRate;
+	if (InInputs.SampleRateOverride > 0)
+	{
+		FinalSampleRate = InInputs.SampleRateOverride;
+	}
+
+	if (TConstArrayView<float> AllowedSampleRates = InInputs.Compressor->GetSortedAllowedSampleRates(); AllowedSampleRates.Num())
+	{
+		// Favor next highest if possible to avoid cutting frequencies.
+		int32 SRIndex = 0;
+		for (; SRIndex < AllowedSampleRates.Num(); SRIndex++)
+		{
+			if (AllowedSampleRates[SRIndex] >= FinalSampleRate)
+			{
+				break;
+			}
+		}
+
+		if (SRIndex == AllowedSampleRates.Num())
+		{
+			SRIndex = AllowedSampleRates.Num() - 1;
+			UE_LOG(LogAudioDerivedData, Log, TEXT("Input sample rate (%.0f) exceeds allows sample rates for format %s, resampling to %.0f"), FinalSampleRate, *WriteToString<32>(InInputs.BaseFormat), AllowedSampleRates[SRIndex]);
+		}
+		else if (AllowedSampleRates[SRIndex] != FinalSampleRate)
+		{
+			UE_LOG(LogAudioDerivedData, Log, TEXT("Input sample rate (%.0f) doesn't match any allowed sample rates for format %s, resampling to %.0f"), FinalSampleRate, *WriteToString<32>(InInputs.BaseFormat), AllowedSampleRates[SRIndex]);
+		}
+
+		FinalSampleRate = AllowedSampleRates[SRIndex];
+	}
+	return FinalSampleRate;
+}
+
+
 /**
  * Cook a simple mono or stereo wave
  */
@@ -1739,7 +1776,9 @@ static void CookSimpleWave(const FAudioCookInputs& Inputs, TArray<uint8>& Output
 	int32 NumBytes = Input.Num();
 	int32 NumSamples = NumBytes / sizeof(int16);
 
-	const bool bNeedsResample = (Inputs.SampleRateOverride > 0 && Inputs.SampleRateOverride != (float)WaveSampleRate);
+	float FinalSampleRate = GetFinalSampleRate(WaveSampleRate, Inputs);
+
+	const bool bNeedsResample = FinalSampleRate != (float)WaveSampleRate;
 	const bool bNeedsToApplyWaveTransformation = (Inputs.WaveTransformations.Num() > 0);
 
 	// Only convert PCM wave data to float if needed. The conversion alters the sample
@@ -1780,9 +1819,9 @@ static void CookSimpleWave(const FAudioCookInputs& Inputs, TArray<uint8>& Output
 		// Resample if necessary
 		if (bNeedsResample)
 		{
-			ResampleWaveData(InputFloatBuffer, NumChannels, WaveSampleRate, Inputs.SampleRateOverride);
+			ResampleWaveData(InputFloatBuffer, NumChannels, WaveSampleRate, FinalSampleRate);
 			
-			WaveSampleRate = Inputs.SampleRateOverride;
+			WaveSampleRate = FinalSampleRate;
 			NumSamples = InputFloatBuffer.Num();
 		}
 
@@ -1971,8 +2010,10 @@ static void CookSurroundWave(const FAudioCookInputs& Inputs,  TArray<uint8>& Out
 	int32 WaveSampleRate = *WaveInfo.pSamplesPerSec;
 	int32 NumFrames = SampleDataSize / sizeof(int16);
 
+	const float FinalSampleRate = GetFinalSampleRate(WaveSampleRate, Inputs);
+
 	// bNeedsResample could change if a transformation changes the sample rate
-	bool bNeedsResample = Inputs.SampleRateOverride > 0 && Inputs.SampleRateOverride != (float)WaveSampleRate;
+	bool bNeedsResample = FinalSampleRate != (float)WaveSampleRate;
 	
 	const bool bContainsTransformations = Inputs.WaveTransformations.Num() > 0;
 	const bool bNeedsDeinterleave = bNeedsResample || bContainsTransformations;
@@ -1995,7 +2036,7 @@ static void CookSurroundWave(const FAudioCookInputs& Inputs,  TArray<uint8>& Out
 		Audio::ArrayInterleave(InputMultichannelBuffer, InterleavedFloatBuffer);
 
 		// run transformations
-		if(bContainsTransformations)
+		if (bContainsTransformations)
 		{
 			Audio::FWaveformTransformationWaveInfo TransformationInfo;
 
@@ -2010,7 +2051,7 @@ static void CookSurroundWave(const FAudioCookInputs& Inputs,  TArray<uint8>& Out
 
 			UE_CLOG(WaveSampleRate != TransformationInfo.SampleRate, LogAudioDerivedData, Warning, TEXT("Wave transformations which alter the sample rate are not supported. Cooked audio for %s may be incorrect"), *Inputs.SoundFullName);
 			WaveSampleRate = TransformationInfo.SampleRate;
-			bNeedsResample = Inputs.SampleRateOverride > 0 && Inputs.SampleRateOverride != (float)WaveSampleRate;
+			bNeedsResample = FinalSampleRate != (float)WaveSampleRate;
 			
 			UE_CLOG(ChannelCount != TransformationInfo.NumChannels, LogAudioDerivedData, Error, TEXT("Wave transformations which alter number of channels are not supported. Cooked audio for %s may be incorrect"), *Inputs.SoundFullName);
 			ChannelCount = TransformationInfo.NumChannels;
@@ -2020,9 +2061,9 @@ static void CookSurroundWave(const FAudioCookInputs& Inputs,  TArray<uint8>& Out
 
 		if (bNeedsResample)
 		{
-			ResampleWaveData(InterleavedFloatBuffer, ChannelCount, WaveSampleRate, Inputs.SampleRateOverride);
+			ResampleWaveData(InterleavedFloatBuffer, ChannelCount, WaveSampleRate, FinalSampleRate);
 			
-			WaveSampleRate = Inputs.SampleRateOverride;
+			WaveSampleRate = FinalSampleRate;
 			NumFrames = InterleavedFloatBuffer.Num() / ChannelCount;
 		}
 
