Make automatic transition rule account for transit durations smaller than delta time value

diff --git a/Engine/Source/Runtime/Engine/Private/Animation/AnimNode_StateMachine.cpp b/Engine/Source/Runtime/Engine/Private/Animation/AnimNode_StateMachine.cpp
--- a/Engine/Source/Runtime/Engine/Private/Animation/AnimNode_StateMachine.cpp
+++ b/Engine/Source/Runtime/Engine/Private/Animation/AnimNode_StateMachine.cpp
@@ -772,7 +772,27 @@ bool FAnimNode_StateMachine::FindValidTransition(const FAnimationUpdateContext&
 			{
 				if (const UAnimationAsset* AnimAsset = RelevantPlayer->GetAnimAsset())
 				{
-					const float AnimTimeRemaining = AnimAsset->GetPlayLength() - RelevantPlayer->GetAccumulatedTime();
+					float AnimTimeRemaining = AnimAsset->GetPlayLength() - RelevantPlayer->GetAccumulatedTime();
+
+					// Check whether the player looped last frame
+					if (RelevantPlayer->IsLooping() && AnimTimeRemaining > 0.f)
+					{
+						if (const FDeltaTimeRecord* DeltaTimeRecord = RelevantPlayer->GetDeltaTimeRecord())
+						{
+							if (DeltaTimeRecord->IsPreviousValid())
+							{
+								const float DeltaTimeFromPlayer = RelevantPlayer->GetAccumulatedTime() - DeltaTimeRecord->GetPrevious();
+								const float DeltaTimeFromRecord = DeltaTimeRecord->Delta;
+
+								// deal with negative play rates
+								if (DeltaTimeFromPlayer * DeltaTimeFromRecord < 0.f)
+								{
+									AnimTimeRemaining = 0.f;
+								}
+							}
+						}
+					}
+
 					const FAnimationTransitionBetweenStates& TransitionInfo = GetTransitionInfo(TransitionRule.TransitionIndex);
 				
 					// For transitions that go to a conduit the user is not able to edit the transition's cross fade duration,
