Change some errors to warnings to prevent misc texture issues from stopping an otherwise good cook.

diff --git a/Engine/Source/Runtime/Engine/Private/TextureDerivedDataTask.cpp b/Engine/Source/Runtime/Engine/Private/TextureDerivedDataTask.cpp
--- a/Engine/Source/Runtime/Engine/Private/TextureDerivedDataTask.cpp
+++ b/Engine/Source/Runtime/Engine/Private/TextureDerivedDataTask.cpp
@@ -1803,7 +1803,7 @@ bool DDC1_BuildTiledClassicTexture(
 			if (bHasLinearDerivedData && 
 				LinearDerivedData.TryLoadMipsWithSizes(0, LinearMipData, LinearMipSizes, TexturePathName) == false)
 			{
-				UE_LOG(LogTexture, Error, TEXT("Tiling texture build was unable to load the linear texture mips after a successful build, bad bug!: %s"), *TexturePathName);
+				UE_LOG(LogTexture, Warning, TEXT("Tiling texture build was unable to load the linear texture mips after a successful build, bad bug!: %s"), *TexturePathName);
 				return false;
 			}
 		}
@@ -1811,7 +1811,7 @@ bool DDC1_BuildTiledClassicTexture(
 
 	if (bHasLinearDerivedData == false)
 	{
-		UE_LOG(LogTexture, Error, TEXT("Tiling texture build was unable to fetch or build the linear texture source: %s"), *TexturePathName);
+		UE_LOG(LogTexture, Warning, TEXT("Tiling texture build was unable to fetch or build the linear texture source: %s"), *TexturePathName);
 		return false;
 	}
 
@@ -2118,7 +2118,7 @@ void FTextureCacheDerivedDataWorker::Finalize()
 
 		if (DDC1_IsTextureDataValid(TextureData, CompositeTextureData) == false)
 		{
-			UE_LOG(LogTexture, Error, TEXT("Unable to get texture source data for synchronous build of %s"), *TexturePathName);
+			UE_LOG(LogTexture, Warning, TEXT("Unable to get texture source data for synchronous build of %s"), *TexturePathName);
 		}
 		else
 		{
