Add CVar to toggle the random delay added to actors' NextUpdateTime.

diff --git a/Engine/Source/Runtime/Engine/Private/NetDriver.cpp b/Engine/Source/Runtime/Engine/Private/NetDriver.cpp
--- a/Engine/Source/Runtime/Engine/Private/NetDriver.cpp
+++ b/Engine/Source/Runtime/Engine/Private/NetDriver.cpp
@@ -442,6 +442,13 @@ static TAutoConsoleVariable<bool> CVarAddNetDriverInfoToNetAnalytics(
 	false,
 	TEXT("Automatically add NetDriver information to the NetAnalytics cache"));
 
+static int32 GNetDisableRandomNetUpdateDelay = 0;
+static FAutoConsoleVariableRef CVarNetDisableRandomNetUpdateDelay(
+	TEXT("net.DisableRandomNetUpdateDelay"),
+	GNetDisableRandomNetUpdateDelay,
+	TEXT("If 1, the server will not add a random delay to an actors' NextUpdateTime."),
+	ECVF_Default);
+
 namespace UE::Net
 {
 	static FString GRequiredEncryptionNetDriverDefNames_Internal = TEXT("all");
@@ -5028,8 +5035,10 @@ void UNetDriver::ServerReplicateActors_BuildConsiderList( TArray<FNetworkObjectI
 
 			const float NextUpdateDelta = bUseAdapativeNetFrequency ? ActorInfo->OptimalNetUpdateDelta : 1.0f / Actor->NetUpdateFrequency;
 
+			float RandDelay = GNetDisableRandomNetUpdateDelay ? 0.0f : UpdateDelayRandomStream.FRand() * ServerTickTime;
+
 			// then set the next update time
-			ActorInfo->NextUpdateTime = World->TimeSeconds + UpdateDelayRandomStream.FRand() * ServerTickTime + NextUpdateDelta;
+			ActorInfo->NextUpdateTime = World->TimeSeconds + RandDelay + NextUpdateDelta;
 
 			// and mark when the actor first requested an update
 			//@note: using ElapsedTime because it's compared against UActorChannel.LastUpdateTime which also uses that value
