* Updated order of operations in World Partition MoveExternalActorsToLevel in order to avoid iterating into Actors that may become invalid mid-cycle.

diff --git a/Engine/Source/Runtime/Engine/Private/WorldPartition/WorldPartitionLevelHelper.cpp b/Engine/Source/Runtime/Engine/Private/WorldPartition/WorldPartitionLevelHelper.cpp
--- a/Engine/Source/Runtime/Engine/Private/WorldPartition/WorldPartitionLevelHelper.cpp
+++ b/Engine/Source/Runtime/Engine/Private/WorldPartition/WorldPartitionLevelHelper.cpp
@@ -219,6 +219,8 @@ void FWorldPartitionLevelHelper::MoveExternalActorsToLevel(const TArray<FWorldPa
 	}
 
 	// Move all actors to Cell level
+	TSet<AActor*> LoadedActors;
+
 	for (const FWorldPartitionRuntimeCellObjectMapping& PackageObjectMapping : InChildPackages)
 	{
 		// We assume actor failed to duplicate if LoadedPath equals NAME_None (warning already logged we can skip this mapping)
@@ -258,7 +260,7 @@ void FWorldPartitionLevelHelper::MoveExternalActorsToLevel(const TArray<FWorldPa
 			}
 			else if (!InLevel->Actors.Contains(Actor))
 			{
-				InLevel->AddLoadedActor(Actor);
+				LoadedActors.Emplace(Actor);
 			}
 			check(Actor->GetPackage() == LevelPackage);
 
@@ -282,7 +284,7 @@ void FWorldPartitionLevelHelper::MoveExternalActorsToLevel(const TArray<FWorldPa
 						}
 						else if (NestedActor && !InLevel->Actors.Contains(NestedActor))
 						{
-							InLevel->AddLoadedActor(NestedActor);
+							LoadedActors.Emplace(NestedActor);
 						}
 						if (NestedActor)
 						{
@@ -311,9 +313,11 @@ void FWorldPartitionLevelHelper::MoveExternalActorsToLevel(const TArray<FWorldPa
 		}
 	}
 
+	InLevel->AddLoadedActors(LoadedActors.Array());
+
 	for (AActor* Actor : InLevel->Actors)
 	{
-		if (Actor && Actor->HasAllFlags(RF_WasLoaded))
+		if (IsValid(Actor) && Actor->HasAllFlags(RF_WasLoaded))
 		{
 			checkf(LevelActors.Contains(Actor->GetFName()), TEXT("Actor %s(%s) was unexpectedly loaded when moving actors to streaming cell"), *Actor->GetActorNameOrLabel(), *Actor->GetName());
 		}
