Reset OverlayColor state on PrimitiveComponent when removing OverlayColor
* Also fix issue where WantsEditorEffects was never set false when OverlayColorOverride is removed.

diff --git a/Engine/Source/Runtime/Engine/Private/Components/PrimitiveComponent.cpp b/Engine/Source/Runtime/Engine/Private/Components/PrimitiveComponent.cpp
--- a/Engine/Source/Runtime/Engine/Private/Components/PrimitiveComponent.cpp
+++ b/Engine/Source/Runtime/Engine/Private/Components/PrimitiveComponent.cpp
@@ -2007,11 +2007,12 @@ void UPrimitiveComponent::SetOverlayColor(FColor InOverlayColor)
 
 void UPrimitiveComponent::RemoveOverlayColor()
 {
-	bWantsEditorEffects = true;
+	OverlayColor = FColor(ForceInitToZero);
+	bWantsEditorEffects = false;
 
 	if (SceneProxy)
 	{
-		SceneProxy->SetOverlayColor_GameThread(FColor{EForceInit::ForceInitToZero});
+		SceneProxy->SetOverlayColor_GameThread(OverlayColor);
 		SceneProxy->SetSelectionOverride_GameThread(false);
 	}
 }
