[Content Bundle] Added an optional flag to allow marking content bundles as "required", which will make them unremovable and bake their actors into the main map to reduce IO stress.

diff --git a/Engine/Source/Runtime/Engine/Private/WorldPartition/ContentBundle/ContentBundleDescriptor.cpp b/Engine/Source/Runtime/Engine/Private/WorldPartition/ContentBundle/ContentBundleDescriptor.cpp
--- a/Engine/Source/Runtime/Engine/Private/WorldPartition/ContentBundle/ContentBundleDescriptor.cpp
+++ b/Engine/Source/Runtime/Engine/Private/WorldPartition/ContentBundle/ContentBundleDescriptor.cpp
@@ -8,10 +8,9 @@
 #include UE_INLINE_GENERATED_CPP_BY_NAME(ContentBundleDescriptor)
 
 UContentBundleDescriptor::UContentBundleDescriptor(const FObjectInitializer& ObjectInitializer)
-	: DebugColor(FColor::Black)
-{
-
-}
+	: bRequired(false)
+	, DebugColor(FColor::Black)
+{}
 
 FString UContentBundleDescriptor::GetPackageRoot() const
 {
diff --git a/Engine/Source/Runtime/Engine/Private/WorldPartition/ContentBundle/ContentBundleEditor.cpp b/Engine/Source/Runtime/Engine/Private/WorldPartition/ContentBundle/ContentBundleEditor.cpp
--- a/Engine/Source/Runtime/Engine/Private/WorldPartition/ContentBundle/ContentBundleEditor.cpp
+++ b/Engine/Source/Runtime/Engine/Private/WorldPartition/ContentBundle/ContentBundleEditor.cpp
@@ -93,6 +93,8 @@ void FContentBundleEditor::DoInjectContent()
 
 			check(GetDescriptor()->GetGuid().IsValid());
 
+			ActorDescContainerInstance->bIsContentBundleRequired = GetDescriptor()->GetRequired();
+
 			if (!ActorDescContainerInstance->GetContainer()->IsEmpty())
 			{
 				WorldDataLayersActorReference = FWorldDataLayersReference(ActorDescContainerInstance.Get(), BuildWorlDataLayersName());
@@ -346,6 +348,11 @@ void FContentBundleEditor::GenerateStreaming(TArray<FString>* OutPackageToGenera
 		return;
 	}
 
+	if (GetDescriptor()->GetRequired())
+	{
+		return;
+	}
+
 	FActorDescContainerInstanceCollection Collection({ TObjectPtr<UActorDescContainerInstance>(ActorDescContainerInstance.Get()) });
 	UWorldPartition::FGenerateStreamingParams Params = UWorldPartition::FGenerateStreamingParams()
 		.SetContainerInstanceCollection(Collection, FStreamingGenerationContainerInstanceCollection::ECollectionType::BaseAsContentBundle);
@@ -388,16 +395,22 @@ void FContentBundleEditor::GenerateStreaming(TArray<FString>* OutPackageToGenera
 
 void FContentBundleEditor::OnBeginCook(IWorldPartitionCookPackageContext& CookContext)
 {
-	check(!bIsInCook);
-	CookContext.RegisterPackageCookPackageGenerator(this);
-	bIsInCook = true;
+	if (!GetDescriptor()->GetRequired())
+	{
+		check(!bIsInCook);
+		CookContext.RegisterPackageCookPackageGenerator(this);
+		bIsInCook = true;
+	}
 }
 
 void FContentBundleEditor::OnEndCook(IWorldPartitionCookPackageContext& CookContext)
 {
-	check(bIsInCook);
-	CookContext.UnregisterPackageCookPackageGenerator(this);
-	bIsInCook = false;
+	if (!GetDescriptor()->GetRequired())
+	{
+		check(bIsInCook);
+		CookContext.UnregisterPackageCookPackageGenerator(this);
+		bIsInCook = false;
+	}
 }
 
 bool FContentBundleEditor::GatherPackagesToCook(class IWorldPartitionCookPackageContext& CookContext)
diff --git a/Engine/Source/Runtime/Engine/Private/WorldPartition/RuntimeSpatialHash/RuntimeSpatialHashHLOD.cpp b/Engine/Source/Runtime/Engine/Private/WorldPartition/RuntimeSpatialHash/RuntimeSpatialHashHLOD.cpp
--- a/Engine/Source/Runtime/Engine/Private/WorldPartition/RuntimeSpatialHash/RuntimeSpatialHashHLOD.cpp
+++ b/Engine/Source/Runtime/Engine/Private/WorldPartition/RuntimeSpatialHash/RuntimeSpatialHashHLOD.cpp
@@ -396,7 +396,7 @@ static void UpdateHLODGridsActors(UWorld* World, const UActorDescContainerInstan
 		ASpatialHashRuntimeGridInfo* GridActor = ExistingGridActors.FindRef(GridName);
 		if (!GridActor)
 		{
-			FContentBundleActivationScope ContentBndleScope(ContainerInstance->GetContentBundleGuid());
+			FContentBundleActivationScope ContentBundleScope(ContainerInstance->IsContentBundleRequired() ? FGuid() : ContainerInstance->GetContentBundleGuid());
 			FScopedOverrideSpawningLevelMountPointObject EDLScope(ContainerInstance->GetExternalDataLayerAsset());
 
 			FActorSpawnParameters SpawnParams;
@@ -405,7 +405,7 @@ static void UpdateHLODGridsActors(UWorld* World, const UActorDescContainerInstan
 			GridActor->Tags.Add(HLODGridTag);
 			bDirty = true;
 
-			check(GridActor->GetContentBundleGuid() == ContainerInstance->GetContentBundleGuid());
+			check(GridActor->GetContentBundleGuid() == (ContainerInstance->IsContentBundleRequired() ? FGuid() : ContainerInstance->GetContentBundleGuid()));
 			check(GridActor->GetExternalDataLayerAsset() == ContainerInstance->GetExternalDataLayerAsset());
 		}
 
diff --git a/Engine/Source/Runtime/Engine/Private/WorldPartition/WorldPartitionStreamingGeneration.cpp b/Engine/Source/Runtime/Engine/Private/WorldPartition/WorldPartitionStreamingGeneration.cpp
--- a/Engine/Source/Runtime/Engine/Private/WorldPartition/WorldPartitionStreamingGeneration.cpp
+++ b/Engine/Source/Runtime/Engine/Private/WorldPartition/WorldPartitionStreamingGeneration.cpp
@@ -2077,7 +2077,7 @@ static void ExtractContentBundleContainerInstances(const FActorDescContainerInst
 {
 	InContainerInstanceCollection->ForEachActorDescContainerInstance([&OutContentBundleContainerInstances, &OutNonContentBundleContainerInstances](const UActorDescContainerInstance* InActorDescContainerInstance)
 	{
-		if (InActorDescContainerInstance->GetContentBundleGuid().IsValid())
+		if (InActorDescContainerInstance->GetContentBundleGuid().IsValid() && !InActorDescContainerInstance->IsContentBundleRequired())
 		{
 			OutContentBundleContainerInstances.Add(InActorDescContainerInstance);
 		}
@@ -2234,6 +2234,12 @@ void FStreamingGenerationContainerInstanceCollection::InitializeCollection()
 	{
 		check(GetActorDescContainerCount() == 1);
 		check(GetContentBundleGuid().IsValid());
+
+		ActorDescContainerInstanceCollection.SetNum(Algo::RemoveIf(ActorDescContainerInstanceCollection, [](const UActorDescContainerInstance* ActorDescContainerInstance)
+		{
+			return ActorDescContainerInstance->GetContentBundleGuid().IsValid() && ActorDescContainerInstance->bIsContentBundleRequired;
+		}));
+
 		ContentBundleStartIdx = 0;
 		return;
 	}
@@ -2246,7 +2252,10 @@ void FStreamingGenerationContainerInstanceCollection::InitializeCollection()
 		// When type is set to BaseAndEDL, we remove ContentBundle containers from the collection.
 		// BaseAndEDL type assumes ContentBundle containers are generated separately one at a time.
 		check(!ShouldRegisterDelegates());
-		ActorDescContainerInstanceCollection.SetNum(Algo::RemoveIf(ActorDescContainerInstanceCollection, [](const UActorDescContainerInstance* ActorDescContainerInstance) { return ActorDescContainerInstance->GetContentBundleGuid().IsValid(); }));
+		ActorDescContainerInstanceCollection.SetNum(Algo::RemoveIf(ActorDescContainerInstanceCollection, [](const UActorDescContainerInstance* ActorDescContainerInstance)
+		{
+			return ActorDescContainerInstance->GetContentBundleGuid().IsValid() && !ActorDescContainerInstance->bIsContentBundleRequired;
+		}));
 	}
 
 	int32 BaseContainerCount = Algo::CountIf(ActorDescContainerInstanceCollection, [](const UActorDescContainerInstance* ActorDescContainerInstance) { return !ActorDescContainerInstance->HasExternalContent(); });
diff --git a/Engine/Source/Runtime/Engine/Public/WorldPartition/ActorDescContainerInstance.h b/Engine/Source/Runtime/Engine/Public/WorldPartition/ActorDescContainerInstance.h
--- a/Engine/Source/Runtime/Engine/Public/WorldPartition/ActorDescContainerInstance.h
+++ b/Engine/Source/Runtime/Engine/Public/WorldPartition/ActorDescContainerInstance.h
@@ -14,6 +14,7 @@ class UActorDescContainer;
 class FWorldPartitionActorDesc;
 class UWorldPartition;
 class UExternalDataLayerAsset;
+class FContentBundleEditor;
 struct FWorldPartitionRuntimeCellPropertyOverride;
 
 class FActorDescInstanceList : public TActorDescList<FWorldPartitionActorDescInstance> { };
@@ -29,11 +30,14 @@ protected:
 		: bIsInitialized(false) 
 		, bRegisteredDelegates(false)
 		, bCreateChildContainerHierarchy(false)
+		, bIsContentBundleRequired(false)
 #endif
 	{}
 
 #if WITH_EDITOR
+	friend class FStreamingGenerationContainerInstanceCollection;
 	friend FWorldPartitionActorDescInstance;
+	friend FContentBundleEditor;
 	friend UWorldPartition;
 		
 public:
@@ -148,6 +152,9 @@ public:
 	ENGINE_API const TMap<FGuid, TObjectPtr<UActorDescContainerInstance>>& GetChildContainerInstances() const { return ChildContainerInstances; }
 		
 	const UActorDescContainerInstance* GetParentContainerInstance() const { return ParentContainerInstance.Get(); }
+
+	bool IsContentBundleRequired() const { return bIsContentBundleRequired; }
+
 protected:
 	virtual void RegisterContainer(const FInitializeParams& InParams);
 	virtual void UnregisterContainer();
@@ -204,5 +211,6 @@ protected:
 	bool														bIsInitialized;
 	bool														bRegisteredDelegates;
 	bool														bCreateChildContainerHierarchy;
+	bool														bIsContentBundleRequired;
 #endif
 };
\ No newline at end of file
diff --git a/Engine/Source/Runtime/Engine/Public/WorldPartition/ContentBundle/ContentBundleDescriptor.h b/Engine/Source/Runtime/Engine/Public/WorldPartition/ContentBundle/ContentBundleDescriptor.h
--- a/Engine/Source/Runtime/Engine/Public/WorldPartition/ContentBundle/ContentBundleDescriptor.h
+++ b/Engine/Source/Runtime/Engine/Public/WorldPartition/ContentBundle/ContentBundleDescriptor.h
@@ -18,6 +18,7 @@ class UContentBundleDescriptor : public UObject
 
 public:
 	const FString& GetDisplayName() const { return DisplayName; }
+	bool GetRequired() const { return bRequired; }
 	const FColor& GetDebugColor() const { return DebugColor; }
 	ENGINE_API FString GetPackageRoot() const;
 	const FGuid& GetGuid() const { return Guid; }
@@ -42,6 +43,9 @@ private:
 	UPROPERTY(EditAnywhere, Category = BaseInformation)
 	FString DisplayName;
 
+	UPROPERTY(EditAnywhere, Category = BaseInformation)
+	bool bRequired;
+
 	UPROPERTY(EditAnywhere, DuplicateTransient, Category = BaseInformation)
 	FColor DebugColor;
 
diff --git a/Engine/Source/Runtime/Engine/Public/WorldPartition/RuntimeSpatialHash/RuntimeSpatialHashGridHelper.h b/Engine/Source/Runtime/Engine/Public/WorldPartition/RuntimeSpatialHash/RuntimeSpatialHashGridHelper.h
--- a/Engine/Source/Runtime/Engine/Public/WorldPartition/RuntimeSpatialHash/RuntimeSpatialHashGridHelper.h
+++ b/Engine/Source/Runtime/Engine/Public/WorldPartition/RuntimeSpatialHash/RuntimeSpatialHashGridHelper.h
@@ -303,7 +303,8 @@ struct FSquare2DGridHelper
 			void AddActorSetInstance(const IStreamingGenerationContext::FActorSetInstance* ActorSetInstance)
 			{
 				const FDataLayersID DataLayersID = FDataLayersID(ActorSetInstance->DataLayers);
-				FGridCellDataChunk& ActorDataChunk = DataChunks.FindOrAddByHash(DataLayersID.GetHash(), FGridCellDataChunk(ActorSetInstance->DataLayers, ActorSetInstance->ContentBundleID));
+				const FGuid ContentBundleID = (ActorSetInstance->ActorSetContainerInstance->ContainerInstanceCollection->GetCollectionType() == FStreamingGenerationContainerInstanceCollection::ECollectionType::BaseAsContentBundle) ? ActorSetInstance->ContentBundleID : FGuid();
+				FGridCellDataChunk& ActorDataChunk = DataChunks.FindOrAddByHash(DataLayersID.GetHash(), FGridCellDataChunk(ActorSetInstance->DataLayers, ContentBundleID));
 				ActorDataChunk.AddActorSetInstance(ActorSetInstance);
 			}
 
