[Backout] - CL33705721
Original CL Desc
-----------------------------------------------------------------
[Content Bundle] Added an optional flag to allow marking content bundles as "required", which will make them unremovable and bake their actors into the main map to reduce IO stress.

diff --git a/Engine/Source/Runtime/Engine/Private/WorldPartition/ContentBundle/ContentBundleDescriptor.cpp b/Engine/Source/Runtime/Engine/Private/WorldPartition/ContentBundle/ContentBundleDescriptor.cpp
--- a/Engine/Source/Runtime/Engine/Private/WorldPartition/ContentBundle/ContentBundleDescriptor.cpp
+++ b/Engine/Source/Runtime/Engine/Private/WorldPartition/ContentBundle/ContentBundleDescriptor.cpp
@@ -8,9 +8,10 @@
 #include UE_INLINE_GENERATED_CPP_BY_NAME(ContentBundleDescriptor)
 
 UContentBundleDescriptor::UContentBundleDescriptor(const FObjectInitializer& ObjectInitializer)
-	: bRequired(false)
-	, DebugColor(FColor::Black)
-{}
+	: DebugColor(FColor::Black)
+{
+
+}
 
 FString UContentBundleDescriptor::GetPackageRoot() const
 {
diff --git a/Engine/Source/Runtime/Engine/Private/WorldPartition/ContentBundle/ContentBundleEditor.cpp b/Engine/Source/Runtime/Engine/Private/WorldPartition/ContentBundle/ContentBundleEditor.cpp
--- a/Engine/Source/Runtime/Engine/Private/WorldPartition/ContentBundle/ContentBundleEditor.cpp
+++ b/Engine/Source/Runtime/Engine/Private/WorldPartition/ContentBundle/ContentBundleEditor.cpp
@@ -93,8 +93,6 @@ void FContentBundleEditor::DoInjectContent()
 
 			check(GetDescriptor()->GetGuid().IsValid());
 
-			ActorDescContainerInstance->bIsContentBundleRequired = GetDescriptor()->GetRequired();
-
 			if (!ActorDescContainerInstance->GetContainer()->IsEmpty())
 			{
 				WorldDataLayersActorReference = FWorldDataLayersReference(ActorDescContainerInstance.Get(), BuildWorlDataLayersName());
@@ -348,11 +346,6 @@ void FContentBundleEditor::GenerateStreaming(TArray<FString>* OutPackageToGenera
 		return;
 	}
 
-	if (GetDescriptor()->GetRequired())
-	{
-		return;
-	}
-
 	FActorDescContainerInstanceCollection Collection({ TObjectPtr<UActorDescContainerInstance>(ActorDescContainerInstance.Get()) });
 	UWorldPartition::FGenerateStreamingParams Params = UWorldPartition::FGenerateStreamingParams()
 		.SetContainerInstanceCollection(Collection, FStreamingGenerationContainerInstanceCollection::ECollectionType::BaseAsContentBundle);
@@ -395,22 +388,16 @@ void FContentBundleEditor::GenerateStreaming(TArray<FString>* OutPackageToGenera
 
 void FContentBundleEditor::OnBeginCook(IWorldPartitionCookPackageContext& CookContext)
 {
-	if (!GetDescriptor()->GetRequired())
-	{
-		check(!bIsInCook);
-		CookContext.RegisterPackageCookPackageGenerator(this);
-		bIsInCook = true;
-	}
+	check(!bIsInCook);
+	CookContext.RegisterPackageCookPackageGenerator(this);
+	bIsInCook = true;
 }
 
 void FContentBundleEditor::OnEndCook(IWorldPartitionCookPackageContext& CookContext)
 {
-	if (!GetDescriptor()->GetRequired())
-	{
-		check(bIsInCook);
-		CookContext.UnregisterPackageCookPackageGenerator(this);
-		bIsInCook = false;
-	}
+	check(bIsInCook);
+	CookContext.UnregisterPackageCookPackageGenerator(this);
+	bIsInCook = false;
 }
 
 bool FContentBundleEditor::GatherPackagesToCook(class IWorldPartitionCookPackageContext& CookContext)
diff --git a/Engine/Source/Runtime/Engine/Private/WorldPartition/RuntimeSpatialHash/RuntimeSpatialHashHLOD.cpp b/Engine/Source/Runtime/Engine/Private/WorldPartition/RuntimeSpatialHash/RuntimeSpatialHashHLOD.cpp
--- a/Engine/Source/Runtime/Engine/Private/WorldPartition/RuntimeSpatialHash/RuntimeSpatialHashHLOD.cpp
+++ b/Engine/Source/Runtime/Engine/Private/WorldPartition/RuntimeSpatialHash/RuntimeSpatialHashHLOD.cpp
@@ -396,7 +396,7 @@ static void UpdateHLODGridsActors(UWorld* World, const UActorDescContainerInstan
 		ASpatialHashRuntimeGridInfo* GridActor = ExistingGridActors.FindRef(GridName);
 		if (!GridActor)
 		{
-			FContentBundleActivationScope ContentBundleScope(ContainerInstance->IsContentBundleRequired() ? FGuid() : ContainerInstance->GetContentBundleGuid());
+			FContentBundleActivationScope ContentBndleScope(ContainerInstance->GetContentBundleGuid());
 			FScopedOverrideSpawningLevelMountPointObject EDLScope(ContainerInstance->GetExternalDataLayerAsset());
 
 			FActorSpawnParameters SpawnParams;
@@ -405,7 +405,7 @@ static void UpdateHLODGridsActors(UWorld* World, const UActorDescContainerInstan
 			GridActor->Tags.Add(HLODGridTag);
 			bDirty = true;
 
-			check(GridActor->GetContentBundleGuid() == (ContainerInstance->IsContentBundleRequired() ? FGuid() : ContainerInstance->GetContentBundleGuid()));
+			check(GridActor->GetContentBundleGuid() == ContainerInstance->GetContentBundleGuid());
 			check(GridActor->GetExternalDataLayerAsset() == ContainerInstance->GetExternalDataLayerAsset());
 		}
 
diff --git a/Engine/Source/Runtime/Engine/Private/WorldPartition/WorldPartitionStreamingGeneration.cpp b/Engine/Source/Runtime/Engine/Private/WorldPartition/WorldPartitionStreamingGeneration.cpp
--- a/Engine/Source/Runtime/Engine/Private/WorldPartition/WorldPartitionStreamingGeneration.cpp
+++ b/Engine/Source/Runtime/Engine/Private/WorldPartition/WorldPartitionStreamingGeneration.cpp
@@ -2077,7 +2077,7 @@ static void ExtractContentBundleContainerInstances(const FActorDescContainerInst
 {
 	InContainerInstanceCollection->ForEachActorDescContainerInstance([&OutContentBundleContainerInstances, &OutNonContentBundleContainerInstances](const UActorDescContainerInstance* InActorDescContainerInstance)
 	{
-		if (InActorDescContainerInstance->GetContentBundleGuid().IsValid() && !InActorDescContainerInstance->IsContentBundleRequired())
+		if (InActorDescContainerInstance->GetContentBundleGuid().IsValid())
 		{
 			OutContentBundleContainerInstances.Add(InActorDescContainerInstance);
 		}
@@ -2234,12 +2234,6 @@ void FStreamingGenerationContainerInstanceCollection::InitializeCollection()
 	{
 		check(GetActorDescContainerCount() == 1);
 		check(GetContentBundleGuid().IsValid());
-
-		ActorDescContainerInstanceCollection.SetNum(Algo::RemoveIf(ActorDescContainerInstanceCollection, [](const UActorDescContainerInstance* ActorDescContainerInstance)
-		{
-			return ActorDescContainerInstance->GetContentBundleGuid().IsValid() && ActorDescContainerInstance->bIsContentBundleRequired;
-		}));
-
 		ContentBundleStartIdx = 0;
 		return;
 	}
@@ -2252,10 +2246,7 @@ void FStreamingGenerationContainerInstanceCollection::InitializeCollection()
 		// When type is set to BaseAndEDL, we remove ContentBundle containers from the collection.
 		// BaseAndEDL type assumes ContentBundle containers are generated separately one at a time.
 		check(!ShouldRegisterDelegates());
-		ActorDescContainerInstanceCollection.SetNum(Algo::RemoveIf(ActorDescContainerInstanceCollection, [](const UActorDescContainerInstance* ActorDescContainerInstance)
-		{
-			return ActorDescContainerInstance->GetContentBundleGuid().IsValid() && !ActorDescContainerInstance->bIsContentBundleRequired;
-		}));
+		ActorDescContainerInstanceCollection.SetNum(Algo::RemoveIf(ActorDescContainerInstanceCollection, [](const UActorDescContainerInstance* ActorDescContainerInstance) { return ActorDescContainerInstance->GetContentBundleGuid().IsValid(); }));
 	}
 
 	int32 BaseContainerCount = Algo::CountIf(ActorDescContainerInstanceCollection, [](const UActorDescContainerInstance* ActorDescContainerInstance) { return !ActorDescContainerInstance->HasExternalContent(); });
diff --git a/Engine/Source/Runtime/Engine/Public/WorldPartition/ActorDescContainerInstance.h b/Engine/Source/Runtime/Engine/Public/WorldPartition/ActorDescContainerInstance.h
--- a/Engine/Source/Runtime/Engine/Public/WorldPartition/ActorDescContainerInstance.h
+++ b/Engine/Source/Runtime/Engine/Public/WorldPartition/ActorDescContainerInstance.h
@@ -14,7 +14,6 @@ class UActorDescContainer;
 class FWorldPartitionActorDesc;
 class UWorldPartition;
 class UExternalDataLayerAsset;
-class FContentBundleEditor;
 struct FWorldPartitionRuntimeCellPropertyOverride;
 
 class FActorDescInstanceList : public TActorDescList<FWorldPartitionActorDescInstance> { };
@@ -30,14 +29,11 @@ protected:
 		: bIsInitialized(false) 
 		, bRegisteredDelegates(false)
 		, bCreateChildContainerHierarchy(false)
-		, bIsContentBundleRequired(false)
 #endif
 	{}
 
 #if WITH_EDITOR
-	friend class FStreamingGenerationContainerInstanceCollection;
 	friend FWorldPartitionActorDescInstance;
-	friend FContentBundleEditor;
 	friend UWorldPartition;
 		
 public:
@@ -152,9 +148,6 @@ public:
 	ENGINE_API const TMap<FGuid, TObjectPtr<UActorDescContainerInstance>>& GetChildContainerInstances() const { return ChildContainerInstances; }
 		
 	const UActorDescContainerInstance* GetParentContainerInstance() const { return ParentContainerInstance.Get(); }
-
-	bool IsContentBundleRequired() const { return bIsContentBundleRequired; }
-
 protected:
 	virtual void RegisterContainer(const FInitializeParams& InParams);
 	virtual void UnregisterContainer();
@@ -211,6 +204,5 @@ protected:
 	bool														bIsInitialized;
 	bool														bRegisteredDelegates;
 	bool														bCreateChildContainerHierarchy;
-	bool														bIsContentBundleRequired;
 #endif
 };
\ No newline at end of file
diff --git a/Engine/Source/Runtime/Engine/Public/WorldPartition/ContentBundle/ContentBundleDescriptor.h b/Engine/Source/Runtime/Engine/Public/WorldPartition/ContentBundle/ContentBundleDescriptor.h
--- a/Engine/Source/Runtime/Engine/Public/WorldPartition/ContentBundle/ContentBundleDescriptor.h
+++ b/Engine/Source/Runtime/Engine/Public/WorldPartition/ContentBundle/ContentBundleDescriptor.h
@@ -18,7 +18,6 @@ class UContentBundleDescriptor : public UObject
 
 public:
 	const FString& GetDisplayName() const { return DisplayName; }
-	bool GetRequired() const { return bRequired; }
 	const FColor& GetDebugColor() const { return DebugColor; }
 	ENGINE_API FString GetPackageRoot() const;
 	const FGuid& GetGuid() const { return Guid; }
@@ -43,9 +42,6 @@ private:
 	UPROPERTY(EditAnywhere, Category = BaseInformation)
 	FString DisplayName;
 
-	UPROPERTY(EditAnywhere, Category = BaseInformation)
-	bool bRequired;
-
 	UPROPERTY(EditAnywhere, DuplicateTransient, Category = BaseInformation)
 	FColor DebugColor;
 
diff --git a/Engine/Source/Runtime/Engine/Public/WorldPartition/RuntimeSpatialHash/RuntimeSpatialHashGridHelper.h b/Engine/Source/Runtime/Engine/Public/WorldPartition/RuntimeSpatialHash/RuntimeSpatialHashGridHelper.h
--- a/Engine/Source/Runtime/Engine/Public/WorldPartition/RuntimeSpatialHash/RuntimeSpatialHashGridHelper.h
+++ b/Engine/Source/Runtime/Engine/Public/WorldPartition/RuntimeSpatialHash/RuntimeSpatialHashGridHelper.h
@@ -303,8 +303,7 @@ struct FSquare2DGridHelper
 			void AddActorSetInstance(const IStreamingGenerationContext::FActorSetInstance* ActorSetInstance)
 			{
 				const FDataLayersID DataLayersID = FDataLayersID(ActorSetInstance->DataLayers);
-				const FGuid ContentBundleID = (ActorSetInstance->ActorSetContainerInstance->ContainerInstanceCollection->GetCollectionType() == FStreamingGenerationContainerInstanceCollection::ECollectionType::BaseAsContentBundle) ? ActorSetInstance->ContentBundleID : FGuid();
-				FGridCellDataChunk& ActorDataChunk = DataChunks.FindOrAddByHash(DataLayersID.GetHash(), FGridCellDataChunk(ActorSetInstance->DataLayers, ContentBundleID));
+				FGridCellDataChunk& ActorDataChunk = DataChunks.FindOrAddByHash(DataLayersID.GetHash(), FGridCellDataChunk(ActorSetInstance->DataLayers, ActorSetInstance->ContentBundleID));
 				ActorDataChunk.AddActorSetInstance(ActorSetInstance);
 			}
 
diff --git a/Engine/Source/Runtime/Engine/Public/WorldPartition/WorldPartitionStreamingGeneration.h b/Engine/Source/Runtime/Engine/Public/WorldPartition/WorldPartitionStreamingGeneration.h
--- a/Engine/Source/Runtime/Engine/Public/WorldPartition/WorldPartitionStreamingGeneration.h
+++ b/Engine/Source/Runtime/Engine/Public/WorldPartition/WorldPartitionStreamingGeneration.h
@@ -243,10 +243,9 @@ public:
 	ENGINE_API FName GetBaseContainerInstancePackageName() const;
 	ENGINE_API FGuid GetContentBundleGuid() const; // @todo_ow : Remove once conversion to ExternalDataLayer is complete. 
 
-	ECollectionType GetCollectionType() const { return CollectionType; }
-
 private:
-	void SetCollectionType(const ECollectionType& InCollectionType) { check(CollectionType == ECollectionType::Invalid); CollectionType = InCollectionType; }	
+	void SetCollectionType(const ECollectionType& InCollectionType) { check(CollectionType == ECollectionType::Invalid); CollectionType = InCollectionType; }
+	ECollectionType GetCollectionType() const { return CollectionType; }
 	void InitializeCollection();
 	FName GetBaseContainerPackageName() const;
 	TArrayView<const UActorDescContainerInstance* const> GetExternalDataLayerContainerInstances() const;
