Fix skeletal meshes rendered as static not marking the relevant RayTracing::GeometryGroupHandle as referenced.
- Register a RayTracingGeometryGroup per FSkeletalMeshRenderData.
- Assign StaticRayTracingGeometry in FSkeletalMeshLODRenderData to the relevant group.
- Modified FSkeletalMeshSceneProxy to override GetRayTracingGeometryGroupHandle().

diff --git a/Engine/Source/Runtime/Engine/Private/SkeletalMesh.cpp b/Engine/Source/Runtime/Engine/Private/SkeletalMesh.cpp
--- a/Engine/Source/Runtime/Engine/Private/SkeletalMesh.cpp
+++ b/Engine/Source/Runtime/Engine/Private/SkeletalMesh.cpp
@@ -6939,12 +6939,29 @@ void FSkeletalMeshSceneProxy::GetDynamicElementsSection(const TArray<const FScen
 	}
 }
 
+void FSkeletalMeshSceneProxy::CreateRenderThreadResources(FRHICommandListBase& RHICmdList)
+{
+#if RHI_RAYTRACING
+	if (IsRayTracingAllowed())
+	{
+		// copy RayTracingGeometryGroupHandle from FSkeletalMeshRenderData since USkeletalMesh can be released before the proxy is destroyed
+		RayTracingGeometryGroupHandle = SkeletalMeshRenderData->RayTracingGeometryGroupHandle;
+	}
+#endif
+}
+
 #if RHI_RAYTRACING
 bool FSkeletalMeshSceneProxy::HasRayTracingRepresentation() const
 {
 	return bRenderStatic;
 }
 
+RayTracing::GeometryGroupHandle FSkeletalMeshSceneProxy::GetRayTracingGeometryGroupHandle() const
+{
+	check(IsInRenderingThread() || IsInParallelRenderingThread());
+	return RayTracingGeometryGroupHandle;
+}
+
 TArray<FRayTracingGeometry*> FSkeletalMeshSceneProxy::GetStaticRayTracingGeometries() const
 {
 	if (IsRayTracingAllowed() && bRenderStatic)
diff --git a/Engine/Source/Runtime/Engine/Private/SkeletalMeshRenderData.cpp b/Engine/Source/Runtime/Engine/Private/SkeletalMeshRenderData.cpp
--- a/Engine/Source/Runtime/Engine/Private/SkeletalMeshRenderData.cpp
+++ b/Engine/Source/Runtime/Engine/Private/SkeletalMeshRenderData.cpp
@@ -10,6 +10,7 @@
 #include "EngineLogs.h"
 #include "UObject/Package.h"
 #include "Rendering/RenderCommandPipes.h"
+#include "Rendering/RayTracingGeometryManager.h"
 
 #if WITH_EDITOR
 #include "ProfilingDebugging/CookStats.h"
@@ -618,6 +619,17 @@ void FSkeletalMeshRenderData::InitResources(bool bNeedsVertexColors, TArray<UMor
 {
 	if (!bInitialized)
 	{
+#if RHI_RAYTRACING
+		if (IsRayTracingAllowed())
+		{
+			ENQUEUE_RENDER_COMMAND(SkeletalMeshInitRayTracingGeometryGroup)(UE::RenderCommandPipe::SkeletalMesh,
+				[this, NumLODs = LODRenderData.Num()]
+				{
+					RayTracingGeometryGroupHandle = GRayTracingGeometryManager->RegisterRayTracingGeometryGroup(NumLODs);
+				});
+		}
+#endif
+
 		// initialize resources for each lod
 		for (int32 LODIndex = 0; LODIndex < LODRenderData.Num(); LODIndex++)
 		{
@@ -652,6 +664,18 @@ void FSkeletalMeshRenderData::ReleaseResources()
 			LODRenderData[LODIndex].ReleaseResources();
 		}
 
+#if RHI_RAYTRACING
+		if (IsRayTracingAllowed())
+		{
+			ENQUEUE_RENDER_COMMAND(SkeletalMeshReleaseRayTracingGeometryGroup)(UE::RenderCommandPipe::SkeletalMesh,
+				[this]
+				{
+					GRayTracingGeometryManager->ReleaseRayTracingGeometryGroup(RayTracingGeometryGroupHandle);
+					RayTracingGeometryGroupHandle = INDEX_NONE;
+				});
+		}
+#endif
+
 		check(NaniteResourcesPtr.IsValid());
 		NaniteResourcesPtr->ReleaseResources();
 
diff --git a/Engine/Source/Runtime/Engine/Private/SkeletalRenderStatic.cpp b/Engine/Source/Runtime/Engine/Private/SkeletalRenderStatic.cpp
--- a/Engine/Source/Runtime/Engine/Private/SkeletalRenderStatic.cpp
+++ b/Engine/Source/Runtime/Engine/Private/SkeletalRenderStatic.cpp
@@ -113,6 +113,7 @@ void FSkeletalMeshObjectStatic::InitResources(USkinnedMeshComponent* InMeshCompo
 								Initializer.Type = ERayTracingGeometryInitializerType::StreamingDestination;
 							}
 
+							RayTracingGeometry.GroupHandle = SkelMeshRenderData->RayTracingGeometryGroupHandle;
 							RayTracingGeometry.LODIndex = LODIndex;
 
 							RayTracingGeometry.SetInitializer(Initializer);
diff --git a/Engine/Source/Runtime/Engine/Public/Rendering/SkeletalMeshRenderData.h b/Engine/Source/Runtime/Engine/Public/Rendering/SkeletalMeshRenderData.h
--- a/Engine/Source/Runtime/Engine/Public/Rendering/SkeletalMeshRenderData.h
+++ b/Engine/Source/Runtime/Engine/Public/Rendering/SkeletalMeshRenderData.h
@@ -44,6 +44,10 @@ public:
 	/** Whether ray tracing acceleration structures should be created for this mesh. Derived from owner USkinnedAsset. */
 	bool bSupportRayTracing;
 
+#if RHI_RAYTRACING
+	RayTracing::GeometryGroupHandle RayTracingGeometryGroupHandle = INDEX_NONE;
+#endif
+
 #if WITH_EDITORONLY_DATA
 	/** UV data used for streaming accuracy debug view modes. In sync for rendering thread */
 	TArray<FMeshUVChannelInfo> UVChannelDataPerMaterial;
diff --git a/Engine/Source/Runtime/Engine/Public/SkeletalMeshSceneProxy.h b/Engine/Source/Runtime/Engine/Public/SkeletalMeshSceneProxy.h
--- a/Engine/Source/Runtime/Engine/Public/SkeletalMeshSceneProxy.h
+++ b/Engine/Source/Runtime/Engine/Public/SkeletalMeshSceneProxy.h
@@ -31,6 +31,8 @@ public:
 	ENGINE_API virtual void DrawStaticElements(FStaticPrimitiveDrawInterface* PDI) override;
 	ENGINE_API virtual void GetDynamicMeshElements(const TArray<const FSceneView*>& Views, const FSceneViewFamily& ViewFamily, uint32 VisibilityMap, FMeshElementCollector& Collector) const override;
 
+	ENGINE_API virtual void CreateRenderThreadResources(FRHICommandListBase& RHICmdList) override;
+
 #if RHI_RAYTRACING
 	ENGINE_API virtual bool HasRayTracingRepresentation() const override;
 
@@ -41,6 +43,8 @@ public:
 		return bRenderStatic;
 	}
 
+	ENGINE_API virtual RayTracing::GeometryGroupHandle GetRayTracingGeometryGroupHandle() const override;
+
 	ENGINE_API virtual TArray<FRayTracingGeometry*> GetStaticRayTracingGeometries() const override;
 
 	ENGINE_API virtual void GetDynamicRayTracingInstances(struct FRayTracingMaterialGatheringContext& Context, TArray<struct FRayTracingInstance>& OutRayTracingInstances) override;
@@ -199,6 +203,10 @@ protected:
 	/** The primitive's pre-skinned local space bounds. */
 	FBoxSphereBounds PreSkinnedLocalBounds;
 
+#if RHI_RAYTRACING
+	RayTracing::GeometryGroupHandle RayTracingGeometryGroupHandle = INDEX_NONE;
+#endif
+
 #if !(UE_BUILD_SHIPPING || UE_BUILD_TEST)
 	/** The color we draw this component in if drawing debug bones */
 	TOptional<FLinearColor> DebugDrawColor;
