Added some checks and ensures to prevent and track down build/update requests for FRayTracingGeometry without valid RHI.
- otherwise checks only trigger inside RHIBuildAccelerationStructures() at which point we don't have much context.

diff --git a/Engine/Source/Runtime/Engine/Private/RayTracingSkinnedGeometry.cpp b/Engine/Source/Runtime/Engine/Private/RayTracingSkinnedGeometry.cpp
--- a/Engine/Source/Runtime/Engine/Private/RayTracingSkinnedGeometry.cpp
+++ b/Engine/Source/Runtime/Engine/Private/RayTracingSkinnedGeometry.cpp
@@ -29,6 +29,8 @@ FAutoConsoleVariableRef CVarSkinnedGeometryMaxRayTracingPrimitivesPerCmdList(
 
 void FRayTracingSkinnedGeometryUpdateQueue::Add(FRayTracingGeometry* InRayTracingGeometry, const FRayTracingAccelerationStructureSize& StructureSize)
 {
+	checkf(InRayTracingGeometry->GetRHI() != nullptr, TEXT("FRayTracingGeometry needs to have a valid RHI to be updated by FRayTracingSkinnedGeometryUpdateQueue."));
+
 	FScopeLock Lock(&CS);
 	FRayTracingUpdateInfo* CurrentUpdateInfo = ToUpdate.Find(InRayTracingGeometry);
 	if (CurrentUpdateInfo == nullptr)
@@ -213,6 +215,11 @@ void FRayTracingSkinnedGeometryUpdateQueue::Commit(FRDGBuilder& GraphBuilder, ER
 			FRayTracingGeometry* RayTracingGeometry = Iter.Key();
 			FRayTracingUpdateInfo& UpdateInfo = Iter.Value();
 
+			if (!ensureMsgf(RayTracingGeometry->GetRHI(), TEXT("Skipping request with invalid ray tracing geometry in FRayTracingSkinnedGeometryUpdateQueue. Geometry->IsEvicted(): %d."), RayTracingGeometry->IsEvicted()))
+			{
+				continue;
+			}
+
 			FRayTracingGeometryBuildParams BuildParams;
 			BuildParams.Geometry = RayTracingGeometry->GetRHI();
 			BuildParams.BuildMode = UpdateInfo.BuildMode;
diff --git a/Engine/Source/Runtime/Engine/Private/Rendering/RayTracingGeometryManager.cpp b/Engine/Source/Runtime/Engine/Private/Rendering/RayTracingGeometryManager.cpp
--- a/Engine/Source/Runtime/Engine/Private/Rendering/RayTracingGeometryManager.cpp
+++ b/Engine/Source/Runtime/Engine/Private/Rendering/RayTracingGeometryManager.cpp
@@ -543,6 +543,7 @@ void FRayTracingGeometryManager::ProcessBuildRequests(FRHIComputeCommandList& In
 void FRayTracingGeometryManager::SetupBuildParams(const FBuildRequest& InBuildRequest, TArray<FRayTracingGeometryBuildParams>& InBuildParams, bool bRemoveFromRequestArray)
 {
 	check(InBuildRequest.RequestIndex != INDEX_NONE && InBuildRequest.Owner->RayTracingBuildRequestIndex != INDEX_NONE);
+	checkf(InBuildRequest.Owner->GetRHI() != nullptr, TEXT("Build request for FRayTracingGeometry without valid RHI. Was the FRayTracingGeometry evicted or released without calling RemoveBuildRequest()?"));
 
 	FRayTracingGeometryBuildParams BuildParam;
 	BuildParam.Geometry = InBuildRequest.Owner->GetRHI();
