Avoid modifying constant value in linear key reduction codec

diff --git a/Engine/Source/Runtime/Engine/Classes/Animation/AnimCompress_RemoveLinearKeys.h b/Engine/Source/Runtime/Engine/Classes/Animation/AnimCompress_RemoveLinearKeys.h
--- a/Engine/Source/Runtime/Engine/Classes/Animation/AnimCompress_RemoveLinearKeys.h
+++ b/Engine/Source/Runtime/Engine/Classes/Animation/AnimCompress_RemoveLinearKeys.h
@@ -59,7 +59,7 @@ class UAnimCompress_RemoveLinearKeys : public UAnimCompress
 	 * A scale value which increases the likelihood that a bone will retain a key if it's parent also had a key at the same time position. 
 	 * Higher values can remove shaking artifacts from the animation, at the cost of compression.
 	 */
-	UPROPERTY(EditAnywhere, Category=LinearKeyRemoval)
+	UPROPERTY(EditAnywhere, Category=LinearKeyRemoval, meta = (ClampMin = "1.0"))
 	float ParentKeyScale;
 
 	/** 
diff --git a/Engine/Source/Runtime/Engine/Private/Animation/AnimCompress_RemoveLinearKeys.cpp b/Engine/Source/Runtime/Engine/Private/Animation/AnimCompress_RemoveLinearKeys.cpp
--- a/Engine/Source/Runtime/Engine/Private/Animation/AnimCompress_RemoveLinearKeys.cpp
+++ b/Engine/Source/Runtime/Engine/Private/Animation/AnimCompress_RemoveLinearKeys.cpp
@@ -608,7 +608,7 @@ void UAnimCompress_RemoveLinearKeys::ProcessAnimationTracks(
 	const bool bHasScale =  (ScaleTracks.Num() > 0);
 
 	// make sure the parent key scale is properly bound to 1.0 or more
-	ParentKeyScale = FMath::Max(ParentKeyScale, 1.0f);
+	const float ClampedParentKeyScale = FMath::Max(ParentKeyScale, 1.0f);
 
 	// generate the raw and compressed skeleton in world-space
 	TArray<FTransform> RawWorldBones;
@@ -857,7 +857,7 @@ void UAnimCompress_RemoveLinearKeys::ProcessAnimationTracks(
 
 			// look for a parent track to reference as a guide
 			int32 GuideTrackIndex = INDEX_NONE;
-			if (ParentKeyScale > 1.0f)
+			if (ClampedParentKeyScale > 1.0f)
 			{
 				for (long FamilyIndex=0; (FamilyIndex < Bone.BonesToRoot.Num()) && (GuideTrackIndex == INDEX_NONE); ++FamilyIndex)
 				{
@@ -899,7 +899,7 @@ void UAnimCompress_RemoveLinearKeys::ProcessAnimationTracks(
 			{
 				FTranslationTrack& GuideTransTrack = PositionTracks[GuideTrackIndex];
 				GuidanceTrack = &GuideTransTrack.Times;
-				GuidanceScale = ParentKeyScale;
+				GuidanceScale = ClampedParentKeyScale;
 			}
 			
 			// if the TargetBoneIndices array is empty, then this bone is an end effector.
