Changes to Ray Tracing Geometry Group managment

Context:
- We were not clearing the relevant entry in FRayTracingGeometryGroup::Geometries when a FRayTracingGeometry is unregistered from FRayTracingGeometryManager.
- This was causing a check to fail in RegisterRayTracingGeometry(...) because we dynamically create/destroy FSkeletalMeshObjectStatic/StaticRayTracingGeometry as necessary (referencing the same RayTracingGeometryGroup).
- Using BeginReleaseResource(FRayTracingGeometry) can end up causing the release commands to be delayed due to batching logic, which can result in FRayTracingGeometry having stale RayTracing::GeometryGroupHandle.

Changes:
- Clear entry in FRayTracingGeometryGroup when releasing RayTracingGeometryHandle of geometry that was assigned to a group.
- Delay releasing the RayTracing::GeometryGroupHandle until all references (geometries/scene proxies) are released to avoid stale handles.

diff --git a/Engine/Source/Runtime/Engine/Private/Rendering/RayTracingGeometryManager.cpp b/Engine/Source/Runtime/Engine/Private/Rendering/RayTracingGeometryManager.cpp
--- a/Engine/Source/Runtime/Engine/Private/Rendering/RayTracingGeometryManager.cpp
+++ b/Engine/Source/Runtime/Engine/Private/Rendering/RayTracingGeometryManager.cpp
@@ -134,6 +134,7 @@ RayTracing::GeometryGroupHandle FRayTracingGeometryManager::RegisterRayTracingGe
 
 	FRayTracingGeometryGroup Group;
 	Group.Geometries.AddDefaulted(NumLODs);
+	Group.NumReferences = 1;
 
 	RayTracing::GeometryGroupHandle Handle = RegisteredGroups.Add(MoveTemp(Group));
 
@@ -148,18 +149,31 @@ void FRayTracingGeometryManager::ReleaseRayTracingGeometryGroup(RayTracing::Geom
 
 	check(RegisteredGroups.IsValidIndex(Handle));
 
-	if (RegisteredGroups[Handle].ProxiesWithCachedRayTracingState.IsEmpty())
+	FRayTracingGeometryGroup& Group = RegisteredGroups[Handle];
+	
+	ReleaseRayTracingGeometryGroupReference(Handle);
+}
+
+void FRayTracingGeometryManager::ReleaseRayTracingGeometryGroupReference(RayTracing::GeometryGroupHandle Handle)
+{
+	FRayTracingGeometryGroup& Group = RegisteredGroups[Handle];
+
+	--Group.NumReferences;
+
+	if (Group.NumReferences == 0)
 	{
+		for (FRayTracingGeometry* Geometry : Group.Geometries)
+		{
+			checkf(Geometry == nullptr, TEXT("All FRayTracingGeometry in a group must be unregistered before releasing the group."));
+		}
+
+		check(Group.ProxiesWithCachedRayTracingState.IsEmpty());
+
 		RegisteredGroups.RemoveAt(Handle);
 		ReferencedGeometryGroups.Remove(Handle);
-	}
-	else
-	{
-		// set flag on group so that it is released once the last primitive is unregistered
-		RegisteredGroups[Handle].bPendingRelease = true;
-	}
 
-	DEC_DWORD_STAT(STAT_RayTracingGeometryGroupCount);
+		DEC_DWORD_STAT(STAT_RayTracingGeometryGroupCount);
+	}
 }
 
 FRayTracingGeometryManager::RayTracingGeometryHandle FRayTracingGeometryManager::RegisterRayTracingGeometry(FRayTracingGeometry* InGeometry)
@@ -185,6 +199,7 @@ FRayTracingGeometryManager::RayTracingGeometryHandle FRayTracingGeometryManager:
 			checkf(Group.Geometries[InGeometry->LODIndex] == nullptr, TEXT("Each LOD inside a FRayTracingGeometryGroup can only be associated with a single FRayTracingGeometry"));
 
 			Group.Geometries[InGeometry->LODIndex] = InGeometry;
+			++Group.NumReferences;
 		}
 		
 		INC_DWORD_STAT(STAT_RayTracingGeometryCount);
@@ -207,6 +222,22 @@ void FRayTracingGeometryManager::ReleaseRayTracingGeometryHandle(RayTracingGeome
 
 		FRegisteredGeometry& RegisteredGeometry = RegisteredGeometries[Handle];
 
+		if (RegisteredGeometry.Geometry->GroupHandle != INDEX_NONE)
+		{
+			// if geometry was assigned to a group, clear the relevant entry so another geometry can be registered later
+
+			checkf(RegisteredGroups.IsValidIndex(RegisteredGeometry.Geometry->GroupHandle), TEXT("FRayTracingGeometry.GroupHandle must be valid"));
+
+			FRayTracingGeometryGroup& Group = RegisteredGroups[RegisteredGeometry.Geometry->GroupHandle];
+
+			checkf(RegisteredGeometry.Geometry->LODIndex >= 0 && RegisteredGeometry.Geometry->LODIndex < Group.Geometries.Num(), TEXT("FRayTracingGeometry assigned to a group must have a valid LODIndex"));
+			checkf(Group.Geometries[RegisteredGeometry.Geometry->LODIndex] == RegisteredGeometry.Geometry, TEXT("Unexpected mismatch of FRayTracingGeometry in FRayTracingGeometryGroup"));
+
+			Group.Geometries[RegisteredGeometry.Geometry->LODIndex] = nullptr;
+
+			ReleaseRayTracingGeometryGroupReference(RegisteredGeometry.Geometry->GroupHandle);
+		}
+
 		int32 NumRemoved = ResidentGeometries.Remove(RegisteredGeometry.Geometry);
 
 		if (NumRemoved > 0)
@@ -572,10 +603,14 @@ void FRayTracingGeometryManager::RegisterProxyWithCachedRayTracingState(FPrimiti
 	checkf(IsRayTracingAllowed(), TEXT("Should only register proxies with FRayTracingGeometryManager when ray tracing is allowed"));
 	checkf(RegisteredGroups.IsValidIndex(InRayTracingGeometryGroupHandle), TEXT("InRayTracingGeometryGroupHandle must be valid"));
 
-	TSet<FPrimitiveSceneProxy*>& ProxiesSet = RegisteredGroups[InRayTracingGeometryGroupHandle].ProxiesWithCachedRayTracingState;
+	FRayTracingGeometryGroup& Group = RegisteredGroups[InRayTracingGeometryGroupHandle];
+
+	TSet<FPrimitiveSceneProxy*>& ProxiesSet = Group.ProxiesWithCachedRayTracingState;
 	check(!ProxiesSet.Contains(Proxy));
 
 	ProxiesSet.Add(Proxy);
+
+	++Group.NumReferences;
 }
 
 void FRayTracingGeometryManager::UnregisterProxyWithCachedRayTracingState(FPrimitiveSceneProxy* Proxy, RayTracing::GeometryGroupHandle InRayTracingGeometryGroupHandle)
@@ -590,10 +625,7 @@ void FRayTracingGeometryManager::UnregisterProxyWithCachedRayTracingState(FPrimi
 
 	verify(ProxiesSet.Remove(Proxy) == 1);
 
-	if (ProxiesSet.IsEmpty() && Group.bPendingRelease)
-	{
-		RegisteredGroups.RemoveAt(InRayTracingGeometryGroupHandle);
-	}
+	ReleaseRayTracingGeometryGroupReference(InRayTracingGeometryGroupHandle);
 }
 
 void FRayTracingGeometryManager::RequestUpdateCachedRenderState(RayTracing::GeometryGroupHandle InRayTracingGeometryGroupHandle)
diff --git a/Engine/Source/Runtime/Engine/Public/Rendering/RayTracingGeometryManager.h b/Engine/Source/Runtime/Engine/Public/Rendering/RayTracingGeometryManager.h
--- a/Engine/Source/Runtime/Engine/Public/Rendering/RayTracingGeometryManager.h
+++ b/Engine/Source/Runtime/Engine/Public/Rendering/RayTracingGeometryManager.h
@@ -63,6 +63,8 @@ private:
 
 	void SetupBuildParams(const FBuildRequest& InBuildRequest, TArray<FRayTracingGeometryBuildParams>& InBuildParams, bool bRemoveFromRequestArray = true);
 
+	void ReleaseRayTracingGeometryGroupReference(RayTracing::GeometryGroupHandle Handle);
+
 	FCriticalSection RequestCS;
 
 	TSparseArray<FBuildRequest> GeometryBuildRequests;
@@ -81,9 +83,10 @@ private:
 
 		TSet<FPrimitiveSceneProxy*> ProxiesWithCachedRayTracingState;
 
-		// flag used to indicate that ReleaseRayTracingGeometryHandle(...) has been called 
-		// group is pending release due to remaining proxies
-		bool bPendingRelease = false;
+		// Due to the way we batch release FRenderResource and SceneProxies, 
+		// ReleaseRayTracingGeometryGroup(...) can end up being called before all FRayTracingGeometry and SceneProxies are actually released.
+		// To deal with this, we keep track of whether the group is still referenced and only release the group handle once all references are released.
+		uint32 NumReferences = 0;
 
 		// TODO: Implement use-after-free checks in RayTracing::GeometryGroupHandle using some bits to identify generation
 	};
