Reset BuildParams in FRayTracingGeometryManager::ForceBuildIfPending(...).
- otherwise check in FRayTracingGeometryManager::ProcessBuildRequests(...) fails.

diff --git a/Engine/Source/Runtime/Engine/Private/Rendering/RayTracingGeometryManager.cpp b/Engine/Source/Runtime/Engine/Private/Rendering/RayTracingGeometryManager.cpp
--- a/Engine/Source/Runtime/Engine/Private/Rendering/RayTracingGeometryManager.cpp
+++ b/Engine/Source/Runtime/Engine/Private/Rendering/RayTracingGeometryManager.cpp
@@ -487,6 +487,8 @@ void FRayTracingGeometryManager::ForceBuildIfPending(FRHIComputeCommandList& InC
 	{
 		InCmdList.BuildAccelerationStructures(BuildParams);
 	}
+
+	BuildParams.Reset();
 }
 
 void FRayTracingGeometryManager::ProcessBuildRequests(FRHIComputeCommandList& InCmdList, bool bInBuildAll)
