* Fix problem with deleted archetype in the replacement map.

diff --git a/Engine/Source/Runtime/Engine/Private/UnrealEngine.cpp b/Engine/Source/Runtime/Engine/Private/UnrealEngine.cpp
--- a/Engine/Source/Runtime/Engine/Private/UnrealEngine.cpp
+++ b/Engine/Source/Runtime/Engine/Private/UnrealEngine.cpp
@@ -17452,9 +17452,16 @@ void UEngine::CopyPropertiesForUnrelatedObjects(UObject* OldObject, UObject* New
 	// Gather references to old instances or objects that need to be replaced after we serialize in saved data
 	TMap<UObject*, UObject*> ReferenceReplacementMap;
 	ReferenceReplacementMap.Add(OldObject, NewObject);
-	if (OldObject->GetArchetype() != NewObject->GetArchetype())
+	UObject* OldArchetype = OldObject->GetArchetype();
+	UObject* NewArchetype = NewObject->GetArchetype();
+	if (OldArchetype != NewArchetype)
 	{
-		ReferenceReplacementMap.Add(OldObject->GetArchetype(), NewObject->GetArchetype());
+		// When an archetype is removed from a cdo, the GetArchetype will not return the right one here and could even point to the CDO so adding that replacement for a wrong archetype isn't good at all
+		// Check the OptionalReplacementMappings if it already contains a mapping to that new archetype and do not add it to the ReferenceReplacementMap
+		if(!Params.OptionalReplacementMappings || !Params.OptionalReplacementMappings->FindKey(NewArchetype))
+		{
+			ReferenceReplacementMap.Add(OldArchetype, NewArchetype);
+		}
 	}
 	if (OldObject->GetClass() != NewObject->GetClass())
 	{
