Fix a case where SourcePose can be used without having been initialized in FAnimInstanceProxy::SlotEvaluatePose

diff --git a/Engine/Source/Runtime/Engine/Private/Animation/AnimInstanceProxy.cpp b/Engine/Source/Runtime/Engine/Private/Animation/AnimInstanceProxy.cpp
--- a/Engine/Source/Runtime/Engine/Private/Animation/AnimInstanceProxy.cpp
+++ b/Engine/Source/Runtime/Engine/Private/Animation/AnimInstanceProxy.cpp
@@ -1978,16 +1978,24 @@ void FAnimInstanceProxy::SlotEvaluatePose(const FName& SlotNodeName, const FAnim
 	}
 
 	// Make sure we have at least one montage here.
-	check((AdditivePoses.Num() > 0) || (NonAdditivePoses.Num() > 0));
+	ensure((AdditivePoses.Num() > 0) || (NonAdditivePoses.Num() > 0));
 
 	// Second pass, blend non additive poses together
 	{
 		// If we're only playing additive animations, just copy source for base pose.
 		if (NonAdditivePoses.Num() == 0)
 		{
-			BlendedPose = SourcePose;
-			BlendedCurve = SourceCurve;
-			BlendedAttributes = SourceAttributes;
+			// If the source weight is 0, SourcePose will be uninitialized and cannot be used
+			if (InSourceWeight > ZERO_ANIMWEIGHT_THRESH)
+			{
+				BlendedPose = SourcePose;
+				BlendedCurve = SourceCurve;
+				BlendedAttributes = SourceAttributes;
+			}
+			else
+			{
+				BlendedPose.ResetToRefPose(); 
+			}
 		}		
 		else // Otherwise we need to blend non additive poses together
 		{
