Add & Remove Modulators functions

diff --git a/Engine/Source/Runtime/Engine/Classes/Components/AudioComponent.h b/Engine/Source/Runtime/Engine/Classes/Components/AudioComponent.h
--- a/Engine/Source/Runtime/Engine/Classes/Components/AudioComponent.h
+++ b/Engine/Source/Runtime/Engine/Classes/Components/AudioComponent.h
@@ -125,24 +125,6 @@ enum class EAudioFaderCurve : uint8
 	Count UMETA(Hidden)
 };
 
-UENUM(BlueprintType)
-enum class EModulationDestination : uint8
-{
-	/* Volume modulation */
-	Volume,
-
-	/* Pitch modulation */
-	Pitch,
-
-	/* Cutoff Frequency of a lowpass filter */
-	Lowpass,
-
-	/* Cutoff Frequency of a highpass filter */
-	Highpass,
-
-	Count UMETA(Hidden)
-};
-
 /**
  * Legacy struct used for storing named parameter for a given AudioComponent.
  */
@@ -738,7 +720,8 @@ public:
 	ENGINE_API bool GetCookedEnvelopeDataForAllPlayingSounds(TArray<FSoundWaveEnvelopeDataPerSound>& OutEnvelopeData);
 
 	/**
-	* Sets the routing for one of the given Audio component's Modulation Destinations.
+	* Sets the routing for one of the given Audio component's Modulation Destinations. This will overwrite the current settings.
+	* To Add modulators without removing previously set ones, use AddModulators.
 	* @param Modulators The set of modulators to apply to the given destination on the component.
 	* @param Destination The destination to assign the modulators to.
 	* @param RoutingMethod The routing method to use for the given modulator.
@@ -746,6 +729,24 @@ public:
 	UFUNCTION(BlueprintCallable, Category = "Audio|Components|Audio", DisplayName = "Set Modulation Routing")
 	ENGINE_API void SetModulationRouting(const TSet<USoundModulatorBase*>& Modulators, const EModulationDestination Destination, const EModulationRouting RoutingMethod = EModulationRouting::Inherit);
 
+	/**
+	* Adds the given set of Modulators to the modulators currently set on the Audio Component. This performs a logical Union, so duplicate modulators will not be added.
+	* To completely overwrite existing Modulation settings, use SetModulationRouting.
+	* @param Modulators The set of modulators to add to the given destination on the component.
+	* @param Destination The destination to add the modulators to.
+	*/
+	UFUNCTION(BlueprintCallable, Category = "Audio|Components|Audio", DisplayName = "Add Modulation Routing")
+	ENGINE_API void AddModulationRouting(const TSet<USoundModulatorBase*>& Modulators, const EModulationDestination Destination);
+
+	/**
+	* Removes the given set of Modulators to the modulators currently set on the Audio Component. This performs a logical subtraction of the sets, so modulators that are not in the given set will stay.
+	* To completely overwrite current Modulation settings, use SetModulationRouting.
+	* @param Modulators The set of modulators to remove from the given destination on the component.
+	* @param Destination The destination to remove the modulators from.
+	*/
+	UFUNCTION(BlueprintCallable, Category = "Audio|Components|Audio", DisplayName = "Remove Modulation Routing")
+	ENGINE_API void RemoveModulationRouting(const TSet<USoundModulatorBase*>& Modulators, const EModulationDestination Destination);
+
 	/**
 	* Gets the set of currently active modulators for a given Modulation Destination.
 	* @param Destination The Destination to retrieve the Modulators from.
diff --git a/Engine/Source/Runtime/Engine/Classes/Sound/SoundModulationDestination.h b/Engine/Source/Runtime/Engine/Classes/Sound/SoundModulationDestination.h
--- a/Engine/Source/Runtime/Engine/Classes/Sound/SoundModulationDestination.h
+++ b/Engine/Source/Runtime/Engine/Classes/Sound/SoundModulationDestination.h
@@ -38,6 +38,24 @@ enum class EModulationRouting : uint8
 	Union
 };
 
+UENUM(BlueprintType)
+enum class EModulationDestination : uint8
+{
+	/* Volume modulation */
+	Volume,
+
+	/* Pitch modulation */
+	Pitch,
+
+	/* Cutoff Frequency of a lowpass filter */
+	Lowpass,
+
+	/* Cutoff Frequency of a highpass filter */
+	Highpass,
+
+	Count UMETA(Hidden)
+};
+
 /** Parameter destination settings allowing modulation control override for parameter destinations opting in to the Modulation System. */
 USTRUCT(BlueprintType)
 struct FSoundModulationDestinationSettings
diff --git a/Engine/Source/Runtime/Engine/Private/ActiveSound.cpp b/Engine/Source/Runtime/Engine/Private/ActiveSound.cpp
--- a/Engine/Source/Runtime/Engine/Private/ActiveSound.cpp
+++ b/Engine/Source/Runtime/Engine/Private/ActiveSound.cpp
@@ -467,6 +467,60 @@ void FActiveSound::SetNewModulationRouting(const FSoundModulationDefaultRoutingS
 	bModulationRoutingUpdated = true;
 }
 
+void FActiveSound::AddModulationRouting(const TSet<TObjectPtr<USoundModulatorBase>>& NewModulators, EModulationDestination Destination)
+{
+	switch (Destination)
+	{
+		case EModulationDestination::Volume:
+			ModulationRouting.VolumeModulationDestination.Modulators.Append(NewModulators);
+			break;
+		case EModulationDestination::Pitch:
+			ModulationRouting.PitchModulationDestination.Modulators.Append(NewModulators);
+			break;
+		case EModulationDestination::Lowpass:
+			ModulationRouting.LowpassModulationDestination.Modulators.Append(NewModulators);
+			break;
+		case EModulationDestination::Highpass:
+			ModulationRouting.HighpassModulationDestination.Modulators.Append(NewModulators);
+			break;
+		default:
+		{
+			static_assert(static_cast<int32>(EModulationDestination::Count) == 4, "Possible missing ELiteralType case coverage.");
+			ensureMsgf(false, TEXT("Failed to set input node default: Literal type not supported"));
+			return;
+		}
+	}
+	
+	bModulationRoutingUpdated = true;
+}
+
+void FActiveSound::RemoveModulationRouting(const TSet<TObjectPtr<USoundModulatorBase>>& NewModulators, EModulationDestination Destination)
+{
+	switch (Destination)
+	{
+		case EModulationDestination::Volume:
+			ModulationRouting.VolumeModulationDestination.Modulators = ModulationRouting.VolumeModulationDestination.Modulators.Difference(NewModulators);
+			break;
+		case EModulationDestination::Pitch:
+			ModulationRouting.PitchModulationDestination.Modulators = ModulationRouting.PitchModulationDestination.Modulators.Difference(NewModulators);
+			break;
+		case EModulationDestination::Lowpass:
+			ModulationRouting.LowpassModulationDestination.Modulators = ModulationRouting.LowpassModulationDestination.Modulators.Difference(NewModulators);
+			break;
+		case EModulationDestination::Highpass:
+			ModulationRouting.HighpassModulationDestination.Modulators = ModulationRouting.HighpassModulationDestination.Modulators.Difference(NewModulators);
+			break;
+		default:
+		{
+			static_assert(static_cast<int32>(EModulationDestination::Count) == 4, "Possible missing ELiteralType case coverage.");
+			ensureMsgf(false, TEXT("Failed to set input node default: Literal type not supported"));
+			return;
+		}
+	}
+
+	bModulationRoutingUpdated = true;
+}
+
 void FActiveSound::Stop()
 {
 	if (AudioDevice)
diff --git a/Engine/Source/Runtime/Engine/Private/Components/AudioComponent.cpp b/Engine/Source/Runtime/Engine/Private/Components/AudioComponent.cpp
--- a/Engine/Source/Runtime/Engine/Private/Components/AudioComponent.cpp
+++ b/Engine/Source/Runtime/Engine/Private/Components/AudioComponent.cpp
@@ -1853,36 +1853,122 @@ void UAudioComponent::SetModulationRouting(const TSet<USoundModulatorBase*>& Mod
 
 	switch (Destination)
 	{
-	case EModulationDestination::Volume:
-		ModulationRouting.VolumeRouting = RoutingMethod;
-		ModulationRouting.VolumeModulationDestination.Modulators = Modulators;
-		break;
-	case EModulationDestination::Pitch:
-		ModulationRouting.PitchRouting = RoutingMethod;
-		ModulationRouting.PitchModulationDestination.Modulators = Modulators;
-		break;
-	case EModulationDestination::Lowpass:
-		ModulationRouting.LowpassRouting = RoutingMethod;
-		ModulationRouting.LowpassModulationDestination.Modulators = Modulators;
-		break;
-	case EModulationDestination::Highpass:
-		ModulationRouting.HighpassRouting = RoutingMethod;
-		ModulationRouting.HighpassModulationDestination.Modulators = Modulators;
-		break;
-	default:
-	{
-		static_assert(static_cast<int32>(EModulationDestination::Count) == 4, "Possible missing ELiteralType case coverage.");
-		ensureMsgf(false, TEXT("Failed to set input node default: Literal type not supported"));
+		case EModulationDestination::Volume:
+			ModulationRouting.VolumeRouting = RoutingMethod;
+			ModulationRouting.VolumeModulationDestination.Modulators = Modulators;
+			break;
+		case EModulationDestination::Pitch:
+			ModulationRouting.PitchRouting = RoutingMethod;
+			ModulationRouting.PitchModulationDestination.Modulators = Modulators;
+			break;
+		case EModulationDestination::Lowpass:
+			ModulationRouting.LowpassRouting = RoutingMethod;
+			ModulationRouting.LowpassModulationDestination.Modulators = Modulators;
+			break;
+		case EModulationDestination::Highpass:
+			ModulationRouting.HighpassRouting = RoutingMethod;
+			ModulationRouting.HighpassModulationDestination.Modulators = Modulators;
+			break;
+		default:
+		{
+			static_assert(static_cast<int32>(EModulationDestination::Count) == 4, "Possible missing ELiteralType case coverage.");
+			ensureMsgf(false, TEXT("Failed to set input node default: Literal type not supported"));
+			return;
+		}
+	}
+
+	// Tell the active sounds on the component to use the new Modulation Routing
+	AudioDevice->SendCommandToActiveSounds(AudioComponentID, [NewRouting = ModulationRouting](FActiveSound& ActiveSound)
+	{
+		ActiveSound.SetNewModulationRouting(NewRouting);
+	});
+
+}
+
+void UAudioComponent::AddModulationRouting(const TSet<USoundModulatorBase*>& Modulators, const EModulationDestination Destination)
+{
+	FAudioDevice* AudioDevice = GetAudioDevice();
+	if (!AudioDevice)
+	{
 		return;
 	}
+
+	TSet<TObjectPtr<USoundModulatorBase>> ConvertedModulators;
+	for (USoundModulatorBase* Modulator : Modulators)
+	{
+		ConvertedModulators.Add(TObjectPtr<USoundModulatorBase>(Modulator));
+	}
+
+	switch (Destination)
+	{
+		case EModulationDestination::Volume:
+			ModulationRouting.VolumeModulationDestination.Modulators.Append(ConvertedModulators);
+			break;
+		case EModulationDestination::Pitch:
+			ModulationRouting.PitchModulationDestination.Modulators.Append(ConvertedModulators);
+			break;
+		case EModulationDestination::Lowpass:
+			ModulationRouting.LowpassModulationDestination.Modulators.Append(ConvertedModulators);
+			break;
+		case EModulationDestination::Highpass:
+			ModulationRouting.HighpassModulationDestination.Modulators.Append(ConvertedModulators);
+			break;
+		default:
+		{
+			static_assert(static_cast<int32>(EModulationDestination::Count) == 4, "Possible missing ELiteralType case coverage.");
+			ensureMsgf(false, TEXT("Failed to set input node default: Literal type not supported"));
+			return;
+		}
 	}
 
 	// Tell the active sounds on the component to use the new Modulation Routing
-	AudioDevice->SendCommandToActiveSounds(AudioComponentID, [NewRouting = ModulationRouting](FActiveSound& ActiveSound)
+	AudioDevice->SendCommandToActiveSounds(AudioComponentID, [NewModulators = ConvertedModulators, Destination](FActiveSound& ActiveSound)
+	{
+		ActiveSound.AddModulationRouting(NewModulators, Destination);
+	});
+}
+
+void UAudioComponent::RemoveModulationRouting(const TSet<USoundModulatorBase*>& Modulators, const EModulationDestination Destination)
+{
+	FAudioDevice* AudioDevice = GetAudioDevice();
+	if (!AudioDevice)
+	{
+		return;
+	}
+
+	TSet<TObjectPtr<USoundModulatorBase>> ConvertedModulators;
+	for (USoundModulatorBase* Modulator : Modulators)
+	{
+		ConvertedModulators.Add(TObjectPtr<USoundModulatorBase>(Modulator));
+	}
+
+	switch (Destination)
+	{
+		case EModulationDestination::Volume:
+			ModulationRouting.VolumeModulationDestination.Modulators = ModulationRouting.VolumeModulationDestination.Modulators.Difference(ConvertedModulators);
+			break;
+		case EModulationDestination::Pitch:
+			ModulationRouting.PitchModulationDestination.Modulators = ModulationRouting.PitchModulationDestination.Modulators.Difference(ConvertedModulators);
+			break;
+		case EModulationDestination::Lowpass:
+			ModulationRouting.LowpassModulationDestination.Modulators = ModulationRouting.LowpassModulationDestination.Modulators.Difference(ConvertedModulators);
+			break;
+		case EModulationDestination::Highpass:
+			ModulationRouting.HighpassModulationDestination.Modulators = ModulationRouting.HighpassModulationDestination.Modulators.Difference(ConvertedModulators);
+			break;
+		default:
 		{
-			ActiveSound.SetNewModulationRouting(NewRouting);
-		});
+			static_assert(static_cast<int32>(EModulationDestination::Count) == 4, "Possible missing ELiteralType case coverage.");
+			ensureMsgf(false, TEXT("Failed to set input node default: Literal type not supported"));
+			return;
+		}
+	}
 
+	// Tell the active sounds on the component to use the new Modulation Routing
+	AudioDevice->SendCommandToActiveSounds(AudioComponentID, [NewModulators = ConvertedModulators, Destination](FActiveSound& ActiveSound)
+	{
+		ActiveSound.RemoveModulationRouting(NewModulators, Destination);
+	});
 }
 
 TSet<USoundModulatorBase*> UAudioComponent::GetModulators(const EModulationDestination Destination)
@@ -1897,24 +1983,24 @@ TSet<USoundModulatorBase*> UAudioComponent::GetModulators(const EModulationDesti
 
 	switch (Destination)
 	{
-	case EModulationDestination::Volume:
-		ModulatorSet = &ModulationRouting.VolumeModulationDestination.Modulators;
-		break;
-	case EModulationDestination::Pitch:
-		ModulatorSet = &ModulationRouting.PitchModulationDestination.Modulators;
-		break;
-	case EModulationDestination::Lowpass:
-		ModulatorSet = &ModulationRouting.LowpassModulationDestination.Modulators;
-		break;
-	case EModulationDestination::Highpass:
-		ModulatorSet = &ModulationRouting.HighpassModulationDestination.Modulators;
-		break;
-	default:
-	{
-		static_assert(static_cast<int32>(EModulationDestination::Count) == 4, "Possible missing ELiteralType case coverage.");
-		ensureMsgf(false, TEXT("Failed to set input node default: Literal type not supported"));
-		return TSet<USoundModulatorBase*>();
-	}
+		case EModulationDestination::Volume:
+			ModulatorSet = &ModulationRouting.VolumeModulationDestination.Modulators;
+			break;
+		case EModulationDestination::Pitch:
+			ModulatorSet = &ModulationRouting.PitchModulationDestination.Modulators;
+			break;
+		case EModulationDestination::Lowpass:
+			ModulatorSet = &ModulationRouting.LowpassModulationDestination.Modulators;
+			break;
+		case EModulationDestination::Highpass:
+			ModulatorSet = &ModulationRouting.HighpassModulationDestination.Modulators;
+			break;
+		default:
+		{
+			static_assert(static_cast<int32>(EModulationDestination::Count) == 4, "Possible missing ELiteralType case coverage.");
+			ensureMsgf(false, TEXT("Failed to set input node default: Literal type not supported"));
+			return TSet<USoundModulatorBase*>();
+		}
 	}
 
 	check(ModulatorSet);
diff --git a/Engine/Source/Runtime/Engine/Public/ActiveSound.h b/Engine/Source/Runtime/Engine/Public/ActiveSound.h
--- a/Engine/Source/Runtime/Engine/Public/ActiveSound.h
+++ b/Engine/Source/Runtime/Engine/Public/ActiveSound.h
@@ -709,9 +709,15 @@ public:
 	/** Resets internal data of new Source Bus Sends */
 	ENGINE_API void ResetNewBusSends();
 
-	/* Gives new Modulation Routing settings to the Active Sound. */
+	/* Gives new Modulation Routing settings to the ActiveSound. This overwrites all previous Modulation Routing settings. */
 	ENGINE_API void SetNewModulationRouting(const FSoundModulationDefaultRoutingSettings& NewRouting);
 
+	/* Adds additional Modulators to the ActiveSound, if possible. To replace existing modulators, use SetNewModulationRouting. */
+	ENGINE_API void AddModulationRouting(const TSet<TObjectPtr<USoundModulatorBase>>& NewModulators, EModulationDestination Destination);
+
+	/* Removes given Modulators from the ActiveSound, if possible. */
+	ENGINE_API void RemoveModulationRouting(const TSet<TObjectPtr<USoundModulatorBase>>& NewModulators, EModulationDestination Destination);
+
 	/* Determines which of the provided listeners is the closest to the sound */
 	ENGINE_API int32 FindClosestListener( const TArray<struct FListener>& InListeners ) const;
 
