Clarify the SetSimulatePhysics tooltip in PrimitiveComponent. This has to cover all the derived classes as well, which might have slightly different behaviour(and currently tooltips can't be overridden in derived classes).

diff --git a/Engine/Source/Runtime/Engine/Classes/Components/PrimitiveComponent.h b/Engine/Source/Runtime/Engine/Classes/Components/PrimitiveComponent.h
--- a/Engine/Source/Runtime/Engine/Classes/Components/PrimitiveComponent.h
+++ b/Engine/Source/Runtime/Engine/Classes/Components/PrimitiveComponent.h
@@ -1483,10 +1483,16 @@ public:
 	ENGINE_API virtual void SetWalkableSlopeOverride(const FWalkableSlopeOverride& NewOverride);
 
 	/** 
-	 *	Sets whether or not a single body should use physics simulation, or should be 'fixed' (kinematic).
-	 *	Note that if this component is currently attached to something, beginning simulation will detach it.
+	 * When this component is a simple/single body, this will enable or disable simulation on that body. In addition, 
+	 * if this component is currently attached to something, beginning simulation will detach it. Note that stopping
+	 * simulation will not reattach the component - that would need to be done explicitly. 
+	 * 
+	 * For more complex components (e.g. skeletal meshes), simulation will apply to the bodies contained by the
+	 * component (e.g. using a physics asset). Since these bodies will be free to move independently of the component, 
+	 * the component will not be automatically detached. If detachment is required, then that can be done by 
+	 * calling DetachFromComponent.
 	 *
-	 *	@param	bSimulate	New simulation state for single body
+	 * @param bSimulate New simulation state for the single body or multiple bodies
 	 */
 	UFUNCTION(BlueprintCallable, Category="Physics")
 	ENGINE_API virtual void SetSimulatePhysics(bool bSimulate);
