Add networks stats for the number of packets in/out and lost per frame and route them to CSV.

diff --git a/Engine/Config/BaseEngine.ini b/Engine/Config/BaseEngine.ini
--- a/Engine/Config/BaseEngine.ini
+++ b/Engine/Config/BaseEngine.ini
@@ -3561,6 +3561,10 @@ IntervalSeconds=1
 +Listeners=(MetricName=SharedSerializationRPCMiss, Class=/Script/Engine.NetworkMetricsCSV_Replication, EnableMode=ENetworkMetricEnableMode::EnableForNonIrisOnly)
 +Listeners=(MetricName=OutgoingReliableMessageQueueMaxSize, Class=/Script/Engine.NetworkMetricsCSV_Replication)
 +Listeners=(MetricName=IncomingReliableMessageQueueMaxSize, Class=/Script/Engine.NetworkMetricsCSV_Replication)
++Listeners=(MetricName=InPacketsPerFrame, Class=/Script/Engine.NetworkMetricsCSV_Replication)
++Listeners=(MetricName=OutPacketsPerFrame, Class=/Script/Engine.NetworkMetricsCSV_Replication)
++Listeners=(MetricName=InLostPacketsFoundPerFrame, Class=/Script/Engine.NetworkMetricsCSV_Replication)
++Listeners=(MetricName=OutLostPacketsFoundPerFrame, Class=/Script/Engine.NetworkMetricsCSV_Replication)
 +Listeners=(MetricName=NumConnections, Class=/Script/Engine.NetworkMetricsPerfCounters)
 +Listeners=(MetricName=AvgPing, Class=/Script/Engine.NetworkMetricsPerfCounters)
 +Listeners=(MetricName=InBunches, Class=/Script/Engine.NetworkMetricsPerfCounters)
diff --git a/Engine/Source/Runtime/Engine/Classes/Engine/NetDriver.h b/Engine/Source/Runtime/Engine/Classes/Engine/NetDriver.h
--- a/Engine/Source/Runtime/Engine/Classes/Engine/NetDriver.h
+++ b/Engine/Source/Runtime/Engine/Classes/Engine/NetDriver.h
@@ -1148,12 +1148,16 @@ public:
 	uint32						NetGUIDOutBytes;
 	/** Incoming rate of NetGUID Bunches */
 	uint32						NetGUIDInBytes;
-	/** todo document */
+	/** The number of packets per second that have been received. */
 	uint32						InPackets;
+	/** The number of packets that were received in the previous frame. */
+	uint32 						PrevInPackets;
 	/** Total packets received since the net driver's creation  */
 	uint32						InTotalPackets;
-	/** todo document */
+	/** The number of packets per second that have been sent. */
 	uint32						OutPackets;
+	/** The number of packets that were sent in the previous frame. */
+	uint32 						PrevOutPackets;
 	/** Total packets sent since the net driver's creation  */
 	uint32						OutTotalPackets;
 	/** todo document */
@@ -1168,12 +1172,16 @@ public:
 	uint32						OutTotalReliableBunches;
 	/** Total number of incoming reliable bunches */
 	uint32						InTotalReliableBunches;
-	/** todo document */
+	/** The percentage of incoming packets that have been found to be lost. */
 	uint32						InPacketsLost;
+	/** The number of incoming packets that were found to be lost in the previous frame. */
+	uint32 						PrevInPacketsLost;
 	/** Total packets lost that have been sent by clients since the net driver's creation  */
 	uint32						InTotalPacketsLost;
-	/** todo document */
+	/** The percentage of outgoing packets that have been found to be lost. */
 	uint32						OutPacketsLost;
+	/** The number of outgoing packets that were found to be lost in the previous frame.  */
+	uint32 						PrevOutPacketsLost;
 	/** Total packets lost that have been sent by the server since the net driver's creation  */
 	uint32						OutTotalPacketsLost;
 	/** Tracks the total number of voice packets sent */
diff --git a/Engine/Source/Runtime/Engine/Private/Net/NetworkMetricsDefs.cpp b/Engine/Source/Runtime/Engine/Private/Net/NetworkMetricsDefs.cpp
--- a/Engine/Source/Runtime/Engine/Private/Net/NetworkMetricsDefs.cpp
+++ b/Engine/Source/Runtime/Engine/Private/Net/NetworkMetricsDefs.cpp
@@ -9,6 +9,10 @@ const FName AddedConnections("AddedConnections");
 const FName ClosedConnectionsDueToReliableBufferOverflow("ClosedConnectionsDueToReliableBufferOverflow");
 const FName InPackets("InPackets");
 const FName OutPackets("OutPackets");
+const FName InPacketsPerFrame("InPacketsPerFrame");
+const FName OutPacketsPerFrame("OutPacketsPerFrame");
+const FName InLostPacketsFoundPerFrame("InLostPacketsFoundPerFrame");
+const FName OutLostPacketsFoundPerFrame("OutLostPacketsFoundPerFrame");
 const FName InPacketsClientPerSecondAvg("InPacketsClientPerSecondAvg");
 const FName InPacketsClientPerSecondMax("InPacketsClientPerSecondMax");
 const FName InPacketsClientPerSecondMin("InPacketsClientPerSecondMin");
diff --git a/Engine/Source/Runtime/Engine/Private/NetDriver.cpp b/Engine/Source/Runtime/Engine/Private/NetDriver.cpp
--- a/Engine/Source/Runtime/Engine/Private/NetDriver.cpp
+++ b/Engine/Source/Runtime/Engine/Private/NetDriver.cpp
@@ -449,6 +449,13 @@ static FAutoConsoleVariableRef CVarNetDisableRandomNetUpdateDelay(
 	TEXT("If 1, the server will not add a random delay to an actors' NextUpdateTime."),
 	ECVF_Default);
 
+static bool GbEnableNetStats = false;
+static FAutoConsoleVariableRef CVarEnableNetStats(
+	TEXT("net.EnableNetStats"),
+	GbEnableNetStats,
+	TEXT("Enable the calculation of network stats and storing them in the network metrics database."),
+	ECVF_Default);
+
 namespace UE::Net
 {
 	static FString GRequiredEncryptionNetDriverDefNames_Internal = TEXT("all");
@@ -552,8 +559,10 @@ UNetDriver::UNetDriver(const FObjectInitializer& ObjectInitializer)
 ,	NetGUIDOutBytes(0)
 ,	NetGUIDInBytes(0)
 ,	InPackets(0)
+,	PrevInPackets(0)
 ,	InTotalPackets(0)
 ,	OutPackets(0)
+,	PrevOutPackets(0)
 ,	OutTotalPackets(0)
 ,	InBunches(0)
 ,	OutBunches(0)
@@ -562,8 +571,10 @@ UNetDriver::UNetDriver(const FObjectInitializer& ObjectInitializer)
 ,	OutTotalReliableBunches(0)
 ,	InTotalReliableBunches(0)
 ,	InPacketsLost(0)
+,	PrevInPacketsLost(0)
 ,	InTotalPacketsLost(0)
 ,	OutPacketsLost(0)
+,	PrevOutPacketsLost(0)
 ,	OutTotalPacketsLost(0)
 ,	StatUpdateTime(0.0)
 ,	StatPeriod(1.f)
@@ -1622,6 +1633,18 @@ void UNetDriver::SetupNetworkMetrics()
 	// The number of outgoing packets per second across all connections.
 	GetMetrics()->CreateInt(UE::Net::Metric::OutPackets, 0);
 
+	// The number of incoming packets per frame across all connections.
+	GetMetrics()->CreateInt(UE::Net::Metric::InPacketsPerFrame, 0);
+
+	// The number of outgoing packets per frame across all connections.
+	GetMetrics()->CreateInt(UE::Net::Metric::OutPacketsPerFrame, 0);
+
+	// The number of incoming packets lost per frame across all connections.
+	GetMetrics()->CreateInt(UE::Net::Metric::InLostPacketsFoundPerFrame, 0);
+
+	// The number of outgoing packets lost per frame across all connections.
+	GetMetrics()->CreateInt(UE::Net::Metric::OutLostPacketsFoundPerFrame, 0);
+
 	// The average/min/max incoming packets per second across all connections.
 	GetMetrics()->CreateInt(UE::Net::Metric::InPacketsClientPerSecondAvg, 0);
 	GetMetrics()->CreateInt(UE::Net::Metric::InPacketsClientPerSecondMax, 0);
@@ -1632,11 +1655,11 @@ void UNetDriver::SetupNetworkMetrics()
 	GetMetrics()->CreateInt(UE::Net::Metric::OutPacketsClientPerSecondMax, 0);
 	GetMetrics()->CreateInt(UE::Net::Metric::OutPacketsClientPerSecondMin, 0);
 	
-	// (Server only) The average/max incoming packets per frame across all connections.
+	// (Server only) The average/max incoming packets per stat period across all connections.
 	GetMetrics()->CreateInt(UE::Net::Metric::InPacketsClientAvg, 0);
 	GetMetrics()->CreateInt(UE::Net::Metric::InPacketsClientMax, 0);
 
-	// (Server only) The average/max outgoing packets per frame across all connections.
+	// (Server only) The average/max outgoing packets per stat period across all connections.
 	GetMetrics()->CreateInt(UE::Net::Metric::OutPacketsClientAvg, 0);
 	GetMetrics()->CreateInt(UE::Net::Metric::OutPacketsClientMax, 0);
 
@@ -5641,8 +5664,8 @@ void UNetDriver::ReportSyncLoad(const FNetSyncLoadReport& Report)
 // -------------------------------------------------------------------------------------------------------------------------
 
 
-// Replication CSV category is enabled by default in server builds
-CSV_DEFINE_CATEGORY(Replication, WITH_SERVER_CODE);
+// Replication CSV category is always enabled (client and server).
+CSV_DEFINE_CATEGORY(Replication, true);
 
 double GReplicationGatherPrioritizeTimeSeconds;
 double GServerReplicateActorTimeSeconds;
@@ -7890,7 +7913,7 @@ void UNetDriver::UpdateNetworkStats()
 	bCollectServerStats = true;
 #endif
 
-	if (bCollectNetStats || bCollectServerStats)
+	if (bCollectNetStats || bCollectServerStats || GbEnableNetStats)
 	{
 		SCOPE_CYCLE_COUNTER(STAT_NetTickFlushGatherStats);
 
@@ -7898,6 +7921,16 @@ void UNetDriver::UpdateNetworkStats()
 
 		NumFramesOverIncomingBunchTimeLimit += HasExceededIncomingBunchFrameProcessingTime() ? 1 : 0;
 
+		GetMetrics()->SetInt(UE::Net::Metric::InPacketsPerFrame, InPackets - PrevInPackets);
+		GetMetrics()->SetInt(UE::Net::Metric::OutPacketsPerFrame, OutPackets - PrevOutPackets);
+		GetMetrics()->SetInt(UE::Net::Metric::InLostPacketsFoundPerFrame, InPacketsLost - PrevInPacketsLost);
+		GetMetrics()->SetInt(UE::Net::Metric::OutLostPacketsFoundPerFrame, OutPacketsLost - PrevOutPacketsLost);
+
+		PrevInPackets = InPackets;
+		PrevOutPackets = OutPackets;
+		PrevInPacketsLost = InPacketsLost;
+		PrevOutPacketsLost = OutPacketsLost;
+
 		const double CurrentRealtimeSeconds = FPlatformTime::Seconds();
 		// Update network stats (only main game net driver for now) if stats or perf counters are used
 		if (NetDriverName == NAME_GameNetDriver &&
@@ -7981,7 +8014,7 @@ void UNetDriver::UpdateNetworkStats()
 #if STATS
 				FThreadStats::IsCollectingData() ||
 #endif
-				bCollectNetStats)
+				bCollectNetStats || GbEnableNetStats)
 			{
 				const float RealTime = CurrentRealtimeSeconds - StatUpdateTime;
 
@@ -8278,11 +8311,15 @@ void UNetDriver::UpdateNetworkStats()
 			NetGUIDOutBytes = 0;
 			NetGUIDInBytes = 0;
 			InPackets = 0;
+			PrevInPackets = 0;
 			OutPackets = 0;
+			PrevOutPackets = 0;
 			InBunches = 0;
 			OutBunches = 0;
 			OutPacketsLost = 0;
+			PrevOutPacketsLost = 0;
 			InPacketsLost = 0;
+			PrevInPacketsLost = 0;
 			VoicePacketsSent = 0;
 			VoiceBytesSent = 0;
 			VoicePacketsRecv = 0;
diff --git a/Engine/Source/Runtime/Engine/Public/Net/NetworkMetricsDefs.h b/Engine/Source/Runtime/Engine/Public/Net/NetworkMetricsDefs.h
--- a/Engine/Source/Runtime/Engine/Public/Net/NetworkMetricsDefs.h
+++ b/Engine/Source/Runtime/Engine/Public/Net/NetworkMetricsDefs.h
@@ -11,6 +11,10 @@ extern ENGINE_API const FName AddedConnections;
 extern ENGINE_API const FName ClosedConnectionsDueToReliableBufferOverflow;
 extern ENGINE_API const FName InPackets;
 extern ENGINE_API const FName OutPackets;
+extern ENGINE_API const FName InPacketsPerFrame;
+extern ENGINE_API const FName OutPacketsPerFrame;
+extern ENGINE_API const FName InLostPacketsFoundPerFrame;
+extern ENGINE_API const FName OutLostPacketsFoundPerFrame;
 extern ENGINE_API const FName InPacketsClientPerSecondAvg;
 extern ENGINE_API const FName InPacketsClientPerSecondMax;
 extern ENGINE_API const FName InPacketsClientPerSecondMin;
