Fixed SplineMeshComponent dirtying the package even if the actor has the flag RF_Transient.

diff --git a/Engine/Source/Runtime/Engine/Private/Components/SplineMeshComponent.cpp b/Engine/Source/Runtime/Engine/Private/Components/SplineMeshComponent.cpp
--- a/Engine/Source/Runtime/Engine/Private/Components/SplineMeshComponent.cpp
+++ b/Engine/Source/Runtime/Engine/Private/Components/SplineMeshComponent.cpp
@@ -828,7 +828,10 @@ bool USplineMeshComponent::Modify(bool bAlwaysMarkDirty)
 
 	if (BodySetup != nullptr)
 	{
-		BodySetup->Modify(bAlwaysMarkDirty);
+		// BodySetup shares the same package as its component.
+		// Rely on the above call to Super::Modify(bAlwaysMarkDirty) to dirty the package if necessary.
+		// Note that UActorComponent::Modify can force bAlwaysMarkDirty to false in some cases (like if the Actor is transient).
+		BodySetup->Modify(false);
 	}
 
 	return bSavedToTransactionBuffer;
