Support of overridable in T3D copy/paste format

diff --git a/Engine/Source/Editor/UnrealEd/Private/EditorObject.cpp b/Engine/Source/Editor/UnrealEd/Private/EditorObject.cpp
--- a/Engine/Source/Editor/UnrealEd/Private/EditorObject.cpp
+++ b/Engine/Source/Editor/UnrealEd/Private/EditorObject.cpp
@@ -33,6 +33,8 @@
 #include "InstancedFoliage.h"
 #include "Components/BrushComponent.h"
 #include "Algo/Transform.h"
+#include "UObject/OverridableManager.h"
+#include "UObject/OverriddenPropertySet.h"
 
 DEFINE_LOG_CATEGORY_STATIC(LogEditorObject, Log, All);
 
@@ -260,6 +262,9 @@ static const TCHAR* ImportProperties(
 	if ( SourceText == NULL )
 		return NULL;
 
+	FOverriddenPropertySet* OverriddenProperties = SubobjectOuter ? FOverridableManager::Get().GetOverriddenProperties(*SubobjectOuter) : nullptr;
+	FEnableOverridableSerializationScope Scope(OverriddenProperties != nullptr, OverriddenProperties);
+
 	// Cannot create subobjects when importing struct defaults, or if SubobjectOuter (used as the Outer for any subobject declarations encountered) is NULL
 	bool bSubObjectsAllowed = !ObjectStruct->IsA(UScriptStruct::StaticClass()) && SubobjectOuter != NULL;
 
@@ -508,6 +513,10 @@ static const TCHAR* ImportProperties(
 			}
 			else 
 			{
+				FName OverriddenOperationName = NAME_None;
+				FParse::Value(Str,TEXT("OverriddenOperation="), OverriddenOperationName);
+				TOptional<EOverriddenPropertyOperation> OverriddenOperation = GetOverriddenOperationFromName(OverriddenOperationName);
+
 				UObject* Archetype = NULL;
 				UObject* ComponentTemplate = NULL;
 
@@ -674,6 +683,15 @@ static const TCHAR* ImportProperties(
 				}
 				FArchiveReplaceObjectRef<UObject> ReplaceAr(SubobjectOuter, ReplacementMap, EArchiveReplaceObjectFlags::IgnoreArchetypeRef);
 
+				if (OverriddenOperation.IsSet())
+				{
+					FOverridableManager::Get().SetOverriddenProperties(*ComponentTemplate, OverriddenOperation.GetValue());
+				}
+				else
+				{
+					FOverridableManager::Get().Disable(*ComponentTemplate);
+				}
+
 				// import the properties for the subobject
 				SourceText = ImportObjectProperties(
 					(uint8*)ComponentTemplate, 
@@ -1038,6 +1056,10 @@ static const TCHAR* ImportCreateSubObjectsStep(
 			FString ExportedObjectFullName;
 			FParse::Value(Str, TEXT("ExportPath="), ExportedObjectFullName);
 
+			FName OverriddenOperationName = NAME_None;
+			const bool bParsed = FParse::Value(Str,TEXT("OverriddenOperation="), OverriddenOperationName);
+			TOptional<EOverriddenPropertyOperation> OverriddenOperation = bParsed ? GetOverriddenOperationFromName(OverriddenOperationName) : TOptional<EOverriddenPropertyOperation>();
+
 			// points to the parent class's template subobject/component, if we are overriding a subobject/component declared in our parent class
 			UObject* BaseTemplate = nullptr;
 			bool bRedefiningSubobject = false;
@@ -1267,6 +1289,15 @@ static const TCHAR* ImportCreateSubObjectsStep(
 					ObjectRemapper->Add(MoveTemp(ExportedObjectFullName), ComponentTemplate);
 				}
 
+				if (OverriddenOperation.IsSet())
+				{
+					FOverridableManager::Get().SetOverriddenProperties(*ComponentTemplate, OverriddenOperation.GetValue());
+				}
+				else
+				{
+					FOverridableManager::Get().Disable(*ComponentTemplate);
+				}
+
 				// Create the subobjects for the subobject
 				CurrentSourceText = ImportCreateObjectsStep(Params);
 			}
@@ -1319,6 +1350,9 @@ static const TCHAR* ImportPropertiesStep(
 		return nullptr;
 	}
 
+	FOverriddenPropertySet* OverriddenProperties = SubobjectOuter ? FOverridableManager::Get().GetOverriddenProperties(*SubobjectOuter) : nullptr;
+	FEnableOverridableSerializationScope Scope(OverriddenProperties != nullptr, OverriddenProperties);
+
 	// Cannot create subobjects when importing struct defaults, or if SubobjectOuter (used as the Outer for any subobject declarations encountered) is NULL
 	bool bSubObjectsAllowed = !ObjectStruct->IsA(UScriptStruct::StaticClass()) && SubobjectOuter != NULL;
 
diff --git a/Engine/Source/Editor/UnrealEd/Private/Factories/EditorFactories.cpp b/Engine/Source/Editor/UnrealEd/Private/Factories/EditorFactories.cpp
--- a/Engine/Source/Editor/UnrealEd/Private/Factories/EditorFactories.cpp
+++ b/Engine/Source/Editor/UnrealEd/Private/Factories/EditorFactories.cpp
@@ -283,6 +283,7 @@
 #include "UDIMUtilities.h"
 #include "FileHelpers.h"
 #include "StaticMeshOperations.h"
+#include "UObject/OverridableManager.h"
 
 
 DEFINE_LOG_CATEGORY(LogEditorFactories);
@@ -5348,6 +5349,10 @@ void FCustomizableTextObjectFactory::ProcessBuffer(UObject* InParent, EObjectFla
 					ObjArchetype = LoadObject<UObject>(nullptr, *ObjArchetypeName, nullptr, LOAD_None, nullptr);
 				}
 
+				FName OverriddenOperationName = NAME_None;
+				const bool bParsed = FParse::Value(Str,TEXT("OverriddenOperation="),OverriddenOperationName);
+				TOptional<EOverriddenPropertyOperation> OverriddenOperation = bParsed ? GetOverriddenOperationFromName(OverriddenOperationName) : TOptional<EOverriddenPropertyOperation>();
+
 				UObject* ObjectParent = InParent ? InParent : GetParentForNewObject(ObjClass);
 
 				// Make sure we are allowed to create this object
@@ -5363,6 +5368,16 @@ void FCustomizableTextObjectFactory::ProcessBuffer(UObject* InParent, EObjectFla
 				// Spawn the object and reset it's archetype
 				UObject* CreatedObject = NewObject<UObject>(ObjectParent, ObjClass, ObjName, Flags, ObjArchetype, !!ObjectParent, &InstanceGraph);
 
+				// Overridable handling
+				if (OverriddenOperation.IsSet())
+				{
+					FOverridableManager::Get().SetOverriddenProperties(*CreatedObject, OverriddenOperation.GetValue());
+				}
+				else
+				{
+					FOverridableManager::Get().Disable(*CreatedObject);
+				}
+
 				// Get property text for the new object.
 				FString PropText, PropLine;
 				FString SubObjText;
diff --git a/Engine/Source/Runtime/CoreUObject/Private/UObject/Class.cpp b/Engine/Source/Runtime/CoreUObject/Private/UObject/Class.cpp
--- a/Engine/Source/Runtime/CoreUObject/Private/UObject/Class.cpp
+++ b/Engine/Source/Runtime/CoreUObject/Private/UObject/Class.cpp
@@ -3303,6 +3303,9 @@ void UScriptStruct::ExportText(FString& ValueStr, const void* Value, const void*
 		{
 			for (int32 Index = 0; Index < It->ArrayDim; Index++)
 			{
+				const bool bStaticArray = It->ArrayDim > 1;
+				FOverridableTextPortPropertyPathScope ScopePath(*It, bStaticArray ? Index : INDEX_NONE, bStaticArray ? EPropertyVisitorInfoType::StaticArrayIndex : EPropertyVisitorInfoType::None);
+
 				FString InnerValue;
 				if (It->ExportText_InContainer(Index, InnerValue, Value, Defaults, OwnerObject, PPF_Delimited | PortFlags, ExportRootScope))
 				{
@@ -3322,13 +3325,27 @@ void UScriptStruct::ExportText(FString& ValueStr, const void* Value, const void*
 
 					const FString PropertyName = (PortFlags & (PPF_ExternalEditor | PPF_BlueprintDebugView)) != 0 ? It->GetAuthoredName() : It->GetName();
 
-					if (It->ArrayDim == 1)
+					FString OverridableOperation;
+					if (FOverriddenPropertySet* OverriddenProperties = FOverridableSerializationLogic::GetOverriddenProperties())
+					{
+						FPropertyVisitorPath* Path = FOverridableSerializationLogic::GetOverriddenPortTextPropertyPath();
+						checkf(Path, TEXT("Expecting a path"));
+						FArchiveSerializedPropertyChain Chain = Path->ToSerializedPropertyChain();
+						
+						const EOverriddenPropertyOperation Operation = OverriddenProperties->GetOverriddenPropertyOperation(&Chain, /*Property*/nullptr);
+						if (Operation != EOverriddenPropertyOperation::None)
+						{
+							OverridableOperation = FString::Printf(TEXT("<%s>"), *GetOverriddenOperationString(Operation));
+						}
+					}
+
+					if (!bStaticArray)
 					{
-						ValueStr += FString::Printf(TEXT("%s="), *PropertyName);
+						ValueStr += FString::Printf(TEXT("%s%s="), *PropertyName, *OverridableOperation);
 					}
 					else
 					{
-						ValueStr += FString::Printf(TEXT("%s[%i]="), *PropertyName, Index);
+						ValueStr += FString::Printf(TEXT("%s[%i]%s="), *PropertyName, Index, *OverridableOperation);
 					}
 					ValueStr += InnerValue;
 				}
diff --git a/Engine/Source/Runtime/CoreUObject/Private/UObject/OverriddenPropertySet.cpp b/Engine/Source/Runtime/CoreUObject/Private/UObject/OverriddenPropertySet.cpp
--- a/Engine/Source/Runtime/CoreUObject/Private/UObject/OverriddenPropertySet.cpp
+++ b/Engine/Source/Runtime/CoreUObject/Private/UObject/OverriddenPropertySet.cpp
@@ -21,12 +21,12 @@ DEFINE_LOG_CATEGORY(LogOverridableObject);
 //----------------------------------------------------------------------//
 thread_local bool FOverridableSerializationLogic::bUseOverridableSerialization = false;
 thread_local FOverriddenPropertySet* FOverridableSerializationLogic::OverriddenProperties = nullptr;
+thread_local FPropertyVisitorPath* FOverridableSerializationLogic::OverriddenPortTextPropertyPath = nullptr;
 
-EOverriddenPropertyOperation FOverridableSerializationLogic::GetOverriddenPropertyOperation(const FArchive& Ar, FProperty* Property /*= nullptr*/, uint8* DataPtr /*= nullptr*/, uint8* DefaultValue)
+EOverriddenPropertyOperation FOverridableSerializationLogic::GetOverriddenPropertyOperation(const int32 PortFlags, const FArchiveSerializedPropertyChain* CurrentPropertyChain, FProperty* Property, const void* DataPtr, const void* DefaultValue)
 {
-	checkf(bUseOverridableSerialization, TEXT("Nobody should use this method if it is not setup to use overridable serialization"));
+	checkf(FOverridableSerializationLogic::bUseOverridableSerialization, TEXT("Nobody should use this method if it is not setup to use overridable serialization"));
 
-	const FArchiveSerializedPropertyChain* CurrentPropertyChain = Ar.GetSerializedPropertyChain();
 	const EOverriddenPropertyOperation OverriddenOperation = OverriddenProperties ? OverriddenProperties->GetOverriddenPropertyOperation(CurrentPropertyChain, Property) : EOverriddenPropertyOperation::None;
 	if (OverriddenOperation != EOverriddenPropertyOperation::None)
 	{
@@ -45,7 +45,7 @@ EOverriddenPropertyOperation FOverridableSerializationLogic::GetOverriddenProper
 		if (OverriddenProperties && OverriddenProperties->IsCDOOwningProperty(*CurrentProperty))
 		{
 			// Only need serialize this value if it is different from the default property value
-			if (!CurrentProperty->Identical(DataPtr, DefaultValue, Ar.GetPortFlags()))
+			if (!CurrentProperty->Identical(DataPtr, DefaultValue, PortFlags))
 			{
 				return 	EOverriddenPropertyOperation::Replace;
 			}
@@ -55,6 +55,36 @@ EOverriddenPropertyOperation FOverridableSerializationLogic::GetOverriddenProper
 	return EOverriddenPropertyOperation::None;
 }
 
+EOverriddenPropertyOperation FOverridableSerializationLogic::GetOverriddenPropertyOperation(const FArchive& Ar, FProperty* Property /*= nullptr*/, uint8* DataPtr /*= nullptr*/, uint8* DefaultValue /*= nullptr*/)
+{
+	const FArchiveSerializedPropertyChain* CurrentPropertyChain = Ar.GetSerializedPropertyChain();
+	return GetOverriddenPropertyOperation(Ar.GetPortFlags(), CurrentPropertyChain, Property, DataPtr, DefaultValue);
+}
+
+EOverriddenPropertyOperation FOverridableSerializationLogic::GetOverriddenPropertyOperationForPortText(const void* DataPtr, const void* DefaultValue, int32 PortFlags)
+{
+	checkf(OverriddenPortTextPropertyPath, TEXT("Expecting an overridden port text path"));
+
+	const FArchiveSerializedPropertyChain CurrentPropertyChain = OverriddenPortTextPropertyPath->ToSerializedPropertyChain();
+	return GetOverriddenPropertyOperation(PortFlags,  &CurrentPropertyChain, nullptr, DataPtr, DefaultValue);
+}
+
+FPropertyVisitorPath* FOverridableSerializationLogic::GetOverriddenPortTextPropertyPath()
+{
+	return OverriddenPortTextPropertyPath;
+}
+
+void FOverridableSerializationLogic::SetOverriddenPortTextPropertyPath(FPropertyVisitorPath& Path)
+{
+	checkf(OverriddenPortTextPropertyPath == nullptr, TEXT("Should not set a path on top of an existing one"));
+	OverriddenPortTextPropertyPath = &Path;
+}
+
+void FOverridableSerializationLogic::ResetOverriddenPortTextPropertyPath()
+{
+	OverriddenPortTextPropertyPath = nullptr;
+}
+
 //----------------------------------------------------------------------//
 // FOverridableSerializationScope
 //----------------------------------------------------------------------//
@@ -85,6 +115,46 @@ FEnableOverridableSerializationScope::~FEnableOverridableSerializationScope()
 	}
 }
 
+//----------------------------------------------------------------------//
+// FOverridableTextImportPropertyPathScope
+//----------------------------------------------------------------------//
+FOverridableTextPortPropertyPathScope::FOverridableTextPortPropertyPathScope(const FProperty* InProperty, int32 InIndex/* = INDEX_NONE*/, EPropertyVisitorInfoType InPropertyInfo /*= EPropertyVisitorInfoType::None*/)
+{
+	if (!FOverridableSerializationLogic::IsEnabled())
+	{
+		return;
+	}
+
+	checkf(InProperty, TEXT("Expecting a valid property ptr"));
+
+	// Save property for comparison in the destructor
+	Property = InProperty;
+
+	FPropertyVisitorPath* Path = FOverridableSerializationLogic::GetOverriddenPortTextPropertyPath();
+	if (!Path)
+	{
+		FOverridableSerializationLogic::SetOverriddenPortTextPropertyPath(DefaultPath);
+		Path = &DefaultPath;
+	}
+
+	Path->Push(FPropertyVisitorInfo(InProperty, InIndex, InPropertyInfo));
+}
+
+FOverridableTextPortPropertyPathScope::~FOverridableTextPortPropertyPathScope()
+{
+	if (Property)
+	{
+		FPropertyVisitorPath* Path = FOverridableSerializationLogic::GetOverriddenPortTextPropertyPath();
+		checkf(Path, TEXT("Expecting a valid path "));
+		checkf(Path->Num(), TEXT("Expecting at least one property in the path"));
+		verifyf(Path->Pop().Property == Property, TEXT("Expecting at the top property to match the one we pushed in the constructor"));
+		if (!Path->Num())
+		{
+			FOverridableSerializationLogic::ResetOverriddenPortTextPropertyPath();
+		}
+	}
+}
+
 //----------------------------------------------------------------------//
 // FOverriddenPropertyNodeID
 //----------------------------------------------------------------------//
@@ -159,8 +229,8 @@ EOverriddenPropertyOperation FOverriddenPropertySet::GetOverriddenPropertyOperat
 		}
 
 
-		// Special handling for for instanced subobjects 
-		if (const FObjectProperty* ObjectProperty = CastField<FObjectProperty>(CurrentProperty))
+		// Special handling for instanced subobjects 
+		if (const FObjectPropertyBase* ObjectProperty = CastField<FObjectPropertyBase>(CurrentProperty))
 		{
 			if (PropertyIterator->GetNextNode())
 			{
@@ -179,7 +249,7 @@ EOverriddenPropertyOperation FOverriddenPropertySet::GetOverriddenPropertyOperat
 
 			// Only special case is instanced subobjects, otherwise we fallback to full array override
 			checkf(ArrayProperty->Inner, TEXT("Expecting an inner type for Arrays"));
-			if (const FObjectProperty* InnerObjectProperty = ArrayProperty->Inner->HasAnyPropertyFlags(CPF_PersistentInstance) ? CastField<FObjectProperty>(ArrayProperty->Inner) : nullptr)
+			if (const FObjectPropertyBase* InnerObjectProperty = ArrayProperty->Inner->HasAnyPropertyFlags(CPF_PersistentInstance) ? CastField<FObjectPropertyBase>(ArrayProperty->Inner) : nullptr)
 			{
 				FScriptArrayHelper ArrayHelper(ArrayProperty, SubValuePtr);
 				if(ArrayHelper.IsValidIndex(ArrayIndex))
@@ -224,7 +294,7 @@ EOverriddenPropertyOperation FOverriddenPropertySet::GetOverriddenPropertyOperat
 				if (PropertyIterator->GetNextNode())
 				{
 					// Forward any sub queries to the subobject
-					if (const FObjectProperty* ValueInstancedObjectProperty = MapProperty->ValueProp->HasAnyPropertyFlags(CPF_PersistentInstance) ? CastField<FObjectProperty>(MapProperty->ValueProp) : nullptr)
+					if (const FObjectPropertyBase* ValueInstancedObjectProperty = MapProperty->ValueProp->HasAnyPropertyFlags(CPF_PersistentInstance) ? CastField<FObjectPropertyBase>(MapProperty->ValueProp) : nullptr)
 					{
 						if (UObject* ValueSubObject = ValueInstancedObjectProperty->GetObjectPropertyValue(MapHelper.GetValuePtr(InternalMapIndex)))
 						{
@@ -236,7 +306,7 @@ EOverriddenPropertyOperation FOverriddenPropertySet::GetOverriddenPropertyOperat
 				{
 					// Caller wants to know about any override state on the reference of the map pair itself
 					checkf(MapProperty->KeyProp, TEXT("Expecting a key type for Maps"));
-					FObjectProperty* KeyObjectProperty = CastField<FObjectProperty>(MapProperty->KeyProp);
+					FObjectPropertyBase* KeyObjectProperty = CastField<FObjectPropertyBase>(MapProperty->KeyProp);
 
 					FOverriddenPropertyNodeID OverriddenKeyID;
 					if (const UObject* OverriddenKeyObject = KeyObjectProperty ? KeyObjectProperty->GetObjectPropertyValue(MapHelper.GetKeyPtr(InternalMapIndex)) : nullptr)
@@ -306,7 +376,7 @@ bool FOverriddenPropertySet::ClearOverriddenProperty(FOverriddenPropertyNode& Pa
 		}
 
 		// Special handling for for instanced subobjects 
-		if (const FObjectProperty* ObjectProperty = CastField<FObjectProperty>(CurrentProperty))
+		if (const FObjectPropertyBase* ObjectProperty = CastField<FObjectPropertyBase>(CurrentProperty))
 		{
 			if (UObject* SubObject = ObjectProperty->GetObjectPropertyValue(SubValuePtr))
 			{
@@ -327,7 +397,7 @@ bool FOverriddenPropertySet::ClearOverriddenProperty(FOverriddenPropertyNode& Pa
 			ArrayIndex = PropertyEvent.GetArrayIndex(CurrentProperty->GetName());
 
 			// Only special case is instanced subobjects, otherwise we fallback to full array override
-			if (FObjectProperty* InnerObjectProperty = CastField<FObjectProperty>(ArrayProperty->Inner))
+			if (FObjectPropertyBase* InnerObjectProperty = CastField<FObjectPropertyBase>(ArrayProperty->Inner))
 			{
 				if (InnerObjectProperty->HasAnyPropertyFlags(CPF_PersistentInstance))
 				{
@@ -378,7 +448,7 @@ bool FOverriddenPropertySet::ClearOverriddenProperty(FOverriddenPropertyNode& Pa
 			FScriptMapHelper MapHelper(MapProperty, SubValuePtr);
 
 			const int32 InternalMapIndex = ArrayIndex != INDEX_NONE ? MapHelper.FindInternalIndex(ArrayIndex) : INDEX_NONE;
-			const FObjectProperty* ValueInstancedObjectProperty = MapProperty->ValueProp->HasAnyPropertyFlags(CPF_PersistentInstance) ? CastField<FObjectProperty>(MapProperty->ValueProp) : nullptr;
+			const FObjectPropertyBase* ValueInstancedObjectProperty = MapProperty->ValueProp->HasAnyPropertyFlags(CPF_PersistentInstance) ? CastField<FObjectPropertyBase>(MapProperty->ValueProp) : nullptr;
 
 			// If there is a next node, it is probably because the map value is holding a instanced subobject and the user is changing value on it.
 			// So forward the call to the instanced subobject
@@ -413,7 +483,7 @@ bool FOverriddenPropertySet::ClearOverriddenProperty(FOverriddenPropertyNode& Pa
 			else if (MapHelper.IsValidIndex(InternalMapIndex) && CurrentOverriddenPropertyNode)
 			{
 				checkf(MapProperty->KeyProp, TEXT("Expecting a key type for Maps"));
-				FObjectProperty* KeyObjectProperty = CastField<FObjectProperty>(MapProperty->KeyProp);
+				FObjectPropertyBase* KeyObjectProperty = CastField<FObjectPropertyBase>(MapProperty->KeyProp);
 
 				// Calculate the node from the key
 				FOverriddenPropertyNodeID OverriddenKeyID;
@@ -535,7 +605,7 @@ void FOverriddenPropertySet::NotifyPropertyChange(FOverriddenPropertyNode* Paren
 	if (const FArrayProperty* ArrayProperty = CastField<FArrayProperty>(Property))
 	{
 		// Only special case is instanced subobjects, otherwise we fallback to full array override
-		if (FObjectProperty* InnerObjectProperty = CastField<FObjectProperty>(ArrayProperty->Inner))
+		if (FObjectPropertyBase* InnerObjectProperty = CastField<FObjectPropertyBase>(ArrayProperty->Inner))
 		{
 			if (InnerObjectProperty->HasAnyPropertyFlags(CPF_PersistentInstance))
 			{
@@ -719,13 +789,13 @@ void FOverriddenPropertySet::NotifyPropertyChange(FOverriddenPropertyNode* Paren
 	{
 		// Special handling of instanced subobjects
 		checkf(MapProperty->KeyProp, TEXT("Expecting a key type for Maps"));
-		FObjectProperty* KeyObjectProperty = CastField<FObjectProperty>(MapProperty->KeyProp);
+		FObjectPropertyBase* KeyObjectProperty = CastField<FObjectPropertyBase>(MapProperty->KeyProp);
 
 		// SubObjects
 		checkf(!KeyObjectProperty || !MapProperty->KeyProp->HasAnyPropertyFlags(CPF_PersistentInstance) || CastField<FClassProperty>(MapProperty->KeyProp), TEXT("Keys as a instanced subobject is not supported yet"));
 
 		checkf(MapProperty->ValueProp, TEXT("Expecting a value type for Maps"));
-		FObjectProperty* ValueInstancedObjectProperty = MapProperty->ValueProp->HasAnyPropertyFlags(CPF_PersistentInstance) ? CastField<FObjectProperty>(MapProperty->ValueProp) : nullptr;
+		FObjectPropertyBase* ValueInstancedObjectProperty = MapProperty->ValueProp->HasAnyPropertyFlags(CPF_PersistentInstance) ? CastField<FObjectPropertyBase>(MapProperty->ValueProp) : nullptr;
 
 		FScriptMapHelper MapHelper(MapProperty, SubValuePtr);
 		int32 LogicalMapIndex = PropertyEvent.GetArrayIndex(Property->GetName());
@@ -970,7 +1040,7 @@ void FOverriddenPropertySet::NotifyPropertyChange(FOverriddenPropertyNode* Paren
 		}
 		return;
 	}
-	else if (const FObjectProperty* ObjectProperty = CastField<FObjectProperty>(Property))
+	else if (const FObjectPropertyBase* ObjectProperty = CastField<FObjectPropertyBase>(Property))
 	{
 		if (!PropertyNode->GetNextNode())
 		{
@@ -1041,7 +1111,7 @@ bool FOverriddenPropertySet::ClearOverriddenProperty(const FPropertyChangedEvent
 	{
 		return ClearOverriddenProperty(*RootNode, PropertyEvent, PropertyNode, Owner);
 	}
-	return false;
+	return true;
 }
 
 void FOverriddenPropertySet::OverrideProperty(const FPropertyChangedEvent& PropertyEvent, const FEditPropertyChain::TDoubleLinkedListNode* PropertyNode, const void* Data)
@@ -1228,3 +1298,59 @@ const FOverriddenPropertyNode* FOverriddenPropertySet::GetOverriddenPropertyNode
 
 	return OverriddenPropertyNode;
 }
+
+int32 UE::OverridableMapUtilities::FindKeyInternalIndex(const FOverriddenPropertyNodeID& KeyIDToFind, FScriptMapHelper& MapHelper)
+{
+	// Special case for object we didn't use the pointer to create the key
+	if (const FObjectProperty* KeyObjectProperty = CastField<FObjectProperty>(MapHelper.KeyProp))
+	{
+		for (FScriptMapHelper::FIterator It(MapHelper); It; ++It) 
+		{
+			if (UObject* CurrentObject = KeyObjectProperty->GetObjectPropertyValue(MapHelper.GetKeyPtr(It)))
+			{
+				if (KeyIDToFind == FOverriddenPropertyNodeID(*CurrentObject))
+				{
+					return It.GetInternalIndex();
+				}
+			}
+		}
+	}
+	else
+	{
+		// Default case, just import the text as key value for comparison
+		void* TempKeyValueStorage = FMemory_Alloca(MapHelper.MapLayout.SetLayout.Size);
+		MapHelper.KeyProp->InitializeValue(TempKeyValueStorage);
+
+		FString KeyToFind(KeyIDToFind.ToString());
+		MapHelper.KeyProp->ImportText_Direct(*KeyToFind, TempKeyValueStorage, nullptr, PPF_None);
+
+		const int32 InternalIndex = MapHelper.FindMapPairIndexFromHash(TempKeyValueStorage);
+
+		MapHelper.KeyProp->DestroyValue(TempKeyValueStorage);
+
+		return InternalIndex;
+	}
+	return INDEX_NONE;
+}
+
+FOverriddenPropertyNodeID UE::OverridableMapUtilities::GetIDFromKey(const FProperty* KeyProp, uint8* KeyData)
+{
+	// Special case for object as we cannot save the pointer value as a key
+	if (const FObjectProperty* KeyObjectProperty = CastField<FObjectProperty>(KeyProp))
+	{
+		if (const UObject* Object = KeyObjectProperty->GetObjectPropertyValue(KeyData))
+		{
+			return FOverriddenPropertyNodeID(*Object);
+		}
+	}
+	else
+	{
+		// Default case, just export key value as text
+		FString KeyString;
+		KeyProp->ExportTextItem_Direct(KeyString, KeyData, /*DefaultValue*/nullptr, /*Parent*/nullptr, PPF_None);
+		return FOverriddenPropertyNodeID(FName(KeyString));
+	}
+
+	checkf(false, TEXT("This case is not handled"))
+	return FOverriddenPropertyNodeID();
+};
diff --git a/Engine/Source/Runtime/CoreUObject/Private/UObject/Property.cpp b/Engine/Source/Runtime/CoreUObject/Private/UObject/Property.cpp
--- a/Engine/Source/Runtime/CoreUObject/Private/UObject/Property.cpp
+++ b/Engine/Source/Runtime/CoreUObject/Private/UObject/Property.cpp
@@ -10,12 +10,14 @@
 #include "Math/Sphere.h"
 #include "Misc/AsciiSet.h"
 #include "Misc/Guid.h"
+#include "Misc/ScopeExit.h"
 #include "Misc/StringBuilder.h"
 #include "Serialization/TestUndeclaredScriptStructObjectReferences.h"
 #include "Templates/Casts.h"
 #include "UObject/Class.h"
 #include "UObject/CoreNetTypes.h"
 #include "UObject/CoreRedirects.h"
+#include "UObject/OverridableManager.h"
 #include "UObject/Package.h"
 #include "UObject/PropertyHelper.h"
 #include "UObject/PropertyTypeName.h"
@@ -1139,7 +1141,9 @@ bool FProperty::ExportText_Direct
 	UObject*	ExportRootScope
 	) const
 {
-	if( Data==Delta || !Identical(Data,Delta,PortFlags) )
+	if( Data==Delta || 
+		(FOverridableSerializationLogic::IsEnabled() && FOverridableSerializationLogic::GetOverriddenPropertyOperationForPortText(Data, Delta, PortFlags) != EOverriddenPropertyOperation::None)  ||
+		(!FOverridableSerializationLogic::IsEnabled() && !Identical(Data,Delta,PortFlags)) )
 	{
 		ExportText_Internal
 			(
@@ -1552,7 +1556,8 @@ static const int32 ReadArrayIndex(const UStruct* ObjectStruct, const TCHAR*& Str
 		FString IndexText(TEXT(""));
 		while ( *Str && *Str != ')' && *Str != ']' )
 		{
-			if ( *Str == TCHAR('=') )
+			// Stop at either we reach the = or the beginning of the overridable operation ex: <replace>
+			if ( *Str == TCHAR('=') || *Str == TCHAR('<'))
 			{
 				// we've encountered an equals sign before the closing bracket
 				Warn->Logf(ELogVerbosity::Warning, TEXT("Missing ')' in default properties subscript: %s"), Start);
@@ -1614,6 +1619,61 @@ static const int32 ReadArrayIndex(const UStruct* ObjectStruct, const TCHAR*& Str
 	return Index;
 }
 
+/**
+ * Attempts to read an array index (xxx) sequence.  Handles const/enum replacements, etc.
+ * @param	ObjectStruct	the scope of the object/struct containing the property we're currently importing
+ * @param	Str				[out] pointer to the the buffer containing the property value to import
+ * @param	Warn			the output device to send errors/warnings to
+ * @return	the array index for this defaultproperties line.  INDEX_NONE if this line doesn't contains an array specifier, or 0 if there was an error parsing the specifier.
+ */
+static TOptional<EOverriddenPropertyOperation> ReadOverriddenOperation(const TCHAR*& Str, FOutputDevice* Warn)
+{
+	const TCHAR* Start = Str;
+	TOptional<EOverriddenPropertyOperation> Operation;
+	SkipWhitespace(Str);
+
+	if (*Str == '<')
+	{
+		Str++;
+		FString IndexText(TEXT(""));
+		while (*Str && *Str != '>')
+		{
+			if (*Str == TCHAR('='))
+			{
+				// we've encountered an equals sign before the closing bracket
+				Warn->Logf(ELogVerbosity::Warning, TEXT("Missing '>' in default properties subscript: %s"), Start);
+				return Operation;
+			}
+
+			IndexText += *Str++;
+		}
+
+		if (*Str++)
+		{
+			if (IndexText.Len() > 0)
+			{
+				Operation = GetOverriddenOperationFromString(IndexText);
+				if (!Operation.IsSet())
+				{
+					// unknown or invalid identifier specified for array subscript
+					Warn->Logf(ELogVerbosity::Warning, TEXT("Invalid overridable in default properties: %s"), Start);
+				}
+			}
+			else
+			{
+				// nothing was specified between the opening and closing parenthesis
+				Warn->Logf(ELogVerbosity::Warning, TEXT("Empty overridable operation: %s"), Start);
+			}
+		}
+		else
+		{
+			Warn->Logf(ELogVerbosity::Warning, TEXT("Missing '>' in default properties subscript: %s"), Start );
+		}
+	}
+	return Operation;
+}
+
+
 /** 
  * Do not attempt to import this property if there is no value for it - i.e. (Prop1=,Prop2=)
  * This normally only happens for empty strings or empty dynamic arrays, and the alternative
@@ -1639,7 +1699,7 @@ const TCHAR* FProperty::ImportSingleProperty( const TCHAR* Str, void* DestData,
 	check(ObjectStruct);
 
 	constexpr FAsciiSet Whitespaces(" \t");
-	constexpr FAsciiSet Delimiters("=([.");
+	constexpr FAsciiSet Delimiters("=([.<");
 
 	// strip leading whitespace
 	const TCHAR* Start = FAsciiSet::Skip(Str, Whitespaces);
@@ -1704,6 +1764,8 @@ const TCHAR* FProperty::ImportSingleProperty( const TCHAR* Str, void* DestData,
 			return Str;
 		}
 
+		FOverridableTextPortPropertyPathScope Scope(Property);
+
 		// Parse an array operation, if present.
 		enum EArrayOp
 		{
@@ -1737,6 +1799,27 @@ const TCHAR* FProperty::ImportSingleProperty( const TCHAR* Str, void* DestData,
 		}
 
 		FArrayProperty* const ArrayProperty = ExactCastField<FArrayProperty>(Property);
+
+		// Parse overridable info
+		TOptional<EOverriddenPropertyOperation> Operation = ReadOverriddenOperation(Str, Warn);
+		FOverriddenPropertySet* OverriddenProperties = nullptr;
+		FPropertyVisitorPath* Path = nullptr;
+		// Skip add and remove operations as they should be handled independently
+		if (Operation.IsSet())
+		{
+			OverriddenProperties = FOverridableSerializationLogic::GetOverriddenProperties();
+			Path = FOverridableSerializationLogic::GetOverriddenPortTextPropertyPath();
+			if (OverriddenProperties &&
+				!Property->HasAllPropertyFlags(CPF_ExperimentalAlwaysOverriden) &&
+				Operation.GetValue() != EOverriddenPropertyOperation::Add && 
+				Operation.GetValue() != EOverriddenPropertyOperation::Remove)
+			{
+				checkf(Path, TEXT("Expecting a path"));
+				FArchiveSerializedPropertyChain Chain = Path->ToSerializedPropertyChain();
+				OverriddenProperties->SetOverriddenPropertyOperation(Operation.GetValue(), &Chain, /*Property*/nullptr);
+			}
+		}
+
 		FMulticastDelegateProperty* const MulticastDelegateProperty = CastField<FMulticastDelegateProperty>(Property);
 		if( MulticastDelegateProperty != NULL && ArrayOp != ADO_None )
 		{
@@ -1933,8 +2016,270 @@ const TCHAR* FProperty::ImportSingleProperty( const TCHAR* Str, void* DestData,
 			// try to read an array index
 			int32 Index = ReadArrayIndex(ObjectStruct, Str, Warn);
 
+			// strip whitespace before =
+			SkipWhitespace(Str);
+			if (*Str++ != '=')
+			{
+				Warn->Logf(ELogVerbosity::Warning, TEXT("Missing '=' in default properties assignment: %s"), Start );
+				return Str;
+			}
+			// strip whitespace after =
+			SkipWhitespace(Str);
+
+			auto ImportText = [SubobjectOuter, PortFlags, Warn, Start, ObjectStruct](const TCHAR* Buffer, const FProperty* Property, void* PropertyPtr)
+			{
+				FStringOutputDevice ImportError;
+				const TCHAR* Result = Property->ImportText_Direct(Buffer, PropertyPtr, SubobjectOuter, PortFlags, &ImportError);
+
+				// Spit any error we had while importing property
+				if (ImportError.Len() > 0)
+				{
+					TArray<FString> ImportErrors;
+					ImportError.ParseIntoArray(ImportErrors,LINE_TERMINATOR,true);
+
+					Warn->Logf(ELogVerbosity::Warning, TEXT("While importing text for property '%s' in '%s':"), *Property->GetName(), *ObjectStruct->GetName());
+					for (int32 ErrorIndex = 0; ErrorIndex < ImportErrors.Num(); ErrorIndex++)
+					{
+						Warn->Logf(ELogVerbosity::Warning, TEXT("%s"), *ImportErrors[ErrorIndex]);
+					}
+				}
+				else if (Result == Buffer && *Buffer == TCHAR('\0'))
+				{
+					Warn->Logf(ELogVerbosity::Warning, TEXT("Invalid property value in defaults: %s"), Start);
+				}
+				return Result;
+			};
+
+			// Handle modifying overridable operations on arrays
+			if (Operation.IsSet() && Operation.GetValue() != EOverriddenPropertyOperation::Replace)
+			{
+				if (ArrayProperty)
+				{
+					FScriptArrayHelper_InContainer ArrayHelper(ArrayProperty, DestData);
+					switch(Operation.GetValue())
+					{
+						case EOverriddenPropertyOperation::Remove:
+						{
+							checkf(ArrayProperty->Inner->HasAnyPropertyFlags(CPF_PersistentInstance), TEXT("Only instanced sub object is supporting remove operation"));
+							const FObjectPropertyBase* InnerObjectProperty = CastFieldChecked<FObjectPropertyBase>(ArrayProperty->Inner);
+
+							uint8* TempValueStorage = (uint8*)FMemory_Alloca(InnerObjectProperty->ElementSize);
+							InnerObjectProperty->InitializeValue(TempValueStorage);
+
+							const TCHAR* Result = ImportText(Str, InnerObjectProperty, TempValueStorage);
+							if (UObject* RemovedSubObject = InnerObjectProperty->GetObjectPropertyValue(TempValueStorage))
+							{
+								const int32 ArrayNum = ArrayHelper.Num();
+								for (int i = 0; i < ArrayNum; ++i)
+								{
+									UObject* CurrentObject = InnerObjectProperty->GetObjectPropertyValue(ArrayHelper.GetElementPtr(i));
+									if (CurrentObject->GetArchetype() == RemovedSubObject)
+									{
+										ArrayHelper.RemoveValues(i);
+										break;
+									}
+								}
+
+								if (Path)
+								{
+									// Need to fetch the ArrayOverriddenPropertyNode every loop as the previous iteration might have reallocated the node.
+									FArchiveSerializedPropertyChain Chain = Path->ToSerializedPropertyChain();
+									if (FOverriddenPropertyNode* ArrayOverriddenPropertyNode = OverriddenProperties ? OverriddenProperties->SetOverriddenPropertyOperation(EOverriddenPropertyOperation::Modified, &Chain, /*Property*/nullptr) : nullptr)
+									{
+										// Rebuild the overridden info
+										const FOverriddenPropertyNodeID RemovedSubObjectID(*RemovedSubObject);
+										OverriddenProperties->SetSubPropertyOperation(EOverriddenPropertyOperation::Remove, *ArrayOverriddenPropertyNode, RemovedSubObjectID);
+									}
+								}
+							}
+
+							InnerObjectProperty->DestroyValue(TempValueStorage);
+							return Result;
+						}
+						case EOverriddenPropertyOperation::Add:
+						{
+							// Special case for instanced sub objects
+							if (const FObjectPropertyBase* InnerObjectProperty =  ArrayProperty->Inner->HasAnyPropertyFlags(CPF_PersistentInstance) ? CastField<FObjectPropertyBase>(ArrayProperty->Inner) : nullptr)
+							{
+								uint8* TempValueStorage = (uint8*)FMemory_Alloca(InnerObjectProperty->ElementSize);
+								InnerObjectProperty->InitializeValue(TempValueStorage);
+
+								const TCHAR* Result = ImportText(Str, InnerObjectProperty, TempValueStorage);
+								if (UObject* AddedSubObject = InnerObjectProperty->GetObjectPropertyValue(TempValueStorage))
+								{
+									UObject* AddedSubObjectArchetype = AddedSubObject->GetArchetype();
+									const int32 ArrayNum = ArrayHelper.Num();
+									for (int i = 0; i < ArrayNum; ++i)
+									{
+										UObject* CurrentObject = InnerObjectProperty->GetObjectPropertyValue(ArrayHelper.GetElementPtr(i));
+										if (CurrentObject == AddedSubObject || (CurrentObject == AddedSubObjectArchetype))
+										{
+											Index = i;
+											break;
+										}
+									}
+
+									if (Index == INDEX_NONE)
+									{
+										Index = ArrayHelper.Num();
+										ArrayHelper.ExpandForIndex(Index);
+									}
+
+									InnerObjectProperty->SetObjectPropertyValue(ArrayHelper.GetRawPtr(Index), AddedSubObject);
+
+									if (Path)
+									{
+										// Need to fetch the ArrayOverriddenPropertyNode every loop as the previous iteration might have reallocated the node.
+										FArchiveSerializedPropertyChain Chain = Path->ToSerializedPropertyChain();
+										if (FOverriddenPropertyNode* ArrayOverriddenPropertyNode = OverriddenProperties ? OverriddenProperties->SetOverriddenPropertyOperation(EOverriddenPropertyOperation::Modified, &Chain, /*Property*/nullptr) : nullptr)
+										{
+											// Rebuild the overridden info
+											const FOverriddenPropertyNodeID AddedSubObjectID(*AddedSubObject);
+											OverriddenProperties->SetSubPropertyOperation(EOverriddenPropertyOperation::Add, *ArrayOverriddenPropertyNode, AddedSubObjectID);
+										}
+									}
+								}
+
+								InnerObjectProperty->DestroyValue(TempValueStorage);
+								return Result;
+							}
+
+							Index = ArrayHelper.Num();
+							ArrayHelper.ExpandForIndex(Index);
+							break;
+						}
+
+						default:
+							checkf(false, TEXT("Unsupported array operation while importing text"));
+							return nullptr;
+					}
+				}
+				else if (FMapProperty* const MapProperty = ExactCastField<FMapProperty>(Property))
+				{
+					FScriptMapHelper_InContainer MapHelper(MapProperty, DestData);
+
+					uint8* TempKeyValueStorage = (uint8*)FMemory_Alloca(MapProperty->MapLayout.SetLayout.Size);
+					MapProperty->KeyProp->InitializeValue(TempKeyValueStorage);
+					MapProperty->ValueProp->InitializeValue(TempKeyValueStorage + MapProperty->MapLayout.ValueOffset);
+					ON_SCOPE_EXIT
+					{
+						MapProperty->KeyProp->DestroyValue(TempKeyValueStorage);
+						MapProperty->ValueProp->DestroyValue(TempKeyValueStorage + MapProperty->MapLayout.ValueOffset);
+					};
+
+					switch(Operation.GetValue())
+					{
+						case EOverriddenPropertyOperation::Remove:
+						{
+							const TCHAR* Result = ImportText(Str, MapHelper.KeyProp, TempKeyValueStorage);
+							MapHelper.RemovePair(TempKeyValueStorage);
+
+							if (Path)
+							{
+								// Need to fetch the MapOverriddenPropertyNode every loop as the previous might have reallocated the node.
+								FArchiveSerializedPropertyChain Chain = Path->ToSerializedPropertyChain();
+								if (FOverriddenPropertyNode* MapOverriddenPropertyNode = OverriddenProperties ? OverriddenProperties->SetOverriddenPropertyOperation(EOverriddenPropertyOperation::Modified, &Chain, /*Property*/nullptr) : nullptr)
+								{
+									// Rebuild the overridden info
+									FOverriddenPropertyNodeID RemovedKeyID = UE::OverridableMapUtilities::GetIDFromKey(MapProperty->KeyProp, TempKeyValueStorage);
+									OverriddenProperties->SetSubPropertyOperation(EOverriddenPropertyOperation::Remove, *MapOverriddenPropertyNode, RemovedKeyID);
+								}
+							}
+
+							return Result;
+						}
+						case EOverriddenPropertyOperation::Modified:
+						{
+							const TCHAR* Result = Str;
+
+							SkipWhitespace(Result);
+							if (*Result++ != TCHAR('('))
+							{
+								return nullptr;
+							}
+							SkipWhitespace(Result);
+
+							Result = ImportText(Result, MapHelper.KeyProp, TempKeyValueStorage);
+							const int32 InternalIndex = MapHelper.FindMapPairIndexFromHash(TempKeyValueStorage);
+
+							SkipWhitespace(Result);
+							if (*Result++ != TCHAR(','))
+							{
+								return nullptr;
+							}
+
+							uint8* ValuePtr = InternalIndex != INDEX_NONE ? MapHelper.GetValuePtr(InternalIndex) : TempKeyValueStorage + MapHelper.MapLayout.ValueOffset;
+							Result = ImportText(Result, MapHelper.ValueProp, ValuePtr);
+
+							SkipWhitespace(Result);
+							if (*Result++ != TCHAR(')'))
+							{
+								return nullptr;
+							}
+
+							if (Path)
+							{
+								// Need to fetch the MapOverriddenPropertyNode every loop as the previous might have reallocated the node.
+								FArchiveSerializedPropertyChain Chain = Path->ToSerializedPropertyChain();
+								if (FOverriddenPropertyNode* MapOverriddenPropertyNode = OverriddenProperties ? OverriddenProperties->SetOverriddenPropertyOperation(EOverriddenPropertyOperation::Modified, &Chain, /*Property*/nullptr) : nullptr)
+								{
+									// Rebuild the overridden info
+									FOverriddenPropertyNodeID ModifiedKeyID = UE::OverridableMapUtilities::GetIDFromKey(MapProperty->KeyProp, TempKeyValueStorage);
+									OverriddenProperties->SetSubPropertyOperation(EOverriddenPropertyOperation::Modified, *MapOverriddenPropertyNode, ModifiedKeyID);
+								}
+							}
+							return Result;
+						}
+						case EOverriddenPropertyOperation::Add:
+						{
+							const TCHAR* Result = Str;
+							SkipWhitespace(Result);
+							if (*Result++ != TCHAR('('))
+							{
+								return nullptr;
+							}
+							SkipWhitespace(Result);
+
+							Result = ImportText(Result, MapHelper.KeyProp, TempKeyValueStorage);
+
+							SkipWhitespace(Result);
+							if (*Result++ != TCHAR(','))
+							{
+								return nullptr;
+							}
+							SkipWhitespace(Result);
+
+							void* ValuePtr =  MapHelper.FindOrAdd(TempKeyValueStorage);;
+							Result = ImportText(Result, MapHelper.ValueProp, ValuePtr);
+
+							SkipWhitespace(Result);
+							if (*Result++ != TCHAR(')'))
+							{
+								return nullptr;
+							}
+
+							if (Path)
+							{
+								// Need to fetch the MapOverriddenPropertyNode every loop as the previous might have reallocated the node.
+								FArchiveSerializedPropertyChain Chain = Path->ToSerializedPropertyChain();
+								if (FOverriddenPropertyNode* MapOverriddenPropertyNode = OverriddenProperties ? OverriddenProperties->SetOverriddenPropertyOperation(EOverriddenPropertyOperation::Modified, &Chain, /*Property*/nullptr) : nullptr)
+								{
+									// Rebuild the overridden info
+									FOverriddenPropertyNodeID AddedKeyID = UE::OverridableMapUtilities::GetIDFromKey(MapProperty->KeyProp, TempKeyValueStorage);
+									OverriddenProperties->SetSubPropertyOperation(EOverriddenPropertyOperation::Add, *MapOverriddenPropertyNode, AddedKeyID);
+								}
+							}
+							return Result;
+						}
+						default:
+							checkf(false, TEXT("Unsupported map operation while importing text"));
+							return nullptr;
+					}
+				}
+			}
+
 			// check for out of bounds on static arrays
-			if (ArrayProperty == NULL && Index >= Property->ArrayDim)
+			if (ArrayProperty == nullptr && Index >= Property->ArrayDim)
 			{
 				Warn->Logf(ELogVerbosity::Warning, TEXT("Out of bound array default property (%i/%i): %s"), Index, Property->ArrayDim, Start);
 				return Str;
@@ -1951,16 +2296,6 @@ const TCHAR* FProperty::ImportSingleProperty( const TCHAR* Str, void* DestData,
 			}
 			DefinedProperties.Add(D);
 
-			// strip whitespace before =
-			SkipWhitespace(Str);
-			if (*Str++ != '=')
-			{
-				Warn->Logf(ELogVerbosity::Warning, TEXT("Missing '=' in default properties assignment: %s"), Start );
-				return Str;
-			}
-			// strip whitespace after =
-			SkipWhitespace(Str);
-
 			if (!IsPropertyValueSpecified(Str) && ArrayProperty == nullptr)
 			{
 				// if we're not importing default properties for classes (i.e. we're pasting something in the editor or something)
@@ -1973,31 +2308,15 @@ const TCHAR* FProperty::ImportSingleProperty( const TCHAR* Str, void* DestData,
 			// not done above with ShouldPort() check because this is intentionally exported so we don't want it to cause errors on import
 			if (Property->GetFName() != NAME_Name || !Property->GetOwnerVariant().IsUObject() || Property->GetOwner<UObject>()->GetFName() != NAME_Object)
 			{
-				if (Index > -1 && ArrayProperty != NULL) //set single dynamic array element
+				if (Index > -1 && ArrayProperty != nullptr) //set single dynamic array element
 				{
 					FScriptArrayHelper_InContainer ArrayHelper(ArrayProperty, DestData);
 
 					ArrayHelper.ExpandForIndex(Index);
 
-					FStringOutputDevice ImportError;
-					const TCHAR* Result = ArrayProperty->Inner->ImportText_Direct(Str, ArrayHelper.GetRawPtr(Index), SubobjectOuter, PortFlags, &ImportError);
-					// Spit any error we had while importing property
-					if (ImportError.Len() > 0)
-					{
-						TArray<FString> ImportErrors;
-						ImportError.ParseIntoArray(ImportErrors,LINE_TERMINATOR,true);
-
-						for ( int32 ErrorIndex = 0; ErrorIndex < ImportErrors.Num(); ErrorIndex++ )
-						{
-							Warn->Logf(ELogVerbosity::Warning, TEXT("%s"), *ImportErrors[ErrorIndex]);
-						}
-					}
-					else if (Result == Str)
-					{
-						Warn->Logf(ELogVerbosity::Warning, TEXT("Invalid property value in defaults: %s"), Start);
-					}
+					const TCHAR* Result = ImportText(Str, ArrayProperty->Inner, ArrayHelper.GetRawPtr(Index));
 					// in the failure case, don't return NULL so the caller can potentially skip less and get values further in the string
-					if (Result != NULL)
+					if (Result != nullptr)
 					{
 						Str = Result;
 					}
@@ -2009,28 +2328,10 @@ const TCHAR* FProperty::ImportSingleProperty( const TCHAR* Str, void* DestData,
 						Index = 0;
 					}
 
-					FStringOutputDevice ImportError;
 
-					const TCHAR* Result = Property->ImportText_Direct(Str, Property->ContainerPtrToValuePtr<void>(DestData, Index), SubobjectOuter, PortFlags, &ImportError);
-					
-					// Spit any error we had while importing property
-					if (ImportError.Len() > 0)
-					{
-						TArray<FString> ImportErrors;
-						ImportError.ParseIntoArray(ImportErrors, LINE_TERMINATOR, true);
-
-						Warn->Logf(ELogVerbosity::Warning, TEXT("While importing text for property '%s' in '%s':"), *Property->GetName(), *ObjectStruct->GetName());
-						for ( int32 ErrorIndex = 0; ErrorIndex < ImportErrors.Num(); ErrorIndex++ )
-						{
-							Warn->Logf(ELogVerbosity::Warning, TEXT("%s"), *ImportErrors[ErrorIndex]);
-						}
-					}
-					else if ((Result == NULL && ArrayProperty == nullptr) || Result == Str)
-					{
-						UE_SUPPRESS(LogExec, Verbose, Warn->Logf(TEXT("Unknown property in %s: %s "), *ObjectStruct->GetName(), Start));
-					}
+					const TCHAR* Result = ImportText(Str, Property, Property->ContainerPtrToValuePtr<void>(DestData, Index));
 					// in the failure case, don't return NULL so the caller can potentially skip less and get values further in the string
-					if (Result != NULL)
+					if (Result != nullptr)
 					{
 						Str = Result;
 					}
diff --git a/Engine/Source/Runtime/CoreUObject/Private/UObject/PropertyMap.cpp b/Engine/Source/Runtime/CoreUObject/Private/UObject/PropertyMap.cpp
--- a/Engine/Source/Runtime/CoreUObject/Private/UObject/PropertyMap.cpp
+++ b/Engine/Source/Runtime/CoreUObject/Private/UObject/PropertyMap.cpp
@@ -308,26 +308,6 @@ void FMapProperty::SerializeItem(FStructuredArchive::FSlot Slot, void* Value, co
 		{
 			checkf(!UnderlyingArchive.ArUseCustomPropertyList, TEXT("Using custom property list is not supported by overridable serialization"));
 
-			auto GetIDFromKey = [&](uint8* KeyData) -> FOverriddenPropertyNodeID
-			{
-				if (FObjectProperty* KeyObjectProperty = CastField<FObjectProperty>(KeyProp))
-				{
-					if (const UObject* Object = KeyObjectProperty->GetObjectPropertyValue(KeyData))
-					{
-						return FOverriddenPropertyNodeID(*Object);
-					}
-				}
-				else
-				{
-					FString KeyString;
-					KeyProp->ExportTextItem_Direct(KeyString, KeyData, /*DefaultValue*/nullptr, /*Parent*/nullptr, PPF_None);
-					return FOverriddenPropertyNodeID(FName(KeyString));
-				}
-		
-				checkf(false, TEXT("This case is not handled"))
-				return FOverriddenPropertyNodeID();
-			};
-
 			if (UnderlyingArchive.IsLoading())
 			{
 				int32 NumReplaced = 0;
@@ -393,7 +373,7 @@ void FMapProperty::SerializeItem(FStructuredArchive::FSlot Slot, void* Value, co
 							if (FOverriddenPropertyNode* MapOverriddenPropertyNode = OverriddenProperties ? OverriddenProperties->SetOverriddenPropertyOperation(EOverriddenPropertyOperation::Modified, UnderlyingArchive.GetSerializedPropertyChain(), /*Property*/nullptr) : nullptr)
 							{
 								// Rebuild the overridden info
-								FOverriddenPropertyNodeID RemovedKeyID = GetIDFromKey(TempKeyValueStorage);
+								FOverriddenPropertyNodeID RemovedKeyID = UE::OverridableMapUtilities::GetIDFromKey(KeyProp, TempKeyValueStorage);
 								OverriddenProperties->SetSubPropertyOperation(EOverriddenPropertyOperation::Remove, *MapOverriddenPropertyNode, RemovedKeyID);
 							}
 						}
@@ -467,7 +447,7 @@ void FMapProperty::SerializeItem(FStructuredArchive::FSlot Slot, void* Value, co
 							if (FOverriddenPropertyNode* MapOverriddenPropertyNode = OverriddenProperties ? OverriddenProperties->SetOverriddenPropertyOperation(EOverriddenPropertyOperation::Modified, UnderlyingArchive.GetSerializedPropertyChain(), /*Property*/nullptr) : nullptr)
 							{
 								// Rebuild the overridden info
-								FOverriddenPropertyNodeID AddedKeyID = GetIDFromKey(TempKeyValueStorage);
+								FOverriddenPropertyNodeID AddedKeyID = UE::OverridableMapUtilities::GetIDFromKey(KeyProp, TempKeyValueStorage);
 								OverriddenProperties->SetSubPropertyOperation(EOverriddenPropertyOperation::Add, *MapOverriddenPropertyNode, AddedKeyID);
 							}
 						}
@@ -476,39 +456,6 @@ void FMapProperty::SerializeItem(FStructuredArchive::FSlot Slot, void* Value, co
 			}
 			else
 			{
-				auto FindKeyInternalIndex = [this](const FOverriddenPropertyNodeID& KeyIDToFind, FScriptMapHelper& MapHelper) -> int32
-				{
-					if (const FObjectProperty* KeyObjectProperty = CastField<FObjectProperty>(KeyProp))
-					{
-						for (FScriptMapHelper::FIterator It(MapHelper); It; ++It) 
-						{
-							if (UObject* CurrentObject = KeyObjectProperty->GetObjectPropertyValue(MapHelper.GetKeyPtr(It.GetInternalIndex())))
-							{
-								if (KeyIDToFind == FOverriddenPropertyNodeID(*CurrentObject))
-								{
-									return It.GetInternalIndex();
-								}
-							}
-						}
-					}
-					else
-					{
-						void* TempKeyValueStorage = FMemory::Malloc(MapLayout.SetLayout.Size);
-						KeyProp->InitializeValue(TempKeyValueStorage);
-
-						FString KeyToFind(KeyIDToFind.ToString());
-						KeyProp->ImportText_Direct(*KeyToFind, TempKeyValueStorage, nullptr, PPF_None);
-
-						const int32 InternalIndex = MapHelper.FindMapPairIndexFromHash(TempKeyValueStorage);
-
-						KeyProp->DestroyValue(TempKeyValueStorage);
-						FMemory::Free(TempKeyValueStorage);
-
-						return InternalIndex;
-					}
-					return INDEX_NONE;
-				};
-
 				// Container for temporarily tracking some indices
 				TArray<int32> RemovedIndices;
 				TArray<int32> AddedIndices;
@@ -561,35 +508,33 @@ void FMapProperty::SerializeItem(FStructuredArchive::FSlot Slot, void* Value, co
 									const EOverriddenPropertyOperation OverrideOp = OverriddenProperties->GetSubPropertyOperation(Pair.Value);
 									switch (OverrideOp)
 									{
-									case EOverriddenPropertyOperation::Remove:
+										case EOverriddenPropertyOperation::Remove:
 										{
-											const int32 InternalIndex = FindKeyInternalIndex(Pair.Key, DefaultsMapHelper);
+											const int32 InternalIndex = UE::OverridableMapUtilities::FindKeyInternalIndex(Pair.Key, DefaultsMapHelper);
 											if (InternalIndex != INDEX_NONE)
 											{
 												RemovedIndices.Add(InternalIndex);
 											}
 											break;
 										}
-									case EOverriddenPropertyOperation::Add:
+										case EOverriddenPropertyOperation::Add:
 										{
-											const int32 InternalIndex = FindKeyInternalIndex(Pair.Key, MapHelper);
+											const int32 InternalIndex = UE::OverridableMapUtilities::FindKeyInternalIndex(Pair.Key, MapHelper);
 											if (InternalIndex != INDEX_NONE)
 											{
 												AddedIndices.Add(InternalIndex);
-												ModifiedIndices.Remove(InternalIndex);
 											}
 											break;
 										}
-									case EOverriddenPropertyOperation::Modified:
-										if (!bAreValuesInstancedSubObjects)
+										case EOverriddenPropertyOperation::Modified:
 										{
-											const int32 InternalIndex = FindKeyInternalIndex(Pair.Key, MapHelper);
+											const int32 InternalIndex = UE::OverridableMapUtilities::FindKeyInternalIndex(Pair.Key, MapHelper);
 											if (InternalIndex != INDEX_NONE)
 											{
 												ModifiedIndices.Add(InternalIndex);
 											}
+											break;
 										}
-										break;
 									default:
 										checkf(false, TEXT("Unsupported map operation"));
 										break;
diff --git a/Engine/Source/Runtime/CoreUObject/Private/UObject/PropertyOptional.cpp b/Engine/Source/Runtime/CoreUObject/Private/UObject/PropertyOptional.cpp
--- a/Engine/Source/Runtime/CoreUObject/Private/UObject/PropertyOptional.cpp
+++ b/Engine/Source/Runtime/CoreUObject/Private/UObject/PropertyOptional.cpp
@@ -699,8 +699,13 @@ EPropertyVisitorControlFlow FOptionalProperty::Visit(FPropertyVisitorPath& Path,
 
 	EPropertyVisitorControlFlow RetVal = Super::Visit(Path, Data, InFunc);
 
-	if (RetVal == EPropertyVisitorControlFlow::StepInto && IsSet(Data))
+	if (RetVal == EPropertyVisitorControlFlow::StepInto)
 	{
+		if(!IsSet(Data))
+		{
+			// There is nothing to StepInto, so continue to next property
+			return EPropertyVisitorControlFlow::StepOver;
+		}
 		checkf(ValueProperty, TEXT("Expecting a valid property value"));
 
 		FPropertyVisitorScope Scope(Path, FPropertyVisitorInfo(ValueProperty));
diff --git a/Engine/Source/Runtime/CoreUObject/Public/UObject/OverriddenPropertySet.h b/Engine/Source/Runtime/CoreUObject/Public/UObject/OverriddenPropertySet.h
--- a/Engine/Source/Runtime/CoreUObject/Public/UObject/OverriddenPropertySet.h
+++ b/Engine/Source/Runtime/CoreUObject/Public/UObject/OverriddenPropertySet.h
@@ -25,7 +25,7 @@ public:
 	 * Call to enable overridable serialization and to set the overridden properties of the current serialized object
 	 * Note this is not re-entrant and it stores information in a thread local storage
 	 * @param InOverriddenProperties of the current serializing object */
-	FORCEINLINE static void Enable(FOverriddenPropertySet* InOverriddenProperties)
+	inline static void Enable(FOverriddenPropertySet* InOverriddenProperties)
 	{
 		checkf(!bUseOverridableSerialization, TEXT("Nobody should use this method if overridable serialization is already enabled"));
 		bUseOverridableSerialization = true;
@@ -35,7 +35,7 @@ public:
 	/**
 	 * Call to disable overridable serialization
 	 * Note this is not re-entrant and it stores information in a thread local storage */
-	FORCEINLINE static void Disable()
+	inline static void Disable()
 	{
 		checkf(bUseOverridableSerialization, TEXT("Expecting overridable serialization to be already enabled"));
 		OverriddenProperties = nullptr;
@@ -45,7 +45,7 @@ public:
 	/**
 	 * Called during the serialization of an object to know to know if it should do overridden serialization logic
 	 * @return true if the overridable serialization is enabled on the current serializing object */
-	FORCEINLINE static bool IsEnabled()
+	inline static bool IsEnabled()
 	{
 		return bUseOverridableSerialization;
 	}
@@ -55,7 +55,7 @@ public:
 	 * Note: Expects the current serialized object to use overridable serialization
 	 * Note this is not re-entrant and it stores information in a thread local storage
 	 * @return the overridden properties of the current object being serialized */
-	FORCEINLINE static FOverriddenPropertySet* GetOverriddenProperties()
+	inline static FOverriddenPropertySet* GetOverriddenProperties()
 	{
 		return OverriddenProperties;
 	}
@@ -67,11 +67,36 @@ public:
 	 * @param DataPtr to the memory of that property
 	 * @param DefaultValue memory pointer of that property
 	 * @return the overridden property operation */
-	static EOverriddenPropertyOperation GetOverriddenPropertyOperation(const FArchive& Ar, FProperty* Property = nullptr, uint8* DataPtr = nullptr, uint8* DefaultValue = nullptr);
+	COREUOBJECT_API static EOverriddenPropertyOperation GetOverriddenPropertyOperation(const FArchive& Ar, FProperty* Property = nullptr, uint8* DataPtr = nullptr, uint8* DefaultValue = nullptr);
+
+	/**
+	 * Use the port text path to retrieve the current overridden property operation to know if it has to be serialized or not
+	 * @param DataPtr to the memory of that property
+	 * @param DefaultValue memory pointer of that property
+	 * @param PortFlags for the import/export text operation
+	 * @return the overridden property operation */
+	COREUOBJECT_API static EOverriddenPropertyOperation GetOverriddenPropertyOperationForPortText(const void* DataPtr, const void* DefaultValue, int32 PortFlags);
+
+	/**
+	 * Call during the import text
+	 * @return the current text import property path */
+	COREUOBJECT_API static FPropertyVisitorPath* GetOverriddenPortTextPropertyPath();
+
+	/**
+	 * Call during the import text to set the property path
+	 * @param Path to start tracking the property path*/
+	COREUOBJECT_API static void SetOverriddenPortTextPropertyPath(FPropertyVisitorPath& Path);
+
+	/** Call during the import text to reset property path */
+	COREUOBJECT_API static void ResetOverriddenPortTextPropertyPath();
 
 private:
+
+	static EOverriddenPropertyOperation GetOverriddenPropertyOperation(const int32 PortFlags, const FArchiveSerializedPropertyChain* CurrentPropertyChain, FProperty* Property, const void* DataPtr, const void* DefaultValue);
+
 	static thread_local bool bUseOverridableSerialization;
 	static thread_local FOverriddenPropertySet* OverriddenProperties;
+	static thread_local FPropertyVisitorPath* OverriddenPortTextPropertyPath;
 };
 
 /*
@@ -82,8 +107,8 @@ private:
 */
 struct FEnableOverridableSerializationScope
 {
-	FEnableOverridableSerializationScope(bool bEnableOverridableSerialization, FOverriddenPropertySet* OverriddenProperties);
-	~FEnableOverridableSerializationScope();
+	COREUOBJECT_API FEnableOverridableSerializationScope(bool bEnableOverridableSerialization, FOverriddenPropertySet* OverriddenProperties);
+	COREUOBJECT_API ~FEnableOverridableSerializationScope();
 
 protected:
 	bool bOverridableSerializationEnabled = false;
@@ -91,6 +116,24 @@ protected:
 	FOverriddenPropertySet* SavedOverriddenProperties = nullptr;
 };
 
+
+/*
+ *************************************************************************************
+ * Overridable serialization is experimental, not supported and use at your own risk *
+ *************************************************************************************
+ * Scope responsible for tracking current property path for text importing
+*/
+struct FOverridableTextPortPropertyPathScope
+{
+	COREUOBJECT_API FOverridableTextPortPropertyPathScope(const FProperty* InProperty, int32 InIndex = INDEX_NONE, EPropertyVisitorInfoType InPropertyInfo = EPropertyVisitorInfoType::None);
+	COREUOBJECT_API ~FOverridableTextPortPropertyPathScope();
+
+protected:
+
+	const FProperty* Property = nullptr;
+	FPropertyVisitorPath DefaultPath;
+};
+
 /*
  * Override operation type for each property node
  */
@@ -104,6 +147,27 @@ enum class EOverriddenPropertyOperation : uint8
 	Remove,		/* this element was removed from the container */
 };
 
+inline TOptional<EOverriddenPropertyOperation> GetOverriddenOperationFromString(const FString& OverriddenOperationString)
+{
+	int64 EnumValue = StaticEnum<EOverriddenPropertyOperation>()->GetValueByNameString(StaticEnum<EOverriddenPropertyOperation>()->GetName() + FString(TEXT("::")) + OverriddenOperationString);
+	if (EnumValue == INDEX_NONE)
+	{
+		TOptional<EOverriddenPropertyOperation>();
+	}
+
+	return TOptional((EOverriddenPropertyOperation)EnumValue);
+}
+
+inline TOptional<EOverriddenPropertyOperation> GetOverriddenOperationFromName(const FName OverriddenOperationName)
+{
+	return GetOverriddenOperationFromString(OverriddenOperationName.ToString());
+}
+
+inline FString GetOverriddenOperationString(EOverriddenPropertyOperation Operation)
+{
+	return UEnum::GetValueAsString(Operation).RightChop(StaticEnum<EOverriddenPropertyOperation>()->GetName().Len() + /*::*/2);
+}
+
 USTRUCT()
 struct FOverriddenPropertyNodeID
 {
@@ -319,3 +383,24 @@ private:
 public:
 	bool bNeedsSubobjectTemplateInstantiation = false;
 };
+
+// Utility methods for maps
+namespace UE::OverridableMapUtilities
+{
+	/**
+	 * Find the index of the key id in the specified map
+	 * @param KeyIDToFind in the map
+	 * @param MapHelper the map to search in
+	 * @return INDEX_NONE if it didn't find it
+	 */
+	COREUOBJECT_API int32 FindKeyInternalIndex(const FOverriddenPropertyNodeID& KeyIDToFind, FScriptMapHelper& MapHelper);
+
+	/**
+	 * Retrieve the key id from the key property and value
+	 * NOTE: This method will check if unable to create the key
+	 * @param KeyProp to use while converting to an key id
+	 * @param KeyData to use while converting to an key id
+	 * @return the key id of the specified key
+	 */
+	COREUOBJECT_API FOverriddenPropertyNodeID GetIDFromKey(const FProperty* KeyProp, uint8* KeyData);
+}
\ No newline at end of file
diff --git a/Engine/Source/Runtime/Engine/Private/UnrealExporter.cpp b/Engine/Source/Runtime/Engine/Private/UnrealExporter.cpp
--- a/Engine/Source/Runtime/Engine/Private/UnrealExporter.cpp
+++ b/Engine/Source/Runtime/Engine/Private/UnrealExporter.cpp
@@ -16,8 +16,11 @@
 #include "Model.h"
 #include "Misc/FeedbackContext.h"
 #include "AssetExportTask.h"
+#include "DataTableUtils.h"
 #include "Misc/AsciiSet.h"
 #include "UObject/GCObjectScopeGuard.h"
+#include "UObject/OverridableManager.h"
+#include "UObject/OverriddenPropertySet.h"
 
 #if WITH_EDITOR
 #include "Editor.h"
@@ -471,6 +474,12 @@ void UExporter::EmitBeginObject( FOutputDevice& Ar, UObject* Obj, uint32 PortFla
 			// we need the fully qualified path for the archetype (so we don't get confused when unpacking this)
 			Ar.Logf(TEXT(" Archetype=%s"), *FObjectPropertyBase::GetExportPath(Archetype, nullptr, /*ExportRootScope =*/nullptr, (PortFlags | PPF_Delimited) & ~PPF_ExportsNotFullyQualified));
 		}
+
+		if (FOverriddenPropertySet* ObjectOverriddenProperties = FOverridableManager::Get().GetOverriddenProperties(*Obj))
+		{
+			const EOverriddenPropertyOperation Operation = ObjectOverriddenProperties->GetOverriddenPropertyOperation((FArchiveSerializedPropertyChain*)nullptr, (FProperty*)nullptr);
+			Ar.Logf(TEXT(" OverriddenOperation=%s"), *GetOverriddenOperationString(Operation));
+		}
 	}
 
 	// When exporting for diffs, export paths can cause false positives. since diff files don't get imported, we can
@@ -672,7 +681,7 @@ void ExportProperties
 	FOutputDevice&	Out,
 	UClass*			ObjectClass,
 	uint8*			Object,
-	int32				Indent,
+	int32			Indent,
 	UClass*			DiffClass,
 	uint8*			Diff,
 	UObject*		Parent,
@@ -682,10 +691,15 @@ void ExportProperties
 {
 	check(ObjectClass != NULL);
 
-	for( FProperty* Property = ObjectClass->PropertyLink; Property; Property = Property->PropertyLinkNext )
+	FOverriddenPropertySet* OverriddenProperties = Object ? FOverridableManager::Get().GetOverriddenProperties(*(UObject*)Object) : nullptr;
+	FEnableOverridableSerializationScope Scope(OverriddenProperties!= nullptr, OverriddenProperties);
+
+	for (FProperty* Property = ObjectClass->PropertyLink; Property; Property = Property->PropertyLinkNext)
 	{
 		if (!Property->ShouldPort(PortFlags))
+		{
 			continue;
+		}
 
 		FString SanitizedPropertyName = Property->GetName();
 		constexpr FAsciiSet Whitespace("\t ");
@@ -697,48 +711,110 @@ void ExportProperties
 		}
 
 		FArrayProperty* ArrayProperty = CastField<FArrayProperty>(Property);
+		FMapProperty* MapProperty = CastField<FMapProperty>(Property);
+
 		FObjectPropertyBase* ExportObjectProp = (Property->PropertyFlags & CPF_ExportObject) != 0 ? CastField<FObjectPropertyBase>(Property) : NULL;
 		const uint32 ExportFlags = PortFlags | PPF_Delimited;
-
-		if ( ArrayProperty != NULL )
+		auto HandleExportObject = [Context, &Out, Indent, Parent, PortFlags, ExportObjectProp](const FProperty* Property, uint8* Data)
 		{
-			// Export dynamic array.
-			FProperty* InnerProp = ArrayProperty->Inner;
-			ExportObjectProp = (Property->PropertyFlags & CPF_ExportObject) != 0 ? CastField<FObjectPropertyBase>(InnerProp) : NULL;
-			// This is used as the default value in the case of an array property that has
-			// fewer elements than the exported object.
-			uint8* StructDefaults = NULL;
-			FStructProperty* StructProperty = CastField<FStructProperty>(InnerProp);
-			if ( StructProperty != NULL )
+			if (ExportObjectProp)
 			{
-				checkSlow(StructProperty->Struct);
-				StructDefaults = (uint8*)FMemory::Malloc(StructProperty->Struct->GetStructureSize());
-				StructProperty->InitializeValue(StructDefaults);
+				UObject* Obj = ExportObjectProp->GetObjectPropertyValue(Data);
+				check(!Obj || Obj->IsValidLowLevel());
+				if (Obj && !Obj->HasAnyMarks(OBJECTMARK_TagImp))
+				{
+					// only export the BEGIN OBJECT block for a component if Parent is the component's Outer....when importing subobject definitions,
+					// (i.e. BEGIN OBJECT), whichever BEGIN OBJECT block a component's BEGIN OBJECT block is located within is the object that will be
+					// used as the Outer to create the component
+
+					// Is this an array of components?
+					if (Property->HasAnyPropertyFlags(CPF_InstancedReference))
+					{
+						if (Obj->GetOuter() == Parent)
+						{
+							// Don't export more than once.
+							Obj->Mark(OBJECTMARK_TagImp);
+							UExporter::ExportToOutputDevice(Context, Obj, NULL, Out, TEXT("T3D"), Indent, PortFlags);
+						}
+						else
+						{
+							// set the OBJECTMARK_TagExp flag so that the calling code knows we wanted to export this object
+							Obj->Mark(OBJECTMARK_TagExp);
+						}
+					}
+					else
+					{
+						// Don't export more than once.
+						Obj->Mark(OBJECTMARK_TagImp);
+						UExporter::ExportToOutputDevice(Context, Obj, NULL, Out, TEXT("T3D"), Indent, PortFlags);
+					}
+				}
 			}
+		};
 
-			for( int32 PropertyArrayIndex=0; PropertyArrayIndex<Property->ArrayDim; PropertyArrayIndex++ )
+		for (int32 PropertyArrayIndex=0; PropertyArrayIndex<Property->ArrayDim; PropertyArrayIndex++)
+		{
+			// Array special case
+			if (ArrayProperty != nullptr)
 			{
+				// Export dynamic array.
+				FProperty* InnerProp = ArrayProperty->Inner;
+				ExportObjectProp = (Property->PropertyFlags & CPF_ExportObject) != 0 ? CastField<FObjectPropertyBase>(InnerProp) : NULL;
+				// This is used as the default value in the case of an array property that has
+				// fewer elements than the exported object.
+				uint8* StructDefaults = nullptr;
+				FStructProperty* StructProperty = CastField<FStructProperty>(InnerProp);
+				if (StructProperty != nullptr)
+				{
+					checkSlow(StructProperty->Struct);
+					StructDefaults = (uint8*)FMemory::Malloc(StructProperty->Struct->GetStructureSize());
+					StructProperty->InitializeValue(StructDefaults);
+				}
+				ON_SCOPE_EXIT
+				{
+					if (StructDefaults)
+					{
+						StructProperty->DestroyValue(StructDefaults);
+						FMemory::Free(StructDefaults);
+					}
+				};
+
 				void* Arr = Property->ContainerPtrToValuePtr<void>(Object, PropertyArrayIndex);
 				FScriptArrayHelper ArrayHelper(ArrayProperty, Arr);
 
-				void*	DiffArr = NULL;
-				if( DiffClass )
+				void* DiffArr = nullptr;
+				if (DiffClass)
 				{
 					DiffArr = Property->ContainerPtrToValuePtrForDefaults<void>(DiffClass, Diff, PropertyArrayIndex);
 				}
 				// we won't use this if DiffArr is NULL, but we have to set it up to something
 				FScriptArrayHelper DiffArrayHelper(ArrayProperty, DiffArr);
 
-				// If the current size of the array is 0 and the default one is not, add in an empty item so on import it will be empty
-				if( ArrayHelper.Num() == 0 && DiffArrayHelper.Num() != 0 )
+				EOverriddenPropertyOperation Operation = EOverriddenPropertyOperation::None;
+				FArchiveSerializedPropertyChain Chain;
+				if (OverriddenProperties)
 				{
-					Out.Logf(TEXT("%s%s=\r\n"), FCString::Spc(Indent), *SanitizedPropertyName);
+					FOverridableTextPortPropertyPathScope ScopePath(Property);
+					Operation = FOverridableSerializationLogic::GetOverriddenPropertyOperationForPortText(Arr, DiffArr, PortFlags);
+					FPropertyVisitorPath* Path = FOverridableSerializationLogic::GetOverriddenPortTextPropertyPath();
+					checkf(Path, TEXT("Expecting a path"));
+					Chain = Path->ToSerializedPropertyChain();
 				}
-				else
+
+				if(!OverriddenProperties && !Property->HasAnyPropertyFlags(CPF_ExperimentalOverridableLogic))
 				{
+					// If the current size of the array is 0 and the default one is not, add in an empty item so on import it will be empty
+					if (ArrayHelper.Num() == 0 && DiffArrayHelper.Num() != 0)
+					{
+						Out.Logf(TEXT("%s%s=\r\n"), FCString::Spc(Indent), *SanitizedPropertyName);
+						continue;
+					}
+
 					// If the array sizes are different, we will need to export each index so on import we maintain the size
 					for (int32 DynamicArrayIndex = 0; DynamicArrayIndex < ArrayHelper.Num(); DynamicArrayIndex++)
 					{
+						FOverridableTextPortPropertyPathScope ScopePath(Property, DynamicArrayIndex, EPropertyVisitorInfoType::ContainerIndex);
+
 						FString	Value;
 
 						// compare each element's value manually so that elements which match the NULL value for the array's inner property type
@@ -753,105 +829,246 @@ void ExportProperties
 						if (bExportItem)
 						{
 							InnerProp->ExportTextItem_Direct(Value, SourceData, DiffData, Parent, ExportFlags, ExportRootScope);
-							if (ExportObjectProp)
+							HandleExportObject(InnerProp, ArrayHelper.GetRawPtr(DynamicArrayIndex));
+
+							Out.Logf(TEXT("%s%s(%i)=%s\r\n"), FCString::Spc(Indent), *SanitizedPropertyName, DynamicArrayIndex, *Value);
+						}
+					}
+					for (int32 DynamicArrayIndex = DiffArrayHelper.Num()-1; DynamicArrayIndex >= ArrayHelper.Num(); --DynamicArrayIndex)
+					{
+						Out.Logf(TEXT("%s%s.RemoveIndex(%d)\r\n"), FCString::Spc(Indent), *SanitizedPropertyName, DynamicArrayIndex);
+					}
+					continue;
+				}
+
+				// Replaced operation goes through the generic property export code
+				if (Operation == EOverriddenPropertyOperation::Modified)
+				{
+					if (const FOverriddenPropertyNode* ArrayOverriddenPropertyNode = OverriddenProperties->GetOverriddenPropertyNode(&Chain))
+					{
+						checkf(Arr && DiffArr, TEXT("Expecting a memory ptr to Array and its default"));
+
+						const FObjectProperty* InnerObjectProperty = CastField<FObjectProperty>(InnerProp);
+
+						// Only array of instanced subobjects are handled here as sort of a set where the matching key is done using the archetype
+						checkf(InnerObjectProperty&& InnerObjectProperty->HasAnyPropertyFlags(CPF_PersistentInstance), TEXT("Expecting only arrays of instanced subobjects"));
+
+						auto FindObject = [InnerObjectProperty](const FOverriddenPropertyNodeID ObjectToFind, FScriptArrayHelper& ArrayHelper) -> int32
+						{
+							const int32 ArrayNum = ArrayHelper.Num();
+							for (int i = 0; i < ArrayNum; ++i)
 							{
-								UObject* Obj = ExportObjectProp->GetObjectPropertyValue(ArrayHelper.GetRawPtr(DynamicArrayIndex));
-								check(!Obj || Obj->IsValidLowLevel());
-								if (Obj && !Obj->HasAnyMarks(OBJECTMARK_TagImp))
+								if (UObject* CurrentObject = InnerObjectProperty->GetObjectPropertyValue(ArrayHelper.GetElementPtr(i)))
 								{
-									// only export the BEGIN OBJECT block for a component if Parent is the component's Outer....when importing subobject definitions,
-									// (i.e. BEGIN OBJECT), whichever BEGIN OBJECT block a component's BEGIN OBJECT block is located within is the object that will be
-									// used as the Outer to create the component
+									if (ObjectToFind == FOverriddenPropertyNodeID(*CurrentObject))
+									{
+										return i;
+									}
+								}
+							}
+							return INDEX_NONE;
+						};
+
+						FScriptArrayHelper DefaultsArrayHelper(ArrayProperty, DiffArr);
 
-									// Is this an array of components?
-									if (InnerProp->HasAnyPropertyFlags(CPF_InstancedReference))
+						TArray<int32> RemovedIndices;
+						TArray<int32> AddedIndices;
+
+						for (const auto& Pair : ArrayOverriddenPropertyNode->SubPropertyNodeKeys)
+						{
+							const EOverriddenPropertyOperation OverrideOp = OverriddenProperties->GetSubPropertyOperation(Pair.Value);
+							switch (OverrideOp)
+							{
+							case EOverriddenPropertyOperation::Remove:
+								{
+									const int32 DefaultIndex = FindObject(Pair.Key, DefaultsArrayHelper);
+									if (DefaultIndex != INDEX_NONE)
 									{
-										if (Obj->GetOuter() == Parent)
-										{
-											// Don't export more than once.
-											Obj->Mark(OBJECTMARK_TagImp);
-											UExporter::ExportToOutputDevice(Context, Obj, NULL, Out, TEXT("T3D"), Indent, PortFlags);
-										}
-										else
-										{
-											// set the OBJECTMARK_TagExp flag so that the calling code knows we wanted to export this object
-											Obj->Mark(OBJECTMARK_TagExp);
-										}
+										RemovedIndices.Add(DefaultIndex);
 									}
-									else
+									break;
+								}
+							case EOverriddenPropertyOperation::Add:
+								{
+									const int32 Index = FindObject(Pair.Key, ArrayHelper);
+									if (Index != INDEX_NONE)
 									{
-										// Don't export more than once.
-										Obj->Mark(OBJECTMARK_TagImp);
-										UExporter::ExportToOutputDevice(Context, Obj, NULL, Out, TEXT("T3D"), Indent, PortFlags);
+										AddedIndices.Add(Index);
 									}
+									break;
 								}
+							default:
+								checkf(false, TEXT("Unsupported operation type"));
+								break;
 							}
+						}
 
-							Out.Logf(TEXT("%s%s(%i)=%s\r\n"), FCString::Spc(Indent), *SanitizedPropertyName, DynamicArrayIndex, *Value);
+						for (int32 i : RemovedIndices)
+						{
+							FString Value;
+							InnerProp->ExportTextItem_Direct(Value, DiffArrayHelper.GetRawPtr(i), nullptr, Parent, ExportFlags, ExportRootScope);
+
+							Out.Logf(TEXT("%s%s<%s>=%s\r\n"), FCString::Spc(Indent), *SanitizedPropertyName, *GetOverriddenOperationString(EOverriddenPropertyOperation::Remove), *Value);
+						}
+
+						for (int32 i : AddedIndices)
+						{
+							FOverridableTextPortPropertyPathScope ScopePath(Property, i, EPropertyVisitorInfoType::ContainerIndex);
+
+							FString Value;
+							InnerProp->ExportTextItem_Direct(Value, ArrayHelper.GetRawPtr(i), nullptr, Parent, ExportFlags, ExportRootScope);
+							HandleExportObject(InnerProp, ArrayHelper.GetRawPtr(i));
+
+							Out.Logf(TEXT("%s%s<%s>=%s\r\n"), FCString::Spc(Indent), *SanitizedPropertyName, *GetOverriddenOperationString(EOverriddenPropertyOperation::Add), *Value);
 						}
 					}
-					for (int32 DynamicArrayIndex = DiffArrayHelper.Num()-1; DynamicArrayIndex >= ArrayHelper.Num(); --DynamicArrayIndex)
-					{
-						Out.Logf(TEXT("%s%s.RemoveIndex(%d)\r\n"), FCString::Spc(Indent), *SanitizedPropertyName, DynamicArrayIndex);
-					}
+					continue;
 				}
 			}
-			if (StructDefaults)
-			{
-				StructProperty->DestroyValue(StructDefaults);
-				FMemory::Free(StructDefaults);
-			}
-		}
-		else
-		{
-			for( int32 PropertyArrayIndex=0; PropertyArrayIndex<Property->ArrayDim; PropertyArrayIndex++ )
+
+			// Map special case
+			if (MapProperty != nullptr)
 			{
-				FString	Value;
-				// Export single element.
 
-				uint8* DiffData = (DiffClass && Property->IsInContainer(DiffClass->GetPropertiesSize())) ? Diff : NULL;
-				if( Property->ExportText_InContainer( PropertyArrayIndex, Value, Object, DiffData, Parent, ExportFlags, ExportRootScope ) )
+				void* Map = Property->ContainerPtrToValuePtr<void>(Object, PropertyArrayIndex);
+				void* DiffMap = nullptr;
+				if (DiffClass)
 				{
-					if ( ExportObjectProp )
+					DiffMap = Property->ContainerPtrToValuePtrForDefaults<void>(DiffClass, Diff, PropertyArrayIndex);
+				}
+
+				EOverriddenPropertyOperation Operation = EOverriddenPropertyOperation::None;
+				FArchiveSerializedPropertyChain Chain;
+				if (OverriddenProperties)
+				{
+					FOverridableTextPortPropertyPathScope ScopePath(Property);
+					Operation = FOverridableSerializationLogic::GetOverriddenPropertyOperationForPortText(Map, DiffMap, PortFlags);
+					FPropertyVisitorPath* Path = FOverridableSerializationLogic::GetOverriddenPortTextPropertyPath();
+					checkf(Path, TEXT("Expecting a path"));
+					Chain = Path->ToSerializedPropertyChain();
+				}
+
+				// Replaced operation goes through the generic property export code
+				if (Operation == EOverriddenPropertyOperation::Modified)
+				{
+					if (const FOverriddenPropertyNode* MapOverriddenPropertyNode = OverriddenProperties->GetOverriddenPropertyNode(&Chain))
 					{
-						UObject* Obj = ExportObjectProp->GetObjectPropertyValue(Property->ContainerPtrToValuePtr<void>(Object, PropertyArrayIndex));
-						if( Obj && !Obj->HasAnyMarks(OBJECTMARK_TagImp) )
+						checkf(Map && DiffMap, TEXT("Expecting memory ptr to the map and its defaults"));
+
+						FScriptMapHelper MapHelper(MapProperty, Map);
+						FScriptMapHelper DiffMapHelper(MapProperty, DiffMap);
+
+						TArray<int32> RemovedIndices;
+						TArray<int32> ModifiedIndices;
+						TArray<int32> AddedIndices;
+
+
+						// Figure out the modifications of the map
+						for (const auto& Pair : MapOverriddenPropertyNode->SubPropertyNodeKeys)
 						{
-							// only export the BEGIN OBJECT block for a component if Parent is the component's Outer....when importing subobject definitions,
-							// (i.e. BEGIN OBJECT), whichever BEGIN OBJECT block a component's BEGIN OBJECT block is located within is the object that will be
-							// used as the Outer to create the component
-							if ( Property->HasAnyPropertyFlags(CPF_InstancedReference) )
+							const EOverriddenPropertyOperation OverrideOp = OverriddenProperties->GetSubPropertyOperation(Pair.Value);
+							switch (OverrideOp)
 							{
-								if ( Obj->GetOuter() == Parent )
+							case EOverriddenPropertyOperation::Remove:
+								{
+									const int32 InternalIndex = UE::OverridableMapUtilities::FindKeyInternalIndex(Pair.Key, DiffMapHelper);
+									if (InternalIndex != INDEX_NONE)
+									{
+										RemovedIndices.Add(InternalIndex);
+									}
+									break;
+								}
+							case EOverriddenPropertyOperation::Add:
 								{
-									// Don't export more than once.
-									Obj->Mark(OBJECTMARK_TagImp);
-									UExporter::ExportToOutputDevice( Context, Obj, NULL, Out, TEXT("T3D"), Indent, PortFlags );
+									const int32 InternalIndex = UE::OverridableMapUtilities::FindKeyInternalIndex(Pair.Key, MapHelper);
+									if (InternalIndex != INDEX_NONE)
+									{
+										AddedIndices.Add(InternalIndex);
+									}
+									break;
 								}
-								else
+							case EOverriddenPropertyOperation::Modified:
 								{
-									// set the OBJECTMARK_TagExp flag so that the calling code knows we wanted to export this object
-									Obj->Mark(OBJECTMARK_TagExp);
+									const int32 InternalIndex = UE::OverridableMapUtilities::FindKeyInternalIndex(Pair.Key, MapHelper);
+									if (InternalIndex != INDEX_NONE)
+									{
+										ModifiedIndices.Add(InternalIndex);
+									}
+									break;
 								}
-							}
-							else
-							{
-								// Don't export more than once.
-								Obj->Mark(OBJECTMARK_TagImp);
-								UExporter::ExportToOutputDevice( Context, Obj, NULL, Out, TEXT("T3D"), Indent, PortFlags );
+							default:
+								checkf(false, TEXT("Unsupported map operation"));
+								break;
 							}
 						}
-					}
 
-					if( Property->ArrayDim == 1 )
-					{
-						Out.Logf( TEXT("%s%s=%s\r\n"), FCString::Spc(Indent), *SanitizedPropertyName, *Value );
+						auto ExportItem = [&MapHelper, &DiffMapHelper, Parent, ExportRootScope, Indent, SanitizedPropertyName, ExportFlags, &Out](int32 i, EOverriddenPropertyOperation Operation)
+						{
+							FString Key;
+							MapHelper.KeyProp->ExportTextItem_Direct(Key, MapHelper.GetKeyPtr(i), Operation == EOverriddenPropertyOperation::Modified ? DiffMapHelper.GetKeyPtr(i) : nullptr, Parent, ExportFlags, ExportRootScope);
+
+							FString Value;
+							MapHelper.ValueProp->ExportTextItem_Direct(Value, MapHelper.GetValuePtr(i), Operation == EOverriddenPropertyOperation::Modified ? DiffMapHelper.GetValuePtr(i) : nullptr, Parent, ExportFlags, ExportRootScope);
+
+							Out.Logf(TEXT("%s%s<%s>=(%s,%s)\r\n"), FCString::Spc(Indent), *SanitizedPropertyName, *GetOverriddenOperationString(Operation), *Key, *Value);
+						};
+
+						for (int32 i : RemovedIndices)
+						{
+							FString Value;
+							DiffMapHelper.KeyProp->ExportTextItem_Direct(Value, DiffMapHelper.GetKeyPtr(i), nullptr, Parent, ExportFlags, ExportRootScope);
+
+							Out.Logf(TEXT("%s%s<%s>=%s\r\n"), FCString::Spc(Indent), *SanitizedPropertyName, *GetOverriddenOperationString(EOverriddenPropertyOperation::Remove), *Value);
+						}
+
+						for (int32 i : ModifiedIndices)
+						{
+							FOverridableTextPortPropertyPathScope ScopePath(Property, i, EPropertyVisitorInfoType::ContainerIndex);
+
+							ExportItem(i, EOverriddenPropertyOperation::Modified);
+						}
+
+						for (int32 i : AddedIndices)
+						{
+							FOverridableTextPortPropertyPathScope ScopePath(Property, i, EPropertyVisitorInfoType::ContainerIndex);
+
+							ExportItem(i, EOverriddenPropertyOperation::Add);
+						}
 					}
-					else
+
+					continue;
+				}
+			}
+
+
+			// Generic property export code
+			FString	Value;
+			const bool bStaticArray = Property->ArrayDim > 1;
+			FOverridableTextPortPropertyPathScope ScopePath(Property, bStaticArray ? PropertyArrayIndex : INDEX_NONE, bStaticArray ? EPropertyVisitorInfoType::StaticArrayIndex : EPropertyVisitorInfoType::None);
+
+			uint8* DiffData = (DiffClass && Property->IsInContainer(DiffClass->GetPropertiesSize())) ? Diff : NULL;
+			if (Property->ExportText_InContainer(PropertyArrayIndex, Value, Object, DiffData, Parent, ExportFlags, ExportRootScope))
+			{
+				HandleExportObject(Property, Property->ContainerPtrToValuePtr<uint8>(Object, PropertyArrayIndex));
+
+				FString OverridableOperation;
+				if (OverriddenProperties)
+				{
+					const EOverriddenPropertyOperation Operation = FOverridableSerializationLogic::GetOverriddenPropertyOperationForPortText(Object, DiffData, PortFlags);
+					if (Operation!=EOverriddenPropertyOperation::None)
 					{
-						Out.Logf( TEXT("%s%s(%i)=%s\r\n"), FCString::Spc(Indent), *SanitizedPropertyName, PropertyArrayIndex, *Value );
+						OverridableOperation = FString::Printf(TEXT("<%s>"), *GetOverriddenOperationString(Operation));
 					}
 				}
+
+				if (!bStaticArray)
+				{
+					Out.Logf( TEXT("%s%s%s=%s\r\n"), FCString::Spc(Indent), *SanitizedPropertyName, *OverridableOperation, *Value );
+				}
+				else
+				{
+					Out.Logf( TEXT("%s%s(%i)%s=%s\r\n"), FCString::Spc(Indent), *SanitizedPropertyName, PropertyArrayIndex, *OverridableOperation, *Value );
+				}
 			}
 		}
 	}
