Enable cook-time (and only cook-time) per-shader caching to both local cache and remote cache by default to reduce the number of shader compiles experienced during cook as iteration happens on shader source files.  For best performance on large projects, this should be used with:
- ZenServer as Local DDC (enabled by default in UE5.4)
- ZenServer as Shared DDC if a shared DDC is needed (production ready for UE5.5)
- UE Cloud DDC if a cloud DDC is needed (production ready for UE5.4)

diff --git a/Engine/Source/Runtime/Engine/Private/ShaderCompiler/ShaderCompiler.cpp b/Engine/Source/Runtime/Engine/Private/ShaderCompiler/ShaderCompiler.cpp
--- a/Engine/Source/Runtime/Engine/Private/ShaderCompiler/ShaderCompiler.cpp
+++ b/Engine/Source/Runtime/Engine/Private/ShaderCompiler/ShaderCompiler.cpp
@@ -150,7 +150,7 @@ static TAutoConsoleVariable<bool> CVarJobCacheDDCPolicy(
 
 static TAutoConsoleVariable<bool> CVarJobCacheDDCCookPolicy(
 	TEXT("r.ShaderCompiler.JobCacheDDCCookEnableRemotePolicy"),
-	false,
+	true,
 	TEXT("If true, individual shader jobs will be cached to remote/shared DDC instances in all cook commandlet only; if false they will only cache to DDC instances on the local machine.\n"),
 	ECVF_ReadOnly);
 
@@ -211,7 +211,7 @@ static FAutoConsoleVariableRef CVarShaderCompilerPerShaderDDCGlobal(
 
 static TAutoConsoleVariable<bool> CVarShaderCompilerPerShaderDDCCook(
 	TEXT("r.ShaderCompiler.PerShaderDDCCook"),
-	false,
+	true,
 	TEXT("If true, per-shader DDC caching will be enabled during cooks."),
 	ECVF_Default
 );
