Skeletal Mesh: Wedge re-use is not allowed on import models

#rnx

diff --git a/Engine/Source/Runtime/Engine/Private/Rendering/SkeletalMeshLODImporterData.cpp b/Engine/Source/Runtime/Engine/Private/Rendering/SkeletalMeshLODImporterData.cpp
--- a/Engine/Source/Runtime/Engine/Private/Rendering/SkeletalMeshLODImporterData.cpp
+++ b/Engine/Source/Runtime/Engine/Private/Rendering/SkeletalMeshLODImporterData.cpp
@@ -2770,26 +2770,39 @@ FSkeletalMeshImportData FSkeletalMeshImportData::CreateFromMeshDescription(const
 	
 	SkelMeshImportData.Wedges.SetNumZeroed(InMeshDescription.VertexInstances().GetArraySize());
 	
-	for (FVertexInstanceID VertexInstanceID: InMeshDescription.VertexInstances().GetElementIDs())
+	auto FillWedge = [&MeshDescription=InMeshDescription, &SkelMeshImportData, &VertexInstanceColors, &VertexInstanceUVs](const FVertexInstanceID InVertexInstanceID, SkeletalMeshImportData::FVertex& OutWedge)
 	{
-		SkeletalMeshImportData::FVertex& Wedge = SkelMeshImportData.Wedges[VertexInstanceID.GetValue()];
-		
-		Wedge.VertexIndex = static_cast<uint32>(InMeshDescription.GetVertexInstanceVertex(VertexInstanceID).GetValue());
-		Wedge.MatIndex = 0;			// We set this later -- not that this is actually used by any internal process.
+		OutWedge.VertexIndex = static_cast<uint32>(MeshDescription.GetVertexInstanceVertex(InVertexInstanceID).GetValue());
+		OutWedge.MatIndex = 0;			// We set this later -- not that this is actually used by any internal process.
 		
 		constexpr bool bSRGB = false; //avoid linear to srgb conversion
-		const FLinearColor VertexColor = VertexInstanceColors[VertexInstanceID];
+		const FLinearColor VertexColor = VertexInstanceColors[InVertexInstanceID];
 		if (VertexColor != FLinearColor::Black)
 		{
 			SkelMeshImportData.bHasVertexColors = true;
 		}
-		Wedge.Color = VertexColor.ToFColor(bSRGB);
+		OutWedge.Color = VertexColor.ToFColor(bSRGB);
 		
 		for (int32 UVChannelIndex = 0; UVChannelIndex < static_cast<int32>(SkelMeshImportData.NumTexCoords); ++UVChannelIndex)
 		{
-			Wedge.UVs[UVChannelIndex] = VertexInstanceUVs.Get(VertexInstanceID, UVChannelIndex);
+			OutWedge.UVs[UVChannelIndex] = VertexInstanceUVs.Get(InVertexInstanceID, UVChannelIndex);
 		}
-	}	
+	};
+	
+	for (FVertexInstanceID VertexInstanceID: InMeshDescription.VertexInstances().GetElementIDs())
+	{
+		SkeletalMeshImportData::FVertex& Wedge = SkelMeshImportData.Wedges[VertexInstanceID.GetValue()];
+
+		FillWedge(VertexInstanceID, Wedge);
+	}
+	
+	// Keep track of which vertex instance ID has been used already. If it gets used again, make a duplicate, so that
+	// we don't end up with two wedges being used by the same face. Downstream algorithms do not like it.
+	// We have to do ensure that the vertex instances are kept in the same order that they got defined originally, because
+	// there are some models that add externally defined morph targets, and they require that the point order is kept
+	// consistent.
+	TBitArray UsedWedges;
+	UsedWedges.SetNum(InMeshDescription.VertexInstances().GetArraySize(), false);
 	
 	//////////////////////////////////////////////////////////////////////////
 	// Copy the triangles
@@ -2833,7 +2846,20 @@ FSkeletalMeshImportData FSkeletalMeshImportData::CreateFromMeshDescription(const
 				Face.TangentY[Corner] = FVector3f::ZeroVector;
 			}
 
-			const int32 WedgeIndex = VertexInstanceID.GetValue();
+			int32 WedgeIndex = VertexInstanceID.GetValue();
+			if (!UsedWedges[WedgeIndex])
+			{
+				// This wedge has not been used before, mark it for first use, so that any subsequent use can be
+				// used to form a duplicate.
+				UsedWedges[WedgeIndex] = true;
+			}
+			else
+			{
+				// Create a new, duplicate wedge.
+				WedgeIndex = SkelMeshImportData.Wedges.Num();
+				FillWedge(VertexInstanceID, SkelMeshImportData.Wedges.AddZeroed_GetRef());
+			}
+			
 			Face.WedgeIndex[Corner] = WedgeIndex;
 			SkelMeshImportData.Wedges[WedgeIndex].MatIndex = Face.MatIndex;
 		}
