OnCommit updating the Actor Class Name Count

diff --git a/Engine/Source/Runtime/Engine/Private/Actor.cpp b/Engine/Source/Runtime/Engine/Private/Actor.cpp
--- a/Engine/Source/Runtime/Engine/Private/Actor.cpp
+++ b/Engine/Source/Runtime/Engine/Private/Actor.cpp
@@ -43,6 +43,7 @@
 #include "ObjectTrace.h"
 #include "Engine/AutoDestroySubsystem.h"
 #include "WorldPartition/WorldPartitionRuntimeCellInterface.h"
+#include "AutoRTFM/AutoRTFM.h"
 #if UE_WITH_IRIS
 #include "Iris/ReplicationSystem/ObjectReplicationBridge.h"
 #include "Iris/ReplicationSystem/ReplicationSystem.h"
@@ -224,6 +225,8 @@ void AActor::InitializeDefaults()
 
 #if (CSV_PROFILER && !UE_BUILD_SHIPPING)
 	// Increment actor class count
+	// Update our ActorClassName count after the transaction finishes
+	UE_AUTORTFM_ONCOMMIT(
 	{
 		if (!HasAnyFlags(RF_ArchetypeObject | RF_ClassDefaultObject))
 		{
@@ -235,7 +238,7 @@ void AActor::InitializeDefaults()
 			CurrentCount++;
 			CSVActorTotalCount++;
 		}
-	}
+	});
 #endif // (CSV_PROFILER && !UE_BUILD_SHIPPING)
 
 #if WITH_EDITORONLY_DATA
@@ -770,6 +773,8 @@ void AActor::BeginDestroy()
 
 #if (CSV_PROFILER && !UE_BUILD_SHIPPING)
 	// Decrement actor class count
+	// Update our ActorClassName count after the transaction finishes
+	UE_AUTORTFM_ONCOMMIT(
 	{
 		if (!HasAnyFlags(RF_ArchetypeObject | RF_ClassDefaultObject))
 		{
@@ -784,7 +789,7 @@ void AActor::BeginDestroy()
 			}
 			CSVActorTotalCount--;
 		}
-	}
+	});
 #endif // (CSV_PROFILER && !UE_BUILD_SHIPPING)
 
 #if WITH_EDITOR
