Fix issue where projectile launching with speed smaller than quantization level will have desync movement
- use actual quantization tolerance for comparing velocity to zero vector

#changelist validated

diff --git a/Engine/Source/Runtime/Engine/Private/Components/ProjectileMovementComponent.cpp b/Engine/Source/Runtime/Engine/Private/Components/ProjectileMovementComponent.cpp
--- a/Engine/Source/Runtime/Engine/Private/Components/ProjectileMovementComponent.cpp
+++ b/Engine/Source/Runtime/Engine/Private/Components/ProjectileMovementComponent.cpp
@@ -184,6 +184,15 @@ void UProjectileMovementComponent::TickComponent(float DeltaTime, enum ELevelTic
 		return;
 	}
 
+	FVector::FReal VelocityTolerance = 0.0;
+	switch (ActorOwner->GetReplicatedMovement().VelocityQuantizationLevel)
+	{
+	case EVectorQuantization::RoundWholeNumber: VelocityTolerance = 1.0; break;
+	case EVectorQuantization::RoundOneDecimal:  VelocityTolerance = 0.1; break;
+	case EVectorQuantization::RoundTwoDecimals: VelocityTolerance = 0.01; break;
+	default:									VelocityTolerance = UE_KINDA_SMALL_NUMBER; break;
+	}
+	
 	float RemainingTime	= DeltaTime;
 	int32 NumImpacts = 0;
 	int32 NumBounces = 0;
@@ -213,7 +222,7 @@ void UProjectileMovementComponent::TickComponent(float DeltaTime, enum ELevelTic
 		Hit.Time = 1.f;
 		const FVector OldVelocity = Velocity;
 		const FVector MoveDelta = ComputeMoveDelta(OldVelocity, TimeTick);
-		FQuat NewRotation = (bRotationFollowsVelocity && !OldVelocity.IsNearlyZero(0.01f)) ? OldVelocity.ToOrientationQuat() : UpdatedComponent->GetComponentQuat();
+		FQuat NewRotation = (bRotationFollowsVelocity && !OldVelocity.IsNearlyZero(VelocityTolerance)) ? OldVelocity.ToOrientationQuat() : UpdatedComponent->GetComponentQuat();
 
 		if (bRotationFollowsVelocity && bRotationRemainsVertical)
 		{
