Adding support for explicit package build dependencies and cvar dependencies for FCookDependency. New use case for these in USoundWave cooking.

diff --git a/Engine/Source/Editor/UnrealEd/Private/TargetDomain/TargetDomainUtils.cpp b/Engine/Source/Editor/UnrealEd/Private/TargetDomain/TargetDomainUtils.cpp
--- a/Engine/Source/Editor/UnrealEd/Private/TargetDomain/TargetDomainUtils.cpp
+++ b/Engine/Source/Editor/UnrealEd/Private/TargetDomain/TargetDomainUtils.cpp
@@ -208,7 +208,7 @@ bool FCookDependencies::TryCalculateCurrentKey(const FAssetPackageData* Override
 					}
 				}
 				bError = true;
-			});
+			}, PackageName);
 
 		for (UE::Cook::FCookDependency& CookDependency : CookDependencies)
 		{
@@ -330,10 +330,6 @@ FCookDependencies FCookDependencies::Collect(UPackage* Package, const ITargetPla
 		BuildDependenciesSet.Append(AssetDependencies);
 	}
 
-	Result.BuildPackageDependencies = BuildDependenciesSet.Array();
-	Result.BuildPackageDependencies.RemoveAllSwap(IsDisallowedBuildDependencyName, EAllowShrinking::Yes);
-	Result.BuildPackageDependencies.Sort(FNameLexicalLess());
-
 #if UE_WITH_CONFIG_TRACKING
 	{
 		using namespace UE::ConfigAccessTracking;
@@ -372,6 +368,10 @@ FCookDependencies FCookDependencies::Collect(UPackage* Package, const ITargetPla
 				}
 				Result.TransitiveBuildDependencies.Add(MoveTemp(CookDependency));
 			}
+			else if (CookDependency.GetType() == UE::Cook::ECookDependency::Package)
+			{
+				BuildDependenciesSet.Add(CookDependency.GetPackageName());
+			}
 			else
 			{
 				Result.CookDependencies.Add(MoveTemp(CookDependency));
@@ -382,6 +382,10 @@ FCookDependencies FCookDependencies::Collect(UPackage* Package, const ITargetPla
 		Algo::Sort(Result.TransitiveBuildDependencies);
 	}
 
+	Result.BuildPackageDependencies = BuildDependenciesSet.Array();
+	Result.BuildPackageDependencies.RemoveAllSwap(IsDisallowedBuildDependencyName, EAllowShrinking::Yes);
+	Result.BuildPackageDependencies.Sort(FNameLexicalLess());
+
 	RuntimeDependencies.RemoveAllSwap(IsDisallowedRuntimeDependencyName, EAllowShrinking::No);
 	RuntimeDependencies.Sort(FNameLexicalLess());
 	RuntimeDependencies.SetNum(Algo::Unique(RuntimeDependencies), EAllowShrinking::Yes);
diff --git a/Engine/Source/Runtime/CoreUObject/Private/Cooker/CookDependency.cpp b/Engine/Source/Runtime/CoreUObject/Private/Cooker/CookDependency.cpp
--- a/Engine/Source/Runtime/CoreUObject/Private/Cooker/CookDependency.cpp
+++ b/Engine/Source/Runtime/CoreUObject/Private/Cooker/CookDependency.cpp
@@ -7,6 +7,7 @@
 #include "Containers/Map.h"
 #include "CoreGlobals.h"
 #include "HAL/FileManager.h"
+#include "HAL/IConsoleManager.h"
 #include "HAL/PlatformMath.h"
 #include "Hash/Blake3.h"
 #include "Misc/AssertionMacros.h"
@@ -37,7 +38,7 @@ namespace UE::Cook
 FCookDependency FCookDependency::File(FStringView InFileName)
 {
 	FCookDependency Result(ECookDependency::File);
-	Result.FileName = InFileName;
+	Result.StringData = InFileName;
 	return Result;
 }
 
@@ -58,6 +59,20 @@ FCookDependency FCookDependency::TransitiveBuildAndRuntime(FName PackageName)
 	return Result;
 }
 
+FCookDependency FCookDependency::Package(FName PackageName)
+{
+	FCookDependency Result(ECookDependency::Package);
+	Result.NameData = PackageName;
+	return Result;
+}
+
+FCookDependency FCookDependency::ConsoleVariable(FStringView VariableName)
+{
+	FCookDependency Result(ECookDependency::ConsoleVariable);
+	Result.StringData = VariableName;
+	return Result;
+}
+
 FCookDependency::FCookDependency()
 	: Type(ECookDependency::None)
 {
@@ -96,7 +111,8 @@ FCookDependency& FCookDependency::operator=(const FCookDependency& Other)
 	case ECookDependency::None:
 		break;
 	case ECookDependency::File:
-		FileName = Other.FileName;
+	case ECookDependency::ConsoleVariable:
+		StringData = Other.StringData;
 		break;
 	case ECookDependency::Function:
 		FunctionData = Other.FunctionData;
@@ -104,6 +120,9 @@ FCookDependency& FCookDependency::operator=(const FCookDependency& Other)
 	case ECookDependency::TransitiveBuild:
 		TransitiveBuildData = Other.TransitiveBuildData;
 		break;
+	case ECookDependency::Package:
+		NameData = Other.NameData;
+		break;
 	default:
 		checkNoEntry();
 		break;
@@ -121,7 +140,8 @@ FCookDependency& FCookDependency::operator=(FCookDependency&& Other)
 	case ECookDependency::None:
 		break;
 	case ECookDependency::File:
-		FileName = MoveTemp(Other.FileName);
+	case ECookDependency::ConsoleVariable:
+		StringData = MoveTemp(Other.StringData);
 		break;
 	case ECookDependency::Function:
 		FunctionData = MoveTemp(Other.FunctionData);
@@ -129,6 +149,9 @@ FCookDependency& FCookDependency::operator=(FCookDependency&& Other)
 	case ECookDependency::TransitiveBuild:
 		TransitiveBuildData = MoveTemp(Other.TransitiveBuildData);
 		break;
+	case ECookDependency::Package:
+		NameData = MoveTemp(Other.NameData);
+		break;
 	default:
 		checkNoEntry();
 		break;
@@ -191,10 +214,28 @@ void FCookDependency::UpdateHash(FCookDependencyContext& Context) const
 		return;
 	}
 	case ECookDependency::TransitiveBuild:
-		// Transitive Build dependencies do not impact the hash; they instead operate by marking the package
+		// Build dependencies do not impact the hash; they instead operate by marking the package
 		// as invalidated based on the invalidation of other packages, in a separate pass after its hash is compared
 		return;
-
+	case ECookDependency::Package:
+		Context.LogError(FString::Printf(
+			TEXT("FCookDependency::Package('%s') failed to UpdateHash: Package dependencies do not implement UpdateHash and it should not be called on them."),
+			*NameData.ToString()));
+		return;
+	case ECookDependency::ConsoleVariable:
+	{
+		IConsoleVariable* VariableInstance = IConsoleManager::Get().FindConsoleVariable(*StringData);
+		if (VariableInstance == nullptr)
+		{
+			Context.LogError(FString::Printf(
+				TEXT("FCookDependency::ConsoleVariable('%s') failed to UpdateHash: could not find console variable."),
+				*StringData));
+			return;
+		}
+		FString VariableAsString = VariableInstance->GetString();
+		Context.Update(VariableAsString.GetCharArray().GetData(), VariableAsString.GetCharArray().Num() * sizeof(FString::ElementType));
+		return;
+	}
 	default:
 		checkNoEntry();
 		return;
@@ -208,7 +249,8 @@ void FCookDependency::Construct()
 	case ECookDependency::None:
 		break;
 	case ECookDependency::File:
-		new(&FileName) FString();
+	case ECookDependency::ConsoleVariable:
+		new(&StringData) FString();
 		break;
 	case ECookDependency::Function:
 		new(&FunctionData) FFunctionData();
@@ -216,6 +258,9 @@ void FCookDependency::Construct()
 	case ECookDependency::TransitiveBuild:
 		new(&TransitiveBuildData) FTransitiveBuildData();
 		break;
+	case ECookDependency::Package:
+		new(&NameData) FName();
+		break;
 	default:
 		checkNoEntry();
 		break;
@@ -229,7 +274,8 @@ void FCookDependency::Destruct()
 	case ECookDependency::None:
 		break;
 	case ECookDependency::File:
-		FileName.~FString();
+	case ECookDependency::ConsoleVariable:
+		StringData.~FString();
 		break;
 	case ECookDependency::Function:
 		FunctionData.~FFunctionData();
@@ -237,6 +283,9 @@ void FCookDependency::Destruct()
 	case ECookDependency::TransitiveBuild:
 		TransitiveBuildData.~FTransitiveBuildData();
 		break;
+	case ECookDependency::Package:
+		NameData.~FName();
+		break;
 	default:
 		checkNoEntry();
 		break;
@@ -252,7 +301,8 @@ void FCookDependency::Save(FCbWriter& Writer) const
 	case ECookDependency::None:
 		break;
 	case ECookDependency::File:
-		Writer << FileName;
+	case ECookDependency::ConsoleVariable:
+		Writer << StringData;
 		break;
 	case ECookDependency::Function:
 		Writer << FunctionData.Name;
@@ -265,6 +315,9 @@ void FCookDependency::Save(FCbWriter& Writer) const
 		Writer << TransitiveBuildData.PackageName;
 		Writer << TransitiveBuildData.bAlsoAddRuntimeDependency;
 		break;
+	case ECookDependency::Package:
+		Writer << NameData;
+		break;
 	default:
 		checkNoEntry();
 		break;
@@ -328,6 +381,26 @@ bool FCookDependency::Load(FCbFieldView Value)
 		TransitiveBuildData.bAlsoAddRuntimeDependency = bLocalAlsoAddRuntimeDependency;
 		return true;
 	}
+	case ECookDependency::Package:
+	{
+		FName LocalPackageName;
+		if (!LoadFromCompactBinary(Field++, LocalPackageName))
+		{
+			return false;
+		}
+		*this = FCookDependency::Package(LocalPackageName);
+		return true;
+	}
+	case ECookDependency::ConsoleVariable:
+	{
+		FString LocalConsoleVariableName;
+		if (!LoadFromCompactBinary(Field++, LocalConsoleVariableName))
+		{
+			return false;
+		}
+		*this = FCookDependency::ConsoleVariable(LocalConsoleVariableName);
+		return true;
+	}
 	default:
 		break;
 	}
diff --git a/Engine/Source/Runtime/CoreUObject/Public/Cooker/CookDependency.h b/Engine/Source/Runtime/CoreUObject/Public/Cooker/CookDependency.h
--- a/Engine/Source/Runtime/CoreUObject/Public/Cooker/CookDependency.h
+++ b/Engine/Source/Runtime/CoreUObject/Public/Cooker/CookDependency.h
@@ -24,7 +24,7 @@ struct FCookDependencyContext
 {
 public:
 	/** InHasher is void* to mask the implementation details of the hashbuilder. See Update function. */
-	explicit FCookDependencyContext(void* InHasher, TUniqueFunction<void(FString&&)>&& InOnLogError);
+	explicit FCookDependencyContext(void* InHasher, TUniqueFunction<void(FString&&)>&& InOnLogError, FName InPackageName);
 
 	/**
 	 * Update the hashbuilder for the key being constructed (e.g. TargetDomainKey for cooked packages)
@@ -63,9 +63,14 @@ public:
 	[[nodiscard]] COREUOBJECT_API FErrorHandlerScope ErrorHandlerScope(
 		TUniqueFunction<FString(FString&&)>&& ErrorHandler);
 
+	/**
+	 * Get the name of the package being considered
+	 */
+	COREUOBJECT_API FName GetPackageName() const { return PackageName; }
 private:
 	TUniqueFunction<void(FString&&)> OnLogError;
 	TArray<TUniqueFunction<FString(FString&&)>, TInlineAllocator<1>> ErrorHandlers;
+	FName PackageName;
 	void* Hasher; // Type is void* to mask the implementation detail
 };
 
@@ -79,6 +84,8 @@ enum class ECookDependency : uint8
 	File = 0x01,
 	Function = 0x02,
 	TransitiveBuild = 0x03,
+	Package = 0x4,
+	ConsoleVariable = 0x5
 };
 
 /**
@@ -119,6 +126,18 @@ public:
 	 */
 	COREUOBJECT_API static FCookDependency TransitiveBuildAndRuntime(FName PackageName);
 
+	/**
+	 * Create a build dependency on the contents of a package.
+	 * Only the bytes of the .uasset/.umap file are considered.
+	 */
+	COREUOBJECT_API static FCookDependency Package(FName PackageName);
+
+	/**
+	 * Create a dependency on the value of a cvar. The cvar will be read and its value (as a string) will be hashed into the oplog data
+	 * If the cvar value is changed, the packages that depend on it will be invalidated
+	 */
+	COREUOBJECT_API static FCookDependency ConsoleVariable(FStringView VariableName);
+
 	/** Construct an empty dependency; it will never be invalidated. */
 	COREUOBJECT_API FCookDependency();
 
@@ -139,7 +158,7 @@ public:
 	/** FunctionArgs if GetType() == Function, else FCbFieldViewIterator(). */
 	FCbFieldViewIterator GetFunctionArgs() const;
 
-	/** PackageName if GetType() == TransitiveBuild, else NAME_None. */
+	/** PackageName if GetType() == TransitiveBuild or GetType() == Package, else NAME_None. */
 	FName GetPackageName() const;
 	/** If GetType() == TransitiveBuild, whether AlsoAddRuntimeDependency was selected, otherwise false. */
 	bool IsAlsoAddRuntimeDependency() const;
@@ -186,9 +205,10 @@ private:
 	};
 	union
 	{
-		FString FileName;
+		FString StringData;
 		FFunctionData FunctionData;
 		FTransitiveBuildData TransitiveBuildData;
+		FName NameData;
 	};
 };
 
@@ -273,7 +293,7 @@ inline ECookDependency FCookDependency::GetType() const
 
 inline FStringView FCookDependency::GetFileName() const
 {
-	return Type == ECookDependency::File ? FileName : FStringView();
+	return Type == ECookDependency::File ? StringData : FStringView();
 }
 
 inline FName FCookDependency::GetFunctionName() const
@@ -288,7 +308,15 @@ inline FCbFieldViewIterator FCookDependency::GetFunctionArgs() const
 
 inline FName FCookDependency::GetPackageName() const
 {
-	return Type == ECookDependency::TransitiveBuild ? TransitiveBuildData.PackageName : NAME_None;
+	switch (GetType())
+	{
+	case ECookDependency::TransitiveBuild:
+		return TransitiveBuildData.PackageName;
+	case ECookDependency::Package:
+		return NameData;
+	default:
+		return NAME_None;
+	};
 }
 
 inline bool FCookDependency::IsAlsoAddRuntimeDependency() const
@@ -308,7 +336,8 @@ inline bool FCookDependency::operator<(const FCookDependency& Other) const
 	case ECookDependency::None:
 		return false;
 	case ECookDependency::File:
-		return FileName.Compare(Other.FileName, ESearchCase::IgnoreCase) < 0;
+	case ECookDependency::ConsoleVariable:
+		return StringData.Compare(Other.StringData, ESearchCase::IgnoreCase) < 0;
 	case ECookDependency::Function:
 	{
 		int32 Compare = FunctionData.Name.Compare(Other.FunctionData.Name);
@@ -340,14 +369,17 @@ inline bool FCookDependency::operator<(const FCookDependency& Other) const
 		}
 		return false;
 	}
+	case ECookDependency::Package:
+		return NameData.Compare(Other.NameData) <  0;
 	default:
 		checkNoEntry();
 		return false;
 	}
 }
 
-inline FCookDependencyContext::FCookDependencyContext(void* InHasher, TUniqueFunction<void(FString&&)>&& InOnLogError)
+inline FCookDependencyContext::FCookDependencyContext(void* InHasher, TUniqueFunction<void(FString&&)>&& InOnLogError, FName InPackageName)
 	: OnLogError(MoveTemp(InOnLogError))
+	, PackageName(InPackageName)
 	, Hasher(InHasher)
 {
 }
diff --git a/Engine/Source/Runtime/Engine/Classes/Sound/SoundWave.h b/Engine/Source/Runtime/Engine/Classes/Sound/SoundWave.h
--- a/Engine/Source/Runtime/Engine/Classes/Sound/SoundWave.h
+++ b/Engine/Source/Runtime/Engine/Classes/Sound/SoundWave.h
@@ -874,6 +874,9 @@ public:
 	*/
 	static ENGINE_API ITargetPlatform* GetRunningPlatform();
 
+	/** Used to ensure that we recook USoundWaves if the default loading behavior has changed */
+	static ENGINE_API const TCHAR* GetDefaultLoadingBehaviorCVarName();
+
 	static ENGINE_API ESoundWaveLoadingBehavior GetDefaultLoadingBehavior();
 
 	/** Async worker that decompresses the audio data on a different thread */
diff --git a/Engine/Source/Runtime/Engine/Classes/Sound/SoundWaveLoadingBehavior.cpp b/Engine/Source/Runtime/Engine/Classes/Sound/SoundWaveLoadingBehavior.cpp
--- a/Engine/Source/Runtime/Engine/Classes/Sound/SoundWaveLoadingBehavior.cpp
+++ b/Engine/Source/Runtime/Engine/Classes/Sound/SoundWaveLoadingBehavior.cpp
@@ -13,6 +13,14 @@
 #include "SoundWave.h"
 #include "UObject/LinkerLoad.h"
 
+#if WITH_EDITOR
+#include "Algo/Sort.h"
+#include "Cooker/CookDependency.h"
+#include "Serialization/CompactBinarySerialization.h"
+#include "Serialization/CompactBinaryWriter.h"
+#include "UObject/ObjectSaveContext.h"
+#endif
+
 #ifndef CASE_ENUM_TO_TEXT
 #define CASE_ENUM_TO_TEXT(TXT) case TXT: return TEXT(#TXT);
 #endif
@@ -45,6 +53,11 @@ FAutoConsoleVariableRef CVAR_SoundWaveLoadingBehaviorUtil_Enable(
 	ECVF_Default
 );
 
+namespace UE::SoundWaveLoadingUtil::Private
+{
+	void HashSoundWaveLoadingBehaviorDependenciesForCook(FCbFieldViewIterator Args, UE::Cook::FCookDependencyContext& Context);
+};
+
 class FSoundWaveLoadingBehaviorUtil : public ISoundWaveLoadingBehaviorUtil
 {
 public:
@@ -59,6 +72,18 @@ public:
 	virtual ~FSoundWaveLoadingBehaviorUtil() override = default;
 
 private:
+
+	struct FAnnotatedClassData
+	{
+		explicit FAnnotatedClassData(USoundClass* InSoundClass) 
+		: BaseData(InSoundClass) 
+		{
+			ClassHierarchy.Add(FName(InSoundClass->GetPackage()->GetFName()));
+		}
+		FClassData BaseData;
+		TArray<FName> ClassHierarchy;
+	};
+
 	void CacheAllClassLoadingBehaviors()
 	{
 		const bool bAssetRegistryInStartup = AssetRegistry.IsSearchAsync() && AssetRegistry.IsLoadingAssets();
@@ -76,9 +101,9 @@ private:
 		}
 	}
 
-	FClassData WalkClassHierarchy(USoundClass* InClass) const
+	FAnnotatedClassData WalkClassHierarchy(USoundClass* InClass) const
 	{
-		FClassData Behavior(InClass);
+		FAnnotatedClassData Behavior(InClass);
 		while (InClass->ParentClass && InClass->Properties.LoadingBehavior == ESoundWaveLoadingBehavior::Inherited)
 		{
 			InClass = InClass->ParentClass;
@@ -90,15 +115,16 @@ private:
 			{
 				InClass->ConditionalPostLoad();
 			}
-			Behavior = FClassData(InClass);
+			Behavior.BaseData = FClassData(InClass);
+			Behavior.ClassHierarchy.Add(FName(InClass->GetPackage()->GetFName()));
 		}
 		
 		// If we failed to find anything other than inherited, use the default which is cvar'd.
-		if (Behavior.LoadingBehavior == ESoundWaveLoadingBehavior::Inherited ||
-			Behavior.LoadingBehavior == ESoundWaveLoadingBehavior::Uninitialized )
+		if (Behavior.BaseData.LoadingBehavior == ESoundWaveLoadingBehavior::Inherited ||
+			Behavior.BaseData.LoadingBehavior == ESoundWaveLoadingBehavior::Uninitialized )
 		{
-			Behavior.LoadingBehavior = USoundWave::GetDefaultLoadingBehavior();
-			Behavior.LengthOfFirstChunkInSeconds = 0;
+			Behavior.BaseData.LoadingBehavior = USoundWave::GetDefaultLoadingBehavior();
+			Behavior.BaseData.LengthOfFirstChunkInSeconds = 0;
 		}
 		return Behavior;
 	}
@@ -107,51 +133,24 @@ private:
 	{
 		if (USoundClass* SoundClass = Cast<USoundClass>(InAssetData.GetAsset()))
 		{
-			FClassData Result = WalkClassHierarchy(SoundClass);
+			FAnnotatedClassData Result = WalkClassHierarchy(SoundClass);
 			CacheClassLoadingBehaviors.Add(InAssetData.PackageName,Result);
-			return Result;
+			return Result.BaseData;
 		}
 		return {};
 	}
-	
-	virtual FClassData FindOwningLoadingBehavior(const USoundWave* InWave, const ITargetPlatform* InTargetPlatform) const override
-	{
-		TRACE_CPUPROFILER_EVENT_SCOPE(FindOwningLoadingBehavior);
-
-		// This code: Given a wave, finds all cues that references it. (reverse lookup)
-		// Then finds the SoundClasses those cues use, traverses the heirarchy (from lookup) to determine the loading behavior.
-		// Then stack ranks the most important behavior. (RetainOnLoad (Highest), PrimeOnLoad (Medium), LoadOnDemand (Lowest))
-		// Which ever wins, we also capture the "SizeOfFirstChunk" to use for that wave.
-
-		if (!InWave)
-		{
-			return {};
-		}
-
-		const bool bIsAssetRegistryStartup = AssetRegistry.IsSearchAsync() && AssetRegistry.IsLoadingAssets();
-
-		// Disallow during startup of registry (cookers will have already done this)
-		if (bIsAssetRegistryStartup)
-		{
-			UE_LOG(LogAudio, Warning, TEXT("FindOwningLoadingBehavior called before AssetRegistry is ready. SoundWave=%s"), *InWave->GetName());
-			return {};
-		}
-	
-		const UPackage* WavePackage = InWave->GetPackage();
-		if (!WavePackage)
-		{
-			return {};
-		}
 
+	void CollectAllRelevantSoundClassAssetData(const FName InWaveName, TSet<FAssetData>& OutClassAssetData) const
+	{
 		TArray<FName> SoundWaveReferencerNames;
-		if (!AssetRegistry.GetReferencers(WavePackage->GetFName(), SoundWaveReferencerNames))
+		if (!AssetRegistry.GetReferencers(InWaveName, SoundWaveReferencerNames))
 		{
-			return {};
+			return;
 		}
 
 		if (SoundWaveReferencerNames.IsEmpty())
 		{
-			return {};
+			return;
 		}
 
 		// Filter on SoundCues.
@@ -164,7 +163,7 @@ private:
 		TArray<FAssetData> ReferencingSoundCueAssetDataArray;
 		if (!AssetRegistry.GetAssets(Filter, ReferencingSoundCueAssetDataArray))
 		{
-			return {};
+			return;
 		}
 
 		// Filter out unwanted classes, see above comment for details.
@@ -178,18 +177,17 @@ private:
 
 		if (ReferencingSoundCueAssetDataArray.IsEmpty())
 		{
-			return {};
+			return;
 		}
 
-		FClassData MostImportantLoadingBehavior;
 		for (const FAssetData& CueAsset : ReferencingSoundCueAssetDataArray)
 		{
 			// Query for class references from the Cue instead of loading and opening it.
 			TArray<FName> SoundCueReferences;
 			if (!AssetRegistry.GetDependencies(CueAsset.PackageName, SoundCueReferences))
 			{
-				UE_LOG(LogAudio, Warning, TEXT("Failed to query SoundCue '%s' for it's dependencies."), 
-				       *CueAsset.PackagePath.ToString() );
+				UE_LOG(LogAudio, Warning, TEXT("Failed to query SoundCue '%s' for it's dependencies."),
+					*CueAsset.PackagePath.ToString());
 				continue;
 			}
 
@@ -207,40 +205,177 @@ private:
 			{
 				UE_LOG(LogAudio, Warning, TEXT("Failed to filter for Soundclasses from the SoundCue dependencies for '%s'"), *CueAsset.PackagePath.ToString());
 				continue;
-			}			
-			
-			// If there's more than one, rank them.
-			for (const FAssetData& Class : ReferencedSoundClasses)
+			}
+			for (FAssetData& SoundClass : ReferencedSoundClasses)
+			{
+				OutClassAssetData.Add(MoveTemp(SoundClass));
+			}
+		}
+	}
+
+	virtual FClassData FindOwningLoadingBehavior(const USoundWave* InWave, const ITargetPlatform* InTargetPlatform) const override
+	{
+		TRACE_CPUPROFILER_EVENT_SCOPE(FindOwningLoadingBehavior);
+
+		// This code: Given a wave, finds all cues that references it. (reverse lookup)
+		// Then finds the SoundClasses those cues use, traverses the heirarchy (from lookup) to determine the loading behavior.
+		// Then stack ranks the most important behavior. (RetainOnLoad (Highest), PrimeOnLoad (Medium), LoadOnDemand (Lowest))
+		// Which ever wins, we also capture the "SizeOfFirstChunk" to use for that wave.
+
+		if (!InWave)
+		{
+			return {};
+		}
+
+		const bool bIsAssetRegistryStartup = AssetRegistry.IsSearchAsync() && AssetRegistry.IsLoadingAssets();
+
+		// Disallow during startup of registry (cookers will have already done this)
+		if (bIsAssetRegistryStartup)
+		{
+			UE_LOG(LogAudio, Warning, TEXT("FindOwningLoadingBehavior called before AssetRegistry is ready. SoundWave=%s"), *InWave->GetName());
+			return {};
+		}
+	
+		const UPackage* WavePackage = InWave->GetPackage();
+		if (!WavePackage)
+		{
+			return {};
+		}
+
+		TSet<FAssetData> SoundClassesToConsider;
+		CollectAllRelevantSoundClassAssetData(WavePackage->GetFName(), SoundClassesToConsider);
+
+		FClassData MostImportantLoadingBehavior;
+
+		// If there's more than one, rank them.
+		FScopeLock Lock(&CacheCS);
+		for (const FAssetData& Class : SoundClassesToConsider)
+		{
+			FClassData CacheLoadingBehavior;
+			if (const FAnnotatedClassData* Found = CacheClassLoadingBehaviors.Find(Class.PackageName))
+			{
+				CacheLoadingBehavior = Found->BaseData;
+			}
+			else
+			{
+				CacheLoadingBehavior = LoadAndCacheClass(Class);
+			}
+
+			// Compare if this is more important
+			if (MostImportantLoadingBehavior.CompareGreater(CacheLoadingBehavior, InTargetPlatform))
 			{
-				FClassData CacheLoadingBehavior;
-				FScopeLock Lock(&CacheCS);
-				if (const FClassData* Found = CacheClassLoadingBehaviors.Find(Class.PackageName))
-				{
-					CacheLoadingBehavior = *Found;
-				}
-				else
-				{
-					CacheLoadingBehavior = LoadAndCacheClass(Class);
-				}
-
-				// Compare if this is more important
-				if (MostImportantLoadingBehavior.CompareGreater(CacheLoadingBehavior, InTargetPlatform))
-				{
-					MostImportantLoadingBehavior = CacheLoadingBehavior;
-				}
+				MostImportantLoadingBehavior = CacheLoadingBehavior;
 			}
 		}
 		
 		// Return the most important one we found.
 		return MostImportantLoadingBehavior;
 	}
-	
+
+	void GetChainOfClassesForLoadingBehaviorInheritance(FName StartingClass, TArray<FName>& OutClasses)
+	{
+		FScopeLock Lock(&CacheCS);
+		const FAnnotatedClassData* CacheEntry = CacheClassLoadingBehaviors.Find(StartingClass);
+		if (CacheEntry)
+		{
+			OutClasses = CacheEntry->ClassHierarchy;
+		}
+		else
+		{
+			OutClasses.Reset();
+		}
+	}
+
+
 	// Make cache here.
 	IAssetRegistry& AssetRegistry;
-	mutable TMap<FName, FClassData> CacheClassLoadingBehaviors;
+	mutable TMap<FName, FAnnotatedClassData> CacheClassLoadingBehaviors;
 	mutable FCriticalSection CacheCS;
+	friend void UE::SoundWaveLoadingUtil::Private::HashSoundWaveLoadingBehaviorDependenciesForCook(FCbFieldViewIterator Args, UE::Cook::FCookDependencyContext& Context);
+	friend void UE::SoundWaveLoadingUtil::Private::RecordSoundWaveLoadingBehaviorDependenciesForCook(FObjectPreSaveContext ObjectSaveContext,	const USoundWave* SoundWave);
 };
 
+namespace UE { namespace SoundWaveLoadingUtil { namespace Private {
+constexpr int32 HashSoundWaveLoadingBehaviorDependenciesForCookArgsVersion = 1;
+void HashSoundWaveLoadingBehaviorDependenciesForCook(FCbFieldViewIterator Args, UE::Cook::FCookDependencyContext& Context)
+{
+	int32 ArgsVersion = -1;
+	bool bValid = false;
+
+	FCbFieldViewIterator ArgField(Args);
+	ArgsVersion = (ArgField++).AsInt32();
+	if (ArgsVersion == HashSoundWaveLoadingBehaviorDependenciesForCookArgsVersion)
+	{
+		bValid = true;
+	}
+	if (!bValid)
+	{
+		Context.LogError(FString::Printf(TEXT("Unsupported arguments version %d."), ArgsVersion));
+		return;
+	}
+
+	TSet<FAssetData> ClassDependenciesFromAssetRegistryQuery;
+	FSoundWaveLoadingBehaviorUtil* Util = static_cast<FSoundWaveLoadingBehaviorUtil*>(FSoundWaveLoadingBehaviorUtil::Get());
+	Util->CollectAllRelevantSoundClassAssetData(Context.GetPackageName(), ClassDependenciesFromAssetRegistryQuery);
+
+	TArray<FName> SortedClassDependencies;
+	for (const FAssetData& AssetData : ClassDependenciesFromAssetRegistryQuery)
+	{
+		SortedClassDependencies.Add(AssetData.PackageName);
+	}
+	Algo::Sort(SortedClassDependencies, FNameLexicalLess());
+
+	for (const FName& ClassDependency : SortedClassDependencies)
+	{
+		TStringBuilder<256> ClassNameString(InPlace, ClassDependency);
+		Context.Update(ClassNameString.GetData(), ClassNameString.Len() * sizeof(ClassNameString.GetData()[0]));
+	}
+}
+
+UE_COOK_DEPENDENCY_FUNCTION(HashSoundWaveLoadingBehaviorDependenciesForCook, HashSoundWaveLoadingBehaviorDependenciesForCook);
+
+void RecordSoundWaveLoadingBehaviorDependenciesForCook(FObjectPreSaveContext ObjectSaveContext,
+	const USoundWave* SoundWave)
+{
+	// Cooking of USoundWaves depends on a bunch of indirect state that we want to cooker to know about for incremental cook invalidation
+	// This is due to how we determine the sound wave loading behavior and associated chunk sizes
+	// Anything that impacts the result of FindOwningLoadingBehavior needs to be encapsulated so that the cooker can
+	// determine whether the results might be different from when the base cook was made
+	FCbWriter Writer;
+	Writer << HashSoundWaveLoadingBehaviorDependenciesForCookArgsVersion;
+
+	// 1. We depend on the set of USoundClass objects which are dependencies of any USoundCue that depends on our USoundWave
+	// Those classes are identified via a series of asset registry queries encapsulated in CollectAllRelevantSoundClassAssetData
+	// HashSoundWaveLoadingBehaviorDependenciesForCook is used to package up that query and make it deterministic for the cooker
+	// dependency checking
+	
+	ObjectSaveContext.AddCookBuildDependency(
+		UE::Cook::FCookDependency::Function(
+			UE_COOK_DEPENDENCY_FUNCTION_CALL(HashSoundWaveLoadingBehaviorDependenciesForCook), Writer.Save()));
+
+	FSoundWaveLoadingBehaviorUtil* Util = static_cast<FSoundWaveLoadingBehaviorUtil*>(FSoundWaveLoadingBehaviorUtil::Get());
+
+	// 2. For each class in 1, we have an additional dependency on a its inheritance hierarchy as identified by 
+	// GetChainOfClassesForLoadingBehaviorInheritance 
+	TSet<FAssetData> ClassDependenciesFromAssetRegistryQuery;
+	Util->CollectAllRelevantSoundClassAssetData(SoundWave->GetPackage()->GetFName(), ClassDependenciesFromAssetRegistryQuery);
+
+	for (const FAssetData& DependencyFromRegistry : ClassDependenciesFromAssetRegistryQuery)
+	{
+		TArray<FName> SoundCueHierarchyDependencies;
+		Util->GetChainOfClassesForLoadingBehaviorInheritance(DependencyFromRegistry.PackageName, SoundCueHierarchyDependencies);
+
+		for (FName SoundCue : SoundCueHierarchyDependencies)
+		{
+			ObjectSaveContext.AddCookBuildDependency(UE::Cook::FCookDependency::Package(SoundCue));
+		}
+	}
+
+	// 3. Finally, we depend on the value of the default loading behavior cvar
+	ObjectSaveContext.AddCookBuildDependency(UE::Cook::FCookDependency::ConsoleVariable(USoundWave::GetDefaultLoadingBehaviorCVarName()));
+}
+}}}
+
 ISoundWaveLoadingBehaviorUtil* ISoundWaveLoadingBehaviorUtil::Get()
 {
 	// Cvar disable system if necessary
diff --git a/Engine/Source/Runtime/Engine/Classes/Sound/SoundWaveLoadingBehavior.h b/Engine/Source/Runtime/Engine/Classes/Sound/SoundWaveLoadingBehavior.h
--- a/Engine/Source/Runtime/Engine/Classes/Sound/SoundWaveLoadingBehavior.h
+++ b/Engine/Source/Runtime/Engine/Classes/Sound/SoundWaveLoadingBehavior.h
@@ -65,4 +65,9 @@ public:
 	virtual FClassData FindOwningLoadingBehavior(const USoundWave*, const ITargetPlatform*) const = 0;
 };
 
+namespace UE::SoundWaveLoadingUtil::Private
+{
+	void RecordSoundWaveLoadingBehaviorDependenciesForCook(FObjectPreSaveContext ObjectSaveContext, const USoundWave* SoundWave);
+};
+
 #endif //WITH_EDITOR
diff --git a/Engine/Source/Runtime/Engine/Private/SoundWave.cpp b/Engine/Source/Runtime/Engine/Private/SoundWave.cpp
--- a/Engine/Source/Runtime/Engine/Private/SoundWave.cpp
+++ b/Engine/Source/Runtime/Engine/Private/SoundWave.cpp
@@ -42,9 +42,12 @@
 #include "UObject/ArchiveCookContext.h"
 #endif
 
+// WARNING: This cvar is used as part of incremental cook invalidation. Changing the name of the variable
+// will invalidate all cooked sound wave packages.
+static const TCHAR* SoundWaveDefaultLoadingBehaviorCVarName = TEXT("au.streamcache.SoundWaveDefaultLoadingBehavior");
 static int32 SoundWaveDefaultLoadingBehaviorCVar = static_cast<int32>(ESoundWaveLoadingBehavior::LoadOnDemand);
 FAutoConsoleVariableRef CVarSoundWaveDefaultLoadingBehavior(
-	TEXT("au.streamcache.SoundWaveDefaultLoadingBehavior"),
+	SoundWaveDefaultLoadingBehaviorCVarName,
 	SoundWaveDefaultLoadingBehaviorCVar,
 	TEXT("This can be set to define the default behavior when a USoundWave is loaded.\n")
 	TEXT("1: Retain audio data on load, 2: prime audio data on load, 3: load on demand (No audio data is loaded until a USoundWave is played or primed)."),
@@ -98,6 +101,11 @@ inline FArchive& operator<<(FArchive& Ar, FSoundWaveCuePoint& CuePoint)
 
 namespace SoundWave_Private
 {
+	static const TCHAR* GetDefaultLoadingBehaviorCVarName()
+	{
+		return SoundWaveDefaultLoadingBehaviorCVarName;
+	}
+
 	static ESoundWaveLoadingBehavior GetDefaultLoadingBehaviorCVar()
 	{
 		// query the default loading behavior CVar
@@ -744,6 +752,11 @@ ITargetPlatform* USoundWave::GetRunningPlatform()
 	}
 }
 
+ENGINE_API const TCHAR* USoundWave::GetDefaultLoadingBehaviorCVarName()
+{
+	return SoundWave_Private::GetDefaultLoadingBehaviorCVarName();
+}
+
 ENGINE_API ESoundWaveLoadingBehavior USoundWave::GetDefaultLoadingBehavior() 
 {
 	return SoundWave_Private::GetDefaultLoadingBehaviorCVar();
@@ -2394,6 +2407,7 @@ void USoundWave::PreSave(FObjectPreSaveContext InSaveContext)
 	{
 		// Populate our cache ahead of Serialize, where we can't do any AssetRegistry queries safely.
 		GetOwnerLoadingBehavior(InSaveContext.GetTargetPlatform());
+		UE::SoundWaveLoadingUtil::Private::RecordSoundWaveLoadingBehaviorDependenciesForCook(InSaveContext, this);
 	}
 	Super::PreSave(InSaveContext);
 }
