Fix crash when a texture with layers is used in a project with VT disabled
add new functions to ask texture if VT is required to build
RequiresVirtualTexturing
IsVirtualTexturingEnabled
CanBuildPlatformData
use CanBuildPlatformData instead of Source.IsValid()

diff --git a/Engine/Source/Runtime/Engine/Classes/Engine/Texture.h b/Engine/Source/Runtime/Engine/Classes/Engine/Texture.h
--- a/Engine/Source/Runtime/Engine/Classes/Engine/Texture.h
+++ b/Engine/Source/Runtime/Engine/Classes/Engine/Texture.h
@@ -34,6 +34,7 @@ namespace FOodleDataCompression {enum class ECompressor : uint8; enum class ECom
 class FTextureReference;
 class FTextureResource;
 class ITargetPlatform;
+class ITargetPlatformSettings;
 class UAssetUserData;
 struct FPropertyChangedEvent;
 
@@ -1708,6 +1709,16 @@ public:
 	/** Returns the virtual texture build settings. */
 	ENGINE_API virtual void GetVirtualTextureBuildSettings(struct FVirtualTextureBuildSettings& OutSettings) const;
 
+#if WITH_EDITORONLY_DATA
+	/** Returns true if this Texture cannot be built unless it has VT enabled.
+	Does not check the VirtualTextureStreaming bool on the texture, this is telling you if that bool must be true. **/
+	ENGINE_API bool RequiresVirtualTexturing() const;
+#endif
+
+	/** Check the Project settings to see if VT is enabled.
+	Optionally also checks the TargetPlatform. **/
+	ENGINE_API static bool IsVirtualTexturingEnabled( const ITargetPlatformSettings * TargetPlatform = nullptr );
+
 	/**
 	 * Textures that use the derived data cache must override this function and
 	 * provide a pointer to the linked list of platform data.
@@ -1760,6 +1771,12 @@ public:
 	 */
 	void CachePlatformData(bool bAsyncCache = false, bool bAllowAsyncBuild = false, bool bAllowAsyncLoading = false, class ITextureCompressorModule* Compressor = nullptr);
 
+	/* Returns if it is possible to build this texture for a given TargetPlatform (or any).
+	This supercedes checking Source.IsValid()
+	When this returns false, CachePlatformData will not attempt to build or cache the PlatformData.
+	*/
+	bool CanBuildPlatformData(const ITargetPlatformSettings * TargetPlatform = nullptr) const;
+
 	/**
 	 * Begins caching platform data in the background for the platform requested
 	 */
diff --git a/Engine/Source/Runtime/Engine/Private/Texture.cpp b/Engine/Source/Runtime/Engine/Private/Texture.cpp
--- a/Engine/Source/Runtime/Engine/Private/Texture.cpp
+++ b/Engine/Source/Runtime/Engine/Private/Texture.cpp
@@ -464,6 +464,33 @@ bool UTexture::IsPostLoadThreadSafe() const
 	return false;
 }
 
+/*static*/ bool UTexture::IsVirtualTexturingEnabled( const ITargetPlatformSettings * TargetPlatform /*= nullptr*/ )
+{
+	// check the project setting cvar for overall VT being on/off :
+	//	(this is in the host platform config)
+	if ( ! CVarVirtualTextures.GetValueOnAnyThread() )
+	{
+		return false;
+	}
+	
+	// note: rendering code (UseVirtualTexturing) also checks FPlatformProperties::SupportsVirtualTextureStreaming()
+	//	that has historically not been checked in the texture code, so we continue to not check it here
+
+	// optionally could do: if ! TargetPlatform , TargetPlatform = CurrentRunning ?
+
+	if ( TargetPlatform )
+	{
+		// this will wind up checking "r.Mobile.VirtualTextures" if TargetPlatform is a mobile platform
+		const bool bPlatformSupportsVirtualTextureStreaming = TargetPlatform->SupportsFeature(ETargetPlatformFeatures::VirtualTextureStreaming);
+		if ( ! bPlatformSupportsVirtualTextureStreaming )
+		{
+			return false;
+		}
+	}
+
+	return true;
+}
+
 #if WITH_EDITOR
 
 bool UTexture::IsDefaultTexture() const
@@ -562,7 +589,6 @@ void UTexture::UpdateOodleTextureSdkVersionToLatest(bool bDoPrePostEditChangeIfC
 	}
 }
 
-
 // we're in WITH_EDITOR but not sure that's right, maybe move out?
 // Beware: while ValidateSettingsAfterImportOrEdit should have been called on all Textures,
 //	 it is not called on load at runtime
@@ -646,7 +672,7 @@ void UTexture::ValidateSettingsAfterImportOrEdit(bool * pRequiresNotifyMaterials
 		// note : checking the VirtualTextureStreaming without checking the TargetPlatform is potentially buggy
 		//	if the VT-enabled-ness of the platforms is not all the same as the Editor host platform
 
-		if ( VirtualTextureStreaming && ! CVarVirtualTextures.GetValueOnAnyThread() )
+		if ( VirtualTextureStreaming && ! IsVirtualTexturingEnabled() )
 		{
 			// VT was turned on for this texture, but off for the project
 			//	fix it now, turn off on the texture:
@@ -657,7 +683,7 @@ void UTexture::ValidateSettingsAfterImportOrEdit(bool * pRequiresNotifyMaterials
 
 		if ( bLargeTextureMustBeVT && ! VirtualTextureStreaming && ( MaxTextureSize == 0 || MaxTextureSize > RHIMaximumDimension ) )
 		{
-			if ( CVarVirtualTextures.GetValueOnAnyThread() )
+			if ( IsVirtualTexturingEnabled() )
 			{
 				if ( GetTextureClass() == ETextureClass::TwoD )
 				{
@@ -2409,12 +2435,12 @@ void FTextureSource::Compress()
 		if ( CompressionFormat != TSCF_UEDELTA )
 		{
 			// change to TSCF_UEDELTA
-
+			
 			// can't just call RemoveCompression here because it uses slow BulkData funcs
 			// RemoveCompression();
 			
 			FSharedBuffer Buffer = Decompress();
-
+			
 			if ( CompressionFormat != TSCF_None )
 			{
 				// set BulkData to the uncompressed data so we can get the hash before compression :
@@ -2425,7 +2451,7 @@ void FTextureSource::Compress()
 
 			// update the Id from the decompressed data :
 			UseHashAsGuid();
-		
+					
 			if ( ! HasLayerColorInfo() )
 			{
 				// since we're changing compression, go ahead and also update channel minmax now if not done
@@ -2445,7 +2471,7 @@ void FTextureSource::Compress()
 			//	we try to keep "Id" == to the hash of the BulkData when it was the raw data
 			//	the invariant
 			//	( Id == UE::Serialization::IoHashToGuid(BulkData.GetPayloadId()) )
-			//	is no longer true after this			
+			//	is no longer true after this	
 		}
 	}
 	else // not ShouldUseUEDeltaForFormat
@@ -3903,8 +3929,7 @@ void UTexture::GetBuiltTextureSize(const ITargetPlatformSettings* TargetPlatform
 		}
 	}
 	
-	// bug? this queries CVarVirtualTextures, but should be asking for that in the TargetPlatform config
-	const bool bVirtualTextureStreaming = CVarVirtualTextures.GetValueOnAnyThread() && TargetPlatformSettings->SupportsFeature(ETargetPlatformFeatures::VirtualTextureStreaming) && VirtualTextureStreaming;
+	const bool bVirtualTextureStreaming = VirtualTextureStreaming && IsVirtualTexturingEnabled(TargetPlatformSettings);
 
 	const UTextureLODSettings& LODSettings = TargetPlatformSettings->GetTextureLODSettings();
  	const uint32 LODBiasNoCinematics = FMath::Max<int32>(LODSettings.CalculateLODBias(SizeX, SizeY, MaxTextureSize, LODGroup, LODBias, 0, MipGenSettings, bVirtualTextureStreaming), 0);
diff --git a/Engine/Source/Runtime/Engine/Private/TextureDerivedData.cpp b/Engine/Source/Runtime/Engine/Private/TextureDerivedData.cpp
--- a/Engine/Source/Runtime/Engine/Private/TextureDerivedData.cpp
+++ b/Engine/Source/Runtime/Engine/Private/TextureDerivedData.cpp
@@ -1104,7 +1104,6 @@ static void GetTextureBuildSettings(
 	)
 {
 	const bool bPlatformSupportsTextureStreaming = TargetPlatform.SupportsFeature(ETargetPlatformFeatures::TextureStreaming);
-	const bool bPlatformSupportsVirtualTextureStreaming = TargetPlatform.SupportsFeature(ETargetPlatformFeatures::VirtualTextureStreaming);
 
 	if (OutBuildResultMetadata)
 	{
@@ -1215,18 +1214,43 @@ static void GetTextureBuildSettings(
 		bBorderColorBlack
 		);
 
-	static const auto CVarVirtualTexturesEnabled = IConsoleManager::Get().FindTConsoleVariableDataInt(TEXT("r.VirtualTextures")); check(CVarVirtualTexturesEnabled);
-	// A ULightMapVirtualTexture2D with multiple layers saved in MapBuildData could be loaded with the r.VirtualTexture disabled, it will generate DDC before we decide to invalidate the light map data, to skip the ensure failure let it generate VT DDC anyway.
-	const bool bForVirtualTextureStreamingBuild = ULightMapVirtualTexture2D::StaticClass() == Texture.GetClass();
-	bool bVirtualTextureStreaming = bForVirtualTextureStreamingBuild || (CVarVirtualTexturesEnabled->GetValueOnAnyThread() && bPlatformSupportsVirtualTextureStreaming && Texture.VirtualTextureStreaming);
-	if (Texture.Availability == ETextureAvailability::CPU && TextureClass == ETextureClass::TwoD)
+	bool bVirtualTextureStreaming = Texture.VirtualTextureStreaming;
+
+	if ( !bVirtualTextureStreaming && Texture.GetClass() == ULightMapVirtualTexture2D::StaticClass() )
+	{
+		// A ULightMapVirtualTexture2D with multiple layers saved in MapBuildData could be loaded with the r.VirtualTexture disabled, it will generate DDC before we decide to invalidate the light map data, to skip the ensure failure let it generate VT DDC anyway.
+		// @@ pretty ugly hack here, this should have been fixed in PostLoad or something
+		bVirtualTextureStreaming = true;
+	}
+	
+	if ( bVirtualTextureStreaming && ! UTexture::IsVirtualTexturingEnabled(&TargetPlatform) )
+	{
+		bVirtualTextureStreaming = false;
+	}
+
+	if ( Texture.RequiresVirtualTexturing() && ! bVirtualTextureStreaming )
+	{
+		// should not get here; earlier call to CanBuildPlatformData() should have returned false
+		UE_LOG(LogTexture, Error, TEXT("Texture RequiresVirtualTexturing but VT is off (%s)"),*Texture.GetName());
+
+		// no way to error out and abort the build from here (this function returns void)
+		// return false;
+		
+		// turn it back on to avoid crashes?
+		//	otherwise you will hit checks on NumLayers because we expect non-VT to always have 1 layer
+		bVirtualTextureStreaming = true;
+	}
+
+	if (Texture.Availability == ETextureAvailability::CPU && TextureClass == ETextureClass::TwoD && 
+		! Texture.RequiresVirtualTexturing())
 	{
 		// We are swapping with a placeholder - don't VT it.
 		OutBuildSettings.bCPUAccessible = true;
 		bVirtualTextureStreaming = false;
 		MipGenSettings = TMGS_NoMipmaps;
 	}
-
+	
+	OutBuildSettings.bVirtualStreamable = bVirtualTextureStreaming;
 
 	// Virtual textures must have mips as VT memory management relies on a 1:1 texel/pixel mapping, which in turn
 	// requires that we be able to swap in lower mips when that density gets too high for a given texture.
@@ -1255,6 +1279,12 @@ static void GetTextureBuildSettings(
 			}
 		}
 	}
+	if ( Texture.Source.GetNumBlocks() > 1 && !bVirtualTextureStreaming )
+	{
+		UE_LOG(LogTexture, Warning, TEXT("Texture %s has UDIM Blocks, but bVirtualTextureStreaming is off; will build just the first block."), 
+			*Texture.GetPathName()
+			);
+	}
 
 	const FIntPoint SourceSize = Texture.Source.GetLogicalSize();
 
@@ -1277,7 +1307,6 @@ static void GetTextureBuildSettings(
 
 	OutBuildSettings.LODBias = TextureLODSettings.CalculateLODBias(SourceSize.X, SourceSize.Y, Texture.MaxTextureSize, Texture.LODGroup, Texture.LODBias, Texture.NumCinematicMipLevels, Texture.MipGenSettings, bVirtualTextureStreaming);
 	OutBuildSettings.LODBiasWithCinematicMips = TextureLODSettings.CalculateLODBias(SourceSize.X, SourceSize.Y, Texture.MaxTextureSize, Texture.LODGroup, Texture.LODBias, 0, Texture.MipGenSettings, bVirtualTextureStreaming);
-	OutBuildSettings.bVirtualStreamable = bVirtualTextureStreaming;
 	OutBuildSettings.PowerOfTwoMode = Texture.PowerOfTwoMode;
 	OutBuildSettings.PaddingColor = Texture.PaddingColor;
 	OutBuildSettings.bPadWithBorderColor = Texture.bPadWithBorderColor;
@@ -3426,6 +3455,35 @@ void UTextureCube::GetMipData(int32 FirstMipToLoad, void** OutMipData)
 	}
 }
 
+#if WITH_EDITORONLY_DATA
+bool UTexture::RequiresVirtualTexturing() const
+{
+	if ( ! Source.IsValid() )
+	{
+		return false;
+	}
+
+	if ( Source.GetNumLayers() > 1 )
+	{
+		return true;
+	}
+	
+	// NOTE: optional: if NumBlocks() > 1 , for UDIM
+	//	it does work as a non-VT and will just show the first block
+	//	we can either say RequiresVirtualTexturing or not in that case
+	/*
+	if ( Source.GetNumBlocks() > 1 )
+	{
+		return true;
+	}
+	*/
+
+	// also check class == ULightMapVirtualTexture2D ?
+
+	return false;
+}
+#endif
+
 int32 UTexture::CalculateLODBias(bool bWithCinematicMipBias) const
 {
 	// Async caching of PlatformData must be done before calling this
@@ -3435,6 +3493,25 @@ int32 UTexture::CalculateLODBias(bool bWithCinematicMipBias) const
 }
 
 #if WITH_EDITOR
+
+bool UTexture::CanBuildPlatformData(const ITargetPlatformSettings * TargetPlatform) const
+{
+	if ( ! Source.IsValid() )
+	{
+		return false;
+	}
+
+	if ( RequiresVirtualTexturing() )
+	{
+		if ( ! IsVirtualTexturingEnabled(TargetPlatform) )
+		{
+			return false;
+		}
+	}
+
+	return true;
+}
+
 void UTexture::CachePlatformData(bool bAsyncCache, bool bAllowAsyncBuild, bool bAllowAsyncLoading, ITextureCompressorModule* Compressor)
 {
 	//
@@ -3448,7 +3525,7 @@ void UTexture::CachePlatformData(bool bAsyncCache, bool bAllowAsyncBuild, bool b
 	if (PlatformDataLinkPtr)
 	{
 		FTexturePlatformData*& PlatformDataLink = *PlatformDataLinkPtr;
-		if (Source.IsValid() && FApp::CanEverRender())
+		if ( FApp::CanEverRender() && CanBuildPlatformData() )
 		{
 			bool bPerformCache = false;
 
@@ -3649,6 +3726,11 @@ void UTexture::BeginCacheForCookedPlatformData( const ITargetPlatform *TargetPla
 {
 	// @todo Oodle : if TargetPlatform->IsServerOnly() early exit?
 
+	if ( ! CanBuildPlatformData(TargetPlatform) )
+	{
+		return;
+	}
+
 	TMap<FString, FTexturePlatformData*>* CookedPlatformDataPtr = GetCookedPlatformData();
 	if (CookedPlatformDataPtr && !GetOutermost()->HasAnyPackageFlags(PKG_FilterEditorOnly))
 	{
@@ -3815,6 +3897,11 @@ bool UTexture::IsCachedCookedPlatformDataLoaded(const ITargetPlatform* TargetPla
 		return true; 
 	}
 
+	if ( ! CanBuildPlatformData(TargetPlatform) )
+	{
+		return false;
+	}
+
 	// CookedPlatformData is keyed off of FetchOrBuild settings.
 	ETextureEncodeSpeed EncodeSpeed = GetDesiredEncodeSpeed();
 
@@ -3938,7 +4025,7 @@ void UTexture::FinishCachePlatformData()
 	{
 		FTexturePlatformData*& RunningPlatformData = *RunningPlatformDataPtr;
 		
-		if (Source.IsValid() && FApp::CanEverRender())
+		if (CanBuildPlatformData() && FApp::CanEverRender())
 		{
 			if ( RunningPlatformData == NULL )
 			{
diff --git a/Engine/Source/Runtime/TextureUtilitiesCommon/Private/TextureImportSettings.cpp b/Engine/Source/Runtime/TextureUtilitiesCommon/Private/TextureImportSettings.cpp
--- a/Engine/Source/Runtime/TextureUtilitiesCommon/Private/TextureImportSettings.cpp
+++ b/Engine/Source/Runtime/TextureUtilitiesCommon/Private/TextureImportSettings.cpp
@@ -56,13 +56,15 @@ bool UTextureImportSettings::IsImportAutoVTEnabled() const
 		return false;
 	}
 
-	static const TConsoleVariableData<int32>* CVarVirtualTexturesEnabled = IConsoleManager::Get().FindTConsoleVariableDataInt(TEXT("r.VirtualTextures"));
-	check(CVarVirtualTexturesEnabled);
+	if ( ! UTexture::IsVirtualTexturingEnabled() )
+	{
+		return false;
+	}
 
 	static const auto CVarVirtualTexturesAutoImportEnabled = IConsoleManager::Get().FindTConsoleVariableDataInt(TEXT("r.VT.EnableAutoImport"));
 	check(CVarVirtualTexturesAutoImportEnabled);
 
-	return ( CVarVirtualTexturesEnabled->GetValueOnAnyThread() && CVarVirtualTexturesAutoImportEnabled->GetValueOnAnyThread() );
+	return !! CVarVirtualTexturesAutoImportEnabled->GetValueOnAnyThread();
 }
 
 int64 UTextureImportSettings::GetAutoLimitPixelCount() const
diff --git a/Engine/Source/Runtime/TextureUtilitiesCommon/Private/TextureImportUtils.cpp b/Engine/Source/Runtime/TextureUtilitiesCommon/Private/TextureImportUtils.cpp
--- a/Engine/Source/Runtime/TextureUtilitiesCommon/Private/TextureImportUtils.cpp
+++ b/Engine/Source/Runtime/TextureUtilitiesCommon/Private/TextureImportUtils.cpp
@@ -616,10 +616,7 @@ bool UE::TextureUtilitiesCommon::IsImportResolutionValid(int64 Width, int64 Heig
 
 	if (Width > MaximumSupportedResolutionNonVT || Height > MaximumSupportedResolutionNonVT)
 	{
-		const TConsoleVariableData<int32>* CVarVirtualTexturesEnabled = IConsoleManager::Get().FindTConsoleVariableDataInt(TEXT("r.VirtualTextures")); check(CVarVirtualTexturesEnabled);
-		check(CVarVirtualTexturesEnabled != nullptr);
-
-		if (!CVarVirtualTexturesEnabled->GetValueOnAnyThread())
+		if (! UTexture::IsVirtualTexturingEnabled() )
 		{
 			const FText VTMessage = NSLOCTEXT("Interchange", "Warning_LargeTextureVTDisabled", "\nWarning: Virtual Textures are disabled in this project.");
 
