Add null and active checks for AsyncPhysicsTickComponents and AsyncPhysicsTickActors

diff --git a/Engine/Source/Runtime/Engine/Private/PhysicsEngine/Experimental/PhysScene_Chaos.cpp b/Engine/Source/Runtime/Engine/Private/PhysicsEngine/Experimental/PhysScene_Chaos.cpp
--- a/Engine/Source/Runtime/Engine/Private/PhysicsEngine/Experimental/PhysScene_Chaos.cpp
+++ b/Engine/Source/Runtime/Engine/Private/PhysicsEngine/Experimental/PhysScene_Chaos.cpp
@@ -200,6 +200,8 @@ public:
 			QUICK_SCOPE_CYCLE_COUNTER(STAT_AsyncPhys_TickComponents);
 			for(UActorComponent* Component : AsyncPhysicsTickComponents)
 			{
+				check(Component && Component->IsActive());
+
 				FScopeCycleCounterUObject ComponentScope(Component);
 				Component->AsyncPhysicsTickComponent(DeltaTime, SimTime);
 			}
@@ -209,6 +211,8 @@ public:
 			QUICK_SCOPE_CYCLE_COUNTER(STAT_AsyncPhys_TickActors);
 			for(AActor* Actor : AsyncPhysicsTickActors)
 			{
+				check(Actor)
+
 				FScopeCycleCounterUObject ActorScope(Actor);
 				Actor->AsyncPhysicsTickActor(DeltaTime, SimTime);
 			}
