Comprehensive support for custom gravity in CharacterMovementComponent fix
#rnx

diff --git a/Engine/Source/Runtime/Engine/Private/Components/CharacterMovementComponent.cpp b/Engine/Source/Runtime/Engine/Private/Components/CharacterMovementComponent.cpp
--- a/Engine/Source/Runtime/Engine/Private/Components/CharacterMovementComponent.cpp
+++ b/Engine/Source/Runtime/Engine/Private/Components/CharacterMovementComponent.cpp
@@ -1064,7 +1064,7 @@ bool UCharacterMovementComponent::DoJump(bool bReplayingMoves)
 	if ( CharacterOwner && CharacterOwner->CanJump() )
 	{
 		// Don't jump if we can't move up/down.
-		if (!bConstrainToPlane || FMath::IsNearlyEqual(FMath::Abs(GetGravitySpaceZ(PlaneConstraintNormal)), 1.f))
+		if (!bConstrainToPlane || !FMath::IsNearlyEqual(FMath::Abs(GetGravitySpaceZ(PlaneConstraintNormal)), 1.f))
 		{
 			if (HasCustomGravity())
 			{
