Fix fatal error when using Editor Preview

diff --git a/Engine/Source/Developer/Windows/ShaderFormatD3D/Private/ShaderFormatD3D.cpp b/Engine/Source/Developer/Windows/ShaderFormatD3D/Private/ShaderFormatD3D.cpp
--- a/Engine/Source/Developer/Windows/ShaderFormatD3D/Private/ShaderFormatD3D.cpp
+++ b/Engine/Source/Developer/Windows/ShaderFormatD3D/Private/ShaderFormatD3D.cpp
@@ -138,7 +138,7 @@ public:
 	static bool IsSM60(const FShaderCompilerInput& Input, ELanguage Language)
 	{
 		return RequiresSM6Features(Input, Language)
-			|| Input.Environment.GetIntegerValue(TEXT("PLATFORM_MAX_SAMPLERS")) > 16
+			|| Input.Environment.GetCompileArgument(TEXT("PLATFORM_MAX_SAMPLERS"), 0) > 16
 			;
 	}
 
diff --git a/Engine/Source/Runtime/Engine/Private/ShaderCompiler/ShaderCompiler.cpp b/Engine/Source/Runtime/Engine/Private/ShaderCompiler/ShaderCompiler.cpp
--- a/Engine/Source/Runtime/Engine/Private/ShaderCompiler/ShaderCompiler.cpp
+++ b/Engine/Source/Runtime/Engine/Private/ShaderCompiler/ShaderCompiler.cpp
@@ -8203,7 +8203,7 @@ void GlobalBeginCompileShader(
 
 	{
 		const uint32 MaxSamplers = FDataDrivenShaderPlatformInfo::GetMaxSamplers((EShaderPlatform)Target.Platform);
-		SET_SHADER_DEFINE(Input.Environment, PLATFORM_MAX_SAMPLERS, MaxSamplers);
+		SET_SHADER_DEFINE_AND_COMPILE_ARGUMENT(Input.Environment, PLATFORM_MAX_SAMPLERS, MaxSamplers);
 	}
 
 	{
