FSkeletalMeshSceneProxy: set view relevance dynamic if bForceWireframe is set

diff --git a/Engine/Source/Runtime/Engine/Private/SkeletalMesh.cpp b/Engine/Source/Runtime/Engine/Private/SkeletalMesh.cpp
--- a/Engine/Source/Runtime/Engine/Private/SkeletalMesh.cpp
+++ b/Engine/Source/Runtime/Engine/Private/SkeletalMesh.cpp
@@ -7174,6 +7174,7 @@ FPrimitiveViewRelevance FSkeletalMeshSceneProxy::GetViewRelevance(const FSceneVi
 			|| EngineShowFlags.Collision
 			|| EngineShowFlags.VisualizeGPUSkinCache
 			|| (IsSelected() && (EngineShowFlags.VertexColors || EngineShowFlags.Bounds))
+			|| bForceWireframe
 #if WITH_EDITORONLY_DATA
 			|| MeshObject->SelectedEditorMaterial != -1
 			|| MeshObject->SelectedEditorSection != -1
