[World Partition Runtime HashSet]
- LHGrid partitions will not output the Z coordinates as part of the cell name when set to 2D.
- Fix cells instance suffix used to compute cells hash instead of just appending it as the spatial hash did.

diff --git a/Engine/Source/Runtime/Engine/Private/WorldPartition/RuntimeHashSet/RuntimePartitionLHGrid.cpp b/Engine/Source/Runtime/Engine/Private/WorldPartition/RuntimeHashSet/RuntimePartitionLHGrid.cpp
--- a/Engine/Source/Runtime/Engine/Private/WorldPartition/RuntimeHashSet/RuntimePartitionLHGrid.cpp
+++ b/Engine/Source/Runtime/Engine/Private/WorldPartition/RuntimeHashSet/RuntimePartitionLHGrid.cpp
@@ -16,7 +16,13 @@ struct FCellCoord
 
 	static FCellCoord Invalid;
 
-	inline FString ToString() const
+	inline FString ToString2D() const
+	{
+		check(!Z);
+		return FString::Printf(TEXT("L%d_X%" INT64_FMT "_Y%" INT64_FMT), Level, X, Y);
+	}
+
+	inline FString ToString3D() const
 	{
 		return FString::Printf(TEXT("L%d_X%" INT64_FMT "_Y%" INT64_FMT "_Z%" INT64_FMT), Level, X, Y, Z);
 	}
@@ -189,7 +195,7 @@ bool URuntimePartitionLHGrid::GenerateStreaming(const FGenerateStreamingParams&
 	{
 		const bool bIsSpatiallyLoaded = CellCoord != FCellCoord::Invalid;
 
-		URuntimePartition::FCellDesc& CellDesc = OutResult.RuntimeCellDescs.Emplace_GetRef(CreateCellDesc(CellCoord.ToString(), bIsSpatiallyLoaded, CellCoord.Level, CellActorSetInstances));
+		URuntimePartition::FCellDesc& CellDesc = OutResult.RuntimeCellDescs.Emplace_GetRef(CreateCellDesc(bIs2D ? CellCoord.ToString2D() : CellCoord.ToString3D(), bIsSpatiallyLoaded, CellCoord.Level, CellActorSetInstances));
 
 		if (bIsSpatiallyLoaded)
 		{
diff --git a/Engine/Source/Runtime/Engine/Private/WorldPartition/RuntimeHashSet/WorldPartitionRuntimeHashSet.cpp b/Engine/Source/Runtime/Engine/Private/WorldPartition/RuntimeHashSet/WorldPartitionRuntimeHashSet.cpp
--- a/Engine/Source/Runtime/Engine/Private/WorldPartition/RuntimeHashSet/WorldPartitionRuntimeHashSet.cpp
+++ b/Engine/Source/Runtime/Engine/Private/WorldPartition/RuntimeHashSet/WorldPartitionRuntimeHashSet.cpp
@@ -749,13 +749,13 @@ FGuid UWorldPartitionRuntimeHashSet::RegisterWorldAssetStreaming(const UWorldPar
 			}
 			else
 			{
-				UE_LOG(LogWorldPartition, Error, TEXT("Error creating streaming cell %s for world asset %s at %s"), *RuntimeCell->GetName(), *WorldAsset.ToString(), *InParams.Transform.ToString());StreamingData.CreatePartitionsSpatialIndex();
+				UE_LOG(LogWorldPartition, Error, TEXT("Error creating streaming cell %s for world asset %s at %s"), *RuntimeCell->GetName(), *WorldAsset.ToString(), *InParams.Transform.ToString());
 				return FGuid();
 			}
 		}		
 		else
 		{
-			UE_LOG(LogWorldPartition, Error, TEXT("Error creating streaming cell %s for world asset %s at %s"), *CellName, *WorldAsset.ToString(), *InParams.Transform.ToString());StreamingData.CreatePartitionsSpatialIndex();
+			UE_LOG(LogWorldPartition, Error, TEXT("Error creating streaming cell %s for world asset %s at %s"), *CellName, *WorldAsset.ToString(), *InParams.Transform.ToString());
 			return FGuid();
 		}
 	}
@@ -1012,7 +1012,6 @@ UWorldPartitionRuntimeHashSet::FCellUniqueId UWorldPartitionRuntimeHashSet::GetC
 		UWorld* OuterWorld = GetTypedOuter<UWorld>();
 		check(OuterWorld);
 
-		FString InstanceSuffix;
 		FString WorldName = FPackageName::GetShortName(OuterWorld->GetPackage());
 
 		if (!IsRunningCookCommandlet() && OuterWorld->IsGameWorld())
@@ -1032,7 +1031,7 @@ UWorldPartitionRuntimeHashSet::FCellUniqueId UWorldPartitionRuntimeHashSet::GetC
 
 				if (int32 Index = InstancePackageName.Find(SourcePackageName); Index != INDEX_NONE)
 				{
-					InstanceSuffix = InstancePackageName.Mid(Index + SourcePackageName.Len());
+					CellUniqueId.InstanceSuffix = InstancePackageName.Mid(Index + SourcePackageName.Len());
 				}
 			}
 		}
@@ -1050,11 +1049,6 @@ UWorldPartitionRuntimeHashSet::FCellUniqueId UWorldPartitionRuntimeHashSet::GetC
 			CellNameBuilder.Appendf(TEXT("_c%s"), *UContentBundleDescriptor::GetContentBundleCompactString(ContentBundleID));
 		}
 
-		if (!InstanceSuffix.IsEmpty())
-		{
-			CellNameBuilder.Appendf(TEXT("_i%s"), *InstanceSuffix);
-		}
-	
 		CellUniqueId.Name = CellNameBuilder.ToString();
 	}
 
diff --git a/Engine/Source/Runtime/Engine/Private/WorldPartition/RuntimeHashSet/WorldPartitionRuntimeHashSetStreamingGeneration.cpp b/Engine/Source/Runtime/Engine/Private/WorldPartition/RuntimeHashSet/WorldPartitionRuntimeHashSetStreamingGeneration.cpp
--- a/Engine/Source/Runtime/Engine/Private/WorldPartition/RuntimeHashSet/WorldPartitionRuntimeHashSetStreamingGeneration.cpp
+++ b/Engine/Source/Runtime/Engine/Private/WorldPartition/RuntimeHashSet/WorldPartitionRuntimeHashSetStreamingGeneration.cpp
@@ -149,7 +149,7 @@ bool UWorldPartitionRuntimeHashSet::GenerateStreaming(UWorldPartitionStreamingPo
 	{
 		const FCellUniqueId CellUniqueId = GetCellUniqueId(CellDescInstance);
 
-		UWorldPartitionRuntimeCell* RuntimeCell = CreateRuntimeCell(CellClass, CellDataClass, CellUniqueId.Name, TEXT(""));
+		UWorldPartitionRuntimeCell* RuntimeCell = CreateRuntimeCell(CellClass, CellDataClass, CellUniqueId.Name, CellUniqueId.InstanceSuffix);
 
 		RuntimeCell->SetDataLayers(CellDescInstance.DataLayerInstances);
 		RuntimeCell->SetContentBundleUID(CellDescInstance.ContentBundleID);
diff --git a/Engine/Source/Runtime/Engine/Public/WorldPartition/RuntimeHashSet/WorldPartitionRuntimeHashSet.h b/Engine/Source/Runtime/Engine/Public/WorldPartition/RuntimeHashSet/WorldPartitionRuntimeHashSet.h
--- a/Engine/Source/Runtime/Engine/Public/WorldPartition/RuntimeHashSet/WorldPartitionRuntimeHashSet.h
+++ b/Engine/Source/Runtime/Engine/Public/WorldPartition/RuntimeHashSet/WorldPartitionRuntimeHashSet.h
@@ -212,6 +212,7 @@ private:
 	struct FCellUniqueId
 	{
 		FString Name;
+		FString InstanceSuffix;
 		FGuid Guid;
 	};
 
