#ue5 - Improved behavior to reduce corrections by allowing server to run longer move time deltas than the client is allowed to send, to account for when client updates are lost or jittery (p.NetServerMaxMoveDeltaTimeScalar adds a scalar on the server, clamped to >= 1). Fixed client timestamp when forcing updates, allowing client time advanced during normal updates, and stopping after duration limit exceeded.

Only stop advancing client timestamp on the server if exceeding the allowable forced update duration.
Only sync client timestamp on the server if recovering from forced updates that exceeded the allowable forced update duration.
Added p.NetServerForcePositionUpdateSyncToClient to be able to restore old behavior to always sync timestamps after forced updates recover. Default false.

#tests PIE with bad networking profile, and forced client movement drops at 1% - 10% rate for 0.1 to 5 seconds.

diff --git a/Engine/Source/Runtime/Engine/Private/Components/CharacterMovementComponent.cpp b/Engine/Source/Runtime/Engine/Private/Components/CharacterMovementComponent.cpp
--- a/Engine/Source/Runtime/Engine/Private/Components/CharacterMovementComponent.cpp
+++ b/Engine/Source/Runtime/Engine/Private/Components/CharacterMovementComponent.cpp
@@ -207,13 +207,20 @@ namespace CharacterMovementCVars
 		TEXT("0: Disable, 1: Enable"),
 		ECVF_Default);
 
-	static bool bNetServerForcePositionUpdateAdvanceClientTimestamp = false;
-	FAutoConsoleVariableRef CVarNetServerForcePositionUpdateAdvanceClientTimestamp(
-		TEXT("p.NetServerForcePositionUpdateAdvanceClientTimestamp"),
-		bNetServerForcePositionUpdateAdvanceClientTimestamp,
-		TEXT("If enabled, the server will advance client timestamps when running ForcePositionUpdate() on client timeout.\n"),
+	static bool bNetServerForcePositionUpdateSyncToClient = false;
+	FAutoConsoleVariableRef CVarNetServerForcePositionUpdateSyncToClient(
+		TEXT("p.NetServerForcePositionUpdateSyncToClient"),
+		bNetServerForcePositionUpdateSyncToClient,
+		TEXT("If enabled, the server will always sync to the client timestamp when receiving a valid move after forcing updates (true is the old default behavior).\n"),
 		ECVF_Default);
 
+	static float NetServerMaxMoveDeltaTimeScalar = 1.75f;
+	FAutoConsoleVariableRef CVarNetServerMaxMoveDeltaTimeScalar(
+		TEXT("p.NetServerMaxMoveDeltaTimeScalar"),
+		NetServerMaxMoveDeltaTimeScalar,
+		TEXT("Multiplier to the MaxMoveDeltaTime allowed on the server, compared to the client. Enforced to be >= 1.0\n")
+		TEXT("Allows the server to accept longer moves than the client is normally allowed to send, to avoid corrections (a value of 1 is old default behavior).\n"),
+		ECVF_Default);
 
 	static int32 ReplayLerpAcceleration = 0;
 	FAutoConsoleVariableRef CVarReplayLerpAcceleration(
@@ -8441,7 +8448,7 @@ bool UCharacterMovementComponent::ForcePositionUpdate(float DeltaTime)
 
 	FNetworkPredictionData_Server_Character* ServerData = GetPredictionData_Server_Character();
 
-	if (CharacterMovementCVars::bNetServerForcePositionUpdateAdvanceClientTimestamp)
+	if (!ServerData->bForcedUpdateDurationExceeded)
 	{
 		// Increment client timestamp so we reject client moves after this new simulated time position.
 		// See handling of CurrentClientTimeStamp in VerifyClientTimeStamp().
@@ -9088,7 +9095,7 @@ bool UCharacterMovementComponent::VerifyClientTimeStamp(float TimeStamp, FNetwor
 			CurrentRootMotion.ApplyTimeStampReset(MinTimeBetweenTimeStampResets);
 		}
 		
-		if (bFirstMoveAfterForcedUpdates)
+		if (bFirstMoveAfterForcedUpdates && (ServerData.bForcedUpdateDurationExceeded || CharacterMovementCVars::bNetServerForcePositionUpdateSyncToClient))
 		{
 			// We have been performing ForcedUpdates because we hadn't received any moves from this connection in a while but we've now received a new move!
 			// Let's sync up to this TimeStamp in order to resolve movement desyncs ASAP
@@ -12137,7 +12144,9 @@ float FNetworkPredictionData_Server_Character::GetServerMoveDeltaTime(float Clie
 
 float FNetworkPredictionData_Server_Character::GetBaseServerMoveDeltaTime(float ClientTimeStamp, float ActorTimeDilation) const
 {
-	const float DeltaTime = FMath::Min(MaxMoveDeltaTime * ActorTimeDilation, ClientTimeStamp - CurrentClientTimeStamp);
+	// Allow server to run longer moves than the client is allowed to send, to avoid corrections when there is a gap in lost client moves.
+	const float ServerScalar = FMath::Max(1.0f, CharacterMovementCVars::NetServerMaxMoveDeltaTimeScalar);
+	const float DeltaTime = FMath::Min(ServerScalar * MaxMoveDeltaTime * ActorTimeDilation, ClientTimeStamp - CurrentClientTimeStamp);
 	return DeltaTime;
 }
 
diff --git a/Engine/Source/Runtime/Engine/Private/PlayerController.cpp b/Engine/Source/Runtime/Engine/Private/PlayerController.cpp
--- a/Engine/Source/Runtime/Engine/Private/PlayerController.cpp
+++ b/Engine/Source/Runtime/Engine/Private/PlayerController.cpp
@@ -5106,17 +5106,9 @@ void APlayerController::TickActor( float DeltaSeconds, ELevelTick TickType, FAct
 
 								if (PawnTimeSinceForcingUpdates > PawnTimeForcedUpdateMaxDuration)
 								{
-									if (ServerData->bLastRequestNeedsForcedUpdates)
-									{
-										// No valid updates, don't reset anything but don't mark as exceeded either
-										// Keep forced updates going until new and valid move request is received
-									}
-									else
-									{
-										// Waiting for ServerTimeStamp to advance from a client move.
-										UE_LOG(LogNetPlayerMovement, Log, TEXT("Setting bForcedUpdateDurationExceeded=true (PawnTimeSinceForcingUpdates %.6f > PawnTimeForcedUpdateMaxDuration %.6f)"), PawnTimeSinceForcingUpdates, PawnTimeForcedUpdateMaxDuration);
-										ServerData->bForcedUpdateDurationExceeded = true;
-									}
+									// Waiting for ServerTimeStamp to advance from a client move.
+									UE_LOG(LogNetPlayerMovement, Log, TEXT("Setting bForcedUpdateDurationExceeded=true (PawnTimeSinceForcingUpdates %.6f > PawnTimeForcedUpdateMaxDuration %.6f) (bLastRequestNeedsForcedUpdates:%d)"), PawnTimeSinceForcingUpdates, PawnTimeForcedUpdateMaxDuration, (int32)ServerData->bLastRequestNeedsForcedUpdates);
+									ServerData->bForcedUpdateDurationExceeded = true;
 								}
 							}
 						}
@@ -5135,7 +5127,8 @@ void APlayerController::TickActor( float DeltaSeconds, ELevelTick TickType, FAct
 						{
 							//UE_LOG(LogPlayerController, Warning, TEXT("ForcedMovementTick. PawnTimeSinceUpdate: %f, DeltaSeconds: %f, DeltaSeconds+: %f"), PawnTimeSinceUpdate, DeltaSeconds, DeltaSeconds+0.06f);
 							const USkeletalMeshComponent* PawnMesh = GetPawn()->FindComponentByClass<USkeletalMeshComponent>();
-							if (!ServerData->bForcedUpdateDurationExceeded && (!PawnMesh || !PawnMesh->IsSimulatingPhysics()))
+							const bool bShouldForceUpdate = !ServerData->bForcedUpdateDurationExceeded || ServerData->bLastRequestNeedsForcedUpdates;
+							if (bShouldForceUpdate && (!PawnMesh || !PawnMesh->IsSimulatingPhysics()))
 							{
 								const bool bDidUpdate = NetworkPredictionInterface->ForcePositionUpdate(PawnTimeSinceUpdate);
 
