Texture.cpp call RemoveCompression instead of duplicating the same logic

diff --git a/Engine/Source/Runtime/Engine/Private/Texture.cpp b/Engine/Source/Runtime/Engine/Private/Texture.cpp
--- a/Engine/Source/Runtime/Engine/Private/Texture.cpp
+++ b/Engine/Source/Runtime/Engine/Private/Texture.cpp
@@ -2435,23 +2435,16 @@ void FTextureSource::Compress()
 		if ( CompressionFormat != TSCF_UEDELTA )
 		{
 			// change to TSCF_UEDELTA
-			
-			// can't just call RemoveCompression here because it uses slow BulkData funcs
-			// RemoveCompression();
-			
-			FSharedBuffer Buffer = Decompress();
-			
-			if ( CompressionFormat != TSCF_None )
-			{
-				// set BulkData to the uncompressed data so we can get the hash before compression :
-				// BulkData.UpdatePayload does a slow hash update
-				BulkData.UpdatePayload(Buffer, Owner);
-				CompressionFormat = TSCF_None;
-			}
 
-			// update the Id from the decompressed data :
-			UseHashAsGuid();
-					
+			// RemoveCompression will update the hash Id from the decompressed data :
+			//	(unfortunately this is slow because the BulkData hash is slow and synchronous)
+			RemoveCompression();
+
+			check( CompressionFormat == TSCF_None );
+			//FGuid IdBefore = GetId();
+
+			FSharedBuffer Buffer = BulkData.GetPayload().Get();
+		
 			if ( ! HasLayerColorInfo() )
 			{
 				// since we're changing compression, go ahead and also update channel minmax now if not done
@@ -2460,6 +2453,8 @@ void FTextureSource::Compress()
 
 			FSharedBuffer DeltaBuffer = DoUEDeltaTransform(Buffer,true);
 			
+			Buffer.Reset(); // release ref
+
 			// note: at this moment it would be easy to try LZ compression on the delta and non-delta data and choose the best
 			// if you care about small uasset size and don't mind a slightly slower encode
 			//  (90% of uasset save time is not in this function)
@@ -2469,9 +2464,10 @@ void FTextureSource::Compress()
 			CompressionFormat = TSCF_UEDELTA;
 			
 			//	we try to keep "Id" == to the hash of the BulkData when it was the raw data
-			//	the invariant
-			//	( Id == UE::Serialization::IoHashToGuid(BulkData.GetPayloadId()) )
-			//	is no longer true after this	
+			//	the invariant ( Id == UE::Serialization::IoHashToGuid(BulkData.GetPayloadId()) )
+			//	is no longer true after this, because we change the BulkData but keep the old Id
+			//FGuid IdAfter = GetId();
+			//check( IdAfter == IdBefore );
 		}
 	}
 	else // not ShouldUseUEDeltaForFormat
