Demote some errors to warnings, they are confusing in the log.

diff --git a/Engine/Source/Runtime/Engine/Private/DataTable.cpp b/Engine/Source/Runtime/Engine/Private/DataTable.cpp
--- a/Engine/Source/Runtime/Engine/Private/DataTable.cpp
+++ b/Engine/Source/Runtime/Engine/Private/DataTable.cpp
@@ -97,7 +97,7 @@ void UDataTable::LoadStructData(FStructuredArchiveSlot Slot)
 	{
 		if (!HasAnyFlags(RF_ClassDefaultObject) && GetOutermost() != GetTransientPackage())
 		{
-			UE_LOG(LogDataTable, Error, TEXT("Missing RowStruct while loading DataTable '%s', NeedLoad: '%s'!"), *GetPathName(), HasAnyFlags(RF_NeedLoad) ? TEXT("true") : TEXT("false"));
+			UE_LOG(LogDataTable, Warning, TEXT("Missing RowStruct while loading DataTable '%s', NeedLoad: '%s'!"), *GetPathName(), HasAnyFlags(RF_NeedLoad) ? TEXT("true") : TEXT("false"));
 		}
 		LoadUsingStruct = FTableRowBase::StaticStruct();
 	}
diff --git a/Engine/Source/Runtime/Json/Private/Dom/JsonValue.cpp b/Engine/Source/Runtime/Json/Private/Dom/JsonValue.cpp
--- a/Engine/Source/Runtime/Json/Private/Dom/JsonValue.cpp
+++ b/Engine/Source/Runtime/Json/Private/Dom/JsonValue.cpp
@@ -344,5 +344,12 @@ TSharedPtr<FJsonValue> FJsonValue::Duplicate(const TSharedPtr<FJsonValue>& Src)
 
 void FJsonValue::ErrorMessage(const FString& InType) const
 {
-	UE_LOG(LogJson, Error, TEXT("Json Value of type '%s' used as a '%s'."), *GetType(), *InType);
+	if (IsNull())
+	{
+		UE_LOG(LogJson, Warning, TEXT("Json Value of type '%s' used as a '%s'."), *GetType(), *InType);
+	}
+	else
+	{
+		UE_LOG(LogJson, Error, TEXT("Json Value of type '%s' used as a '%s'."), *GetType(), *InType);
+	}
 }
