Updated LLM reporting for Animation Graph. Added new LLM scope "Animation/Graph" to cover full graph evaluation. Tightened scope of "Animation/RigidBody" and "Animation/RigidBodyWithControl" to capture all of the allocations. Removed "Animation/DeadBlending" and "Animation/Intertialization" as these were previously capturing allocations unrelated to these nodes.

diff --git a/Engine/Plugins/Experimental/PhysicsControl/Source/PhysicsControl/Private/AnimNode_RigidBodyWithControl.cpp b/Engine/Plugins/Experimental/PhysicsControl/Source/PhysicsControl/Private/AnimNode_RigidBodyWithControl.cpp
--- a/Engine/Plugins/Experimental/PhysicsControl/Source/PhysicsControl/Private/AnimNode_RigidBodyWithControl.cpp
+++ b/Engine/Plugins/Experimental/PhysicsControl/Source/PhysicsControl/Private/AnimNode_RigidBodyWithControl.cpp
@@ -437,6 +437,7 @@ void FAnimNode_RigidBodyWithControl::CalculateSimulationSpace(
 
 void FAnimNode_RigidBodyWithControl::RunPhysicsSimulation(float DeltaSeconds, const FVector& SimSpaceGravity)
 {
+	LLM_SCOPE_BYNAME(TEXT("Animation/RigidBodyWithControl"));
 	SCOPE_CYCLE_COUNTER(STAT_RigidBodyNodeWithControl_Simulation);
 	CSV_SCOPED_TIMING_STAT(Animation, RigidBodyNodeWithControlSimulation);
 	FScopeCycleCounterUObject AdditionalScope(PhysicsAssetToUse, GET_STATID(STAT_RigidBodyNodeWithControl_Simulation));
@@ -495,6 +496,7 @@ void FAnimNode_RigidBodyWithControl::SetupControls(USkeletalMeshComponent* const
 
 void FAnimNode_RigidBodyWithControl::EvaluateSkeletalControl_AnyThread(FComponentSpacePoseContext& Output, TArray<FBoneTransform>& OutBoneTransforms)
 {
+	LLM_SCOPE_BYNAME(TEXT("Animation/RigidBodyWithControl"));
 	DECLARE_SCOPE_HIERARCHICAL_COUNTER_ANIMNODE(EvaluateSkeletalControl_AnyThread)
 	SCOPE_CYCLE_COUNTER(STAT_RigidBodyNodeWithControl_Eval);
 	CSV_SCOPED_TIMING_STAT(Animation, RigidBodyEval);
diff --git a/Engine/Source/Runtime/AnimGraphRuntime/Private/BoneControllers/AnimNode_RigidBody.cpp b/Engine/Source/Runtime/AnimGraphRuntime/Private/BoneControllers/AnimNode_RigidBody.cpp
--- a/Engine/Source/Runtime/AnimGraphRuntime/Private/BoneControllers/AnimNode_RigidBody.cpp
+++ b/Engine/Source/Runtime/AnimGraphRuntime/Private/BoneControllers/AnimNode_RigidBody.cpp
@@ -492,6 +492,7 @@ DECLARE_CYCLE_STAT(TEXT("FAnimNode_RigidBody::EvaluateSkeletalControl_AnyThread"
 
 void FAnimNode_RigidBody::RunPhysicsSimulation(float DeltaSeconds, const FVector& SimSpaceGravity)
 {
+	LLM_SCOPE_BYNAME(TEXT("Animation/RigidBody"));
 	SCOPE_CYCLE_COUNTER(STAT_RigidBodyNode_Simulation);
 	CSV_SCOPED_TIMING_STAT(Animation, RigidBodyNodeSimulation);
 	FScopeCycleCounterUObject AdditionalScope(UsePhysicsAsset, GET_STATID(STAT_RigidBodyNode_Simulation));
@@ -522,6 +523,7 @@ void FAnimNode_RigidBody::DestroyPhysicsSimulation()
 
 void FAnimNode_RigidBody::EvaluateSkeletalControl_AnyThread(FComponentSpacePoseContext& Output, TArray<FBoneTransform>& OutBoneTransforms)
 {
+	LLM_SCOPE_BYNAME(TEXT("Animation/RigidBody"));
 	DECLARE_SCOPE_HIERARCHICAL_COUNTER_ANIMNODE(EvaluateSkeletalControl_AnyThread)
 	SCOPE_CYCLE_COUNTER(STAT_RigidBody_Eval);
 	CSV_SCOPED_TIMING_STAT(Animation, RigidBodyEval);
@@ -1778,6 +1780,7 @@ void FAnimNode_RigidBody::UpdateWorldObjects(const FTransform& SpaceTransform)
 {
 	LLM_SCOPE_BYNAME(TEXT("Animation/RigidBody")); 
 
+
 	if (SimulationSpace != ESimulationSpace::WorldSpace)
 	{
 		for (const auto& WorldEntry : ComponentsInSim)
diff --git a/Engine/Source/Runtime/Engine/Private/Animation/AnimInstanceProxy.cpp b/Engine/Source/Runtime/Engine/Private/Animation/AnimInstanceProxy.cpp
--- a/Engine/Source/Runtime/Engine/Private/Animation/AnimInstanceProxy.cpp
+++ b/Engine/Source/Runtime/Engine/Private/Animation/AnimInstanceProxy.cpp
@@ -44,6 +44,8 @@
 
 #undef DO_ANIMSTAT_PROCESSING
 
+LLM_DEFINE_TAG(Animation_Graph);
+
 #define LOCTEXT_NAMESPACE "AnimInstance"
 
 const FName NAME_AnimBlueprintLog(TEXT("AnimBlueprintLog"));
@@ -121,6 +123,8 @@ void FAnimInstanceProxy::UpdateAnimationNode(const FAnimationUpdateContext& InCo
 
 void FAnimInstanceProxy::UpdateAnimationNode_WithRoot(const FAnimationUpdateContext& InContext, FAnimNode_Base* InRootNode, FName InLayerName)
 {
+	LLM_SCOPE_BYNAME(TEXT("Animation/Graph"));
+
 	TRACE_SCOPED_ANIM_GRAPH(InContext)
 	TRACE_SCOPED_ANIM_NODE(InContext)
 	
@@ -374,6 +378,8 @@ void FAnimInstanceProxy::InitializeRootNode(bool bInDeferRootNodeInitialization)
 
 void FAnimInstanceProxy::InitializeRootNode_WithRoot(FAnimNode_Base* InRootNode)
 {
+	LLM_SCOPE_BYNAME(TEXT("Animation/Graph"));
+
 	DECLARE_SCOPE_HIERARCHICAL_COUNTER_FUNC()
 
 	if (InRootNode != nullptr)
@@ -1411,6 +1417,8 @@ void FAnimInstanceProxy::EvaluateAnimation_WithRoot(FPoseContext& Output, FAnimN
 
 void FAnimInstanceProxy::CacheBones()
 {
+	LLM_SCOPE_BYNAME(TEXT("Animation/Graph"));
+
 	DECLARE_SCOPE_HIERARCHICAL_COUNTER_FUNC()
 
 	// If bone caches have been invalidated, have AnimNodes refresh those.
@@ -1438,6 +1446,8 @@ void FAnimInstanceProxy::CacheBones()
 
 void FAnimInstanceProxy::CacheBones_WithRoot(FAnimNode_Base* InRootNode)
 {
+	LLM_SCOPE_BYNAME(TEXT("Animation/Graph"));
+
 	// If bone caches have been invalidated, have AnimNodes refresh those.
 	if (bBoneCachesInvalidated && InRootNode)
 	{
@@ -1468,6 +1478,8 @@ void FAnimInstanceProxy::EvaluateAnimationNode(FPoseContext& Output)
 
 void FAnimInstanceProxy::EvaluateAnimationNode_WithRoot(FPoseContext& Output, FAnimNode_Base* InRootNode)
 {
+	LLM_SCOPE_BYNAME(TEXT("Animation/Graph"));
+
 	if (InRootNode != nullptr)
 	{
 		ANIM_MT_SCOPE_CYCLE_COUNTER(EvaluateAnimGraph, !IsInGameThread());
@@ -2145,6 +2157,8 @@ void FAnimInstanceProxy::GatherDebugData(FNodeDebugData& DebugData)
 
 void FAnimInstanceProxy::GatherDebugData_WithRoot(FNodeDebugData& DebugData, FAnimNode_Base* InRootNode, FName InLayerName)
 {
+	LLM_SCOPE_BYNAME(TEXT("Animation/Graph"));
+
 	// Gather debug data for Root Node
 	if(InRootNode != nullptr)
 	{
@@ -3216,6 +3230,8 @@ void FAnimInstanceProxy::RecordStateWeight(const int32 InMachineClassIndex, cons
 
 void FAnimInstanceProxy::ResetDynamics(ETeleportType InTeleportType)
 {
+	LLM_SCOPE_BYNAME(TEXT("Animation/Graph"));
+
 	DECLARE_SCOPE_HIERARCHICAL_COUNTER_FUNC()
 	for(FAnimNode_Base* Node : DynamicResetNodes)
 	{
@@ -3225,6 +3241,8 @@ void FAnimInstanceProxy::ResetDynamics(ETeleportType InTeleportType)
 
 void FAnimInstanceProxy::ResetDynamics()
 {
+	LLM_SCOPE_BYNAME(TEXT("Animation/Graph"));
+
 	ResetDynamics(ETeleportType::ResetPhysics);
 }
 
diff --git a/Engine/Source/Runtime/Engine/Private/Animation/AnimNode_DeadBlending.cpp b/Engine/Source/Runtime/Engine/Private/Animation/AnimNode_DeadBlending.cpp
--- a/Engine/Source/Runtime/Engine/Private/Animation/AnimNode_DeadBlending.cpp
+++ b/Engine/Source/Runtime/Engine/Private/Animation/AnimNode_DeadBlending.cpp
@@ -12,8 +12,6 @@
 
 #include UE_INLINE_GENERATED_CPP_BY_NAME(AnimNode_DeadBlending)
 
-LLM_DEFINE_TAG(Animation_DeadBlending);
-
 #define LOCTEXT_NAMESPACE "AnimNode_DeadBlending"
 
 TAutoConsoleVariable<int32> CVarAnimDeadBlendingEnable(TEXT("a.AnimNode.DeadBlending.Enable"), 1, TEXT("Enable / Disable DeadBlending"));
@@ -651,7 +649,6 @@ void FAnimNode_DeadBlending::RequestInertialization(const FInertializationReques
 
 void FAnimNode_DeadBlending::Initialize_AnyThread(const FAnimationInitializeContext& Context)
 {
-	LLM_SCOPE_BYNAME(TEXT("Animation/DeadBlending"));
 	DECLARE_SCOPE_HIERARCHICAL_COUNTER_ANIMNODE(Initialize_AnyThread);
 
 	FAnimNode_Base::Initialize_AnyThread(Context);
@@ -724,7 +721,6 @@ void FAnimNode_DeadBlending::CacheBones_AnyThread(const FAnimationCacheBonesCont
 
 void FAnimNode_DeadBlending::Update_AnyThread(const FAnimationUpdateContext& Context)
 {
-	LLM_SCOPE_BYNAME(TEXT("Animation/DeadBlending"));
 	DECLARE_SCOPE_HIERARCHICAL_COUNTER_ANIMNODE(Update_AnyThread);
 
 	const bool bNeedsReset =
@@ -758,7 +754,6 @@ void FAnimNode_DeadBlending::Update_AnyThread(const FAnimationUpdateContext& Con
 
 void FAnimNode_DeadBlending::Evaluate_AnyThread(FPoseContext& Output)
 {
-	LLM_SCOPE_BYNAME(TEXT("Animation/DeadBlending"));
 	DECLARE_SCOPE_HIERARCHICAL_COUNTER_ANIMNODE(Evaluate_AnyThread);
 
 	// Evaluate the Input and write it to the Output
diff --git a/Engine/Source/Runtime/Engine/Private/Animation/AnimNode_Inertialization.cpp b/Engine/Source/Runtime/Engine/Private/Animation/AnimNode_Inertialization.cpp
--- a/Engine/Source/Runtime/Engine/Private/Animation/AnimNode_Inertialization.cpp
+++ b/Engine/Source/Runtime/Engine/Private/Animation/AnimNode_Inertialization.cpp
@@ -14,8 +14,6 @@
 
 #include UE_INLINE_GENERATED_CPP_BY_NAME(AnimNode_Inertialization)
 
-LLM_DEFINE_TAG(Animation_Inertialization);
-
 #define LOCTEXT_NAMESPACE "AnimNode_Inertialization"
 
 IMPLEMENT_ANIMGRAPH_MESSAGE(UE::Anim::IInertializationRequester);
@@ -299,7 +297,6 @@ void FAnimNode_Inertialization::RequestInertialization(const FInertializationReq
 
 void FAnimNode_Inertialization::Initialize_AnyThread(const FAnimationInitializeContext& Context)
 {
-	LLM_SCOPE_BYNAME(TEXT("Animation/Inertialization"));
 	DECLARE_SCOPE_HIERARCHICAL_COUNTER_ANIMNODE(Initialize_AnyThread);
 
 	FAnimNode_Base::Initialize_AnyThread(Context);
@@ -366,7 +363,6 @@ void FAnimNode_Inertialization::CacheBones_AnyThread(const FAnimationCacheBonesC
 
 void FAnimNode_Inertialization::Update_AnyThread(const FAnimationUpdateContext& Context)
 {
-	LLM_SCOPE_BYNAME(TEXT("Animation/Inertialization"));
 	DECLARE_SCOPE_HIERARCHICAL_COUNTER_ANIMNODE(Update_AnyThread);
 
 	const bool bNeedsReset =
@@ -434,7 +430,6 @@ void FAnimNode_Inertialization::Update_AnyThread(const FAnimationUpdateContext&
 
 void FAnimNode_Inertialization::Evaluate_AnyThread(FPoseContext& Output)
 {
-	LLM_SCOPE_BYNAME(TEXT("Animation/Inertialization"));
 	DECLARE_SCOPE_HIERARCHICAL_COUNTER_ANIMNODE(Evaluate_AnyThread);
 	ANIM_MT_SCOPE_CYCLE_COUNTER_VERBOSE(Inertialization, !IsInGameThread());
 
