[shaders] deprecate GetIntegerValue function on FShaderCompilerEnvironment; this was missed when restricting access to content of definitions and usage of it in a shader format implementation snuck in.

diff --git a/Engine/Source/Runtime/Engine/Private/ShaderCompiler/ShaderGenerationUtil.cpp b/Engine/Source/Runtime/Engine/Private/ShaderCompiler/ShaderGenerationUtil.cpp
--- a/Engine/Source/Runtime/Engine/Private/ShaderCompiler/ShaderGenerationUtil.cpp
+++ b/Engine/Source/Runtime/Engine/Private/ShaderCompiler/ShaderGenerationUtil.cpp
@@ -39,17 +39,19 @@ bool NeedsVelocityDepth(EShaderPlatform TargetPlatform)
 
 #define SET_COMPILE_BOOL_IF_TRUE(X) { if (DerivedDefines.X) { OutEnvironment.SetDefine(TEXT(#X), 1); } }
 
+// disable deprecation warnings due to FShaderCompilerDefinitions moving to internal. This can be removed when the move is complete.
+PRAGMA_DISABLE_DEPRECATION_WARNINGS
 #define FETCH_COMPILE_BOOL(X) {																\
 	static FShaderCompilerDefineNameCache Cache_##X(TEXT(#X));								\
-	SrcDefines.X = Environment.GetIntegerValue(Cache_##X, SrcDefines.X) != 0 ? 1 : 0; }
+	SrcDefines.X = Container.GetIntegerValue(Cache_##X, SrcDefines.X) != 0 ? 1 : 0; }
 
 #define FETCH_COMPILE_INT(X) {																\
 	static FShaderCompilerDefineNameCache Cache_##X(TEXT(#X));								\
-	SrcDefines.X = Environment.GetIntegerValue(Cache_##X, SrcDefines.X); }
+	SrcDefines.X = Container.GetIntegerValue(Cache_##X, SrcDefines.X); }
 
 
-template<typename EnvironmentType>
-void ApplyFetchEnvironmentInternal(FShaderGlobalDefines& SrcDefines, const EnvironmentType& Environment)
+template<typename ContainerType>
+void ApplyFetchEnvironmentInternal(FShaderGlobalDefines& SrcDefines, const ContainerType& Container)
 {
 	FETCH_COMPILE_BOOL(USES_BASE_PASS_VELOCITY);
 	FETCH_COMPILE_BOOL(GBUFFER_HAS_VELOCITY);
@@ -79,8 +81,8 @@ void ApplyFetchEnvironmentInternal(FShaderGlobalDefines& SrcDefines, const Envir
 	FETCH_COMPILE_BOOL(PLATFORM_ALLOW_SCENE_DATA_COMPRESSED_TRANSFORMS);
 }
 
-template<typename EnvironmentType>
-void ApplyFetchEnvironmentInternal(FShaderLightmapPropertyDefines& SrcDefines, const EnvironmentType& Environment)
+template<typename ContainerType>
+void ApplyFetchEnvironmentInternal(FShaderLightmapPropertyDefines& SrcDefines, const ContainerType& Container)
 {
 	FETCH_COMPILE_BOOL(LQ_TEXTURE_LIGHTMAP);
 	FETCH_COMPILE_BOOL(HQ_TEXTURE_LIGHTMAP);
@@ -107,8 +109,8 @@ void ApplyFetchEnvironmentInternal(FShaderLightmapPropertyDefines& SrcDefines, c
 
 }
 
-template<typename EnvironmentType>
-void ApplyFetchEnvironmentInternal(FShaderMaterialPropertyDefines& SrcDefines, const EnvironmentType& Environment)
+template<typename ContainerType>
+void ApplyFetchEnvironmentInternal(FShaderMaterialPropertyDefines& SrcDefines, const ContainerType& Container)
 {
 	FETCH_COMPILE_BOOL(MATERIAL_ENABLE_TRANSLUCENCY_FOGGING);
 	FETCH_COMPILE_BOOL(MATERIALBLENDING_ANY_TRANSLUCENT);
@@ -227,8 +229,8 @@ void ApplyFetchEnvironmentInternal(FShaderMaterialPropertyDefines& SrcDefines, c
 	FETCH_COMPILE_INT(GBUFFER_LAYOUT);
 }
 
-template<typename EnvironmentType>
-void ApplyFetchEnvironmentInternal(FShaderCompilerDefines& SrcDefines, const EnvironmentType& Environment)
+template<typename ContainerType>
+void ApplyFetchEnvironmentInternal(FShaderCompilerDefines& SrcDefines, const ContainerType& Container)
 {
 	FETCH_COMPILE_BOOL(COMPILER_GLSL_ES3_1);
 	FETCH_COMPILE_BOOL(ES3_1_PROFILE);
@@ -244,11 +246,12 @@ void ApplyFetchEnvironmentInternal(FShaderCompilerDefines& SrcDefines, const Env
 
 	FETCH_COMPILE_BOOL(PLATFORM_SUPPORTS_DEVELOPMENT_SHADERS);
 }
+PRAGMA_ENABLE_DEPRECATION_WARNINGS
 
 
 void FShaderCompileUtilities::ApplyFetchEnvironment(FShaderGlobalDefines& SrcDefines, const FShaderCompilerEnvironment& Environment, const EShaderPlatform Platform)
 {
-	ApplyFetchEnvironmentInternal(SrcDefines, Environment);
+	ApplyFetchEnvironmentInternal(SrcDefines, *Environment.Definitions.Get());
 
 	// note that we are doing an if so that if we call ApplyFetchEnvironment() twice, we get the logical OR of bSupportsDualBlending support
 	if (RHISupportsDualSourceBlending(Platform))
@@ -258,19 +261,19 @@ void FShaderCompileUtilities::ApplyFetchEnvironment(FShaderGlobalDefines& SrcDef
 }
 void FShaderCompileUtilities::ApplyFetchEnvironment(FShaderLightmapPropertyDefines& SrcDefines, const FShaderCompilerEnvironment& Environment)
 {
-	ApplyFetchEnvironmentInternal(SrcDefines, Environment);
+	ApplyFetchEnvironmentInternal(SrcDefines, *Environment.Definitions.Get());
 }
 void FShaderCompileUtilities::ApplyFetchEnvironment(FShaderMaterialPropertyDefines& SrcDefines, const FShaderCompilerEnvironment& Environment)
 {
-	ApplyFetchEnvironmentInternal(SrcDefines, Environment);
+	ApplyFetchEnvironmentInternal(SrcDefines, *Environment.Definitions.Get());
 }
 void FShaderCompileUtilities::ApplyFetchEnvironment(FShaderCompilerDefines& SrcDefines, const FShaderCompilerEnvironment& Environment)
 {
-	ApplyFetchEnvironmentInternal(SrcDefines, Environment);
+	ApplyFetchEnvironmentInternal(SrcDefines, *Environment.Definitions.Get());
 }
 
-/** Dummy "environment" used to gather names of defines used in the ApplyFetchEnvironmentInternal functions */
-struct FDefineNameGatherEnvironment
+/** Used to gather names of defines used in the ApplyFetchEnvironmentInternal functions */
+struct FDefineNameGather
 {
 	mutable TArray<FName> Names;
 
@@ -291,7 +294,7 @@ struct FShaderInitialDefinesInitializer
 		FShaderGlobalDefines GlobalDefines = {};
 		FShaderCompilerDefines CompilerDefines = {};
 
-		FDefineNameGatherEnvironment GatherNames;
+		FDefineNameGather GatherNames;
 
 		ApplyFetchEnvironmentInternal(GlobalDefines, GatherNames);
 		ApplyFetchEnvironmentInternal(MaterialDefines, GatherNames);
@@ -2235,7 +2238,9 @@ void FShaderCompileUtilities::GenerateBrdfHeaders(const FName& ShaderFormat)
 
 EGBufferLayout FShaderCompileUtilities::FetchGBufferLayout(const FShaderCompilerEnvironment& Environment)
 {
-	const uint32 Layout = Environment.GetIntegerValue(TEXT("GBUFFER_LAYOUT"));
+	PRAGMA_DISABLE_DEPRECATION_WARNINGS
+	const uint32 Layout = Environment.Definitions->GetIntegerValue(TEXT("GBUFFER_LAYOUT"));
+	PRAGMA_ENABLE_DEPRECATION_WARNINGS
 	if (Layout >= GBL_Num)
 	{
 		return GBL_Default;
diff --git a/Engine/Source/Runtime/RenderCore/Public/ShaderCore.h b/Engine/Source/Runtime/RenderCore/Public/ShaderCore.h
--- a/Engine/Source/Runtime/RenderCore/Public/ShaderCore.h
+++ b/Engine/Source/Runtime/RenderCore/Public/ShaderCore.h
@@ -573,7 +573,11 @@ struct FShaderCompilerEnvironment
 	RENDERCORE_API void SetDefine(FShaderCompilerDefineNameCache& Name, bool Value);
 	RENDERCORE_API void SetDefine(FShaderCompilerDefineNameCache& Name, float Value);
 
+	UE_DEPRECATED(5.5,	"Accessing definition values via FShaderCompilerEnvironment is no longer allowed outside of the core shader system.\n"
+						"Please use Get/SetCompileArgument mechanism for additional properties required by compilation.")
 	RENDERCORE_API int32 GetIntegerValue(FName Name) const;
+	UE_DEPRECATED(5.5,	"Accessing definition values via FShaderCompilerEnvironment is no longer allowed outside of the core shader system.\n"
+						"Please use Get/SetCompileArgument mechanism for additional properties required by compilation.")
 	RENDERCORE_API int32 GetIntegerValue(FShaderCompilerDefineNameCache& NameCache, int32 ResultIfNotFound = 0) const;
 
 	UE_DEPRECATED(5.5, "ContainsDefinition will be made private in the future and should not be called by downstream code.")
