Move some classes that use very common names to the UE namespace to avoid potential clashes with projects

#rnx

diff --git a/Engine/Source/Runtime/Engine/Private/UnrealEngine.cpp b/Engine/Source/Runtime/Engine/Private/UnrealEngine.cpp
--- a/Engine/Source/Runtime/Engine/Private/UnrealEngine.cpp
+++ b/Engine/Source/Runtime/Engine/Private/UnrealEngine.cpp
@@ -4649,28 +4649,32 @@ static void ShowSubobjectGraph( FOutputDevice& Ar, UObject* CurrentObject, const
 		}
 	}
 }
+
+namespace UE
+{
+
 struct FItem
 {
-	UClass*	Class;
+	UClass* Class;
 	int32 Count;
 	SIZE_T Num;
 	SIZE_T Max;
 	/** Only exclusive resource size, the truer resource size. */
 	FResourceSizeEx TrueResourceSize;
 
-	FItem( UClass* InClass=NULL )
+	FItem(UClass* InClass = nullptr)
 		: Class(InClass), Count(0), Num(0), Max(0), TrueResourceSize()
 	{}
 
-	FItem( UClass* InClass, int32 InCount, SIZE_T InNum, SIZE_T InMax, FResourceSizeEx InTrueResourceSize ) :
-		Class( InClass ),
-		Count( InCount ),
-		Num( InNum ), 
-		Max( InMax ), 
-		TrueResourceSize( InTrueResourceSize )
+	FItem(UClass* InClass, int32 InCount, SIZE_T InNum, SIZE_T InMax, FResourceSizeEx InTrueResourceSize) :
+		Class(InClass),
+		Count(InCount),
+		Num(InNum),
+		Max(InMax),
+		TrueResourceSize(InTrueResourceSize)
 	{}
 
-	void Add( FArchiveCountMem& Ar, FResourceSizeEx InTrueResourceSize )
+	void Add(FArchiveCountMem& Ar, FResourceSizeEx InTrueResourceSize)
 	{
 		Count++;
 		Num += Ar.GetNum();
@@ -4699,11 +4703,13 @@ struct FSubItem
 	/** Only exclusive resource size, the truer resource size. */
 	FResourceSizeEx TrueResourceSize;
 
-	FSubItem( UObject* InObject, SIZE_T InNum, SIZE_T InMax, FResourceSizeEx InTrueResourceSize )
-		: Object( InObject ), Num( InNum ), Max( InMax ), TrueResourceSize( InTrueResourceSize )
+	FSubItem(UObject* InObject, SIZE_T InNum, SIZE_T InMax, FResourceSizeEx InTrueResourceSize)
+		: Object(InObject), Num(InNum), Max(InMax), TrueResourceSize(InTrueResourceSize)
 	{}
 };
 
+} // namespace UE
+
 #endif // !UE_BUILD_SHIPPING
 
 MSVC_PRAGMA(warning(push))
@@ -8517,7 +8523,7 @@ struct FHierarchy
 			AddClass((UClass*)This);
 		}
 	}
-	FHierarchyNode& Compute(UObject* This, TMap<UObject*, FSubItem> const& Objects, bool bCountItems)
+	FHierarchyNode& Compute(UObject* This, TMap<UObject*, UE::FSubItem> const& Objects, bool bCountItems)
 	{
 		FHierarchyNode& Node = Nodes.FindChecked(This);
 		if (Node.Inc < 0)
@@ -8526,7 +8532,7 @@ struct FHierarchy
 			Node.ExcCount = 1;
 			if (This)
 			{
-				FSubItem const& Item = Objects.FindChecked(This);
+				UE::FSubItem const& Item = Objects.FindChecked(This);
 				Node.Exc += Item.Max;
 				Node.Exc += Item.TrueResourceSize.GetTotalMemoryBytes();
 				if (bCountItems)
@@ -8700,7 +8706,7 @@ bool UEngine::HandleObjCommand( const TCHAR* Cmd, FOutputDevice& Ar )
 	{
 		struct FCompareByInclusiveSize
 		{
-			FORCEINLINE bool operator()( const FItem& A, const FItem& B ) const 
+			FORCEINLINE bool operator()( const UE::FItem& A, const UE::FItem& B ) const 
 			{ 
 				return A.Max > B.Max; 
 			}
@@ -8733,7 +8739,7 @@ bool UEngine::HandleObjCommand( const TCHAR* Cmd, FOutputDevice& Ar )
 			ClassToCheck = UObject::StaticClass();
 		}
 
-		TMap<UClass*,FItem> ObjectsByClass;
+		TMap<UClass*,UE::FItem> ObjectsByClass;
 
 		Ar.Logf( TEXT("**********************************************") );
 		Ar.Logf( TEXT("Obj MemSub for class '%s'"), *ClassToCheck->GetName() );
@@ -8752,7 +8758,7 @@ bool UEngine::HandleObjCommand( const TCHAR* Cmd, FOutputDevice& Ar )
 			FLocal::GetReferencedObjs( Obj, ReferencedObjects );
 
 			// Calculate memory usage.
-			FItem ThisObject( Obj->GetClass() );
+			UE::FItem ThisObject( Obj->GetClass() );
 			for( UObject*& RefObj : ReferencedObjects )
 			{
 				FArchiveCountMem Count( RefObj );
@@ -8761,7 +8767,7 @@ bool UEngine::HandleObjCommand( const TCHAR* Cmd, FOutputDevice& Ar )
 				ThisObject.Add( Count, TrueResourceSize );
 			}
 
-			FItem& ClassObjects = ObjectsByClass.FindOrAdd( ThisObject.Class );
+			UE::FItem& ClassObjects = ObjectsByClass.FindOrAdd( ThisObject.Class );
 			ClassObjects.Count++;
 			ClassObjects.Num += ThisObject.Num;
 			ClassObjects.Max += ThisObject.Max;
@@ -8784,12 +8790,12 @@ bool UEngine::HandleObjCommand( const TCHAR* Cmd, FOutputDevice& Ar )
 			TEXT("Count") 
 		);
 
-		FItem Total;
-		FItem Culled;
+		UE::FItem Total;
+		UE::FItem Culled;
 		for( const auto& It : ObjectsByClass )
 		{
 			UClass* Class = It.Key;
-			const FItem& ClassObjects = It.Value;
+			const UE::FItem& ClassObjects = It.Value;
 
 			if( ClassObjects.Max < (SIZE_T)Limit )
 			{
@@ -8871,13 +8877,13 @@ bool UEngine::HandleObjCommand( const TCHAR* Cmd, FOutputDevice& Ar )
 		FHierarchy Outers(Limit);
 		FHierarchy Flat(Limit);
 
-		TMap<UObject*, FSubItem> Objects;
+		TMap<UObject*, UE::FSubItem> Objects;
 		for( FThreadSafeObjectIterator It; It; ++It )
 		{
 			FArchiveCountMem Count( *It );
 			FResourceSizeEx TrueResourceSize = FResourceSizeEx(EResourceSizeMode::Exclusive);
 			It->GetResourceSizeEx(TrueResourceSize);
-			Objects.Add(*It, FSubItem(*It, Count.GetNum(), Count.GetMax(), TrueResourceSize));
+			Objects.Add(*It, UE::FSubItem(*It, Count.GetNum(), Count.GetMax(), TrueResourceSize));
 			Classes.AddClassInstance(*It);
 			Outers.AddOuter(*It);
 			Flat.AddFlat(*It);
@@ -8956,9 +8962,9 @@ bool UEngine::HandleObjCommand( const TCHAR* Cmd, FOutputDevice& Ar )
 		FString ObjectNameSubString;
 		FParse::Value(Cmd, TEXT("NAMESUB="), ObjectNameSubString);
 
-		TMap<UClass*, FItem> ObjectsByClass;
-		TArray<FSubItem> Objects;
-		FItem Total;
+		TMap<UClass*, UE::FItem> ObjectsByClass;
+		TArray<UE::FSubItem> Objects;
+		UE::FItem Total;
 
 		// support specifying metaclasses when listing class objects
 		if ( CheckType && CheckType->IsChildOf(UClass::StaticClass()) )
@@ -9120,12 +9126,12 @@ bool UEngine::HandleObjCommand( const TCHAR* Cmd, FOutputDevice& Ar )
 					}
 				}
 
-				FItem& ClassData = ObjectsByClass.FindOrAdd(ClassToUse);
+				UE::FItem& ClassData = ObjectsByClass.FindOrAdd(ClassToUse);
 				ClassData.Class = ClassToUse;
 
 				if( bShowDetailedObjectInfo )
 				{
-					Objects.Add(FSubItem(*It, NumBytes, MaxBytes, TrueResourceSize));
+					Objects.Add(UE::FSubItem(*It, NumBytes, MaxBytes, TrueResourceSize));
 				}
 				ClassData.Add(NumBytes, MaxBytes, TrueResourceSize);
 				Total.Add(NumBytes, MaxBytes, TrueResourceSize);
@@ -9149,7 +9155,7 @@ bool UEngine::HandleObjCommand( const TCHAR* Cmd, FOutputDevice& Ar )
 					  bResourceSizeSort( InResourceSizeSort )
 				{}
 
-				FORCEINLINE bool operator()( const FSubItem& A, const FSubItem& B ) const
+				FORCEINLINE bool operator()( const UE::FSubItem& A, const UE::FSubItem& B ) const
 				{
 					if (bAlphaSort)
 					{
@@ -9183,7 +9189,7 @@ bool UEngine::HandleObjCommand( const TCHAR* Cmd, FOutputDevice& Ar )
 				);
 			}
 
-			for (const FSubItem& ObjItem : Objects)
+			for (const UE::FSubItem& ObjItem : Objects)
 			{
 				if (bCSV)
 				{
@@ -9231,12 +9237,12 @@ bool UEngine::HandleObjCommand( const TCHAR* Cmd, FOutputDevice& Ar )
 					, bCountSort( InCountSort )
 					, bResourceSizeSort( InResourceSizeSort )
 				{}
-				FORCEINLINE bool operator()( const FItem& A, const FItem& B ) const
+				FORCEINLINE bool operator()( const UE::FItem& A, const UE::FItem& B ) const
 				{
 					return bAlphaSort ? (A.Class->GetName() < B.Class->GetName()) : bCountSort ? (B.Count < A.Count) : bResourceSizeSort ? (B.TrueResourceSize.GetTotalMemoryBytes() < A.TrueResourceSize.GetTotalMemoryBytes()) : (B.Max < A.Max);
 				}
 			};
-			TArray<FItem> SortedClasses;
+			TArray<UE::FItem> SortedClasses;
 			ObjectsByClass.GenerateValueArray(SortedClasses);
 			SortedClasses.Sort(FCompareFItem(bAlphaSort, bCountSort, bResourceSizeSort));
 			
@@ -9273,7 +9279,7 @@ bool UEngine::HandleObjCommand( const TCHAR* Cmd, FOutputDevice& Ar )
 				}
 			}
 
-			for (const FItem& ClassData : SortedClasses)
+			for (const UE::FItem& ClassData : SortedClasses)
 			{
 				if (bCSV)
 				{
