PerQualityLevelProperties: Cooking: To support cooks of two platforms in the same process, when modifying the struct for the cook serialization, restore it after cook serialization is complete, so that it is still valid for serializing the next platform.
Fix Typo: Qualtiy -> Quality
#rnx

diff --git a/Engine/Source/Runtime/Engine/Classes/Engine/SkeletalMesh.h b/Engine/Source/Runtime/Engine/Classes/Engine/SkeletalMesh.h
--- a/Engine/Source/Runtime/Engine/Classes/Engine/SkeletalMesh.h
+++ b/Engine/Source/Runtime/Engine/Classes/Engine/SkeletalMesh.h
@@ -972,7 +972,7 @@ public:
 	{
 #if WITH_EDITORONLY_DATA
 		WaitUntilAsyncPropertyReleased(ESkeletalMeshAsyncProperties::MinLod);
-		MinQualityLevelLOD.PerQuality = QualityLevelProperty::ConvertQualtiyLevelData(QualityLevelMinimumLODs);
+		MinQualityLevelLOD.PerQuality = QualityLevelProperty::ConvertQualityLevelData(QualityLevelMinimumLODs);
 		MinQualityLevelLOD.Default = Default >=0 ? Default : MinQualityLevelLOD.Default;
 #endif
 	}
@@ -982,7 +982,7 @@ public:
 	{
 #if WITH_EDITORONLY_DATA
 		WaitUntilAsyncPropertyReleased(ESkeletalMeshAsyncProperties::MinLod, ESkinnedAssetAsyncPropertyLockType::ReadOnly);
-		QualityLevelMinimumLODs = QualityLevelProperty::ConvertQualtiyLevelData(MinQualityLevelLOD.PerQuality);
+		QualityLevelMinimumLODs = QualityLevelProperty::ConvertQualityLevelData(MinQualityLevelLOD.PerQuality);
 		Default = MinQualityLevelLOD.Default;
 #endif
 	}
diff --git a/Engine/Source/Runtime/Engine/Classes/Engine/StaticMesh.h b/Engine/Source/Runtime/Engine/Classes/Engine/StaticMesh.h
--- a/Engine/Source/Runtime/Engine/Classes/Engine/StaticMesh.h
+++ b/Engine/Source/Runtime/Engine/Classes/Engine/StaticMesh.h
@@ -769,7 +769,7 @@ public:
 	{
 #if WITH_EDITORONLY_DATA
 		WaitUntilAsyncPropertyReleased(EStaticMeshAsyncProperties::MinLOD);
-		MinQualityLevelLOD.PerQuality = QualityLevelProperty::ConvertQualtiyLevelData(QualityLevelMinimumLODs);
+		MinQualityLevelLOD.PerQuality = QualityLevelProperty::ConvertQualityLevelData(QualityLevelMinimumLODs);
 		MinQualityLevelLOD.Default = Default >= 0 ? Default : MinQualityLevelLOD.Default;
 #endif
 	}
@@ -779,7 +779,7 @@ public:
 	{
 #if WITH_EDITORONLY_DATA
 		WaitUntilAsyncPropertyReleased(EStaticMeshAsyncProperties::MinLOD);
-		QualityLevelMinimumLODs = QualityLevelProperty::ConvertQualtiyLevelData(MinQualityLevelLOD.PerQuality);
+		QualityLevelMinimumLODs = QualityLevelProperty::ConvertQualityLevelData(MinQualityLevelLOD.PerQuality);
 		Default = MinQualityLevelLOD.Default;
 #endif
 	}
diff --git a/Engine/Source/Runtime/Engine/Private/PerQualityLevelProperties.cpp b/Engine/Source/Runtime/Engine/Private/PerQualityLevelProperties.cpp
--- a/Engine/Source/Runtime/Engine/Private/PerQualityLevelProperties.cpp
+++ b/Engine/Source/Runtime/Engine/Private/PerQualityLevelProperties.cpp
@@ -9,6 +9,7 @@
 #include "PlatformInfo.h"
 #include "DeviceProfiles/DeviceProfile.h"
 #include "DeviceProfiles/DeviceProfileManager.h"
+#include "Misc/Optional.h"
 #endif
 
 
@@ -18,7 +19,7 @@ namespace QualityLevelProperty
 	static FString QualityLevelMappingStr = TEXT("QualityLevelMapping");
 
 	template<typename _ValueType>
-	TMap<int32, _ValueType> ConvertQualtiyLevelData(const TMap<EPerQualityLevels, _ValueType>& Data)
+	TMap<int32, _ValueType> ConvertQualityLevelData(const TMap<EPerQualityLevels, _ValueType>& Data)
 	{
 		TMap<int32, _ValueType> ConvertedData;
 
@@ -31,7 +32,7 @@ namespace QualityLevelProperty
 	}
 
 	template<typename _ValueType>
-	TMap<EPerQualityLevels, _ValueType> ConvertQualtiyLevelData(const TMap<int32, _ValueType>& Data)
+	TMap<EPerQualityLevels, _ValueType> ConvertQualityLevelData(const TMap<int32, _ValueType>& Data)
 	{
 		TMap<EPerQualityLevels, _ValueType> ConvertedData;
 
@@ -169,6 +170,29 @@ namespace QualityLevelProperty
 			return *CachedMappingQualitLevelInfo;
 		}
 	}
+
+	template <typename ValueType>
+	struct FSavedData
+	{
+		int32 Default = 0;
+		TMap<int32, ValueType> PerQuality;
+	};
+	template<typename StructType, typename ValueType>
+	void SaveQualityLevelForCooking(const StructType& Property, TOptional<FSavedData<ValueType>>& OutSavedData)
+	{
+		OutSavedData.Emplace();
+		OutSavedData->Default = Property.Default;
+		OutSavedData->PerQuality = Property.PerQuality;
+	}
+	template<typename StructType, typename ValueType>
+	void RestoreQualityLevelForCooking(TOptional<FSavedData<ValueType>>& InSavedData, StructType& OutProperty)
+	{
+		if (InSavedData.IsSet())
+		{
+			OutProperty.Default = InSavedData->Default;
+			OutProperty.PerQuality = MoveTemp(InSavedData->PerQuality);
+		}
+	}
 #endif
 }
 
@@ -321,7 +345,7 @@ FSupportedQualityLevelArray FPerQualityLevelProperty<StructType, ValueType, _Bas
 }
 
 template<typename StructType, typename ValueType, EName _BasePropertyName>
-void FPerQualityLevelProperty<StructType, ValueType, _BasePropertyName>::StripQualtiyLevelForCooking(const TCHAR* InPlatformName)
+void FPerQualityLevelProperty<StructType, ValueType, _BasePropertyName>::StripQualityLevelForCooking(const TCHAR* InPlatformName)
 {
 	StructType* This = StaticCast<StructType*>(this);
 	if (This->PerQuality.Num() > 0 && !CVarName.IsEmpty())
@@ -420,11 +444,13 @@ ENGINE_API FArchive& operator<<(FArchive& Ar, FPerQualityLevelProperty<StructTyp
 	bool bCooked = false;
 	StructType* This = StaticCast<StructType*>(&Property);
 #if WITH_EDITOR
+	TOptional<QualityLevelProperty::FSavedData<ValueType>> SavedData;
 	if (Ar.IsCooking())
 	{
 		bCooked = true;
+		QualityLevelProperty::SaveQualityLevelForCooking(*This, SavedData);
 		const FDataDrivenPlatformInfo& PlatformInfo = Ar.CookingTarget()->GetPlatformInfo();
-		This->StripQualtiyLevelForCooking(*(PlatformInfo.IniPlatformName.ToString()));
+		This->StripQualityLevelForCooking(*(PlatformInfo.IniPlatformName.ToString()));
 	}
 #endif
 	{
@@ -432,6 +458,10 @@ ENGINE_API FArchive& operator<<(FArchive& Ar, FPerQualityLevelProperty<StructTyp
 		Ar << This->Default;
 		Ar << This->PerQuality;
 	}
+#if WITH_EDITOR
+	QualityLevelProperty::RestoreQualityLevelForCooking(SavedData, *This);
+#endif
+
 	return Ar;
 }
 
@@ -446,10 +476,12 @@ ENGINE_API void operator<<(FStructuredArchive::FSlot Slot, FPerQualityLevelPrope
 	StructType* This = StaticCast<StructType*>(&Property);
 
 #if WITH_EDITOR
+	TOptional<QualityLevelProperty::FSavedData<ValueType>> SavedData;
 	if (UnderlyingArchive.IsCooking())
 	{
 		bCooked = true;
-		This->StripQualtiyLevelForCooking(*(UnderlyingArchive.CookingTarget()->GetPlatformInfo().IniPlatformName.ToString()));
+		QualityLevelProperty::SaveQualityLevelForCooking(*This, SavedData);
+		This->StripQualityLevelForCooking(*(UnderlyingArchive.CookingTarget()->GetPlatformInfo().IniPlatformName.ToString()));
 	}
 #endif
 	{
@@ -457,6 +489,9 @@ ENGINE_API void operator<<(FStructuredArchive::FSlot Slot, FPerQualityLevelPrope
 		Record << SA_VALUE(TEXT("Value"), This->Default);
 		Record << SA_VALUE(TEXT("PerQuality"), This->PerQuality);
 	}
+#if WITH_EDITOR
+	QualityLevelProperty::RestoreQualityLevelForCooking(SavedData, *This);
+#endif
 }
 // 
 template ENGINE_API FArchive& operator<<(FArchive&, FPerQualityLevelProperty<FPerQualityLevelInt, int32, NAME_IntProperty>&);
@@ -465,13 +500,13 @@ template ENGINE_API void operator<<(FStructuredArchive::FSlot Slot, FPerQualityL
 #if WITH_EDITOR
 template ENGINE_API int32 FPerQualityLevelProperty<FPerQualityLevelInt, int32, NAME_IntProperty>::GetValueForPlatform(const ITargetPlatform* TargetPlatform) const;
 template ENGINE_API FSupportedQualityLevelArray FPerQualityLevelProperty<FPerQualityLevelInt, int32, NAME_IntProperty>::GetSupportedQualityLevels(const TCHAR* InPlatformName) const;
-template ENGINE_API void FPerQualityLevelProperty<FPerQualityLevelInt, int32, NAME_IntProperty>::StripQualtiyLevelForCooking(const TCHAR* InPlatformName);
+template ENGINE_API void FPerQualityLevelProperty<FPerQualityLevelInt, int32, NAME_IntProperty>::StripQualityLevelForCooking(const TCHAR* InPlatformName);
 template ENGINE_API bool FPerQualityLevelProperty<FPerQualityLevelInt, int32, NAME_IntProperty>::IsQualityLevelValid(int32 QualityLevel) const;
 template ENGINE_API void FPerQualityLevelProperty<FPerQualityLevelInt, int32, NAME_IntProperty>::ConvertQualityLevelData(const TMap<FName, int32>& PlaformData, const TMultiMap<FName, FName>& PerPlatformToQualityLevel, int32 Default);
 template ENGINE_API void FPerQualityLevelProperty<FPerQualityLevelInt, int32, NAME_IntProperty>::ConvertQualityLevelDataUsingCVar(const TMap<FName, int32>& PlaformData, int32 Default, bool);
 #endif
-template TMap<int32, int32> QualityLevelProperty::ConvertQualtiyLevelData(const TMap<EPerQualityLevels, int32>& Data);
-template TMap<EPerQualityLevels, int32> QualityLevelProperty::ConvertQualtiyLevelData(const TMap<int32, int32>& Data);
+template TMap<int32, int32> QualityLevelProperty::ConvertQualityLevelData(const TMap<EPerQualityLevels, int32>& Data);
+template TMap<EPerQualityLevels, int32> QualityLevelProperty::ConvertQualityLevelData(const TMap<int32, int32>& Data);
 template ENGINE_API void FPerQualityLevelProperty<FPerQualityLevelInt, int32, NAME_IntProperty>::Init(const TCHAR* InCVarName, const TCHAR* InSection);
 template ENGINE_API int32 FPerQualityLevelProperty<FPerQualityLevelInt, int32, NAME_IntProperty>::GetValue(int32 QualityLevel) const;
 
@@ -499,13 +534,13 @@ template ENGINE_API void operator<<(FStructuredArchive::FSlot Slot, FPerQualityL
 #if WITH_EDITOR
 template int32 FPerQualityLevelProperty<FPerQualityLevelFloat, float, NAME_FloatProperty>::GetValueForPlatform(const ITargetPlatform* TargetPlatform) const;
 template FSupportedQualityLevelArray FPerQualityLevelProperty<FPerQualityLevelFloat, float, NAME_FloatProperty>::GetSupportedQualityLevels(const TCHAR* InPlatformName) const;
-template void FPerQualityLevelProperty<FPerQualityLevelFloat, float, NAME_FloatProperty>::StripQualtiyLevelForCooking(const TCHAR* InPlatformName);
+template void FPerQualityLevelProperty<FPerQualityLevelFloat, float, NAME_FloatProperty>::StripQualityLevelForCooking(const TCHAR* InPlatformName);
 template bool FPerQualityLevelProperty<FPerQualityLevelFloat, float, NAME_FloatProperty>::IsQualityLevelValid(int32 QualityLevel) const;
 template ENGINE_API void FPerQualityLevelProperty<FPerQualityLevelFloat, float, NAME_FloatProperty>::ConvertQualityLevelData(const TMap<FName, float>& PlaformData, const TMultiMap<FName, FName>& PerPlatformToQualityLevel, float Default);
 template ENGINE_API void FPerQualityLevelProperty<FPerQualityLevelFloat, float, NAME_FloatProperty>::ConvertQualityLevelDataUsingCVar(const TMap<FName, float>& PlaformData, float Default, bool);
 #endif
-template TMap<int32, float> QualityLevelProperty::ConvertQualtiyLevelData(const TMap<EPerQualityLevels, float>&Data);
-template TMap<EPerQualityLevels, float> QualityLevelProperty::ConvertQualtiyLevelData(const TMap<int32, float>& Data);
+template TMap<int32, float> QualityLevelProperty::ConvertQualityLevelData(const TMap<EPerQualityLevels, float>&Data);
+template TMap<EPerQualityLevels, float> QualityLevelProperty::ConvertQualityLevelData(const TMap<int32, float>& Data);
 template ENGINE_API void FPerQualityLevelProperty<FPerQualityLevelFloat, float, NAME_FloatProperty>::Init(const TCHAR* InCVarName, const TCHAR* InSection);
 template ENGINE_API float FPerQualityLevelProperty<FPerQualityLevelFloat, float, NAME_FloatProperty>::GetValue(int32 QualityLevel) const;
 
diff --git a/Engine/Source/Runtime/Engine/Public/PerQualityLevelProperties.h b/Engine/Source/Runtime/Engine/Public/PerQualityLevelProperties.h
--- a/Engine/Source/Runtime/Engine/Public/PerQualityLevelProperties.h
+++ b/Engine/Source/Runtime/Engine/Public/PerQualityLevelProperties.h
@@ -47,10 +47,22 @@ namespace QualityLevelProperty
 	ENGINE_API FName QualityLevelToFName(int32 QL);
 	ENGINE_API int32 FNameToQualityLevel(FName QL);
 	template<typename _ValueType>
-	ENGINE_API TMap<int32, _ValueType> ConvertQualtiyLevelData(const TMap<EPerQualityLevels, _ValueType>& Data);
+	ENGINE_API TMap<int32, _ValueType> ConvertQualityLevelData(const TMap<EPerQualityLevels, _ValueType>& Data);
 	template<typename _ValueType>
-	ENGINE_API TMap<EPerQualityLevels, _ValueType> ConvertQualtiyLevelData(const TMap<int32, _ValueType>& Data);
-	
+	UE_DEPRECATED(5.5, "Use ConvertQualityLevelData.")
+	TMap<int32, _ValueType> ConvertQualtiyLevelData(const TMap<EPerQualityLevels, _ValueType>& Data)
+	{
+		return ConvertQualityLevelData(Data);
+	}
+	template<typename _ValueType>
+	ENGINE_API TMap<EPerQualityLevels, _ValueType> ConvertQualityLevelData(const TMap<int32, _ValueType>& Data);
+	template<typename _ValueType>
+	UE_DEPRECATED(5.5, "Use ConvertQualityLevelData.")
+	TMap<EPerQualityLevels, _ValueType> ConvertQualtiyLevelData(const TMap<int32, _ValueType>& Data)
+	{
+		return ConvertQualityLevelData(Data);
+	}
+
 #if WITH_EDITOR
 	ENGINE_API TArray<FName> GetEnginePlatformsForPlatformOrGroupName(const FString& InPlatformName);
 	ENGINE_API FSupportedQualityLevelArray PerPlatformOverrideMapping(FString& InPlatformName);
@@ -91,7 +103,12 @@ struct FPerQualityLevelProperty
 #if WITH_EDITOR
 	int32 GetValueForPlatform(const ITargetPlatform* TargetPlatform) const;
 	FSupportedQualityLevelArray GetSupportedQualityLevels(const TCHAR* InPlatformName = nullptr) const;
-	void StripQualtiyLevelForCooking(const TCHAR* InPlatformName = nullptr);
+	void StripQualityLevelForCooking(const TCHAR* InPlatformName = nullptr);
+	UE_DEPRECATED(5.5, "Use StripQualityLevelForCooking")
+	void StripQualtiyLevelForCooking(const TCHAR* InPlatformName = nullptr)
+	{
+		StripQualityLevelForCooking(InPlatformName);
+	}
 	bool IsQualityLevelValid(int32 QualityLevel) const;
 	void ConvertQualityLevelData(const TMap<FName, _ValueType>& PlatformData, const TMultiMap<FName, FName>& PerPlatformToQualityLevel, _ValueType Default);
 	UE_DEPRECATED(5.5, "Use ConvertQualityLevelData")
@@ -114,7 +131,7 @@ struct FPerQualityLevelProperty
 		CVarName = FString(InCVarName);
 	}
 
-	UE_DEPRECATED(5.4, "If no cvar is associated with the property, all quality levels will be keept when cooking. Call SetQualtiyLevelCVarForCooking to strip unsupported quality levels when cooking")
+	UE_DEPRECATED(5.4, "If no cvar is associated with the property, all quality levels will be keept when cooking. Call SetQualityLevelCVarForCooking to strip unsupported quality levels when cooking")
 	void Init(const TCHAR* InCVarName, const TCHAR* InSection)
 	{
 		SetQualityLevelCVarForCooking(InCVarName, InSection);
