[World Partition Class Descriptors] Filter out blueprint changes that aren't compilations.
In some cases, a property change event would trigger for a BP being renamed, which registered the BP with the old name.

diff --git a/Engine/Source/Runtime/Engine/Private/WorldPartition/WorldPartitionClassDescRegistry.cpp b/Engine/Source/Runtime/Engine/Private/WorldPartition/WorldPartitionClassDescRegistry.cpp
--- a/Engine/Source/Runtime/Engine/Private/WorldPartition/WorldPartitionClassDescRegistry.cpp
+++ b/Engine/Source/Runtime/Engine/Private/WorldPartition/WorldPartitionClassDescRegistry.cpp
@@ -549,7 +549,7 @@ void FWorldPartitionClassDescRegistry::OnObjectPropertyChanged(UObject* InObject
 		if (UBlueprint* Blueprint = Cast<UBlueprint>(InObject))
 		{
 			// The generated class is invalid in some situations, like renaming a blueprint, etc.
-			if (Blueprint->GeneratedClass && Blueprint->GeneratedClass->IsChildOf<AActor>())
+			if (Blueprint->IsUpToDate() && Blueprint->GeneratedClass && Blueprint->GeneratedClass->IsChildOf<AActor>())
 			{
 				PrefetchClassDesc(Blueprint->GeneratedClass);
 				UpdateClassDescriptor(InObject, true);
