Bring back CL33705721 with fixes
Original CL Desc
-----------------------------------------------------------------
[Content Bundle] Added an optional flag to allow marking content bundles as "required", which will make them unremovable and bake their actors into the main map to reduce IO stress.

diff --git a/Engine/Source/Runtime/Engine/Private/WorldPartition/ContentBundle/ContentBundleDescriptor.cpp b/Engine/Source/Runtime/Engine/Private/WorldPartition/ContentBundle/ContentBundleDescriptor.cpp
--- a/Engine/Source/Runtime/Engine/Private/WorldPartition/ContentBundle/ContentBundleDescriptor.cpp
+++ b/Engine/Source/Runtime/Engine/Private/WorldPartition/ContentBundle/ContentBundleDescriptor.cpp
@@ -8,10 +8,9 @@
 #include UE_INLINE_GENERATED_CPP_BY_NAME(ContentBundleDescriptor)
 
 UContentBundleDescriptor::UContentBundleDescriptor(const FObjectInitializer& ObjectInitializer)
-	: DebugColor(FColor::Black)
-{
-
-}
+	: bRequired(false)
+	, DebugColor(FColor::Black)
+{}
 
 FString UContentBundleDescriptor::GetPackageRoot() const
 {
diff --git a/Engine/Source/Runtime/Engine/Private/WorldPartition/ContentBundle/ContentBundleEditor.cpp b/Engine/Source/Runtime/Engine/Private/WorldPartition/ContentBundle/ContentBundleEditor.cpp
--- a/Engine/Source/Runtime/Engine/Private/WorldPartition/ContentBundle/ContentBundleEditor.cpp
+++ b/Engine/Source/Runtime/Engine/Private/WorldPartition/ContentBundle/ContentBundleEditor.cpp
@@ -93,6 +93,8 @@ void FContentBundleEditor::DoInjectContent()
 
 			check(GetDescriptor()->GetGuid().IsValid());
 
+			ActorDescContainerInstance->bIsContentBundleRequired = GetDescriptor()->GetRequired();
+
 			if (!ActorDescContainerInstance->GetContainer()->IsEmpty())
 			{
 				WorldDataLayersActorReference = FWorldDataLayersReference(ActorDescContainerInstance.Get(), BuildWorlDataLayersName());
@@ -346,6 +348,11 @@ void FContentBundleEditor::GenerateStreaming(TArray<FString>* OutPackageToGenera
 		return;
 	}
 
+	if (GetDescriptor()->GetRequired())
+	{
+		return;
+	}
+
 	FActorDescContainerInstanceCollection Collection({ TObjectPtr<UActorDescContainerInstance>(ActorDescContainerInstance.Get()) });
 	UWorldPartition::FGenerateStreamingParams Params = UWorldPartition::FGenerateStreamingParams()
 		.SetContainerInstanceCollection(Collection, FStreamingGenerationContainerInstanceCollection::ECollectionType::BaseAsContentBundle);
@@ -388,16 +395,22 @@ void FContentBundleEditor::GenerateStreaming(TArray<FString>* OutPackageToGenera
 
 void FContentBundleEditor::OnBeginCook(IWorldPartitionCookPackageContext& CookContext)
 {
-	check(!bIsInCook);
-	CookContext.RegisterPackageCookPackageGenerator(this);
-	bIsInCook = true;
+	if (!GetDescriptor()->GetRequired())
+	{
+		check(!bIsInCook);
+		CookContext.RegisterPackageCookPackageGenerator(this);
+		bIsInCook = true;
+	}
 }
 
 void FContentBundleEditor::OnEndCook(IWorldPartitionCookPackageContext& CookContext)
 {
-	check(bIsInCook);
-	CookContext.UnregisterPackageCookPackageGenerator(this);
-	bIsInCook = false;
+	if (!GetDescriptor()->GetRequired())
+	{
+		check(bIsInCook);
+		CookContext.UnregisterPackageCookPackageGenerator(this);
+		bIsInCook = false;
+	}
 }
 
 bool FContentBundleEditor::GatherPackagesToCook(class IWorldPartitionCookPackageContext& CookContext)
diff --git a/Engine/Source/Runtime/Engine/Private/WorldPartition/RuntimeSpatialHash/RuntimeSpatialHashHLOD.cpp b/Engine/Source/Runtime/Engine/Private/WorldPartition/RuntimeSpatialHash/RuntimeSpatialHashHLOD.cpp
--- a/Engine/Source/Runtime/Engine/Private/WorldPartition/RuntimeSpatialHash/RuntimeSpatialHashHLOD.cpp
+++ b/Engine/Source/Runtime/Engine/Private/WorldPartition/RuntimeSpatialHash/RuntimeSpatialHashHLOD.cpp
@@ -396,7 +396,7 @@ static void UpdateHLODGridsActors(UWorld* World, const UActorDescContainerInstan
 		ASpatialHashRuntimeGridInfo* GridActor = ExistingGridActors.FindRef(GridName);
 		if (!GridActor)
 		{
-			FContentBundleActivationScope ContentBndleScope(ContainerInstance->GetContentBundleGuid());
+			FContentBundleActivationScope ContentBundleScope(ContainerInstance->IsContentBundleRequired() ? FGuid() : ContainerInstance->GetContentBundleGuid());
 			FScopedOverrideSpawningLevelMountPointObject EDLScope(ContainerInstance->GetExternalDataLayerAsset());
 
 			FActorSpawnParameters SpawnParams;
@@ -405,7 +405,7 @@ static void UpdateHLODGridsActors(UWorld* World, const UActorDescContainerInstan
 			GridActor->Tags.Add(HLODGridTag);
 			bDirty = true;
 
-			check(GridActor->GetContentBundleGuid() == ContainerInstance->GetContentBundleGuid());
+			check(GridActor->GetContentBundleGuid() == (ContainerInstance->IsContentBundleRequired() ? FGuid() : ContainerInstance->GetContentBundleGuid()));
 			check(GridActor->GetExternalDataLayerAsset() == ContainerInstance->GetExternalDataLayerAsset());
 		}
 
diff --git a/Engine/Source/Runtime/Engine/Private/WorldPartition/WorldPartitionStreamingGeneration.cpp b/Engine/Source/Runtime/Engine/Private/WorldPartition/WorldPartitionStreamingGeneration.cpp
--- a/Engine/Source/Runtime/Engine/Private/WorldPartition/WorldPartitionStreamingGeneration.cpp
+++ b/Engine/Source/Runtime/Engine/Private/WorldPartition/WorldPartitionStreamingGeneration.cpp
@@ -2073,42 +2073,56 @@ URuntimeHashExternalStreamingObjectBase* UWorldPartition::FlushStreamingToExtern
 	return ExternalStreamingObject;
 }
 
-static void ExtractContentBundleContainerInstances(const FActorDescContainerInstanceCollection* InContainerInstanceCollection, TArray<const UActorDescContainerInstance*>& OutContentBundleContainerInstances, TArray<const UActorDescContainerInstance*>& OutNonContentBundleContainerInstances)
+void UWorldPartition::SetupHLODActors(const FSetupHLODActorsParams& Params)
 {
-	InContainerInstanceCollection->ForEachActorDescContainerInstance([&OutContentBundleContainerInstances, &OutNonContentBundleContainerInstances](const UActorDescContainerInstance* InActorDescContainerInstance)
+	TArray<const UActorDescContainerInstance*> ContentBundleContainerInstances;
+	TArray<const UActorDescContainerInstance*> BaseAndEDLContainerInstances;
+
+	ForEachActorDescContainerInstance([&ContentBundleContainerInstances, &BaseAndEDLContainerInstances](const UActorDescContainerInstance* InActorDescContainerInstance)
 	{
+		// If the content bundle is required, we need to add it to both lists, and we'll filter out all actors so the HLOd generation process can delete HLOD actor files.
 		if (InActorDescContainerInstance->GetContentBundleGuid().IsValid())
 		{
-			OutContentBundleContainerInstances.Add(InActorDescContainerInstance);
+			ContentBundleContainerInstances.Add(InActorDescContainerInstance);
 		}
-		else
+		
+		if (!InActorDescContainerInstance->GetContentBundleGuid().IsValid() || InActorDescContainerInstance->IsContentBundleRequired())
 		{
-			OutNonContentBundleContainerInstances.Add(InActorDescContainerInstance);
+			BaseAndEDLContainerInstances.Add(InActorDescContainerInstance);
 		}
 	});
-}
-
-void UWorldPartition::SetupHLODActors(const FSetupHLODActorsParams& Params)
-{
-	TArray<const UActorDescContainerInstance*> ContentBundleContainerInstances;
-	TArray<const UActorDescContainerInstance*> BaseAndEDLContainerInstances;
-	ExtractContentBundleContainerInstances(this, ContentBundleContainerInstances, BaseAndEDLContainerInstances);
 
-	auto SetupHLODActorsForCollection = [this, &Params](const FStreamingGenerationContainerInstanceCollection& InContainerInstanceCollection)
+	auto SetupHLODActorsForCollection = [this, &Params](const FStreamingGenerationContainerInstanceCollection& InContainerInstanceCollection, bool bFilterAllActors)
 	{
 		TErrorHandlerSelector<FStreamingGenerationLogErrorHandler> ErrorHandlerSelector;
 		FWorldPartitionStreamingGenerator::FWorldPartitionStreamingGeneratorParams StreamingGeneratorParams = FWorldPartitionStreamingGenerator::FWorldPartitionStreamingGeneratorParams()
 			.SetWorldPartitionContext(this)
 			.SetErrorHandler(ErrorHandlerSelector.Get())
 			.SetEnableStreaming(IsStreamingEnabled())
-			.SetFilteredClasses({ AWorldPartitionHLOD::StaticClass() })
+			.SetFilteredClasses({ bFilterAllActors ? AActor::StaticClass() : AWorldPartitionHLOD::StaticClass() })
 			.SetIsValidGrid([this](FName GridName, const UClass* ActorClass) { return RuntimeHash->IsValidGrid(GridName, ActorClass); })
 			.SetIsValidHLODLayer([this](FName GridName, const FSoftObjectPath& HLODLayerPath) { return RuntimeHash->IsValidHLODLayer(GridName, HLODLayerPath); });
 
 		FWorldPartitionStreamingGenerator StreamingGenerator(StreamingGeneratorParams);
 		StreamingGenerator.PreparationPhase(InContainerInstanceCollection);
 
-		TUniquePtr<FArchive> LogFileAr = FWorldPartitionStreamingGenerator::CreateDumpStateLogArchive(TEXT("HLOD"));
+		const FString ContainerPackageName = InContainerInstanceCollection.GetBaseContainerInstancePackageName().ToString();
+		FString ContainerShortName = FPackageName::GetShortName(ContainerPackageName);
+		if (!ContainerPackageName.StartsWith(TEXT("/Game/")))
+		{
+			TArray<FString> SplitContainerPath;
+			if (ContainerPackageName.ParseIntoArray(SplitContainerPath, TEXT("/")))
+			{
+				ContainerShortName += TEXT(".");
+				ContainerShortName += SplitContainerPath[0];
+			}
+		}
+
+		TStringBuilder<256> StateLogSuffix;
+		StateLogSuffix += TEXT("HLOD_");
+		StateLogSuffix += ContainerShortName;
+
+		TUniquePtr<FArchive> LogFileAr = FWorldPartitionStreamingGenerator::CreateDumpStateLogArchive(*StateLogSuffix);
 		if (LogFileAr.IsValid())
 		{
 			FHierarchicalLogArchive HierarchicalLogAr(*LogFileAr);
@@ -2122,14 +2136,14 @@ void UWorldPartition::SetupHLODActors(const FSetupHLODActorsParams& Params)
 	for (const UActorDescContainerInstance* ContentBundleContainerInstance : ContentBundleContainerInstances)
 	{
 		FStreamingGenerationContainerInstanceCollection ContentBundleCollection({ ContentBundleContainerInstance }, FStreamingGenerationContainerInstanceCollection::ECollectionType::BaseAsContentBundle);
-		SetupHLODActorsForCollection(ContentBundleCollection);
+		SetupHLODActorsForCollection(ContentBundleCollection, ContentBundleContainerInstance->IsContentBundleRequired());
 	}
 
 	// Single pass for base and EDL container instances
 	if (!BaseAndEDLContainerInstances.IsEmpty())
 	{
 		FStreamingGenerationContainerInstanceCollection Collection(BaseAndEDLContainerInstances, FStreamingGenerationContainerInstanceCollection::ECollectionType::BaseAndEDLs);
-		SetupHLODActorsForCollection(Collection);
+		SetupHLODActorsForCollection(Collection, false);
 	}
 }
 
@@ -2232,8 +2246,15 @@ void FStreamingGenerationContainerInstanceCollection::InitializeCollection()
 
 	if (CollectionType == ECollectionType::BaseAsContentBundle)
 	{
+#if DO_CHECK
 		check(GetActorDescContainerCount() == 1);
 		check(GetContentBundleGuid().IsValid());
+		Algo::ForEach(ActorDescContainerInstanceCollection, [](const UActorDescContainerInstance* ActorDescContainerInstance)
+		{
+			check(ActorDescContainerInstance->GetContentBundleGuid().IsValid());
+			check(!ActorDescContainerInstance->IsContentBundleRequired());
+		});
+#endif
 		ContentBundleStartIdx = 0;
 		return;
 	}
@@ -2246,7 +2267,10 @@ void FStreamingGenerationContainerInstanceCollection::InitializeCollection()
 		// When type is set to BaseAndEDL, we remove ContentBundle containers from the collection.
 		// BaseAndEDL type assumes ContentBundle containers are generated separately one at a time.
 		check(!ShouldRegisterDelegates());
-		ActorDescContainerInstanceCollection.SetNum(Algo::RemoveIf(ActorDescContainerInstanceCollection, [](const UActorDescContainerInstance* ActorDescContainerInstance) { return ActorDescContainerInstance->GetContentBundleGuid().IsValid(); }));
+		ActorDescContainerInstanceCollection.SetNum(Algo::RemoveIf(ActorDescContainerInstanceCollection, [](const UActorDescContainerInstance* ActorDescContainerInstance)
+		{
+			return ActorDescContainerInstance->GetContentBundleGuid().IsValid() && !ActorDescContainerInstance->IsContentBundleRequired();
+		}));
 	}
 
 	int32 BaseContainerCount = Algo::CountIf(ActorDescContainerInstanceCollection, [](const UActorDescContainerInstance* ActorDescContainerInstance) { return !ActorDescContainerInstance->HasExternalContent(); });
@@ -2343,7 +2367,18 @@ void UWorldPartition::CheckForErrors(const FCheckForErrorsParams& InParams)
 	// @todo_ow : Once content bundles are remove, we will only do this validation in 1 pass
 	TArray<const UActorDescContainerInstance*> ContentBundleContainerInstances;
 	TArray<const UActorDescContainerInstance*> BaseAndEDLContainerInstances;
-	ExtractContentBundleContainerInstances(InParams.ActorDescContainerInstanceCollection, ContentBundleContainerInstances, BaseAndEDLContainerInstances);
+
+	InParams.ActorDescContainerInstanceCollection->ForEachActorDescContainerInstance([&ContentBundleContainerInstances, &BaseAndEDLContainerInstances](const UActorDescContainerInstance* InActorDescContainerInstance)
+	{
+		if (InActorDescContainerInstance->GetContentBundleGuid().IsValid() && !InActorDescContainerInstance->IsContentBundleRequired())
+		{
+			ContentBundleContainerInstances.Add(InActorDescContainerInstance);
+		}
+		else
+		{
+			BaseAndEDLContainerInstances.Add(InActorDescContainerInstance);
+		}
+	});
 
 	if (!BaseAndEDLContainerInstances.IsEmpty())
 	{
diff --git a/Engine/Source/Runtime/Engine/Public/WorldPartition/ActorDescContainerInstance.h b/Engine/Source/Runtime/Engine/Public/WorldPartition/ActorDescContainerInstance.h
--- a/Engine/Source/Runtime/Engine/Public/WorldPartition/ActorDescContainerInstance.h
+++ b/Engine/Source/Runtime/Engine/Public/WorldPartition/ActorDescContainerInstance.h
@@ -14,6 +14,7 @@ class UActorDescContainer;
 class FWorldPartitionActorDesc;
 class UWorldPartition;
 class UExternalDataLayerAsset;
+class FContentBundleEditor;
 struct FWorldPartitionRuntimeCellPropertyOverride;
 
 class FActorDescInstanceList : public TActorDescList<FWorldPartitionActorDescInstance> { };
@@ -29,11 +30,14 @@ protected:
 		: bIsInitialized(false) 
 		, bRegisteredDelegates(false)
 		, bCreateChildContainerHierarchy(false)
+		, bIsContentBundleRequired(false)
 #endif
 	{}
 
 #if WITH_EDITOR
+	friend class FStreamingGenerationContainerInstanceCollection;
 	friend FWorldPartitionActorDescInstance;
+	friend FContentBundleEditor;
 	friend UWorldPartition;
 		
 public:
@@ -148,6 +152,9 @@ public:
 	ENGINE_API const TMap<FGuid, TObjectPtr<UActorDescContainerInstance>>& GetChildContainerInstances() const { return ChildContainerInstances; }
 		
 	const UActorDescContainerInstance* GetParentContainerInstance() const { return ParentContainerInstance.Get(); }
+
+	bool IsContentBundleRequired() const { return bIsContentBundleRequired; }
+
 protected:
 	virtual void RegisterContainer(const FInitializeParams& InParams);
 	virtual void UnregisterContainer();
@@ -204,5 +211,6 @@ protected:
 	bool														bIsInitialized;
 	bool														bRegisteredDelegates;
 	bool														bCreateChildContainerHierarchy;
+	bool														bIsContentBundleRequired;
 #endif
 };
\ No newline at end of file
diff --git a/Engine/Source/Runtime/Engine/Public/WorldPartition/ContentBundle/ContentBundleDescriptor.h b/Engine/Source/Runtime/Engine/Public/WorldPartition/ContentBundle/ContentBundleDescriptor.h
--- a/Engine/Source/Runtime/Engine/Public/WorldPartition/ContentBundle/ContentBundleDescriptor.h
+++ b/Engine/Source/Runtime/Engine/Public/WorldPartition/ContentBundle/ContentBundleDescriptor.h
@@ -18,6 +18,7 @@ class UContentBundleDescriptor : public UObject
 
 public:
 	const FString& GetDisplayName() const { return DisplayName; }
+	bool GetRequired() const { return bRequired; }
 	const FColor& GetDebugColor() const { return DebugColor; }
 	ENGINE_API FString GetPackageRoot() const;
 	const FGuid& GetGuid() const { return Guid; }
@@ -42,6 +43,9 @@ private:
 	UPROPERTY(EditAnywhere, Category = BaseInformation)
 	FString DisplayName;
 
+	UPROPERTY(EditAnywhere, Category = BaseInformation)
+	bool bRequired;
+
 	UPROPERTY(EditAnywhere, DuplicateTransient, Category = BaseInformation)
 	FColor DebugColor;
 
