Prevent linking non blueprint classes on linked anim graphs (fixes a crash when trying to debug a graph where a non-blueprint class has been linked)

diff --git a/Engine/Source/Runtime/Engine/Private/Animation/AnimNode_LinkedAnimGraph.cpp b/Engine/Source/Runtime/Engine/Private/Animation/AnimNode_LinkedAnimGraph.cpp
--- a/Engine/Source/Runtime/Engine/Private/Animation/AnimNode_LinkedAnimGraph.cpp
+++ b/Engine/Source/Runtime/Engine/Private/Animation/AnimNode_LinkedAnimGraph.cpp
@@ -337,6 +337,13 @@ void FAnimNode_LinkedAnimGraph::SetAnimClass(TSubclassOf<UAnimInstance> InClass,
 {
 	UClass* NewClass = InClass.Get();
 
+	// Make sure this is a valid blueprint class
+	if (NewClass && NewClass->ClassGeneratedBy == nullptr)
+	{
+		ensure(false);
+		return;
+	}
+
 	// Verified OK, so set it now
 	TSubclassOf<UAnimInstance> OldClass = InstanceClass;
 	InstanceClass = InClass;
