Fix build - ClassGeneratedBy not available outside of editor builds

diff --git a/Engine/Source/Runtime/Engine/Private/Animation/AnimNode_LinkedAnimGraph.cpp b/Engine/Source/Runtime/Engine/Private/Animation/AnimNode_LinkedAnimGraph.cpp
--- a/Engine/Source/Runtime/Engine/Private/Animation/AnimNode_LinkedAnimGraph.cpp
+++ b/Engine/Source/Runtime/Engine/Private/Animation/AnimNode_LinkedAnimGraph.cpp
@@ -338,7 +338,7 @@ void FAnimNode_LinkedAnimGraph::SetAnimClass(TSubclassOf<UAnimInstance> InClass,
 	UClass* NewClass = InClass.Get();
 
 	// Make sure this is a valid blueprint class
-	if (NewClass && NewClass->ClassGeneratedBy == nullptr)
+	if (NewClass && IAnimClassInterface::GetFromClass(NewClass) == nullptr)
 	{
 		ensure(false);
 		return;
