Move ObjectPropertyTrace from Editor to Runtime code

diff --git a/Engine/Plugins/Animation/GameplayInsights/Source/GameplayInsights/Private/GameplayInsightsModule.cpp b/Engine/Plugins/Animation/GameplayInsights/Source/GameplayInsights/Private/GameplayInsightsModule.cpp
--- a/Engine/Plugins/Animation/GameplayInsights/Source/GameplayInsights/Private/GameplayInsightsModule.cpp
+++ b/Engine/Plugins/Animation/GameplayInsights/Source/GameplayInsights/Private/GameplayInsightsModule.cpp
@@ -15,7 +15,6 @@
 #include "Widgets/Docking/SDockTab.h"
 #include "Trace/StoreClient.h"
 #include "Stats/Stats.h"
-#include "ObjectPropertyTrace.h"
 #include "SMontageView.h"
 #include "AnimCurvesTrack.h"
 #include "InertializationsTrack.h"
@@ -35,7 +34,7 @@
 #include "GameplayInsightsStyle.h"
 #include "SSubobjectInstanceEditor.h"
 #include "ProfilingDebugging/TraceAuxiliary.h"
-
+#include "ObjectPropertyTrace.h"
 #endif
 
 #if WITH_ENGINE
@@ -207,20 +206,12 @@ void FGameplayInsightsModule::StartupModule()
 	FOnRegisterMajorTabExtensions& TimingProfilerExtension = UnrealInsightsModule.OnRegisterMajorTabExtension(FInsightsManagerTabs::TimingProfilerTabId);
 	TimingProfilerExtension.AddRaw(this, &FGameplayInsightsModule::RegisterTimingProfilerLayoutExtensions);
 #endif
-
-#if OBJECT_PROPERTY_TRACE_ENABLED
-	FObjectPropertyTrace::Init();
-#endif
 }
 
 void FGameplayInsightsModule::ShutdownModule()
 {
 	LLM_SCOPE_BYNAME(TEXT("Insights/GameplayInsights"));
 
-#if OBJECT_PROPERTY_TRACE_ENABLED
-	FObjectPropertyTrace::Destroy();
-#endif
-
 #if WITH_EDITOR
 	IAnimationBlueprintEditorModule* AnimationBlueprintEditorModule = FModuleManager::GetModulePtr<IAnimationBlueprintEditorModule>("AnimationBlueprintEditor");
 	if(AnimationBlueprintEditorModule)
diff --git a/Engine/Plugins/Animation/GameplayInsights/Source/GameplayInsights/Private/ObjectEventsTrack.cpp b/Engine/Plugins/Animation/GameplayInsights/Source/GameplayInsights/Private/ObjectEventsTrack.cpp
--- a/Engine/Plugins/Animation/GameplayInsights/Source/GameplayInsights/Private/ObjectEventsTrack.cpp
+++ b/Engine/Plugins/Animation/GameplayInsights/Source/GameplayInsights/Private/ObjectEventsTrack.cpp
@@ -10,10 +10,13 @@
 #include "Insights/ViewModels/TimingEventSearch.h"
 #include "TraceServices/Model/Frames.h"
 #include "VariantTreeNode.h"
-#include "ObjectPropertyTrace.h"
 #include "Framework/MultiBox/MultiBoxBuilder.h"
 #include "Framework/Commands/UICommandInfo.h"
 
+#if WITH_EDITOR
+#include "ObjectPropertyTrace.h"
+#endif
+
 #define LOCTEXT_NAMESPACE "ObjectEventsTrack"
 
 INSIGHTS_IMPLEMENT_RTTI(FObjectEventsTrack)
@@ -171,6 +174,7 @@ void FObjectEventsTrack::BuildContextMenu(FMenuBuilder& MenuBuilder)
 {
 	FGameplayTimingEventsTrack::BuildContextMenu(MenuBuilder);
 
+#if WITH_EDITOR
 #if OBJECT_PROPERTY_TRACE_ENABLED
 	MenuBuilder.BeginSection("Trace", LOCTEXT("TraceHeader", "Trace"));
 	{
@@ -214,6 +218,7 @@ void FObjectEventsTrack::BuildContextMenu(FMenuBuilder& MenuBuilder)
 	}
 	MenuBuilder.EndSection();
 #endif
+#endif
 }
 
 #undef LOCTEXT_NAMESPACE
diff --git a/Engine/Plugins/Animation/GameplayInsights/Source/GameplayInsights/Private/ObjectPropertyTrace.cpp b/Engine/Plugins/Animation/GameplayInsights/Source/GameplayInsights/Private/ObjectPropertyTrace.cpp
deleted file mode 100644
diff --git a/Engine/Plugins/Animation/GameplayInsights/Source/GameplayInsights/Private/ObjectPropertyTrace.h b/Engine/Plugins/Animation/GameplayInsights/Source/GameplayInsights/Private/ObjectPropertyTrace.h
deleted file mode 100644
diff --git a/Engine/Source/Runtime/Engine/Private/ObjectPropertyTrace.cpp b/Engine/Source/Runtime/Engine/Private/ObjectPropertyTrace.cpp
new file mode 100644
--- /dev/null
+++ b/Engine/Source/Runtime/Engine/Private/ObjectPropertyTrace.cpp
@@ -0,0 +1,348 @@
+// Copyright Epic Games, Inc. All Rights Reserved.
+
+#include "ObjectPropertyTrace.h"
+#include "Trace/Trace.h"
+#include "Trace/Trace.inl"
+#include "Containers/Ticker.h"
+
+#if WITH_ENGINE
+
+#if OBJECT_PROPERTY_TRACE_ENABLED
+UE_TRACE_CHANNEL(ObjectProperties)
+
+UE_TRACE_EVENT_BEGIN(Object, ClassPropertyStringId, NoSync|Important)
+	UE_TRACE_EVENT_FIELD(uint32, Id)
+	UE_TRACE_EVENT_FIELD(UE::Trace::WideString, Value)
+UE_TRACE_EVENT_END()
+
+UE_TRACE_EVENT_BEGIN(Object, PropertiesStart2)
+	UE_TRACE_EVENT_FIELD(uint64, Cycle)
+	UE_TRACE_EVENT_FIELD(uint64, ObjectId)
+	UE_TRACE_EVENT_FIELD(double, RecordingTime)
+UE_TRACE_EVENT_END()
+
+UE_TRACE_EVENT_BEGIN(Object, PropertiesEnd)
+	UE_TRACE_EVENT_FIELD(uint64, Cycle)
+	UE_TRACE_EVENT_FIELD(uint64, ObjectId)
+UE_TRACE_EVENT_END()
+
+UE_TRACE_EVENT_BEGIN(Object, PropertyValue3)
+	UE_TRACE_EVENT_FIELD(uint64, Cycle)
+	UE_TRACE_EVENT_FIELD(uint64, ObjectId)
+	UE_TRACE_EVENT_FIELD(int32, ParentId)
+	UE_TRACE_EVENT_FIELD(uint32, TypeId)
+	UE_TRACE_EVENT_FIELD(uint32, NameId)
+	UE_TRACE_EVENT_FIELD(uint32, ParentNameId)
+	UE_TRACE_EVENT_FIELD(UE::Trace::WideString, Value)
+UE_TRACE_EVENT_END()
+
+namespace ObjectPropertyTrace
+{
+	static FTSTicker::FDelegateHandle TickerHandle;
+	static TArray<TWeakObjectPtr<const UObject>> Objects;
+
+	static uint32 CurrentClassPropertyStringId = 0;
+	static TMap<FString, uint32> StringIdMap;
+	typedef TFunctionRef<void(const FString& /*InTypeString*/, const FString& /*InValue*/, int32 /*InId*/, int32 /*InParentId*/, FStringBuilderBase & /*InPropertyNameBuilder*/)> IterateFunction;
+
+	static uint32 TraceStringId(const FString& InString)
+	{
+		if(const uint32* ExistingIdPtr = StringIdMap.Find(InString))
+		{
+			return *ExistingIdPtr;
+		}
+
+		if (InString.IsEmpty())
+		{
+			return 0;	// Return invalid id.
+		}
+		else
+		{
+			const uint32 NewId = StringIdMap.Add(InString, ++CurrentClassPropertyStringId);
+
+			UE_TRACE_LOG(Object, ClassPropertyStringId, ObjectProperties, InString.Len() * sizeof(TCHAR))
+				<< ClassPropertyStringId.Id(NewId)
+				<< ClassPropertyStringId.Value(*InString, InString.Len());
+
+			return NewId;
+		}
+	}
+
+	/**
+	 * Helper to ensure FStringBuilderBase clears any added strings at the end of a scope.
+	 * 
+	 * Used to support creating one single FStringBuilderBase when tracing an object's properties recursively.
+	 */
+	struct FScopedPropertyName
+	{
+		FScopedPropertyName(FStringBuilderBase & InBuilder) : StartSize(InBuilder.Len()), Builder(InBuilder)
+		{
+		}
+
+		~FScopedPropertyName()
+		{
+			Builder.RemoveSuffix(FMath::Max(Builder.Len() - StartSize, 0));
+		}
+		
+		int32 StartSize;
+		FStringBuilderBase & Builder;
+	};
+	
+	static void IteratePropertiesRecursive(FProperty* InProperty, const void* InContainer, const FString& InKey, IterateFunction InFunction, int32& InId, int32 InParentId, FStringBuilderBase & InPropertyNameBuilder)
+	{
+		// Handle container properties
+		if (const FArrayProperty* ArrayProperty = CastField<FArrayProperty>(InProperty))
+		{
+			FScriptArrayHelper_InContainer Helper(ArrayProperty, InContainer);
+
+			int32 ArrayRootId = ++InId;
+			FString SizeString = FString::Printf(TEXT("{Num = %d}"), Helper.Num());
+			
+			// Build property name id
+			FScopedPropertyName ScopedPropertyName(InPropertyNameBuilder);
+			ScopedPropertyName.Builder.AppendChar(TEXT('.'));
+			ScopedPropertyName.Builder.Append(InProperty->GetName());
+
+			// Trace array property
+			InFunction(InProperty->GetCPPType(), SizeString, ArrayRootId, InParentId, InPropertyNameBuilder);
+
+			// Trace container items
+			for (int32 DynamicIndex = 0; DynamicIndex < Helper.Num(); ++DynamicIndex)
+			{
+				const void* ValuePtr = Helper.GetRawPtr(DynamicIndex);
+				FString KeyString = FString::Printf(TEXT("[%d]"), DynamicIndex);
+				
+				IteratePropertiesRecursive(ArrayProperty->Inner, ValuePtr, KeyString, InFunction, InId, ArrayRootId, InPropertyNameBuilder);
+			}
+		}
+		else if (const FMapProperty* MapProperty = CastField<FMapProperty>(InProperty))
+		{
+			FScriptMapHelper_InContainer Helper(MapProperty, InContainer);
+
+			int32 MapRootId = ++InId;
+			FString SizeString = FString::Printf(TEXT("{Num = %d}"), Helper.Num());
+
+			// Build property name id
+			FScopedPropertyName ScopedPropertyName(InPropertyNameBuilder);
+			ScopedPropertyName.Builder.AppendChar(TEXT('.'));
+			ScopedPropertyName.Builder.Append(InProperty->GetName());
+			
+			InFunction(InProperty->GetCPPType(), SizeString, MapRootId, InParentId, InPropertyNameBuilder);
+			
+			for (FScriptMapHelper::FIterator It(Helper); It; ++It)
+			{
+				int32 MapEntryId = ++InId;
+				FString KeyString = FString::Printf(TEXT("[%d]"), It.GetLogicalIndex());
+				FString TypeString = FString::Printf(TEXT("{%s, %s}"), *MapProperty->KeyProp->GetCPPType(), *MapProperty->ValueProp->GetCPPType());
+
+				// Build pair parent name id
+				FScopedPropertyName ScopedParentPairName(InPropertyNameBuilder);
+				ScopedParentPairName.Builder.AppendChar(TEXT('.'));
+				ScopedParentPairName.Builder.Append(KeyString);
+
+				InFunction(TypeString, TEXT("{...}"), MapEntryId, MapRootId, InPropertyNameBuilder);
+
+				const void* PairPtr = Helper.GetPairPtr(It);
+
+				IteratePropertiesRecursive(MapProperty->KeyProp, PairPtr, MapProperty->KeyProp->GetName(), InFunction, InId, MapEntryId, InPropertyNameBuilder);
+
+				IteratePropertiesRecursive(MapProperty->ValueProp, PairPtr, MapProperty->ValueProp->GetName(), InFunction, InId, MapEntryId, InPropertyNameBuilder);
+			}
+		}
+		else if (const FSetProperty* SetProperty = CastField<FSetProperty>(InProperty))
+		{
+			FScriptSetHelper_InContainer Helper(SetProperty, InContainer);
+
+			int32 SetRootId = ++InId;
+			FString SizeString = FString::Printf(TEXT("{Num = %d}"), Helper.Num());
+
+			// Build property name id
+			FScopedPropertyName ScopedPropertyName(InPropertyNameBuilder);
+			ScopedPropertyName.Builder.AppendChar(TEXT('.'));
+			ScopedPropertyName.Builder.Append(InProperty->GetName());
+
+			// Trace set property
+			InFunction(InProperty->GetCPPType(), SizeString, SetRootId, InParentId, InPropertyNameBuilder);
+
+			// Trace container items
+			for (FScriptSetHelper::FIterator It(Helper); It; ++It)
+			{
+				const void* ValuePtr = Helper.GetElementPtr(It);
+				FString KeyString = FString::Printf(TEXT("[%d]"), It.GetLogicalIndex());
+
+				IteratePropertiesRecursive(SetProperty->ElementProp, ValuePtr, KeyString, InFunction, InId, SetRootId, InPropertyNameBuilder);
+			}
+		}
+		else if (const FStructProperty* StructProperty = CastField<FStructProperty>(InProperty))
+		{
+			int32 StructRootId = ++InId;
+
+			// Build property name id
+			FScopedPropertyName ScopedPropertyName(InPropertyNameBuilder);
+			ScopedPropertyName.Builder.AppendChar(TEXT('.'));
+			ScopedPropertyName.Builder.Append(InProperty->GetName());
+
+			// Trace struct property
+			InFunction(InProperty->GetCPPType(), TEXT("{...}"), StructRootId, InParentId, InPropertyNameBuilder);
+			
+			// Recurse
+			const void* StructContainer = StructProperty->ContainerPtrToValuePtr<const void>(InContainer);
+			for(TFieldIterator<FProperty> It(StructProperty->Struct); It; ++It)
+			{
+				IteratePropertiesRecursive(*It, StructContainer, It->GetName(), InFunction, InId, StructRootId, InPropertyNameBuilder);
+			}
+		}
+		else
+		{
+			// Normal property
+			int32 PropertyParentId = InParentId;
+			if(InProperty->ArrayDim > 1)
+			{
+				// Handle static array header
+				PropertyParentId = ++InId;
+				FString SizeString = FString::Printf(TEXT("{Num = %d}"), InProperty->ArrayDim);
+
+				// Build property name id
+				FScopedPropertyName ScopedPropertyName(InPropertyNameBuilder);
+				ScopedPropertyName.Builder.AppendChar(TEXT('.'));
+				ScopedPropertyName.Builder.Append(InProperty->GetName());
+
+				// Trace property
+				InFunction(InProperty->GetCPPType(), SizeString, PropertyParentId, InParentId,  InPropertyNameBuilder);
+			}
+
+			for (int32 StaticIndex = 0; StaticIndex != InProperty->ArrayDim; ++StaticIndex)
+			{
+				const void* ValuePtr = InProperty->ContainerPtrToValuePtr<const void>(InContainer, StaticIndex);
+
+				FString KeyString = InProperty->ArrayDim == 1 ? InKey : FString::Printf(TEXT("[%d]"), StaticIndex);
+
+				// Build property name id
+				FScopedPropertyName ScopedPropertyName(InPropertyNameBuilder);
+				ScopedPropertyName.Builder.AppendChar(TEXT('.'));
+				ScopedPropertyName.Builder.Append(KeyString);
+
+				// Trace property
+				FString ValueString;
+				InProperty->ExportText_Direct(ValueString, ValuePtr, ValuePtr, nullptr, PPF_None);
+				InFunction(InProperty->GetCPPType(), ValueString, ++InId, PropertyParentId,  InPropertyNameBuilder);
+			}
+		}
+	}
+
+	static void IterateProperties(UStruct* InStruct, const void* InContainer, IterateFunction InFunction)
+	{
+		FStringBuilderBase ParentNameBuilder;
+		
+		int32 Id = INDEX_NONE;
+		for(TFieldIterator<FProperty> It(InStruct); It; ++It)
+		{
+			IteratePropertiesRecursive(*It, InContainer, It->GetName(), InFunction, Id, INDEX_NONE, ParentNameBuilder);
+		}
+	}
+
+	static void TraceObjects()
+	{
+		TRACE_CPUPROFILER_EVENT_SCOPE(FObjectPropertyTrace::TraceObjects);
+		for(const TWeakObjectPtr<const UObject>& WeakObject : ObjectPropertyTrace::Objects)
+		{
+			if(const UObject* TracedObject = WeakObject.Get())
+			{
+				uint64 StartCycle = FPlatformTime::Cycles64();
+				uint64 ObjectId = FObjectTrace::GetObjectId(TracedObject);
+
+				UE_TRACE_LOG(Object, PropertiesStart2, ObjectProperties)
+					<< PropertiesStart2.Cycle(StartCycle)
+					<< PropertiesStart2.ObjectId(ObjectId)
+					<< PropertiesStart2.RecordingTime(FObjectTrace::GetWorldElapsedTime(TracedObject->GetWorld()));
+
+				uint64 ClassId = FObjectTrace::GetObjectId(TracedObject->GetClass());
+
+				IterateProperties(TracedObject->GetClass(), TracedObject, [StartCycle, ClassId, ObjectId](const FString& InTypeString, const FString& InValue, int32 InId, int32 InParentId, FStringBuilderBase & InPropertyNameBuilder)
+				{
+					const FString NameString = InPropertyNameBuilder.ToString();
+
+					// Compute parent id from property id 
+					int32 ParentEndIndex;
+					NameString.FindLastChar(TEXT('.'), ParentEndIndex);
+					const FString ParentString = NameString.Left(ParentEndIndex);
+					
+					const uint32 TypeId = TraceStringId(InTypeString);
+					const uint32 NameId = TraceStringId(NameString);
+					const uint32 ParentId = TraceStringId(ParentString);
+	
+					UE_TRACE_LOG(Object, PropertyValue3, ObjectProperties)
+						<< PropertyValue3.Cycle(StartCycle)
+						<< PropertyValue3.ObjectId(ObjectId)
+						<< PropertyValue3.Value(*InValue)
+						<< PropertyValue3.ParentId(InParentId)
+						<< PropertyValue3.TypeId(TypeId)
+						<< PropertyValue3.NameId(NameId)
+						<< PropertyValue3.ParentNameId(ParentId);
+				});
+
+				uint64 EndCycle = FPlatformTime::Cycles64();
+
+				UE_TRACE_LOG(Object, PropertiesEnd, ObjectProperties)
+					<< PropertiesEnd.Cycle(EndCycle)
+					<< PropertiesEnd.ObjectId(ObjectId);
+			}
+		}
+	}
+}
+
+void FObjectPropertyTrace::Init()
+{
+	check(!ObjectPropertyTrace::TickerHandle.IsValid());
+	ObjectPropertyTrace::TickerHandle = FTSTicker::GetCoreTicker().AddTicker(TEXT("ObjectPropertyTrace"), 0.0f, [](float InDelta)
+	{
+		if(UE_TRACE_CHANNELEXPR_IS_ENABLED(ObjectProperties))
+		{
+			ObjectPropertyTrace::TraceObjects();
+		}
+
+		return true;
+	});
+}
+
+void FObjectPropertyTrace::Destroy()
+{
+	check(ObjectPropertyTrace::TickerHandle.IsValid());
+	FTSTicker::GetCoreTicker().RemoveTicker(ObjectPropertyTrace::TickerHandle);
+}
+
+bool FObjectPropertyTrace::IsEnabled()
+{
+	return UE_TRACE_CHANNELEXPR_IS_ENABLED(ObjectProperties);
+}
+
+void FObjectPropertyTrace::ToggleObjectRegistration(const UObject* InObject)
+{
+	if(IsObjectRegistered(InObject))
+	{
+		UnregisterObject(InObject);
+	}
+	else
+	{
+		RegisterObject(InObject);
+	}
+}
+
+void FObjectPropertyTrace::RegisterObject(const UObject* InObject)
+{
+	ObjectPropertyTrace::Objects.AddUnique(InObject);
+}
+
+void FObjectPropertyTrace::UnregisterObject(const UObject* InObject)
+{
+	ObjectPropertyTrace::Objects.Remove(InObject);
+}
+
+bool FObjectPropertyTrace::IsObjectRegistered(const UObject* InObject)
+{
+	return ObjectPropertyTrace::Objects.Contains(InObject);
+}
+
+#endif
+
+#endif
\ No newline at end of file
diff --git a/Engine/Source/Runtime/Engine/Private/UnrealEngine.cpp b/Engine/Source/Runtime/Engine/Private/UnrealEngine.cpp
--- a/Engine/Source/Runtime/Engine/Private/UnrealEngine.cpp
+++ b/Engine/Source/Runtime/Engine/Private/UnrealEngine.cpp
@@ -237,6 +237,7 @@ UnrealEngine.cpp: Implements the UEngine class and helpers.
 
 #include "Particles/ParticleSystemManager.h"
 #include "ObjectTrace.h"
+#include "ObjectPropertyTrace.h"
 #include "StudioAnalytics.h"
 #include "Animation/SkinWeightProfileManager.h"
 
@@ -315,6 +316,10 @@ void FEngineModule::StartupModule()
 	FObjectTrace::Init();
 #endif
 
+#if OBJECT_PROPERTY_TRACE_ENABLED
+	FObjectPropertyTrace::Init();
+#endif
+
 #if TRACE_FILTERING_ENABLED
 	FTraceFilter::Init();
 #endif
@@ -361,6 +366,9 @@ void FEngineModule::ShutdownModule()
 #if OBJECT_TRACE_ENABLED
 	FObjectTrace::Destroy();
 #endif
+#if OBJECT_PROPERTY_TRACE_ENABLED
+	FObjectPropertyTrace::Destroy();
+#endif
 
 	FParticleSystemWorldManager::OnShutdown();
 
diff --git a/Engine/Source/Runtime/Engine/Public/ObjectPropertyTrace.h b/Engine/Source/Runtime/Engine/Public/ObjectPropertyTrace.h
new file mode 100644
--- /dev/null
+++ b/Engine/Source/Runtime/Engine/Public/ObjectPropertyTrace.h
@@ -0,0 +1,53 @@
+// Copyright Epic Games, Inc. All Rights Reserved.
+
+#pragma once
+
+#include "CoreTypes.h"
+#include "CoreMinimal.h"
+
+#if WITH_ENGINE
+#include "ObjectTrace.h"
+#endif
+
+#if WITH_ENGINE
+#define OBJECT_PROPERTY_TRACE_ENABLED OBJECT_TRACE_ENABLED
+#else
+#define OBJECT_PROPERTY_TRACE_ENABLED 0
+#endif
+
+#if OBJECT_PROPERTY_TRACE_ENABLED
+
+struct ENGINE_API FObjectPropertyTrace
+{
+	/** Initialize object property tracing */
+	static void Init();
+
+	/** Shut down object property tracing */
+	static void Destroy();
+
+	/** Check whether object property tracing is enabled */
+	static bool IsEnabled();
+
+	/** Toggle registration for a UObject being traced by the system */
+	static void ToggleObjectRegistration(const UObject* InObject);
+
+	/** Register a UObject to be traced by the system */
+	static void RegisterObject(const UObject* InObject);
+
+	/** Unregister a UObject to be traced by the system */
+	static void UnregisterObject(const UObject* InObject);
+
+	/** Check whether an object is registered */
+	static bool IsObjectRegistered(const UObject* InObject);
+};
+
+
+#define TRACE_OBJECT_PROPERTIES_BEGIN(Object) FObjectPropertyTrace::RegisterObject(Object);
+#define TRACE_OBJECT_PROPERTIES_END(Object) FObjectPropertyTrace::UnregisterObject(Object);
+
+#else
+
+#define TRACE_OBJECT_PROPERTIES_BEGIN(Object)
+#define TRACE_OBJECT_PROPERTIES_END(Object)
+
+#endif
