Skeletal Mesh: Don't remove unused materials when recreating import data

#rnx

diff --git a/Engine/Source/Runtime/Engine/Private/Rendering/SkeletalMeshLODImporterData.cpp b/Engine/Source/Runtime/Engine/Private/Rendering/SkeletalMeshLODImporterData.cpp
--- a/Engine/Source/Runtime/Engine/Private/Rendering/SkeletalMeshLODImporterData.cpp
+++ b/Engine/Source/Runtime/Engine/Private/Rendering/SkeletalMeshLODImporterData.cpp
@@ -2957,8 +2957,6 @@ FSkeletalMeshImportData FSkeletalMeshImportData::CreateFromMeshDescription(const
 		}
 	}
 	
-	SkelMeshImportData.CleanUpUnusedMaterials();
-	
 	// Copy any non-reserved float vertex attributes.
 	InMeshDescription.VertexAttributes().ForEachByType<float>(FCreateAndCopyAttributeValues<float>(SkelMeshImportData, InMeshDescription.Vertices()));
 	InMeshDescription.VertexAttributes().ForEachByType<FVector2f>(FCreateAndCopyAttributeValues<FVector2f>(SkelMeshImportData, InMeshDescription.Vertices()));
