In AddModulationRouting, force routing method to Union

diff --git a/Engine/Source/Runtime/Engine/Classes/Components/AudioComponent.h b/Engine/Source/Runtime/Engine/Classes/Components/AudioComponent.h
--- a/Engine/Source/Runtime/Engine/Classes/Components/AudioComponent.h
+++ b/Engine/Source/Runtime/Engine/Classes/Components/AudioComponent.h
@@ -731,6 +731,7 @@ public:
 
 	/**
 	* Adds the given set of Modulators to the modulators currently set on the Audio Component. This performs a logical Union, so duplicate modulators will not be added.
+	* This function automatically sets the destination's routing method to Union, overriding whatever was set before.
 	* To completely overwrite existing Modulation settings, use SetModulationRouting.
 	* @param Modulators The set of modulators to add to the given destination on the component.
 	* @param Destination The destination to add the modulators to.
diff --git a/Engine/Source/Runtime/Engine/Private/ActiveSound.cpp b/Engine/Source/Runtime/Engine/Private/ActiveSound.cpp
--- a/Engine/Source/Runtime/Engine/Private/ActiveSound.cpp
+++ b/Engine/Source/Runtime/Engine/Private/ActiveSound.cpp
@@ -472,15 +472,19 @@ void FActiveSound::AddModulationRouting(const TSet<TObjectPtr<USoundModulatorBas
 	switch (Destination)
 	{
 		case EModulationDestination::Volume:
+			ModulationRouting.VolumeRouting = EModulationRouting::Union;
 			ModulationRouting.VolumeModulationDestination.Modulators.Append(NewModulators);
 			break;
 		case EModulationDestination::Pitch:
+			ModulationRouting.PitchRouting = EModulationRouting::Union;
 			ModulationRouting.PitchModulationDestination.Modulators.Append(NewModulators);
 			break;
 		case EModulationDestination::Lowpass:
+			ModulationRouting.LowpassRouting = EModulationRouting::Union;
 			ModulationRouting.LowpassModulationDestination.Modulators.Append(NewModulators);
 			break;
 		case EModulationDestination::Highpass:
+			ModulationRouting.HighpassRouting = EModulationRouting::Union;
 			ModulationRouting.HighpassModulationDestination.Modulators.Append(NewModulators);
 			break;
 		default:
diff --git a/Engine/Source/Runtime/Engine/Public/ActiveSound.h b/Engine/Source/Runtime/Engine/Public/ActiveSound.h
--- a/Engine/Source/Runtime/Engine/Public/ActiveSound.h
+++ b/Engine/Source/Runtime/Engine/Public/ActiveSound.h
@@ -712,7 +712,10 @@ public:
 	/* Gives new Modulation Routing settings to the ActiveSound. This overwrites all previous Modulation Routing settings. */
 	ENGINE_API void SetNewModulationRouting(const FSoundModulationDefaultRoutingSettings& NewRouting);
 
-	/* Adds additional Modulators to the ActiveSound, if possible. To replace existing modulators, use SetNewModulationRouting. */
+	/* Adds additional Modulators to the ActiveSound, if possible. 
+	 * This function forces the Destination's routing method to Union.
+	 * To replace existing modulators or set a different routing method, use SetNewModulationRouting.
+	 */
 	ENGINE_API void AddModulationRouting(const TSet<TObjectPtr<USoundModulatorBase>>& NewModulators, EModulationDestination Destination);
 
 	/* Removes given Modulators from the ActiveSound, if possible. */
