Fix skeletal mesh bone transforms for Ragdolls when the root bone has a scale and the root body is not simulated.

diff --git a/Engine/Source/Runtime/Engine/Private/PhysicsEngine/PhysAnim.cpp b/Engine/Source/Runtime/Engine/Private/PhysicsEngine/PhysAnim.cpp
--- a/Engine/Source/Runtime/Engine/Private/PhysicsEngine/PhysAnim.cpp
+++ b/Engine/Source/Runtime/Engine/Private/PhysicsEngine/PhysAnim.cpp
@@ -148,9 +148,6 @@ void USkeletalMeshComponent::PerformBlendPhysicsBones(
 	TArray<FTransform>& InOutBoneSpaceTransforms)
 {
 	SCOPE_CYCLE_COUNTER(STAT_BlendInPhysics);
-	// Get drawscale from Owner (if there is one)
-	FVector TotalScale3D = GetComponentTransform().GetScale3D();
-	FVector RecipScale3D = TotalScale3D.Reciprocal();
 
 	UPhysicsAsset * const PhysicsAsset = GetPhysicsAsset();
 	check( PhysicsAsset );
@@ -160,6 +157,11 @@ void USkeletalMeshComponent::PerformBlendPhysicsBones(
 		return;
 	}
 
+	if (InOutBoneSpaceTransforms.Num() == 0)
+	{
+		return;
+	}
+
 	// Get the scene, and do nothing if we can't get one.
 	FPhysScene* PhysScene = nullptr;
 	if (GetWorld() != nullptr)
@@ -184,7 +186,6 @@ void USkeletalMeshComponent::PerformBlendPhysicsBones(
 	// required bone, leaving any others unchanged.
 	FPhysicsCommand::ExecuteRead(this, [&]()
 	{
-		bool bSetParentScale = false;
 		// Note that IsInstanceSimulatingPhysics returns false for kinematic bodies, so
 		// bSimulatedRootBody means "is the root physics body dynamic" (not the same as the root bone)
 		const bool bSimulatedRootBody = Bodies.IsValidIndex(RootBodyData.BodyIndex) && 
@@ -194,6 +195,18 @@ void USkeletalMeshComponent::PerformBlendPhysicsBones(
 		// set to SimulationUpatesComponentTransform then this will come from the simulation.
 		const FTransform NewComponentTransform = GetComponentTransformFromBodyInstance(Bodies[RootBodyData.BodyIndex]);
 
+		// Get component scale. If we have a body on the root bone, this will be the scale applied 
+		// to its relative physics position.
+		FVector TotalScale3D = GetComponentTransform().GetScale3D();
+		FVector RecipScale3D = TotalScale3D.Reciprocal();
+
+		// For bodies that are not at the root bone, we will also need to apply the inverse of
+		// the root bone scale to the body positions reported by physics. We will update the
+		// scale below when we have processed the body on the root bone (if there is one)
+		const FVector RootBoneScale = InOutBoneSpaceTransforms[0].GetScale3D();
+		const bool bUnitScaledRootBone = RootBoneScale.Equals(FVector(1, 1, 1), UE_SMALL_NUMBER);
+		bool bHaveSetParentScale = bUnitScaledRootBone;	// We don't need to change the scale if root has no scale
+
 		// For each bone:
 		// * Update the WorldBoneTMs entry
 		// * Update InOutBoneSpaceTransforms, using the physics blend weights
@@ -242,7 +255,10 @@ void USkeletalMeshComponent::PerformBlendPhysicsBones(
 					PhysicsAssetBodyInstance->IsValidBodyInstance();
 
 				// Only process this bone if it is simulated, or required due to the bDriveMeshWhenKinematic flag
-				if (PhysicsAssetBodyInstance->IsInstanceSimulatingPhysics() || bDriveMeshWhenKinematic)
+				// Also, if the root bone is scaled, we need to update WorldBoneTMs, even for Kinematics that are 
+				// not blended into the output because they may be a parent body of a dynamic body
+				const bool bIsSimulatingPhysics = PhysicsAssetBodyInstance->IsInstanceSimulatingPhysics();
+				if (bIsSimulatingPhysics || bDriveMeshWhenKinematic || !bUnitScaledRootBone)
 				{
 					FTransform PhysTM = PhysicsAssetBodyInstance->GetUnrealWorldTransform_AssumesLocked();
 
@@ -255,6 +271,23 @@ void USkeletalMeshComponent::PerformBlendPhysicsBones(
 						continue;
 					}
 
+					// NOTE: if the root bone is scaled we need to update the WorldBoneTMs above for both
+					// kinematic and dynamic bodies so that the parent pose calculation below is correct. 
+					// For kinematics that are not blended into the output pose, we are done here
+					if (!bIsSimulatingPhysics && !bDriveMeshWhenKinematic)
+					{
+						continue;
+					}
+
+					// Bodies that are not on the root bone must have the position reported by physics
+					// inverse scaled by the root bone scale.
+					if ((BoneIndex > 0) && !bHaveSetParentScale)
+					{
+						TotalScale3D *= RootBoneScale;
+						RecipScale3D = TotalScale3D.Reciprocal();
+						bHaveSetParentScale = true;
+					}
+
 					// Note that bBlendPhysics is a flag that is used to force use of the physics
 					// body, irrespective of whether they are simulated or not, or the value of the
 					// physics blend weight.
@@ -299,10 +332,9 @@ void USkeletalMeshComponent::PerformBlendPhysicsBones(
 							ParentWorldTM = WorldBoneTMs[ParentIndex].TM;
 						}
 
-
 						// Calculate the relative transform of the current and parent (physical) bones.
 						FTransform RelTM = PhysTM.GetRelativeTransform(ParentWorldTM);
-						RelTM.RemoveScaling();
+						RelTM.RemoveScaling();	// NOTE: this also normalizes the rotation
 						FQuat RelRot(RelTM.GetRotation());
 						FVector RelPos =  RecipScale3D * RelTM.GetLocation();
 						FTransform PhysicalBoneSpaceTransform = FTransform(
@@ -311,15 +343,6 @@ void USkeletalMeshComponent::PerformBlendPhysicsBones(
 						// Now blend in this atom. See if we are forcing this bone to always be blended in
 						InOutBoneSpaceTransforms[BoneIndex].Blend(
 							InOutBoneSpaceTransforms[BoneIndex], PhysicalBoneSpaceTransform, PhysicsBlendWeight);
-
-						if (!bSetParentScale)
-						{
-							// We must update RecipScale3D based on the scale of the root
-							TotalScale3D *= InOutBoneSpaceTransforms[0].GetScale3D();
-							RecipScale3D = TotalScale3D.Reciprocal();
-							bSetParentScale = true;
-						}
-
 					}
 				}
 			}
