Define FORWARD_SHADING for mobile renderer

diff --git a/Engine/Source/Runtime/Engine/Private/ShaderCompiler/ShaderCompiler.cpp b/Engine/Source/Runtime/Engine/Private/ShaderCompiler/ShaderCompiler.cpp
--- a/Engine/Source/Runtime/Engine/Private/ShaderCompiler/ShaderCompiler.cpp
+++ b/Engine/Source/Runtime/Engine/Private/ShaderCompiler/ShaderCompiler.cpp
@@ -8211,7 +8211,11 @@ void GlobalBeginCompileShader(
 
 	bool bForwardShading = false;
 	{
-		if (TargetPlatform)
+		if (bIsMobilePlatform)
+		{
+			bForwardShading = !IsMobileDeferredShadingEnabled((EShaderPlatform)Target.Platform);
+		}
+		else if (TargetPlatform)
 		{
 			bForwardShading = TargetPlatform->UsesForwardShading();
 		}
diff --git a/Engine/Source/Runtime/Engine/Private/ShaderCompiler/ShaderGenerationUtil.cpp b/Engine/Source/Runtime/Engine/Private/ShaderCompiler/ShaderGenerationUtil.cpp
--- a/Engine/Source/Runtime/Engine/Private/ShaderCompiler/ShaderGenerationUtil.cpp
+++ b/Engine/Source/Runtime/Engine/Private/ShaderCompiler/ShaderGenerationUtil.cpp
@@ -372,6 +372,11 @@ static FShaderGlobalDefines FetchShaderGlobalDefines(EShaderPlatform TargetPlatf
 	}
 #endif
 
+	if (bIsMobilePlatform)
+	{
+		Ret.FORWARD_SHADING = !IsMobileDeferredShadingEnabled((EShaderPlatform)TargetPlatform);
+	}
+
 	{
 		static IConsoleVariable* CVar = IConsoleManager::Get().FindConsoleVariable(TEXT("r.VertexFoggingForOpaque"));
 		Ret.PROJECT_VERTEX_FOGGING_FOR_OPAQUE = Ret.FORWARD_SHADING && (CVar ? (CVar->GetInt() != 0) : 0);
