Override material ShadingModel if it was filtered out by platform settings

diff --git a/Engine/Source/Runtime/Engine/Private/Materials/Material.cpp b/Engine/Source/Runtime/Engine/Private/Materials/Material.cpp
--- a/Engine/Source/Runtime/Engine/Private/Materials/Material.cpp
+++ b/Engine/Source/Runtime/Engine/Private/Materials/Material.cpp
@@ -3972,6 +3972,12 @@ void UMaterial::PostLoad()
 	{
 		// Filter out ShadingModels field to a current platform settings
 		FilterOutPlatformShadingModels(GMaxRHIShaderPlatform, ShadingModels);
+		// Override material shader model if it was filtered out
+		if (!ShadingModels.HasShadingModel(ShadingModel) &&
+			(ShadingModel != MSM_FromMaterialExpression || !AllowPerPixelShadingModels(GMaxRHIShaderPlatform)))
+		{
+			ShadingModel = ShadingModels.GetFirstShadingModel();
+		}
 	}
 
 #if WITH_EDITOR
