Add IsStaticLightingAllowed check to ASkeletalMeshActor to prevent unnecessary warnings when static lighting is not enabled in the project

diff --git a/Engine/Source/Runtime/Engine/Private/Animation/SkeletalMeshActor.cpp b/Engine/Source/Runtime/Engine/Private/Animation/SkeletalMeshActor.cpp
--- a/Engine/Source/Runtime/Engine/Private/Animation/SkeletalMeshActor.cpp
+++ b/Engine/Source/Runtime/Engine/Private/Animation/SkeletalMeshActor.cpp
@@ -9,6 +9,7 @@
 #include "Misc/MapErrors.h"
 #include "Engine/CollisionProfile.h"
 #include "Logging/MessageLog.h"
+#include "RenderUtils.h"
 
 #include UE_INLINE_GENERATED_CPP_BY_NAME(SkeletalMeshActor)
 
@@ -50,7 +51,8 @@ void ASkeletalMeshActor::CheckForErrors()
 		if (!PhysicsAsset)
 		{
 			if (SkeletalMeshComponent->CastShadow
-				&& SkeletalMeshComponent->bCastDynamicShadow)
+				&& SkeletalMeshComponent->bCastDynamicShadow
+				&& IsStaticLightingAllowed())
 			{
 				FFormatNamedArguments Arguments;
 				Arguments.Add(TEXT("DetailedInfo"), FText::FromString(GetDetailedInfo()));
@@ -64,7 +66,8 @@ void ASkeletalMeshActor::CheckForErrors()
 		if (SkeletalMeshComponent->CastShadow
 			&& SkeletalMeshComponent->bCastDynamicShadow
 			&& SkeletalMeshComponent->IsRegistered()
-			&& SkeletalMeshComponent->Bounds.SphereRadius > 2000.0f)
+			&& SkeletalMeshComponent->Bounds.SphereRadius > 2000.0f
+			&& IsStaticLightingAllowed())
 		{
 			FFormatNamedArguments Arguments;
 			Arguments.Add(TEXT("DetailedInfo"), FText::FromString(GetDetailedInfo()));
