[World Partition] [Runtime Spatial Hash]
Handle the problematic case where the world bounds is valid but empty,

diff --git a/Engine/Source/Runtime/Engine/Private/WorldPartition/RuntimeSpatialHash/RuntimeSpatialHashGridHelper.cpp b/Engine/Source/Runtime/Engine/Private/WorldPartition/RuntimeSpatialHash/RuntimeSpatialHashGridHelper.cpp
--- a/Engine/Source/Runtime/Engine/Private/WorldPartition/RuntimeSpatialHash/RuntimeSpatialHashGridHelper.cpp
+++ b/Engine/Source/Runtime/Engine/Private/WorldPartition/RuntimeSpatialHash/RuntimeSpatialHashGridHelper.cpp
@@ -61,10 +61,12 @@ FSquare2DGridHelper::FSquare2DGridHelper(const FBox& InWorldBounds, const FVecto
 			{
 				GridSize = FMath::Pow(2, FMath::CeilToDouble(FMath::Log2(static_cast<double>(GridSize))));
 			}
+			GridSize = FMath::Max(GridSize, 1);
 			GridLevelCount = FMath::FloorLog2_64(GridSize) + 1;
 		}
 	}
 
+	check(GridSize);
 	check(FMath::IsPowerOfTwo(GridSize));
 
 	Levels.Reserve(GridLevelCount);
