World Partition - HLOD: Avoid map check warnings related to collisions
* Collisions (ECC_Visibility + ECC_Camera) can be enabled on HLODs in the editor. Before enabling, test if the mesh actually has what's needed to ensure this is possible (sections with collision data or dedicated collision primitives)

diff --git a/Engine/Source/Runtime/Engine/Private/WorldPartition/HLOD/HLODActor.cpp b/Engine/Source/Runtime/Engine/Private/WorldPartition/HLOD/HLODActor.cpp
--- a/Engine/Source/Runtime/Engine/Private/WorldPartition/HLOD/HLODActor.cpp
+++ b/Engine/Source/Runtime/Engine/Private/WorldPartition/HLOD/HLODActor.cpp
@@ -315,10 +315,22 @@ void AWorldPartitionHLOD::PreRegisterAllComponents()
 			SetActorEnableCollision(bShouldEnableCollision);
 			ForEachComponent<UPrimitiveComponent>(false, [bShouldEnableCollision](UPrimitiveComponent* PrimitiveComponent)
 			{
-				PrimitiveComponent->SetCollisionEnabled(bShouldEnableCollision ? ECollisionEnabled::QueryOnly : ECollisionEnabled::NoCollision);
+				bool bShouldEnableCollisionForComponent = bShouldEnableCollision;
+				if (bShouldEnableCollisionForComponent)
+				{
+					if (UStaticMeshComponent* StaticMeshComponent = Cast<UStaticMeshComponent>(PrimitiveComponent))
+					{
+						UStaticMesh* StaticMesh = StaticMeshComponent->GetStaticMesh();
+						int32 NumSectionsWithCollision = StaticMesh ? StaticMesh->GetNumSectionsWithCollision() : 0;
+						int32 NumCollisionPrims = StaticMesh ? (StaticMesh->GetBodySetup() ? StaticMesh->GetBodySetup()->AggGeom.GetElementCount() : 0) : 0;
+						bShouldEnableCollisionForComponent = NumSectionsWithCollision != 0 || NumCollisionPrims != 0;
+					}
+				}
+
+				PrimitiveComponent->SetCollisionEnabled(bShouldEnableCollisionForComponent ? ECollisionEnabled::QueryOnly : ECollisionEnabled::NoCollision);
 				PrimitiveComponent->SetCollisionResponseToAllChannels(ECollisionResponse::ECR_Ignore);
-				PrimitiveComponent->SetCollisionResponseToChannel(ECC_Visibility, bShouldEnableCollision ? ECR_Block : ECR_Ignore);
-				PrimitiveComponent->SetCollisionResponseToChannel(ECC_Camera, bShouldEnableCollision ? ECR_Block : ECR_Ignore);
+				PrimitiveComponent->SetCollisionResponseToChannel(ECC_Visibility, bShouldEnableCollisionForComponent ? ECR_Block : ECR_Ignore);
+				PrimitiveComponent->SetCollisionResponseToChannel(ECC_Camera, bShouldEnableCollisionForComponent ? ECR_Block : ECR_Ignore);
 			});
 		}
 	}	
