[PCG] Fix "not initialized properly even though its struct probably has a custom default constructor" warning in FSplineMeshComponentDescriptorBase

#rnx

diff --git a/Engine/Source/Runtime/Engine/Classes/Engine/SplineMeshComponentDescriptor.h b/Engine/Source/Runtime/Engine/Classes/Engine/SplineMeshComponentDescriptor.h
--- a/Engine/Source/Runtime/Engine/Classes/Engine/SplineMeshComponentDescriptor.h
+++ b/Engine/Source/Runtime/Engine/Classes/Engine/SplineMeshComponentDescriptor.h
@@ -140,9 +140,6 @@ public:
 	UPROPERTY(EditAnywhere, Category = "Component Settings")
 	uint8 bUseAsOccluder : 1;
 
-	UPROPERTY(EditAnywhere, Category = "Component Settings")
-	uint8 bEnableDensityScaling : 1;
-
 	UPROPERTY(EditAnywhere, Category = "Component Settings")
 	uint8 bEnableDiscardOnLoad : 1;
 
