PerQualityLevelProperties: Cooking: Fix bug in previous changelist that was truncating floats in FPerQualityLevelFloat. FPerQualityLevelProperty::Default is a ValueType, not an int.
#rnx

diff --git a/Engine/Source/Runtime/Engine/Private/PerQualityLevelProperties.cpp b/Engine/Source/Runtime/Engine/Private/PerQualityLevelProperties.cpp
--- a/Engine/Source/Runtime/Engine/Private/PerQualityLevelProperties.cpp
+++ b/Engine/Source/Runtime/Engine/Private/PerQualityLevelProperties.cpp
@@ -174,7 +174,7 @@ namespace QualityLevelProperty
 	template <typename ValueType>
 	struct FSavedData
 	{
-		int32 Default = 0;
+		ValueType Default = 0;
 		TMap<int32, ValueType> PerQuality;
 	};
 	template<typename StructType, typename ValueType>
