World Partition - Content Bundle: Test that each CB has a unique GUID in the map check validation

diff --git a/Engine/Source/Editor/UnrealEd/Private/EditorServer.cpp b/Engine/Source/Editor/UnrealEd/Private/EditorServer.cpp
--- a/Engine/Source/Editor/UnrealEd/Private/EditorServer.cpp
+++ b/Engine/Source/Editor/UnrealEd/Private/EditorServer.cpp
@@ -124,6 +124,7 @@
 #include "MovieSceneCaptureModule.h"
 
 #include "WorldPartition/WorldPartition.h"
+#include "WorldPartition/ContentBundle/ContentBundleWorldSubsystem.h"
 #include "WorldPartition/ErrorHandling/WorldPartitionStreamingGenerationMapCheckErrorHandler.h"
 
 #include "Kismet2/KismetEditorUtilities.h"
@@ -3933,6 +3934,11 @@ bool UEditorEngine::Map_Check( UWorld* InWorld, const TCHAR* Str, FOutputDevice&
 		WorldPartition->CheckForErrors(&MapCheckErrorHandler);
 	}
 
+	if (InWorld->ContentBundleManager)
+	{
+		InWorld->ContentBundleManager->CheckForErrors();
+	}
+
 	GWarn->StatusUpdate( 0, ProgressDenominator, CheckMapLocText );
 
 	int32 LastUpdateCount = 0;
diff --git a/Engine/Source/Runtime/Engine/Private/WorldPartition/ContentBundle/ContentBundleContainer.cpp b/Engine/Source/Runtime/Engine/Private/WorldPartition/ContentBundle/ContentBundleContainer.cpp
--- a/Engine/Source/Runtime/Engine/Private/WorldPartition/ContentBundle/ContentBundleContainer.cpp
+++ b/Engine/Source/Runtime/Engine/Private/WorldPartition/ContentBundle/ContentBundleContainer.cpp
@@ -241,6 +241,11 @@ FContentBundleBase& FContentBundleContainer::InitializeContentBundle(TSharedPtr<
 #if WITH_EDITOR
 	if (UseEditorContentBundle())
 	{
+		if (TSharedPtr<FContentBundleEditor> CBE = GetEditorContentBundle(ContentBundleClient->GetDescriptor()->GetGuid()))
+		{
+			UE_LOG(LogContentBundle, Error, TEXT("Found duplicate content bundle GUIDs: %s for content bundles %s AND %s"), *ContentBundleClient->GetDescriptor()->GetGuid().ToString(), *CBE->GetDescriptor()->GetPackage()->GetPathName(), *ContentBundleClient->GetDescriptor()->GetPackage()->GetPathName());
+		}
+
 		ContentBundle = GetEditorContentBundles().Emplace_GetRef(MakeShared<FContentBundleEditor>(ContentBundleClient, GetInjectedWorld())).Get();
 	}
 	else
diff --git a/Engine/Source/Runtime/Engine/Private/WorldPartition/ContentBundle/ContentBundleWorldSubsystem.cpp b/Engine/Source/Runtime/Engine/Private/WorldPartition/ContentBundle/ContentBundleWorldSubsystem.cpp
--- a/Engine/Source/Runtime/Engine/Private/WorldPartition/ContentBundle/ContentBundleWorldSubsystem.cpp
+++ b/Engine/Source/Runtime/Engine/Private/WorldPartition/ContentBundle/ContentBundleWorldSubsystem.cpp
@@ -15,6 +15,8 @@
 #if WITH_EDITOR
 #include "WorldPartition/ContentBundle/ContentBundleEditor.h"
 #include "Editor.h"
+#include "Logging/MessageLog.h"
+#include "Misc/MapErrors.h"
 #else
 #include "Engine/Engine.h"
 #endif
@@ -171,6 +173,28 @@ TSharedPtr<FContentBundleEditor> UContentBundleManager::GetEditorContentBundle(c
 	return nullptr;
 }
 
+void UContentBundleManager::CheckForErrors() const
+{
+	TMap<FGuid, const TSharedPtr<FContentBundleEditor>> ContentBundlesEditor;
+
+	for (const TUniquePtr<FContentBundleContainer>& ContentBundleContainer : ContentBundleContainers)
+	{
+		for (const TSharedPtr<FContentBundleEditor>& ContentBundleEditor : ContentBundleContainer->GetEditorContentBundles())
+		{
+			const TSharedPtr<FContentBundleEditor>& ExistingValue = ContentBundlesEditor.FindOrAdd(ContentBundleEditor->GetDescriptor()->GetGuid(), ContentBundleEditor);
+			if (ExistingValue.Get() != ContentBundleEditor.Get())
+			{
+				const FText DuplicateGuidText = FText::FromString(ContentBundleEditor->GetDescriptor()->GetGuid().ToString());
+
+				FMessageLog("MapCheck").Error()
+					->AddToken(FTextToken::Create(FText::Format(NSLOCTEXT("ContentBundle", "ContentBundleErrorDuplicateGUIDs", "Found content bundles sharing the same GUID ({0}):"), DuplicateGuidText)))
+					->AddToken(FAssetNameToken::Create(ExistingValue->GetDescriptor()->GetPathName()))
+					->AddToken(FAssetNameToken::Create(ContentBundleEditor->GetDescriptor()->GetPathName()));
+			}
+		}
+	}
+}
+
 #endif
 
 uint32 UContentBundleManager::GetContentBundleContainerIndex(const UWorld* InjectedWorld) const
diff --git a/Engine/Source/Runtime/Engine/Public/WorldPartition/ContentBundle/ContentBundleWorldSubsystem.h b/Engine/Source/Runtime/Engine/Public/WorldPartition/ContentBundle/ContentBundleWorldSubsystem.h
--- a/Engine/Source/Runtime/Engine/Public/WorldPartition/ContentBundle/ContentBundleWorldSubsystem.h
+++ b/Engine/Source/Runtime/Engine/Public/WorldPartition/ContentBundle/ContentBundleWorldSubsystem.h
@@ -42,6 +42,8 @@ public:
 	ENGINE_API TSharedPtr<FContentBundleEditor> GetEditorContentBundle(const FGuid& ContentBundleGuid) const;
 
 	UContentBundleDuplicateForPIEHelper* GetPIEDuplicateHelper() const { return PIEDuplicateHelper; }
+
+	ENGINE_API void CheckForErrors() const;
 #endif
 
 	ENGINE_API const FContentBundleBase* GetContentBundle(const UWorld* InWorld, const FGuid& Guid) const;
