AlphaHoldout: Relax compile errors to exclude post-process materials.

#rnx

diff --git a/Engine/Source/Runtime/Engine/Private/Materials/HLSLMaterialTranslator.cpp b/Engine/Source/Runtime/Engine/Private/Materials/HLSLMaterialTranslator.cpp
--- a/Engine/Source/Runtime/Engine/Private/Materials/HLSLMaterialTranslator.cpp
+++ b/Engine/Source/Runtime/Engine/Private/Materials/HLSLMaterialTranslator.cpp
@@ -1821,7 +1821,7 @@ void FHLSLMaterialTranslator::DoTranslate()
 		}
 	}
 
-	if (IsAlphaHoldoutBlendMode(BlendMode))
+	if (IsAlphaHoldoutBlendMode(BlendMode) && !Material->IsPostProcessMaterial())
 	{
 		if (!MaterialShadingModels.IsUnlit())
 		{
