SkeletalMesh:
- Made the OnBoneTransformsFinalized delegates available to all SkinnedMeshComponent type classes.
- Removed deprecated OnBoneTransformsFinalized single delegate member.

Chaos Cloth Asset:
- Added OnBoneTransformsFinalized delegate implementations.

diff --git a/Engine/Plugins/ChaosClothAsset/Source/ChaosClothAssetEngine/Private/ChaosClothAsset/ClothComponent.cpp b/Engine/Plugins/ChaosClothAsset/Source/ChaosClothAssetEngine/Private/ChaosClothAsset/ClothComponent.cpp
--- a/Engine/Plugins/ChaosClothAsset/Source/ChaosClothAssetEngine/Private/ChaosClothAsset/ClothComponent.cpp
+++ b/Engine/Plugins/ChaosClothAsset/Source/ChaosClothAssetEngine/Private/ChaosClothAsset/ClothComponent.cpp
@@ -488,6 +488,23 @@ void UChaosClothComponent::GetAdditionalRequiredBonesForLeader(int32 LeaderLODIn
 	}
 }
 
+void UChaosClothComponent::FinalizeBoneTransform()
+{
+	Super::FinalizeBoneTransform();
+
+	OnBoneTransformsFinalizedMC.Broadcast();
+}
+
+FDelegateHandle UChaosClothComponent::RegisterOnBoneTransformsFinalizedDelegate(const FOnBoneTransformsFinalizedMultiCast::FDelegate& Delegate)
+{
+	return OnBoneTransformsFinalizedMC.Add(Delegate);
+}
+
+void UChaosClothComponent::UnregisterOnBoneTransformsFinalizedDelegate(const FDelegateHandle& DelegateHandle)
+{
+	OnBoneTransformsFinalizedMC.Remove(DelegateHandle);
+}
+
 TSharedPtr<UE::Chaos::ClothAsset::FClothSimulationProxy> UChaosClothComponent::CreateClothSimulationProxy()
 {
 	using namespace UE::Chaos::ClothAsset;
diff --git a/Engine/Plugins/ChaosClothAsset/Source/ChaosClothAssetEngine/Public/ChaosClothAsset/ClothComponent.h b/Engine/Plugins/ChaosClothAsset/Source/ChaosClothAssetEngine/Public/ChaosClothAsset/ClothComponent.h
--- a/Engine/Plugins/ChaosClothAsset/Source/ChaosClothAssetEngine/Public/ChaosClothAsset/ClothComponent.h
+++ b/Engine/Plugins/ChaosClothAsset/Source/ChaosClothAssetEngine/Public/ChaosClothAsset/ClothComponent.h
@@ -148,6 +148,9 @@ protected:
 	virtual void GetUpdateClothSimulationData_AnyThread(TMap<int32, FClothSimulData>& OutClothSimulData, FMatrix& OutLocalToWorld, float& OutBlendWeight) override;
 	virtual void SetSkinnedAssetAndUpdate(USkinnedAsset* InSkinnedAsset, bool bReinitPose = true) override;
 	virtual void GetAdditionalRequiredBonesForLeader(int32 LODIndex, TArray<FBoneIndexType>& InOutRequiredBones) const override;
+	virtual void FinalizeBoneTransform() override;
+	virtual FDelegateHandle RegisterOnBoneTransformsFinalizedDelegate(const FOnBoneTransformsFinalizedMultiCast::FDelegate& Delegate) override;
+	virtual void UnregisterOnBoneTransformsFinalizedDelegate(const FDelegateHandle& DelegateHandle) override;
 	//~ End USkinnedMeshComponent Interface
 
 	/** Override this function for setting up custom simulation proxies when the component is registered. */
@@ -216,4 +219,7 @@ private:
 	TArray<TSharedPtr<::Chaos::Softs::FCollectionPropertyFacade>> CollectionPropertyFacades;
 
 	TSharedPtr<UE::Chaos::ClothAsset::FClothSimulationProxy> ClothSimulationProxy;
+
+	/** Multicaster fired when this component bone transforms are finalized */
+	FOnBoneTransformsFinalizedMultiCast OnBoneTransformsFinalizedMC;
 };
diff --git a/Engine/Source/Runtime/Engine/Classes/Components/SkeletalMeshComponent.h b/Engine/Source/Runtime/Engine/Classes/Components/SkeletalMeshComponent.h
--- a/Engine/Source/Runtime/Engine/Classes/Components/SkeletalMeshComponent.h
+++ b/Engine/Source/Runtime/Engine/Classes/Components/SkeletalMeshComponent.h
@@ -55,10 +55,9 @@ DECLARE_DYNAMIC_MULTICAST_DELEGATE(FOnAnimInitialized);
 DECLARE_MULTICAST_DELEGATE(FOnSkelMeshTeleportedMultiCast);
 typedef FOnSkelMeshTeleportedMultiCast::FDelegate FOnSkelMeshTeleported;
 
+class UE_DEPRECATED(5.5, "use FOnBoneTransformsFinalizedMultiCast instead (see SkinnedMeshComponent.h).") FOnBoneTransformsFinalized;
 DECLARE_DYNAMIC_MULTICAST_DELEGATE(FOnBoneTransformsFinalized);  // Deprecated, use FOnBoneTransformsFinalizedMultiCast instead
 
-DECLARE_MULTICAST_DELEGATE(FOnBoneTransformsFinalizedMultiCast);
-
 DECLARE_TS_MULTICAST_DELEGATE_ThreeParams(FOnLODRequiredBonesUpdateMulticast, USkeletalMeshComponent*, int32, const TArray<FBoneIndexType>&);
 typedef FOnLODRequiredBonesUpdateMulticast::FDelegate FOnLODRequiredBonesUpdate;
 
@@ -1943,9 +1942,6 @@ public:
 	ENGINE_API virtual void ClearRefPoseOverride() override;
 	//~ End USkinnedMeshComponent Interface
 
-	UE_DEPRECATED(4.27, "Use RegisterOnBoneTransformsFinalizedDelegate/UnregisterOnBoneTransformsFinalizedDelegate instead")
-	FOnBoneTransformsFinalized OnBoneTransformsFinalized;
-
 	// Conditions used to gate when post process events happen
 	ENGINE_API bool ShouldUpdatePostProcessInstance() const;
 	ENGINE_API bool ShouldPostUpdatePostProcessInstance() const;
@@ -2586,8 +2582,8 @@ public:
 	ENGINE_API void UnregisterOnTeleportDelegate(const FDelegateHandle& DelegateHandle);
 
 	/** Register/Unregister for OnBoneTransformsFinalized callback */
-	ENGINE_API FDelegateHandle RegisterOnBoneTransformsFinalizedDelegate(const FOnBoneTransformsFinalizedMultiCast::FDelegate& Delegate);
-	ENGINE_API void UnregisterOnBoneTransformsFinalizedDelegate(const FDelegateHandle& DelegateHandle);
+	ENGINE_API virtual FDelegateHandle RegisterOnBoneTransformsFinalizedDelegate(const FOnBoneTransformsFinalizedMultiCast::FDelegate& Delegate) override;
+	ENGINE_API virtual void UnregisterOnBoneTransformsFinalizedDelegate(const FDelegateHandle& DelegateHandle) override;
 
 	/** Register/Unregister for OnLODRequiredBonesUpdate callback */
 	ENGINE_API FDelegateHandle RegisterOnLODRequiredBonesUpdate(const FOnLODRequiredBonesUpdate& Delegate);
diff --git a/Engine/Source/Runtime/Engine/Classes/Components/SkinnedMeshComponent.h b/Engine/Source/Runtime/Engine/Classes/Components/SkinnedMeshComponent.h
--- a/Engine/Source/Runtime/Engine/Classes/Components/SkinnedMeshComponent.h
+++ b/Engine/Source/Runtime/Engine/Classes/Components/SkinnedMeshComponent.h
@@ -47,6 +47,7 @@ namespace UE::Anim
 
 DECLARE_DELEGATE_OneParam(FOnAnimUpdateRateParamsCreated, FAnimUpdateRateParameters*)
 DECLARE_MULTICAST_DELEGATE_ThreeParams(FOnTickPose, USkinnedMeshComponent* /*SkinnedMeshComponent*/, float /*DeltaTime*/, bool /*bNeedsValidRootMotion*/)
+DECLARE_MULTICAST_DELEGATE(FOnBoneTransformsFinalizedMultiCast);
 
 //
 // Bone Visibility.
@@ -2029,6 +2030,19 @@ public:
 	 * @param InOutRequiredBones - the array in which to merge the additional required bones.
 	 */
 	virtual void GetAdditionalRequiredBonesForLeader(int32 LODIndex, TArray<FBoneIndexType>& InOutRequiredBones) const {}
+
+	/**
+	 * Register an OnBoneTransformsFinalized callback which can be called in FinalizeBoneTransform().
+	 * Note the inherited class has to implement the Broadcast call in FinalizeBoneTransform() if it wants to mimic the same behavior as in USkeletalMeshComponent.
+	 * @param Delegate - the delegate to be broadcasted in FinalizeBoneTransform by the inherited class.
+	 */
+	virtual FDelegateHandle RegisterOnBoneTransformsFinalizedDelegate(const FOnBoneTransformsFinalizedMultiCast::FDelegate& /*Delegate*/) { return FDelegateHandle(); }
+
+	/**
+	 * Unregister an OnBoneTransformsFinalized callback.
+	 * @param DelegateHandle - the handle of the delegate to remove from the list.
+	 */
+	virtual void UnregisterOnBoneTransformsFinalizedDelegate(const FDelegateHandle& /*DelegateHandle*/) {}
 };
 
 class FRenderStateRecreator
diff --git a/Engine/Source/Runtime/Engine/Private/Components/SkeletalMeshComponent.cpp b/Engine/Source/Runtime/Engine/Private/Components/SkeletalMeshComponent.cpp
--- a/Engine/Source/Runtime/Engine/Private/Components/SkeletalMeshComponent.cpp
+++ b/Engine/Source/Runtime/Engine/Private/Components/SkeletalMeshComponent.cpp
@@ -4350,10 +4350,6 @@ void USkeletalMeshComponent::FinalizeBoneTransform()
 
 	ConditionallyDispatchQueuedAnimEvents();
 
-	PRAGMA_DISABLE_DEPRECATION_WARNINGS
-	OnBoneTransformsFinalized.Broadcast();  // Deprecated in 4.27
-	PRAGMA_ENABLE_DEPRECATION_WARNINGS
-
 	OnBoneTransformsFinalizedMC.Broadcast();
 
 	TRACE_SKELETAL_MESH_COMPONENT(this);
