Disable ActiveGlobalRecreateRenderStateContextCount (which is used to reduce the amount of times component need primitive updates on recursive calls) because it doesn't work well with the version which provides list of components to recreate causing possible crashes due to stale skin cache data

diff --git a/Engine/Source/Runtime/Engine/Private/Components/ActorComponent.cpp b/Engine/Source/Runtime/Engine/Private/Components/ActorComponent.cpp
--- a/Engine/Source/Runtime/Engine/Private/Components/ActorComponent.cpp
+++ b/Engine/Source/Runtime/Engine/Private/Components/ActorComponent.cpp
@@ -336,7 +336,7 @@ FGlobalComponentRecreateRenderStateContext::FGlobalComponentRecreateRenderStateC
 	{
 		TRACE_CPUPROFILER_EVENT_SCOPE(FGlobalComponentRecreateRenderStateContext::FGlobalComponentRecreateRenderStateContext);
 
-		ActiveGlobalRecreateRenderStateContextCount++;
+		//ActiveGlobalRecreateRenderStateContextCount++;
 
 		// wait until resources are released
 		FlushRenderingCommands();
@@ -369,7 +369,7 @@ FGlobalComponentRecreateRenderStateContext::FGlobalComponentRecreateRenderStateC
 
 FGlobalComponentRecreateRenderStateContext::FGlobalComponentRecreateRenderStateContext(const TArray<UActorComponent*>& InComponents)
 {
-	if (FApp::CanEverRender() && ++ActiveGlobalRecreateRenderStateContextCount == 1)
+	if (FApp::CanEverRender() /*&& ++ActiveGlobalRecreateRenderStateContextCount == 1*/)
 	{
 		TRACE_CPUPROFILER_EVENT_SCOPE(FGlobalComponentRecreateRenderStateContext::FGlobalComponentRecreateRenderStateContext);
 
@@ -395,10 +395,10 @@ FGlobalComponentRecreateRenderStateContext::~FGlobalComponentRecreateRenderState
 
 	if (FApp::CanEverRender())
 	{
-		check(ActiveGlobalRecreateRenderStateContextCount > 0);
+		//check(ActiveGlobalRecreateRenderStateContextCount > 0);
 
 		// Check if this is the last active context
-		if (--ActiveGlobalRecreateRenderStateContextCount == 0)
+		//if (--ActiveGlobalRecreateRenderStateContextCount == 0)
 		{
 			// Clear the PSO material request cache to make sure PSO collection happens again on possible changed data
 			ClearMaterialPSORequests();
