[SkeletalMesh]
 - Fix thread-safety issue with BuildStringID

#rnx

diff --git a/Engine/Source/Runtime/Engine/Private/SkeletalMeshModel.cpp b/Engine/Source/Runtime/Engine/Private/SkeletalMeshModel.cpp
--- a/Engine/Source/Runtime/Engine/Private/SkeletalMeshModel.cpp
+++ b/Engine/Source/Runtime/Engine/Private/SkeletalMeshModel.cpp
@@ -166,8 +166,9 @@ FString FSkeletalMeshModel::GetLODModelIdString() const
 	for (int32 LODIndex = 0; LODIndex < LODModels.Num(); ++LODIndex)
 	{
 		const FSkeletalMeshLODModel& Model = LODModels[LODIndex];
-		Model.BuildStringID = Model.GetLODModelDeriveDataKey();
-		TArray<TCHAR, FString::AllocatorType> IDArray = Model.BuildStringID.GetCharArray();
+		const FString BuildStringID = Model.GetLODModelDeriveDataKey();
+		Model.BuildStringID = BuildStringID;
+		TArray<TCHAR, FString::AllocatorType> IDArray = BuildStringID.GetCharArray();
 		Sha.Update((uint8*)IDArray.GetData(), IDArray.Num() * IDArray.GetTypeSize());
 	}
 	Sha.Final();
diff --git a/Engine/Source/Runtime/Engine/Public/Rendering/SkeletalMeshLODModel.h b/Engine/Source/Runtime/Engine/Public/Rendering/SkeletalMeshLODModel.h
--- a/Engine/Source/Runtime/Engine/Public/Rendering/SkeletalMeshLODModel.h
+++ b/Engine/Source/Runtime/Engine/Public/Rendering/SkeletalMeshLODModel.h
@@ -16,6 +16,8 @@
 #include "Rendering/SkeletalMeshLODImporterData.h"
 #include "Animation/SkinWeightProfile.h"
 #include "CoreTypes.h"
+#include "Async/Mutex.h"
+#include "Async/UniqueLock.h"
 #include "HAL/CriticalSection.h"
 
 //
@@ -529,10 +531,70 @@ public:
 	*/
 	ENGINE_API void SyncronizeUserSectionsDataArray(bool bResetNonUsedSection = false);
 
+	// Offers protected access while supporting copy.
+	struct FThreadSafeBuildStringID
+	{
+	private:
+		FString    BuildStringID;
+		mutable UE::FMutex Mutex;
+	public:
+		FThreadSafeBuildStringID() = default;
+		FThreadSafeBuildStringID(const TCHAR* InText)
+			: BuildStringID(InText)
+		{
+		}
+
+		FThreadSafeBuildStringID(const FThreadSafeBuildStringID& Other)
+		{
+			UE::TUniqueLock OtherLock(Other.Mutex);
+			UE::TUniqueLock Lock(Mutex);
+			BuildStringID = Other.BuildStringID;
+		}
+
+		FThreadSafeBuildStringID& operator=(const FThreadSafeBuildStringID& Other)
+		{
+			if (&Other != this)
+			{
+				UE::TUniqueLock OtherLock(Other.Mutex);
+				UE::TUniqueLock Lock(Mutex);
+				BuildStringID = Other.BuildStringID;
+			}
+			return *this;
+		}
+
+		bool operator==(const FString& Other) const
+		{
+			UE::TUniqueLock Lock(Mutex);
+			return BuildStringID == Other;
+		}
+
+		bool operator==(const FThreadSafeBuildStringID& Other) const
+		{
+			UE::TUniqueLock OtherLock(Other.Mutex);
+			UE::TUniqueLock Lock(Mutex);
+			return BuildStringID == Other.BuildStringID;
+		}
+
+		void operator=(const FString& Other)
+		{
+			UE::TUniqueLock Lock(Mutex);
+			BuildStringID = Other;
+		}
+
+		operator FString() const
+		{
+			UE::TUniqueLock Lock(Mutex);
+			return BuildStringID;
+		}
+
+		FThreadSafeBuildStringID(FThreadSafeBuildStringID&&) = delete;
+		FThreadSafeBuildStringID& operator=(FThreadSafeBuildStringID&&) = delete;
+	};
+
 	//Temporary build String ID
 	//We use this string to store the LOD model data so we can know if the LOD need to be rebuild
 	//This GUID is set when we Cache the render data (build function)
-	mutable FString BuildStringID;
+	mutable FThreadSafeBuildStringID BuildStringID;
 
 	/**
 	* Build a derive data key with the user section data (UserSectionsData) and the original bulk data
