[SkeletalMesh]
 - Fix potential deadlock introduced with the BuildStringID fix

#rnx

diff --git a/Engine/Source/Runtime/Engine/Public/Rendering/SkeletalMeshLODModel.h b/Engine/Source/Runtime/Engine/Public/Rendering/SkeletalMeshLODModel.h
--- a/Engine/Source/Runtime/Engine/Public/Rendering/SkeletalMeshLODModel.h
+++ b/Engine/Source/Runtime/Engine/Public/Rendering/SkeletalMeshLODModel.h
@@ -570,6 +570,11 @@ public:
 
 		bool operator==(const FThreadSafeBuildStringID& Other) const
 		{
+			if (&Other == this)
+			{
+				return true;
+			}
+
 			UE::TUniqueLock OtherLock(Other.Mutex);
 			UE::TUniqueLock Lock(Mutex);
 			return BuildStringID == Other.BuildStringID;
