[World Partition] Fix FStreamingSourceInfo keeping an on-stack reference to the streaming source shape.

diff --git a/Engine/Source/Runtime/Engine/Private/WorldPartition/WorldPartitionRuntimeSpatialHash.cpp b/Engine/Source/Runtime/Engine/Private/WorldPartition/WorldPartitionRuntimeSpatialHash.cpp
--- a/Engine/Source/Runtime/Engine/Private/WorldPartition/WorldPartitionRuntimeSpatialHash.cpp
+++ b/Engine/Source/Runtime/Engine/Private/WorldPartition/WorldPartitionRuntimeSpatialHash.cpp
@@ -378,7 +378,7 @@ void FSpatialHashStreamingGrid::GetCells(const TArray<FWorldPartitionStreamingSo
 			, SourceShape(InSourceShape)
 		{}
 		const FWorldPartitionStreamingSource& Source;
-		const FSphericalSector& SourceShape;
+		const FSphericalSector SourceShape;
 	};
 
 	typedef TMap<FGridCellCoord, TArray<FStreamingSourceInfo>> FIntersectingCells;
