[World Partition]
- Actor Descriptors now stores the editor bounds, allowing some classes to have different bounds for loading in the editor and runtime.
- Actor components can now implement GetStreamingBoundsEditor to provide a different loading bounds in the editor.

diff --git a/Engine/Source/Editor/WorldPartitionEditor/Private/WorldPartition/SWorldPartitionEditorGrid2D.cpp b/Engine/Source/Editor/WorldPartitionEditor/Private/WorldPartition/SWorldPartitionEditorGrid2D.cpp
--- a/Engine/Source/Editor/WorldPartitionEditor/Private/WorldPartition/SWorldPartitionEditorGrid2D.cpp
+++ b/Engine/Source/Editor/WorldPartitionEditor/Private/WorldPartition/SWorldPartitionEditorGrid2D.cpp
@@ -1399,7 +1399,7 @@ uint32 SWorldPartitionEditorGrid2D::PaintActors(const FGeometry& AllottedGeometr
 			{
 				Guid = InActor->GetActorGuid();
 				Label = *InActor->GetActorLabel(false);
-				ActorBounds = InActor->GetStreamingBounds();
+				ActorBounds = InActor->GetStreamingBoundsEditor();
 				DescBounds = InActorDescInstance ? InActorDescInstance->GetEditorBounds() : ActorBounds;
 				bIsSpatiallyLoaded = InActor->GetIsSpatiallyLoaded();
 			}
@@ -2193,7 +2193,7 @@ void SWorldPartitionEditorGrid2D::FocusSelection()
 	{
 		if (AActor* Actor = Cast<AActor>(*It))
 		{
-			SelectionBox += Actor->GetStreamingBounds();
+			SelectionBox += Actor->GetStreamingBoundsEditor();
 		}
 	}
 
diff --git a/Engine/Source/Runtime/Core/Public/UObject/FortniteMainBranchObjectVersion.h b/Engine/Source/Runtime/Core/Public/UObject/FortniteMainBranchObjectVersion.h
--- a/Engine/Source/Runtime/Core/Public/UObject/FortniteMainBranchObjectVersion.h
+++ b/Engine/Source/Runtime/Core/Public/UObject/FortniteMainBranchObjectVersion.h
@@ -480,6 +480,9 @@ struct FFortniteMainBranchObjectVersion
 		// Added multiple section import to the cloth asset skeletal mesh import node
 		ClothAssetSkeletalMeshMultiSectionImport,
 
+		// Serialize EditorBounds in WorldPartitionActorDesc
+		WorldPartitionActorDescSerializeEditorBounds,
+
 		// -----<new versions can be added above this line>-------------------------------------------------
 		VersionPlusOne,
 		LatestVersion = VersionPlusOne - 1
diff --git a/Engine/Source/Runtime/Engine/Classes/Components/ActorComponent.h b/Engine/Source/Runtime/Engine/Classes/Components/ActorComponent.h
--- a/Engine/Source/Runtime/Engine/Classes/Components/ActorComponent.h
+++ b/Engine/Source/Runtime/Engine/Classes/Components/ActorComponent.h
@@ -1069,6 +1069,7 @@ public:
 	ENGINE_API virtual bool IsSelectedInEditor() const override;
 	virtual void SetPackageExternal(bool bExternal, bool bShouldDirty) {}
 	virtual FBox GetStreamingBounds() const { return FBox(ForceInit); }
+	virtual FBox GetStreamingBoundsEditor() const { return GetStreamingBounds(); }
 	virtual bool ForceActorNonSpatiallyLoaded() const { return false; }
 	virtual bool ForceActorNoDataLayers() const { return false; }
 #endif // WITH_EDITOR
diff --git a/Engine/Source/Runtime/Engine/Classes/GameFramework/Actor.h b/Engine/Source/Runtime/Engine/Classes/GameFramework/Actor.h
--- a/Engine/Source/Runtime/Engine/Classes/GameFramework/Actor.h
+++ b/Engine/Source/Runtime/Engine/Classes/GameFramework/Actor.h
@@ -2434,7 +2434,7 @@ public:
 	}
 
 	/**
-	 * Returns the location and the bounding box of all components that make up this Actor.
+	 * Returns the bounding box of all components that make up this Actor for loading at runtime.
 	 *
 	 * This function differs from GetActorBounds because it will return a valid origin and an empty extent if this actor
 	 * doesn't have primitive components.
@@ -2442,6 +2442,17 @@ public:
 	 * @see GetActorBounds()
 	 */
 	ENGINE_API virtual FBox GetStreamingBounds() const;
+
+	/**
+	 * Returns the bounding box of all components that make up this Actor for loading in the editor.
+	 *
+	 * This function differs from GetActorBounds because it will return a valid origin and an empty extent if this actor
+	 * doesn't have primitive components.
+	 *
+	 * @see GetActorBounds()
+	 */
+	ENGINE_API virtual FBox GetStreamingBoundsEditor() const;
+
 #endif
 
 
diff --git a/Engine/Source/Runtime/Engine/Private/Actor.cpp b/Engine/Source/Runtime/Engine/Private/Actor.cpp
--- a/Engine/Source/Runtime/Engine/Private/Actor.cpp
+++ b/Engine/Source/Runtime/Engine/Private/Actor.cpp
@@ -1140,57 +1140,60 @@ void AActor::ProcessEvent(UFunction* Function, void* Parameters)
 }
 
 #if WITH_EDITOR
-static bool IsComponentStreamingRelevant(const AActor* InActor, const UActorComponent* InComponent, FBox& OutStreamingBounds)
+template <class F>
+static bool ForEachStreamingRelevantComponent(const AActor* InActor, bool bForEditor, F Func)
 {
-	check(InActor);
-	check(InComponent);
-
-	if (!InComponent->IsRegistered())
+	auto GetComponentStreamingBounds = [InActor, bForEditor](const UActorComponent* InComponent, FBox& OutStreamingBounds) -> bool
 	{
-		return false;
-	}
+		check(InActor);
+		check(InComponent);
 
-	// Transient components shoudn't be part of the streeaming bounds, unless the actor itself is transient.
-	if (!InActor->HasAnyFlags(RF_Transient) && InComponent->HasAnyFlags(RF_Transient))
-	{
-		return false;
-	}
+		if (!InComponent->IsRegistered())
+		{
+			return false;
+		}
 
-	// Editor-only components shoudn't be part of the streeaming bounds, unless the actor itself is editor-only.
-	if (!InActor->IsEditorOnly() && InComponent->IsEditorOnly())
-	{
-		return false;
-	}
+		// Transient components shouldn't be part of the streeaming bounds, unless the actor itself is transient.
+		// This is to allow transient actors to be loaded in PIE.
+		if (!InActor->HasAnyFlags(RF_Transient) && InComponent->HasAnyFlags(RF_Transient))
+		{
+			return false;
+		}
 
-	OutStreamingBounds = InComponent->GetStreamingBounds();
-	return !!OutStreamingBounds.IsValid;
-}
+		// Editor-only components shouldn't be part of the streeaming bounds, unless the actor itself is editor-only.
+		// This is to allow editor-only actors to be loaded in PIE.
+		if (!bForEditor && !InActor->IsEditorOnly() && InComponent->IsEditorOnly())
+		{
+			return false;
+		}
 
-template <class F>
-static bool ForEachStreamingRelevantComponent(const AActor* InActor, F Func)
-{
-	bool bHasStreamingRelevantComponents = false;
+		OutStreamingBounds = bForEditor ? InComponent->GetStreamingBoundsEditor() : InComponent->GetStreamingBounds();
+		return !!OutStreamingBounds.IsValid;
+	};
 
-	auto HandleComponent = [InActor, &bHasStreamingRelevantComponents, &Func](const UActorComponent* Component)
+	auto HandleComponent = [&GetComponentStreamingBounds, &Func](const UActorComponent* Component) -> bool
 	{
 		FBox ComponentStreamingBound;
-		if (IsComponentStreamingRelevant(InActor, Component, ComponentStreamingBound))
+		if (GetComponentStreamingBounds(Component, ComponentStreamingBound))
 		{
 			Func(Component, ComponentStreamingBound);
-			bHasStreamingRelevantComponents = true;
+			return true;
 		}
+		return false;
 	};
 
-	InActor->ForEachComponent<UPrimitiveComponent>(true, [&HandleComponent](UActorComponent* Component)
+	bool bHasStreamingRelevantComponents = false;
+
+	InActor->ForEachComponent<UPrimitiveComponent>(true, [&bHasStreamingRelevantComponents, &HandleComponent](UActorComponent* Component)
 	{
-		HandleComponent(Component);
+		bHasStreamingRelevantComponents |= HandleComponent(Component);
 	});
 
 	if (!bHasStreamingRelevantComponents)
 	{
-		InActor->ForEachComponent<UActorComponent>(false, [&HandleComponent](UActorComponent* Component)
+		InActor->ForEachComponent<UActorComponent>(false, [&bHasStreamingRelevantComponents, &HandleComponent](UActorComponent* Component)
 		{
-			HandleComponent(Component);
+			bHasStreamingRelevantComponents |= HandleComponent(Component);
 		});
 	}
 
@@ -1200,7 +1203,7 @@ static bool ForEachStreamingRelevantComponent(const AActor* InActor, F Func)
 static bool HasComponentForceActorNonSpatiallyLoaded(const AActor* InActor)
 {
 	bool bHasComponentForceActorNonSpatiallyLoaded = false;
-	ForEachStreamingRelevantComponent(InActor, [&bHasComponentForceActorNonSpatiallyLoaded](const UActorComponent* Component, const FBox& StreamingBound)
+	ForEachStreamingRelevantComponent(InActor, false, [&bHasComponentForceActorNonSpatiallyLoaded](const UActorComponent* Component, const FBox& StreamingBound)
 	{
 		bHasComponentForceActorNonSpatiallyLoaded |= Component->ForceActorNonSpatiallyLoaded();
 	});
@@ -1210,7 +1213,17 @@ static bool HasComponentForceActorNonSpatiallyLoaded(const AActor* InActor)
 FBox AActor::GetStreamingBounds() const
 {
 	FBox StreamingBounds(ForceInit);
-	ForEachStreamingRelevantComponent(this, [&StreamingBounds](const UActorComponent* Component, const FBox& StreamingBound)
+	ForEachStreamingRelevantComponent(this, false, [&StreamingBounds](const UActorComponent* Component, const FBox& StreamingBound)
+	{
+		StreamingBounds += StreamingBound;
+	});
+	return StreamingBounds;
+}
+
+FBox AActor::GetStreamingBoundsEditor() const
+{
+	FBox StreamingBounds(ForceInit);
+	ForEachStreamingRelevantComponent(this, true, [&StreamingBounds](const UActorComponent* Component, const FBox& StreamingBound)
 	{
 		StreamingBounds += StreamingBound;
 	});
diff --git a/Engine/Source/Runtime/Engine/Private/WorldPartition/ActorPartition/PartitionActorDesc.cpp b/Engine/Source/Runtime/Engine/Private/WorldPartition/ActorPartition/PartitionActorDesc.cpp
--- a/Engine/Source/Runtime/Engine/Private/WorldPartition/ActorPartition/PartitionActorDesc.cpp
+++ b/Engine/Source/Runtime/Engine/Private/WorldPartition/ActorPartition/PartitionActorDesc.cpp
@@ -5,6 +5,7 @@
 
 #if WITH_EDITOR
 #include "UObject/UE5MainStreamObjectVersion.h"
+#include "UObject/FortniteMainBranchObjectVersion.h"
 
 FPartitionActorDesc::FPartitionActorDesc()
 	: GridSize(0)
@@ -28,12 +29,15 @@ void FPartitionActorDesc::Init(const AActor* InActor)
 	
 		const FVector ActorLocation = InActor->GetActorLocation();
 		SetGridIndices(ActorLocation.X, ActorLocation.Y, ActorLocation.Z);
+
+		EditorBounds = UActorPartitionSubsystem::FCellCoord::GetCellBounds({ GridIndexX, GridIndexY, GridIndexZ, 0 }, GridSize);
 	}
 }
 
 void FPartitionActorDesc::Serialize(FArchive& Ar)
 {
 	Ar.UsingCustomVersion(FUE5MainStreamObjectVersion::GUID);
+	Ar.UsingCustomVersion(FFortniteMainBranchObjectVersion::GUID);
 
 	FWorldPartitionActorDesc::Serialize(Ar);
 	
@@ -41,6 +45,11 @@ void FPartitionActorDesc::Serialize(FArchive& Ar)
 	{
 		Ar << GridSize << GridIndexX << GridIndexY << GridIndexZ;
 
+		if (Ar.CustomVer(FFortniteMainBranchObjectVersion::GUID) < FFortniteMainBranchObjectVersion::WorldPartitionActorDescSerializeEditorBounds)
+		{
+			EditorBounds = UActorPartitionSubsystem::FCellCoord::GetCellBounds({ GridIndexX, GridIndexY, GridIndexZ, 0 }, GridSize);
+		}
+
 		if (Ar.CustomVer(FUE5MainStreamObjectVersion::GUID) >= FUE5MainStreamObjectVersion::PartitionActorDescSerializeGridGuid)
 		{
 			Ar << GridGuid;
@@ -65,12 +74,6 @@ bool FPartitionActorDesc::Equals(const FWorldPartitionActorDesc* Other) const
 	return false;
 }
 
-FBox FPartitionActorDesc::GetEditorBounds() const
-{
-	const UActorPartitionSubsystem::FCellCoord CellCoord(GridIndexX, GridIndexY, GridIndexZ, 0);
-	return UActorPartitionSubsystem::FCellCoord::GetCellBounds(CellCoord, GridSize);
-}
-
 void FPartitionActorDesc::TransferWorldData(const FWorldPartitionActorDesc* From)
 {
 	FWorldPartitionActorDesc::TransferWorldData(From);
diff --git a/Engine/Source/Runtime/Engine/Private/WorldPartition/HLOD/HLODActorDesc.cpp b/Engine/Source/Runtime/Engine/Private/WorldPartition/HLOD/HLODActorDesc.cpp
--- a/Engine/Source/Runtime/Engine/Private/WorldPartition/HLOD/HLODActorDesc.cpp
+++ b/Engine/Source/Runtime/Engine/Private/WorldPartition/HLOD/HLODActorDesc.cpp
@@ -19,7 +19,6 @@
 #include "WorldPartition/WorldPartitionActorDescInstance.h"
 
 FHLODActorDesc::FHLODActorDesc()
-	: EditorBounds(ForceInit)
 {}
 
 int64 FHLODActorDesc::GetStat(FName InStatName) const
@@ -161,11 +160,10 @@ void FHLODActorDesc::Serialize(FArchive& Ar)
 		
 		if (Ar.CustomVer(FFortniteMainBranchObjectVersion::GUID) >= FFortniteMainBranchObjectVersion::WorldPartitionHLODActorDescSerializeEditorBounds)
 		{
-			Ar << EditorBounds;
-		}
-		else
-		{
-			EditorBounds = GetRuntimeBounds();
+			if (Ar.CustomVer(FFortniteMainBranchObjectVersion::GUID) < FFortniteMainBranchObjectVersion::WorldPartitionActorDescSerializeEditorBounds)
+			{
+				Ar << EditorBounds;
+			}
 		}
 	}
 }
diff --git a/Engine/Source/Runtime/Engine/Private/WorldPartition/Landscape/LandscapeActorDesc.cpp b/Engine/Source/Runtime/Engine/Private/WorldPartition/Landscape/LandscapeActorDesc.cpp
--- a/Engine/Source/Runtime/Engine/Private/WorldPartition/Landscape/LandscapeActorDesc.cpp
+++ b/Engine/Source/Runtime/Engine/Private/WorldPartition/Landscape/LandscapeActorDesc.cpp
@@ -12,12 +12,13 @@ void FLandscapeActorDesc::Init(const AActor* InActor)
 {
 	FPartitionActorDesc::Init(InActor);
 
-	const ALandscapeProxy* LandscapeProxy = CastChecked<ALandscapeProxy>(InActor);
-	check(LandscapeProxy);
-	SetGridIndices(LandscapeProxy->LandscapeSectionOffset.X, LandscapeProxy->LandscapeSectionOffset.Y, 0);
-
 	if (!bIsDefaultActorDesc)
 	{
+		const ALandscapeProxy* LandscapeProxy = CastChecked<ALandscapeProxy>(InActor);
+		check(LandscapeProxy);
+
+		SetGridIndices(LandscapeProxy->LandscapeSectionOffset.X, LandscapeProxy->LandscapeSectionOffset.Y, 0);
+
 		const ALandscape* LandscapeActor = LandscapeProxy->GetLandscapeActor();
 		if (LandscapeActor)
 		{
@@ -47,12 +48,6 @@ void FLandscapeActorDesc::Serialize(FArchive& Ar)
 	}
 }
 
-FBox FLandscapeActorDesc::GetEditorBounds() const
-{
-	// We need to skip super class since we aren't using grid indices as it should be (it's in Landscape space).
-	return FWorldPartitionActorDesc::GetEditorBounds();
-}
-
 bool FLandscapeActorDesc::Equals(const FWorldPartitionActorDesc* Other) const
 {
 	if (FPartitionActorDesc::Equals(Other))
diff --git a/Engine/Source/Runtime/Engine/Private/WorldPartition/LevelInstance/LevelInstanceActorDesc.cpp b/Engine/Source/Runtime/Engine/Private/WorldPartition/LevelInstance/LevelInstanceActorDesc.cpp
--- a/Engine/Source/Runtime/Engine/Private/WorldPartition/LevelInstance/LevelInstanceActorDesc.cpp
+++ b/Engine/Source/Runtime/Engine/Private/WorldPartition/LevelInstance/LevelInstanceActorDesc.cpp
@@ -130,9 +130,7 @@ void FLevelInstanceActorDesc::UpdateBounds()
 {
 	if (UActorDescContainer* ChildContainerPtr = ChildContainer.Get())
 	{
-		FBox ContainerBounds = UActorDescContainerSubsystem::GetChecked().GetContainerBounds(GetChildContainerName()).TransformBy(GetChildContainerTransform());
-
-		ContainerBounds.GetCenterAndExtents(BoundsLocation, BoundsExtent);
+		RuntimeBounds = UActorDescContainerSubsystem::GetChecked().GetContainerBounds(GetChildContainerName()).TransformBy(GetChildContainerTransform());
 	}
 }
 
@@ -395,11 +393,7 @@ void FLevelInstanceActorDesc::Serialize(FArchive& Ar)
 			{
 				if (!IsChildContainerInstance())
 				{
-					FBox OutBounds;
-					if (ULevelInstanceSubsystem::GetLevelInstanceBoundsFromPackage(ActorTransform, GetChildContainerPackage(), OutBounds))
-					{
-						OutBounds.GetCenterAndExtents(BoundsLocation, BoundsExtent);
-					}
+					ULevelInstanceSubsystem::GetLevelInstanceBoundsFromPackage(ActorTransform, GetChildContainerPackage(), RuntimeBounds);
 				}
 			}
 		}
diff --git a/Engine/Source/Runtime/Engine/Private/WorldPartition/WorldPartitionActorDesc.cpp b/Engine/Source/Runtime/Engine/Private/WorldPartition/WorldPartitionActorDesc.cpp
--- a/Engine/Source/Runtime/Engine/Private/WorldPartition/WorldPartitionActorDesc.cpp
+++ b/Engine/Source/Runtime/Engine/Private/WorldPartition/WorldPartitionActorDesc.cpp
@@ -44,14 +44,14 @@ static FGuid GetDefaultActorDescGuid(const FWorldPartitionActorDesc* ActorDesc)
 }
 
 FWorldPartitionActorDesc::FWorldPartitionActorDesc()
-	: bIsSpatiallyLoaded(false)
+	: RuntimeBounds(ForceInit)
+	, bIsSpatiallyLoaded(false)
 	, bActorIsEditorOnly(false)
 	, bActorIsRuntimeOnly(false)
 	, bActorIsMainWorldOnly(false)
 	, bActorIsHLODRelevant(false)
 	, bActorIsListedInSceneOutliner(true)
 	, bIsUsingDataLayerAsset(false)
-	, bIsBoundsValid(false)
 	, ActorNativeClass(nullptr)
 	, Container(nullptr)
 	, bIsDefaultActorDesc(false)
@@ -91,10 +91,8 @@ void FWorldPartitionActorDesc::Init(const AActor* InActor)
 
 	ActorTransform = InActor->GetActorTransform();
 
-	const FBox StreamingBounds = !bIsDefaultActorDesc ? InActor->GetStreamingBounds() : FBox(ForceInit);
-	StreamingBounds.GetCenterAndExtents(BoundsLocation, BoundsExtent);
-	bIsBoundsValid = StreamingBounds.IsValid == 1;
-
+	RuntimeBounds = !bIsDefaultActorDesc ? InActor->GetStreamingBounds() : FBox(ForceInit);
+	EditorBounds = !bIsDefaultActorDesc ? InActor->GetStreamingBoundsEditor() : FBox(ForceInit);
 	RuntimeGrid = InActor->GetRuntimeGrid();
 	bIsSpatiallyLoaded = InActor->GetIsSpatiallyLoaded();
 	bActorIsEditorOnly = InActor->IsEditorOnly();
@@ -300,10 +298,9 @@ bool FWorldPartitionActorDesc::Equals(const FWorldPartitionActorDesc* Other) con
 		ActorPackage == Other->ActorPackage &&
 		ActorPath == Other->ActorPath &&
 		ActorLabel == Other->ActorLabel &&
-		bIsBoundsValid == Other->bIsBoundsValid &&
 		ActorTransform.Equals(Other->ActorTransform, 0.1f) &&
-		BoundsLocation.Equals(Other->BoundsLocation, 0.1f) &&
-		BoundsExtent.Equals(Other->BoundsExtent, 0.1f) &&
+		RuntimeBounds.Equals(Other->RuntimeBounds, 0.1f) &&
+		EditorBounds.Equals(Other->EditorBounds, 0.1f) &&
 		RuntimeGrid == Other->RuntimeGrid &&
 		bIsSpatiallyLoaded == Other->bIsSpatiallyLoaded &&
 		bActorIsEditorOnly == Other->bActorIsEditorOnly &&
@@ -336,7 +333,8 @@ bool FWorldPartitionActorDesc::ShouldResave(const FWorldPartitionActorDesc* Othe
 		bActorIsEditorOnly != Other->bActorIsEditorOnly ||
 		bActorIsRuntimeOnly != Other->bActorIsRuntimeOnly ||
 		bActorIsMainWorldOnly != Other->bActorIsMainWorldOnly ||
-		bIsBoundsValid != Other->bIsBoundsValid ||
+		RuntimeBounds.IsValid != Other->RuntimeBounds.IsValid||
+		EditorBounds.IsValid != Other->EditorBounds.IsValid||
 		HLODLayer != Other->HLODLayer ||
 		ParentActor != Other->ParentActor ||
 		ContentBundleGuid != Other->ContentBundleGuid ||
@@ -350,7 +348,7 @@ bool FWorldPartitionActorDesc::ShouldResave(const FWorldPartitionActorDesc* Othe
 	}
 
 	// Tolerate up to 5% for bounds change
-	if (bIsBoundsValid)
+	if (RuntimeBounds.IsValid)
 	{
 		const FBox ThisBounds = GetRuntimeBounds();
 		const FBox OtherBounds = Other->GetRuntimeBounds();
@@ -439,17 +437,15 @@ FString FWorldPartitionActorDesc::ToString(EToStringMode Mode) const
 	{
 		FString BoundsStr;
 
-		if (bIsBoundsValid)
+		if (RuntimeBounds.IsValid)
 		{
-			const FBox EditorBounds = GetEditorBounds();
-			const FBox RuntimeBounds = GetRuntimeBounds();
-			if (EditorBounds.Equals(RuntimeBounds))
+			if (GetEditorBounds().Equals(RuntimeBounds))
 			{
 				BoundsStr = RuntimeBounds.ToString();
 			}
 			else
 			{
-				BoundsStr = *FString::Printf(TEXT("(Editor:%s Runtime:%s)"), *EditorBounds.ToString(), *RuntimeBounds.ToString());
+				BoundsStr = *FString::Printf(TEXT("(Editor:%s Runtime:%s)"), *GetEditorBounds().ToString(), *RuntimeBounds.ToString());
 			}
 		}
 		else
@@ -566,24 +562,33 @@ void FWorldPartitionActorDesc::Serialize(FArchive& Ar)
 	{
 		FVector3f BoundsLocationFlt, BoundsExtentFlt;
 		Ar << BoundsLocationFlt << BoundsExtentFlt;
-		BoundsLocation = FVector(BoundsLocationFlt);
-		BoundsExtent = FVector(BoundsExtentFlt);
-		bIsBoundsValid = true;
+		RuntimeBounds = FBox(FVector(BoundsLocationFlt - BoundsExtentFlt), FVector(BoundsLocationFlt + BoundsExtentFlt));
+		EditorBounds = RuntimeBounds;
 	}
 	else if (!bIsDefaultActorDesc)
 	{
-		if (Ar.CustomVer(FFortniteMainBranchObjectVersion::GUID) < FFortniteMainBranchObjectVersion::WorldPartitionActorDescSerializeInvalidBounds)
+		if (Ar.CustomVer(FFortniteMainBranchObjectVersion::GUID) < FFortniteMainBranchObjectVersion::WorldPartitionActorDescSerializeEditorBounds)
 		{
-			bIsBoundsValid = true;
+			bool bIsBoundsValid = true;
+			if (Ar.CustomVer(FFortniteMainBranchObjectVersion::GUID) >= FFortniteMainBranchObjectVersion::WorldPartitionActorDescSerializeInvalidBounds)
+			{
+				Ar << bIsBoundsValid;
+			}
+
+			if (bIsBoundsValid)
+			{
+				FVector BoundsLocation;
+				FVector BoundsExtent;
+
+				Ar << BoundsLocation << BoundsExtent;
+
+				RuntimeBounds = FBox(BoundsLocation - BoundsExtent, BoundsLocation + BoundsExtent);
+				EditorBounds = RuntimeBounds;
+			}
 		}
 		else
 		{
-			Ar << bIsBoundsValid;
-		}
-
-		if (bIsBoundsValid)
-		{
-			Ar << BoundsLocation << BoundsExtent;
+			Ar << RuntimeBounds << EditorBounds;
 		}
 	}
 	
@@ -761,12 +766,12 @@ PRAGMA_ENABLE_DEPRECATION_WARNINGS
 
 FBox FWorldPartitionActorDesc::GetEditorBounds() const
 {
-	return bIsBoundsValid ? FBox(BoundsLocation - BoundsExtent, BoundsLocation + BoundsExtent) : FBox(ForceInit);
+	return EditorBounds.IsValid ? EditorBounds : RuntimeBounds;
 }
 
 FBox FWorldPartitionActorDesc::GetRuntimeBounds() const
 {
-	return bIsBoundsValid ? FBox(BoundsLocation - BoundsExtent, BoundsLocation + BoundsExtent) : FBox(ForceInit);
+	return RuntimeBounds;
 }
 
 FName FWorldPartitionActorDesc::GetActorName() const
diff --git a/Engine/Source/Runtime/Engine/Public/WorldPartition/ActorPartition/PartitionActorDesc.h b/Engine/Source/Runtime/Engine/Public/WorldPartition/ActorPartition/PartitionActorDesc.h
--- a/Engine/Source/Runtime/Engine/Public/WorldPartition/ActorPartition/PartitionActorDesc.h
+++ b/Engine/Source/Runtime/Engine/Public/WorldPartition/ActorPartition/PartitionActorDesc.h
@@ -20,12 +20,12 @@ public:
 	FGuid GridGuid;
 
 	ENGINE_API FPartitionActorDesc();
+
 protected:
 	ENGINE_API virtual void Init(const AActor* InActor) override;
 	ENGINE_API virtual bool Equals(const FWorldPartitionActorDesc* Other) const override;
 	virtual uint32 GetSizeOf() const override { return sizeof(FPartitionActorDesc); }
 	ENGINE_API virtual void Serialize(FArchive& Ar) override;
-	ENGINE_API virtual FBox GetEditorBounds() const override;
 	ENGINE_API virtual void TransferWorldData(const FWorldPartitionActorDesc* From) override;
 
 	ENGINE_API void SetGridIndices(double LocationX, double LocationY, double LocationZ);
diff --git a/Engine/Source/Runtime/Engine/Public/WorldPartition/HLOD/HLODActorDesc.h b/Engine/Source/Runtime/Engine/Public/WorldPartition/HLOD/HLODActorDesc.h
--- a/Engine/Source/Runtime/Engine/Public/WorldPartition/HLOD/HLODActorDesc.h
+++ b/Engine/Source/Runtime/Engine/Public/WorldPartition/HLOD/HLODActorDesc.h
@@ -30,10 +30,6 @@ public:
 	inline const FStats& GetStats() const { return HLODStats; }
 	ENGINE_API int64 GetStat(FName InStatName) const;
 
-	//~ Begin FWorldPartitionActorDesc Interface.
-	virtual FBox GetEditorBounds() const override { return EditorBounds; }
-	//~ End FWorldPartitionActorDesc Interface.
-
 protected:
 	ENGINE_API FHLODActorDesc();
 
@@ -50,7 +46,5 @@ protected:
 	FTopLevelAssetPath SourceHLODLayer;
 
 	FStats HLODStats;
-
-	FBox EditorBounds;
 };
 #endif
diff --git a/Engine/Source/Runtime/Engine/Public/WorldPartition/Landscape/LandscapeActorDesc.h b/Engine/Source/Runtime/Engine/Public/WorldPartition/Landscape/LandscapeActorDesc.h
--- a/Engine/Source/Runtime/Engine/Public/WorldPartition/Landscape/LandscapeActorDesc.h
+++ b/Engine/Source/Runtime/Engine/Public/WorldPartition/Landscape/LandscapeActorDesc.h
@@ -16,7 +16,7 @@ public:
 	ENGINE_API virtual void Init(const AActor* InActor) override;
 	ENGINE_API virtual bool Equals(const FWorldPartitionActorDesc* Other) const override;
 	ENGINE_API virtual const FGuid& GetSceneOutlinerParent() const override;
-	ENGINE_API virtual FBox GetEditorBounds() const override;
+
 protected:
 	virtual uint32 GetSizeOf() const override { return sizeof(FLandscapeActorDesc); }
 	ENGINE_API virtual void Serialize(FArchive& Ar) override;
diff --git a/Engine/Source/Runtime/Engine/Public/WorldPartition/WorldPartitionActorDesc.h b/Engine/Source/Runtime/Engine/Public/WorldPartition/WorldPartitionActorDesc.h
--- a/Engine/Source/Runtime/Engine/Public/WorldPartition/WorldPartitionActorDesc.h
+++ b/Engine/Source/Runtime/Engine/Public/WorldPartition/WorldPartitionActorDesc.h
@@ -24,13 +24,11 @@ struct FWorldPartitionActorDescInitData
 {
 	FWorldPartitionActorDescInitData()
 		: DataSource(TInPlaceType<TArray<uint8>>(), TArray<uint8>())
-	{
-	}
+	{}
 
 	FWorldPartitionActorDescInitData(FActorDescArchive* InArchive)
 		: DataSource(TInPlaceType<FActorDescArchive*>(), InArchive)
-	{
-	}
+	{}
 
 	UClass* NativeClass;
 	FName PackageName;
@@ -66,7 +64,6 @@ struct FWorldPartitionActorDescInitData
 private:
 	// Provide SerializedData or an already initialized Archive
 	TVariant<TArray<uint8>, FActorDescArchive*> DataSource;
-
 };
 
 struct FWorldPartitionAssetDataPatcher
@@ -165,7 +162,7 @@ public:
 	inline UClass* GetActorNativeClass() const { return ActorNativeClass; }
 
 	inline FName GetRuntimeGrid() const { return RuntimeGrid; }
-	inline bool GetIsSpatiallyLoaded() const { return bIsBoundsValid ? bIsSpatiallyLoaded : false; }
+	inline bool GetIsSpatiallyLoaded() const { return RuntimeBounds.IsValid ? bIsSpatiallyLoaded : false; }
 	inline bool GetIsSpatiallyLoadedRaw() const { return bIsSpatiallyLoaded; }
 	inline bool GetActorIsEditorOnly() const { return bActorIsEditorOnly; }
 	inline bool GetActorIsRuntimeOnly() const { return bActorIsRuntimeOnly; }
@@ -185,7 +182,7 @@ public:
 	inline const FGuid& GetFolderGuid() const { return FolderGuid; }
 	inline const FTransform& GetActorTransform() const { return ActorTransform; }
 
-	ENGINE_API virtual FBox GetEditorBounds() const;
+	ENGINE_API FBox GetEditorBounds() const;
 	ENGINE_API FBox GetRuntimeBounds() const;
 
 	inline const FGuid& GetParentActor() const { return ParentActor; }
@@ -215,7 +212,7 @@ public:
 	ENGINE_API FGuid GetContentBundleGuid() const;
 
 	virtual const FGuid& GetSceneOutlinerParent() const { return GetParentActor(); }
-	virtual bool IsResaveNeeded() const { return bIsSpatiallyLoaded && !bIsBoundsValid; }
+	virtual bool IsResaveNeeded() const { return bIsSpatiallyLoaded && !RuntimeBounds.IsValid; }
 
 	ENGINE_API virtual void CheckForErrors(const IWorldPartitionActorDescInstanceView* InActorDescView, IStreamingGenerationErrorHandler* ErrorHandler) const;
 
@@ -397,9 +394,8 @@ protected:
 	virtual void TransferWorldData(const FWorldPartitionActorDesc* From)
 	{
 		ActorTransform = From->ActorTransform;
-		BoundsLocation = From->BoundsLocation;
-		BoundsExtent = From->BoundsExtent;
-		bIsBoundsValid = From->bIsBoundsValid;
+		RuntimeBounds = From->RuntimeBounds;
+		EditorBounds = From->EditorBounds;
 	}
 
 	virtual uint32 GetSizeOf() const { return sizeof(FWorldPartitionActorDesc); }
@@ -414,8 +410,8 @@ protected:
 	FSoftObjectPath					ActorPath;		// Not serialized, comes from initialization data
 	FName							ActorLabel;
 	FTransform						ActorTransform;
-	FVector							BoundsLocation;
-	FVector							BoundsExtent;
+	FBox							RuntimeBounds;
+	FBox							EditorBounds;
 	FName							RuntimeGrid;
 	bool							bIsSpatiallyLoaded;
 	bool							bActorIsEditorOnly;
@@ -424,7 +420,6 @@ protected:
 	bool							bActorIsHLODRelevant;
 	bool							bActorIsListedInSceneOutliner;
 	bool							bIsUsingDataLayerAsset; // Used to know if DataLayers array represents DataLayers Asset paths or the FNames of the deprecated version of Data Layers
-	bool							bIsBoundsValid;
 	FSoftObjectPath					HLODLayer;
 	TArray<FName>					DataLayers;
 	FSoftObjectPath					ExternalDataLayerAsset;
