Revise API for IActiveSoundUpdateInterface
+ Fix mismatched Add/Delete pairs; "add virtualized" was missing
+ Rename to Create/Deleting to clearly communicate object lifespan

#rnx

diff --git a/Engine/Plugins/AudioGameplay/Source/AudioGameplay/Private/SoundHandleSubsystem.cpp b/Engine/Plugins/AudioGameplay/Source/AudioGameplay/Private/SoundHandleSubsystem.cpp
--- a/Engine/Plugins/AudioGameplay/Source/AudioGameplay/Private/SoundHandleSubsystem.cpp
+++ b/Engine/Plugins/AudioGameplay/Source/AudioGameplay/Private/SoundHandleSubsystem.cpp
@@ -24,7 +24,17 @@ void USoundHandleSubsystem::Deinitialize()
 	Super::Deinitialize();
 }
 
-void USoundHandleSubsystem::OnNotifyPendingDelete(const FActiveSound& ActiveSound)
+void USoundHandleSubsystem::NotifyActiveSoundDeleting(const FActiveSound& ActiveSound)
+{
+	NotifySoundDeleting(ActiveSound);
+}
+
+void USoundHandleSubsystem::NotifyVirtualizedSoundDeleting(const FActiveSound& ActiveSound)
+{
+	NotifySoundDeleting(ActiveSound);
+}
+
+void USoundHandleSubsystem::NotifySoundDeleting(const FActiveSound& ActiveSound)
 {
 	check(IsInAudioThread());
 	if (ActiveHandles.Num() > 0)
diff --git a/Engine/Plugins/AudioGameplay/Source/AudioGameplay/Public/SoundHandleSubsystem.h b/Engine/Plugins/AudioGameplay/Source/AudioGameplay/Public/SoundHandleSubsystem.h
--- a/Engine/Plugins/AudioGameplay/Source/AudioGameplay/Public/SoundHandleSubsystem.h
+++ b/Engine/Plugins/AudioGameplay/Source/AudioGameplay/Public/SoundHandleSubsystem.h
@@ -31,7 +31,8 @@ public:
 	//~ End USubsystem interface
 	
 	//~ Begin IActiveSoundUpdateInterface
-	virtual void OnNotifyPendingDelete(const FActiveSound& ActiveSound) override;
+	virtual void NotifyActiveSoundDeleting(const FActiveSound& ActiveSound) override;
+	virtual void NotifyVirtualizedSoundDeleting(const FActiveSound& ActiveSound) override;
 	//~ End IActiveSoundUpdateInterface
 
 	//~ Begin ISoundHandleSystem
@@ -49,4 +50,6 @@ protected:
 	
 	TMap<Audio::FSoundHandleID, FSoundHandle> ActiveHandles;
 	TMap<Audio::FSoundHandleID, ISoundHandleOwner*> Owners;
+
+	void NotifySoundDeleting(const FActiveSound& ActiveSound);
 };
diff --git a/Engine/Plugins/AudioGameplayVolume/Source/AudioGameplayVolume/Private/AudioGameplayVolumeSubsystem.cpp b/Engine/Plugins/AudioGameplayVolume/Source/AudioGameplayVolume/Private/AudioGameplayVolumeSubsystem.cpp
--- a/Engine/Plugins/AudioGameplayVolume/Source/AudioGameplayVolume/Private/AudioGameplayVolumeSubsystem.cpp
+++ b/Engine/Plugins/AudioGameplayVolume/Source/AudioGameplayVolume/Private/AudioGameplayVolumeSubsystem.cpp
@@ -296,7 +296,7 @@ void UAudioGameplayVolumeSubsystem::ApplyInteriorSettings(const FActiveSound& Ac
 	ParseParams.AmbientZoneFilterFrequency = ActiveSoundInfo->CurrentInteriorLPF;
 }
 
-void UAudioGameplayVolumeSubsystem::OnNotifyPendingDelete(const FActiveSound& ActiveSound)
+void UAudioGameplayVolumeSubsystem::NotifyActiveSoundDeleting(const FActiveSound& ActiveSound)
 {
 	check(IsInAudioThread());
 	if (ActiveSound.bApplyInteriorVolumes && ActiveSoundData.Num() > 0)
diff --git a/Engine/Plugins/AudioGameplayVolume/Source/AudioGameplayVolume/Public/AudioGameplayVolumeSubsystem.h b/Engine/Plugins/AudioGameplayVolume/Source/AudioGameplayVolume/Public/AudioGameplayVolumeSubsystem.h
--- a/Engine/Plugins/AudioGameplayVolume/Source/AudioGameplayVolume/Public/AudioGameplayVolumeSubsystem.h
+++ b/Engine/Plugins/AudioGameplayVolume/Source/AudioGameplayVolume/Public/AudioGameplayVolumeSubsystem.h
@@ -136,7 +136,7 @@ public:
 	//~ Begin IActiveSoundUpdateInterface
 	virtual void GatherInteriorData(const FActiveSound& ActiveSound, FSoundParseParameters& ParseParams) override;
 	virtual void ApplyInteriorSettings(const FActiveSound& ActiveSound, FSoundParseParameters& ParseParams) override;
-	virtual void OnNotifyPendingDelete(const FActiveSound& ActiveSound) override;
+	virtual void NotifyActiveSoundDeleting(const FActiveSound& ActiveSound) override;
 	//~ End IActiveSoundUpdateInterface
 	
 	/** Add a volume to the system */
diff --git a/Engine/Source/Runtime/Engine/Private/AudioDevice.cpp b/Engine/Source/Runtime/Engine/Private/AudioDevice.cpp
--- a/Engine/Source/Runtime/Engine/Private/AudioDevice.cpp
+++ b/Engine/Source/Runtime/Engine/Private/AudioDevice.cpp
@@ -3639,20 +3639,53 @@ void FAudioDevice::ApplyInteriorSettings(FActiveSound& ActiveSound, FSoundParseP
 	});
 }
 
-void FAudioDevice::NotifyAddActiveSound(FActiveSound& ActiveSound) const
+void FAudioDevice::NotifySubsystemsActiveSoundCreated(FActiveSound& ActiveSound) const
 {
 	SubsystemCollection.ForEachSubsystem<IActiveSoundUpdateInterface>([&ActiveSound](IActiveSoundUpdateInterface* ActiveSoundUpdate)
 	{
+		ActiveSoundUpdate->NotifyActiveSoundCreated(ActiveSound);
+
+		PRAGMA_DISABLE_DEPRECATION_WARNINGS
 		ActiveSoundUpdate->OnNotifyAddActiveSound(ActiveSound);
+		PRAGMA_ENABLE_DEPRECATION_WARNINGS
+
+		return true;
+	});
+}
+
+void FAudioDevice::NotifySubsystemsActiveSoundDeleting(FActiveSound& ActiveSound) const
+{
+	SubsystemCollection.ForEachSubsystem<IActiveSoundUpdateInterface>([&ActiveSound](IActiveSoundUpdateInterface* ActiveSoundUpdate)
+	{
+		ActiveSoundUpdate->NotifyActiveSoundDeleting(ActiveSound);
+
+		PRAGMA_DISABLE_DEPRECATION_WARNINGS
+		ActiveSoundUpdate->OnNotifyPendingDelete(ActiveSound);
+		PRAGMA_ENABLE_DEPRECATION_WARNINGS
+
+		return true;
+	});
+}
+
+void FAudioDevice::NotifySubsystemsVirtualizedSoundCreated(FActiveSound& ActiveSound) const
+{
+	SubsystemCollection.ForEachSubsystem<IActiveSoundUpdateInterface>([&ActiveSound](IActiveSoundUpdateInterface* ActiveSoundUpdate)
+	{
+		ActiveSoundUpdate->NotifyVirtualizedSoundCreated(ActiveSound);
 		return true;
 	});
 }
 
-void FAudioDevice::NotifyPendingDeleteInternal(FActiveSound& ActiveSound) const
+void FAudioDevice::NotifySubsystemsVirtualizedSoundDeleting(FActiveSound& ActiveSound) const
 {
 	SubsystemCollection.ForEachSubsystem<IActiveSoundUpdateInterface>([&ActiveSound](IActiveSoundUpdateInterface* ActiveSoundUpdate)
 	{
+		ActiveSoundUpdate->NotifyVirtualizedSoundDeleting(ActiveSound);
+
+		PRAGMA_DISABLE_DEPRECATION_WARNINGS
 		ActiveSoundUpdate->OnNotifyPendingDelete(ActiveSound);
+		PRAGMA_ENABLE_DEPRECATION_WARNINGS
+
 		return true;
 	});
 }
@@ -5246,7 +5279,7 @@ void FAudioDevice::AddNewActiveSoundInternal(const FActiveSound& InNewActiveSoun
 		ActiveSoundArray.AddUnique(ActiveSound);
 	}
 
-	NotifyAddActiveSound(*ActiveSound);
+	NotifySubsystemsActiveSoundCreated(*ActiveSound);
 }
 
 void FAudioDevice::ReportSoundFailedToStart(const uint64 AudioComponentID, FAudioVirtualLoop* VirtualLoop)
@@ -5358,6 +5391,8 @@ void FAudioDevice::AddVirtualLoop(const FAudioVirtualLoop& InVirtualLoop)
 #endif // !(UE_BUILD_SHIPPING || UE_BUILD_TEST)
 
 	VirtualLoops.Add(&ActiveSound, MoveTemp(VirtualLoop));
+
+	NotifySubsystemsVirtualizedSoundCreated(ActiveSound);
 }
 
 bool FAudioDevice::RemoveVirtualLoop(FActiveSound& InActiveSound)
@@ -5453,7 +5488,7 @@ void FAudioDevice::ProcessingPendingActiveSoundStops(bool bForceDelete)
 				}
 				ActiveSound->ClearTransmitter();
 
-				NotifyPendingDeleteInternal(*ActiveSound);
+				NotifySubsystemsActiveSoundDeleting(*ActiveSound);
 				delete ActiveSound;
 			}
 		}
@@ -5469,6 +5504,7 @@ void FAudioDevice::ProcessingPendingActiveSoundStops(bool bForceDelete)
 		{
 			check(ActiveSound);
 			bool bDeleteActiveSound = false;
+			bool bWasVirtualized = false;
 
 			// If the request was to stop an ActiveSound that
 			// is set to re-trigger but is not playing, remove
@@ -5476,6 +5512,7 @@ void FAudioDevice::ProcessingPendingActiveSoundStops(bool bForceDelete)
 			if (RemoveVirtualLoop(*ActiveSound))
 			{
 				bDeleteActiveSound = true;
+				bWasVirtualized = true;
 			}
 			else
 			{
@@ -5513,11 +5550,17 @@ void FAudioDevice::ProcessingPendingActiveSoundStops(bool bForceDelete)
 				}
 				ActiveSound->ClearTransmitter();
 
-				NotifyPendingDeleteInternal(*ActiveSound);
-
 				// Remove from the list of pending sounds to stop
 				PendingSoundsToStop.Remove(ActiveSound);
 	
+				if (bWasVirtualized)
+				{
+					NotifySubsystemsVirtualizedSoundDeleting(*ActiveSound);
+				}
+				else
+				{
+					NotifySubsystemsActiveSoundDeleting(*ActiveSound);
+				}
 				delete ActiveSound;
 			}
 			else
diff --git a/Engine/Source/Runtime/Engine/Public/ActiveSoundUpdateInterface.h b/Engine/Source/Runtime/Engine/Public/ActiveSoundUpdateInterface.h
--- a/Engine/Source/Runtime/Engine/Public/ActiveSoundUpdateInterface.h
+++ b/Engine/Source/Runtime/Engine/Public/ActiveSoundUpdateInterface.h
@@ -41,18 +41,53 @@ public:
 	virtual void ApplyInteriorSettings(const FActiveSound& ActiveSound, FSoundParseParameters& ParseParams) {}
 
 	/**
-	 * Called when an active sound is being added to the audio engine
+	 * Called when an active sound has just been added to the audio engine,
+	 * both for brand new sounds and for virtualized sounds that have just become active.
+	 * In the latter case, a corresponding NotifyVirtualizedSoundDeleting will be received.
+	 * You can correlate the two objects by matching their GetPlayOrder() value.
+	 *
 	 * NOTE! Called on the AudioThread
 	 *
-	 * @param ActiveSound	The active sound being added
+	 * @param ActiveSound	The active sound being created
 	 */
-	virtual void OnNotifyAddActiveSound(FActiveSound& ActiveSound) {}
+	virtual void NotifyActiveSoundCreated(FActiveSound& ActiveSound) {}
+
+	/**
+	 * Called when an active sound has just been removed from the audio engine, by being stopped or virtualized.
+	 * In either case, the referenced ActiveSound object is about to be deleted; any pointers to it should be discarded.
+	 * 
+	 * NOTE! Called on the AudioThread
+	 *
+	 * @param ActiveSound	The active sound being deleted
+	 */
+	virtual void NotifyActiveSoundDeleting(const FActiveSound& ActiveSound) {}
+
+	/**
+	 * Called when a virtualized sound has just been added to the audio engine,
+	 * both for brand new sounds and for active sounds that have just become virtualized.
+	 * When virtualizing, the corresponding NotifyActiveSoundDeleting will arrive after any fade-out has finished.
+	 * You can correlate the two objects by matching their GetPlayOrder() value.
+	 * 
+	 * NOTE! Called on the AudioThread
+	 *
+	 * @param ActiveSound	The virtualized sound being created
+	 */
+	virtual void NotifyVirtualizedSoundCreated(FActiveSound& ActiveSound) {}
 
 	/**
-	 * Called when the active sound is being removed from the audio engine
+	 * Called when a virtualized sound has just been removed from the audio engine, by being stopped or re-triggered.
+	 * In either case, the referenced ActiveSound object is about to be deleted; any pointers to it should be discarded.
+	 * 
 	 * NOTE! Called on the AudioThread
 	 *
-	 * @param ActiveSound	The active sound
+	 * @param ActiveSound	The active sound being removed
 	 */
+	virtual void NotifyVirtualizedSoundDeleting(const FActiveSound& ActiveSound) {}
+
+
+	UE_DEPRECATED(5.5, "OnNotifyAddActiveSound is deprecated. Use NotifyActiveSoundCreated instead.")
+	virtual void OnNotifyAddActiveSound(FActiveSound& ActiveSound) {}
+
+	UE_DEPRECATED(5.5, "OnNotifyPendingDelete is deprecated. Use NotifyActiveSoundDeleting and/or NotifyVirtualizedSoundDeleting instead.")
 	virtual void OnNotifyPendingDelete(const FActiveSound& ActiveSound) {}
 };
diff --git a/Engine/Source/Runtime/Engine/Public/AudioDevice.h b/Engine/Source/Runtime/Engine/Public/AudioDevice.h
--- a/Engine/Source/Runtime/Engine/Public/AudioDevice.h
+++ b/Engine/Source/Runtime/Engine/Public/AudioDevice.h
@@ -864,22 +864,50 @@ public:
 	ENGINE_API void ApplyInteriorSettings(FActiveSound& ActiveSound, FSoundParseParameters& ParseParams) const;
 
 	/**
-	 * Notifies subsystems an active sound is about to be deleted (called on audio thread) - Deprecated, see NotifyPendingDeleteInternal
+	 * Notifies subsystems an active sound is about to be deleted (called on audio thread) - Deprecated, see NotifySubsystemsActiveSoundDeleting
 	 */
-	UE_DEPRECATED(5.3, "NotifyPending is deprecated in public scope. Use IActiveSoundUpdateInterface::OnNotifyPendingDelete instead.")
+	UE_DEPRECATED(5.3, "NotifyPending is deprecated in public scope. Use IActiveSoundUpdateInterface::NotifyActiveSoundDeleting instead.")
 	void NotifyPendingDelete(FActiveSound& ActiveSound) const {}
 
 protected:
 
 	/**
-	 * Notifies subsystems an active sound has been added (called on audio thread)
+	 * Notifies subsystems an active sound has just been added (called on audio thread).
+	 * Called both for brand new sounds and for virtualized sounds that have just become active.
 	 */
-	ENGINE_API void NotifyAddActiveSound(FActiveSound& ActiveSound) const;
+	ENGINE_API void NotifySubsystemsActiveSoundCreated(FActiveSound& ActiveSound) const;
 
 	/**
-	 * Notifies subsystems an active sound is about to be deleted (called on audio thread)
+	 * Notifies subsystems an active sound is about to be deleted (called on audio thread).
+	 * Called when a sound is either stopped or virtualized. In either case,
+	 * the referenced ActiveSound object no longer exists after this call; any pointers to it should be discarded.
 	 */
-	ENGINE_API void NotifyPendingDeleteInternal(FActiveSound& ActiveSound) const;
+	ENGINE_API void NotifySubsystemsActiveSoundDeleting(FActiveSound& ActiveSound) const;
+
+	/**
+	 * Notifies subsystems a virtualized sound has just been added (called on audio thread).
+	 * Called both for brand new sounds and for active sounds that have just become virtualized.
+	 */
+	ENGINE_API void NotifySubsystemsVirtualizedSoundCreated(FActiveSound& ActiveSound) const;
+
+	/**
+	 * Notifies subsystems a virtualized sound is about to be deleted (called on audio thread).
+	 * Called when a sound is either stopped or re-triggered. In either case,
+	 * the referenced ActiveSound object no longer exists after this call; any pointers to it should be discarded.
+	 */
+	ENGINE_API void NotifySubsystemsVirtualizedSoundDeleting(FActiveSound& ActiveSound) const;
+
+	UE_DEPRECATED(5.5, "NotifyAddActiveSound is deprecated. Use NotifySubsystemsActiveSoundCreated instead.")
+	ENGINE_API void NotifyAddActiveSound(FActiveSound& ActiveSound) const
+	{
+		NotifySubsystemsActiveSoundCreated(ActiveSound);
+	}
+
+	UE_DEPRECATED(5.5, "NotifyPendingDeleteInternal is deprecated. Use NotifySubsystemsActiveSoundDeleting and/or NotifySubsystemsVirtualizedSoundDeleting instead.")
+	ENGINE_API void NotifyPendingDeleteInternal(FActiveSound& ActiveSound) const
+	{
+		NotifySubsystemsActiveSoundDeleting(ActiveSound);
+	}
 
 public:
 
@@ -1968,7 +1996,7 @@ public:
 	}
 
 	/**
-	 * Performs the given operation on all subsytems of the given class. It's safe to create new subsystems during this operation, but not to remove subsystems.
+	 * Performs the given operation on all subsystems of the given class. It's safe to create new subsystems during this operation, but not to remove subsystems.
 	 */
 	template <typename TSubsystemClass>
 	void ForEachSubsystem(TFunctionRef<void(TSubsystemClass*)> Operation) const
