Only expose AspectRatioAxisConstraint on the camera component if it is actually going to do anything

diff --git a/Engine/Source/Runtime/Engine/Classes/Camera/CameraComponent.h b/Engine/Source/Runtime/Engine/Classes/Camera/CameraComponent.h
--- a/Engine/Source/Runtime/Engine/Classes/Camera/CameraComponent.h
+++ b/Engine/Source/Runtime/Engine/Classes/Camera/CameraComponent.h
@@ -94,7 +94,7 @@ class UCameraComponent : public USceneComponent
 	void SetAspectRatio(float InAspectRatio) { AspectRatio = InAspectRatio; }
 
 	/** Override for the default aspect ratio axis constraint defined on the local player */
-	UPROPERTY(Interp, EditAnywhere, BlueprintReadWrite, Category = CameraOptions, meta = (EditCondition = "bOverrideAspectRatioAxisConstraint"))
+	UPROPERTY(Interp, EditAnywhere, BlueprintReadWrite, Category = CameraOptions, meta = (EditCondition = "!bConstrainAspectRatio && bOverrideAspectRatioAxisConstraint"))
 	TEnumAsByte<EAspectRatioAxisConstraint> AspectRatioAxisConstraint = EAspectRatioAxisConstraint::AspectRatio_MaintainXFOV;
 	UFUNCTION(BlueprintCallable, Category = Camera)
 	void SetAspectRatioAxisConstraint(EAspectRatioAxisConstraint InAspectRatioAxisConstraint) { AspectRatioAxisConstraint = InAspectRatioAxisConstraint; }
