Fix decal material creation with Substrate.

The domain and blend mode where not set correctly when introducing a ConvertToDecal node.

diff --git a/Engine/Source/Runtime/Engine/Private/Materials/Material.cpp b/Engine/Source/Runtime/Engine/Private/Materials/Material.cpp
--- a/Engine/Source/Runtime/Engine/Private/Materials/Material.cpp
+++ b/Engine/Source/Runtime/Engine/Private/Materials/Material.cpp
@@ -5048,6 +5048,20 @@ void UMaterial::RebuildShadingModelField()
 					bSanitizeMaterial = true;
 				}
 			}
+			else if (SubstrateMaterialInfo.CountShadingModels() == 2 && SubstrateMaterialInfo.HasShadingModel(ESubstrateShadingModel::SSM_Decal))
+			{
+				// If material has SSM_Decal it has to have 'decal' domain and DefaultLit shading model
+				if (MaterialDomain != MD_DeferredDecal || !SubstrateMaterialInfo.HasShadingModel(SSM_DefaultLit))
+				{
+					bSanitizeMaterial = true;
+				}
+
+				// If blend mode is not one of the translucent blend modes, force sanitization
+				if (!(IsTranslucentOnlyBlendMode(BlendMode) || BlendMode == BLEND_AlphaComposite || IsModulateBlendMode(BlendMode)))
+				{
+					bSanitizeMaterial = true;
+				}
+			}
 			else if (SubstrateMaterialInfo.CountShadingModels() > 1 && MaterialDomain == MD_Surface)
 			{
 				// Case with SSS Profile or SSS MFP are already been handled by above cases. Simply fallback onto DefaultLit
@@ -5061,14 +5075,6 @@ void UMaterial::RebuildShadingModelField()
 					bSanitizeMaterial = true;
 				}
 			}
-			else if (SubstrateMaterialInfo.CountShadingModels() == 2 && SubstrateMaterialInfo.HasShadingModel(ESubstrateShadingModel::SSM_Decal))
-			{
-				// If material has SSM_Decal it has to have 'decal' domain and DefaultLit shading model
-				if (MaterialDomain != MD_DeferredDecal || !SubstrateMaterialInfo.HasShadingModel(SSM_DefaultLit))
-				{
-					bSanitizeMaterial = true;
-				}
-			}
 			else
 			{
 				// Clear the material to default Lit
@@ -5086,15 +5092,11 @@ void UMaterial::RebuildShadingModelField()
 
 		if (bSanitizeMaterial)
 		{
+			const FSubstrateMaterialInfo PreSanitizeSubstrateMaterialInfo = SubstrateMaterialInfo;
 			SubstrateMaterialInfo = FSubstrateMaterialInfo();
 			SubstrateMaterialInfo.AddShadingModel(SSM_DefaultLit);
 
-			if (MaterialDomain == MD_Surface)
-			{
-				// Nothing to do, the node should have added its own type. And if not type but from expression, we are going to generate that below.
-				//AddSurfaceSubstrateShadingModelFromMaterialShadingModel(SubstrateMaterialInfo, ShadingModel);
-			}
-			else if (MaterialDomain == MD_DeferredDecal)
+			if (MaterialDomain == MD_DeferredDecal || PreSanitizeSubstrateMaterialInfo.HasShadingModel(SSM_Decal))
 			{
 				SubstrateMaterialInfo.AddShadingModel(SSM_Decal);
 			}
@@ -5118,6 +5120,11 @@ void UMaterial::RebuildShadingModelField()
 			{
 				// TODO
 			}
+			else if (MaterialDomain == MD_Surface)
+			{
+				// Nothing to do, the node should have added its own type. And if not type but from expression, we are going to generate that below.
+				//AddSurfaceSubstrateShadingModelFromMaterialShadingModel(SubstrateMaterialInfo, ShadingModel);
+			}
 		}
 		
 		if (SubstrateMaterialInfo.HasShadingModelFromExpression())
@@ -5255,6 +5262,11 @@ void UMaterial::RebuildShadingModelField()
 				// Decal can have multiple shading model
 				MaterialDomain = EMaterialDomain::MD_DeferredDecal;
 				ShadingModel = MSM_DefaultLit;
+
+				if (!(IsTranslucentOnlyBlendMode(BlendMode) || BlendMode == BLEND_AlphaComposite || IsModulateBlendMode(BlendMode)))
+				{
+					BlendMode = BLEND_TranslucentGreyTransmittance;
+				}
 			}
 
 			// Also update the ShadingModels for remaining pipeline operation
