[WorldPartitionResaveActorsBuilder] Now outputs diff files in nested format instead of single lines, for better readability.

diff --git a/Engine/Source/Editor/UnrealEd/Private/WorldPartition/WorldPartitionResaveActorsBuilder.cpp b/Engine/Source/Editor/UnrealEd/Private/WorldPartition/WorldPartitionResaveActorsBuilder.cpp
--- a/Engine/Source/Editor/UnrealEd/Private/WorldPartition/WorldPartitionResaveActorsBuilder.cpp
+++ b/Engine/Source/Editor/UnrealEd/Private/WorldPartition/WorldPartitionResaveActorsBuilder.cpp
@@ -465,8 +465,8 @@ bool UWorldPartitionResaveActorsBuilder::RunInternal(UWorld* World, const FCellI
 
 					if (bDiffDirtyActorDescs)
 					{
-						DirtyActorDescsOld.Add(ActorDesc->ToString(FWorldPartitionActorDesc::EToStringMode::Full));
-						DirtyActorDescsNew.Add(NewActorDesc->ToString(FWorldPartitionActorDesc::EToStringMode::Full));
+						DirtyActorDescsOld.Add(ActorDesc->ToString(FWorldPartitionActorDesc::EToStringMode::ForDiff));
+						DirtyActorDescsNew.Add(NewActorDesc->ToString(FWorldPartitionActorDesc::EToStringMode::ForDiff));
 					}
 
 					UE_LOG(LogWorldPartitionResaveActorsBuilder, Log, TEXT("Package %s needs to be resaved."), *Package->GetName());
diff --git a/Engine/Source/Runtime/Engine/Private/WorldPartition/WorldPartitionActorDesc.cpp b/Engine/Source/Runtime/Engine/Private/WorldPartition/WorldPartitionActorDesc.cpp
--- a/Engine/Source/Runtime/Engine/Private/WorldPartition/WorldPartitionActorDesc.cpp
+++ b/Engine/Source/Runtime/Engine/Private/WorldPartition/WorldPartitionActorDesc.cpp
@@ -425,100 +425,109 @@ void FWorldPartitionActorDesc::RegisterActorDescDeprecator(TSubclassOf<AActor> A
 
 FString FWorldPartitionActorDesc::ToString(EToStringMode Mode) const
 {
-	auto GetBoolStr = [](bool bValue) -> const TCHAR*
-	{
-		return bValue ? TEXT("1") : TEXT("0");
-	};
-
 	TStringBuilder<1024> Result;
-	Result.Appendf(TEXT("Guid:%s"), *Guid.ToString());
+	const TCHAR LineStart = (Mode == EToStringMode::ForDiff) ? TEXT('\t') : TEXT(' ');
+	const TCHAR* LineEnd = (Mode == EToStringMode::ForDiff) ? LINE_TERMINATOR : nullptr;
+	
+	Result.Appendf(TEXT("Guid:%s%s"), *Guid.ToString(), LineEnd ? LineEnd : TEXT(""));
 
-	if (Mode >= EToStringMode::Compact)
+	auto Append = [&Result, LineStart, LineEnd](const TCHAR* Name, const TCHAR* Value)
 	{
-		FString BoundsStr;
+		Result.AppendChar(LineStart);
+		Result.Append(Name);
+		Result.AppendChar(TEXT(':'));
+		Result.Append(Value);
 
-		if (RuntimeBounds.IsValid)
+		if (LineEnd)
 		{
-			if (GetEditorBounds().Equals(RuntimeBounds))
-			{
-				BoundsStr = RuntimeBounds.ToString();
-			}
-			else
-			{
-				BoundsStr = *FString::Printf(TEXT("(Editor:%s Runtime:%s)"), *GetEditorBounds().ToString(), *RuntimeBounds.ToString());
-			}
+			Result.Append(LineEnd);
 		}
-		else
+	};
+
+	auto AppendFromString = [&Append](const TCHAR* Name, const FString& Value)
+	{
+		Append(Name, *Value);
+	};
+
+	auto AppendToString = [&AppendFromString]<typename Type>(const TCHAR* Name, const Type& Value)
+	{
+		AppendFromString(Name, Value.ToString());
+	};
+
+	auto AppendFromBool = [&Append, Mode](const TCHAR* Name, bool bValue)
+	{
+		Append(Name, (Mode >= EToStringMode::Compact) ? (bValue ? TEXT("true") : TEXT("false")) : (bValue ? TEXT("1") : TEXT("0")));
+	};
+
+	if (Mode >= EToStringMode::Compact)
+	{
+		if (BaseClass.IsValid())
 		{
-			BoundsStr = TEXT("Invalid");
+			AppendToString(TEXT("BaseClass"), BaseClass);
 		}
 
-		Result.Appendf(
-			TEXT(" BaseClass:%s NativeClass:%s Name:%s Label:%s SpatiallyLoaded:%s Bounds:%s RuntimeGrid:%s EditorOnly:%s RuntimeOnly:%s HLODRelevant:%s ListedInSceneOutliner:%s IsMainWorldOnly:%s"),
-			*BaseClass.ToString(), 
-			*NativeClass.ToString(), 
-			*GetActorNameString(),
-			*GetActorLabel().ToString(),
-			GetBoolStr(bIsSpatiallyLoaded),
-			*BoundsStr,
-			*RuntimeGrid.ToString(),
-			GetBoolStr(bActorIsEditorOnly),
-			GetBoolStr(bActorIsRuntimeOnly),
-			GetBoolStr(bActorIsHLODRelevant),
-			GetBoolStr(bActorIsListedInSceneOutliner),
-			GetBoolStr(IsMainWorldOnly())
-		);
+		AppendToString(TEXT("NativeClass"), NativeClass);
+		AppendFromString(TEXT("Name"), GetActorNameString());
+		AppendToString(TEXT("Label"), GetActorLabel());
+		AppendFromBool(TEXT("SpatiallyLoaded"), bIsSpatiallyLoaded);
+		AppendToString(TEXT("EditorBounds"), EditorBounds);
+		AppendToString(TEXT("RuntimeBounds"), RuntimeBounds);
+		AppendToString(TEXT("RuntimeGrid"), RuntimeGrid);
+		AppendFromBool(TEXT("EditorOnly"), bActorIsEditorOnly);
+		AppendFromBool(TEXT("RuntimeOnly"), bActorIsRuntimeOnly);
+		AppendFromBool(TEXT("ListedInSceneOutliner"), bActorIsListedInSceneOutliner);
+		AppendFromBool(TEXT("IsMainWorldOnly"), IsMainWorldOnly());
 
 		if (ParentActor.IsValid())
 		{
-			Result.Appendf(TEXT(" Parent:%s"), *ParentActor.ToString());
+			AppendToString(TEXT("Parent"), ParentActor);
 		}
 
 		if (HLODLayer.IsValid())
 		{
-			Result.Appendf(TEXT(" HLODLayer:%s"), *HLODLayer.ToString());
+			AppendToString(TEXT("HLODLayer"), HLODLayer);
 		}
 
 		if (!FolderPath.IsNone())
 		{
-			Result.Appendf(TEXT(" FolderPath:%s"), *FolderPath.ToString());
+			AppendToString(TEXT("FolderPath"), FolderPath);
 		}
 
 		if (FolderGuid.IsValid())
 		{
-			Result.Appendf(TEXT(" FolderGuid:%s"), *FolderGuid.ToString());
+			AppendToString(TEXT("FolderGuid"), FolderGuid);
 		}
 
 		if (Mode >= EToStringMode::Full)
 		{
 			if (References.Num())
 			{
-				Result.Appendf(TEXT(" References:%s"), *FString::JoinBy(References, TEXT(","), [&](const FGuid& ReferenceGuid) { return ReferenceGuid.ToString(); }));
+				AppendFromString(TEXT("References"), FString::JoinBy(References, TEXT(","), [&](const FGuid& ReferenceGuid) { return ReferenceGuid.ToString(); }));
 			}
 
 			if (EditorOnlyReferences.Num())
 			{
-				Result.Appendf(TEXT(" EditorOnlyReferences:%s"), *FString::JoinBy(EditorOnlyReferences, TEXT(","), [&](const FGuid& ReferenceGuid) { return ReferenceGuid.ToString(); }));
+				AppendFromString(TEXT("EditorOnlyReferences"), FString::JoinBy(EditorOnlyReferences, TEXT(","), [&](const FGuid& ReferenceGuid) { return ReferenceGuid.ToString(); }));
 			}
 
 			if (Tags.Num())
 			{
-				Result.Appendf(TEXT(" Tags:%s"), *FString::JoinBy(Tags, TEXT(","), [&](const FName& TagName) { return TagName.ToString(); }));
+				AppendFromString(TEXT("Tags"), FString::JoinBy(Tags, TEXT(","), [&](const FName& TagName) { return TagName.ToString(); }));
 			}
 
 			if (Properties.Num())
 			{
-				Result.Appendf(TEXT(" Properties:%s"), *Properties.ToString());
+				AppendToString(TEXT("Properties"), Properties);
 			}
 
 			if (DataLayers.Num())
 			{
-				Result.Appendf(TEXT(" DataLayers:%s"), *FString::JoinBy(DataLayers, TEXT(","), [&](const FName& DataLayerName) { return DataLayerName.ToString(); }));
+				AppendFromString(TEXT("DataLayers"), FString::JoinBy(DataLayers, TEXT(","), [&](const FName& DataLayerName) { return DataLayerName.ToString(); }));
 			}
 
 			if (ExternalDataLayerAsset.IsValid())
 			{
-				Result.Appendf(TEXT(" ExternalDataLayerAsset:%s"), *ExternalDataLayerAsset.ToString());
+				AppendToString(TEXT("ExternalDataLayerAsset"), ExternalDataLayerAsset);
 			}
 		}
 	}
@@ -863,4 +872,4 @@ bool FWorldPartitionActorDesc::IsRuntimeRelevant(const FWorldPartitionActorDescI
 
 #undef LOCTEXT_NAMESPACE
 
-#endif
+#endif
\ No newline at end of file
diff --git a/Engine/Source/Runtime/Engine/Public/WorldPartition/WorldPartitionActorDesc.h b/Engine/Source/Runtime/Engine/Public/WorldPartition/WorldPartitionActorDesc.h
--- a/Engine/Source/Runtime/Engine/Public/WorldPartition/WorldPartitionActorDesc.h
+++ b/Engine/Source/Runtime/Engine/Public/WorldPartition/WorldPartitionActorDesc.h
@@ -354,7 +354,8 @@ public:
 	{
 		Guid,
 		Compact,
-		Full
+		Full,
+		ForDiff
 	};
 
 	ENGINE_API FString ToString(EToStringMode Mode = EToStringMode::Compact) const;
