Fix for TObjectIterator iterating over garbage objects in material-related code

diff --git a/Engine/Source/Editor/MaterialEditor/Private/MaterialEditingLibrary.cpp b/Engine/Source/Editor/MaterialEditor/Private/MaterialEditingLibrary.cpp
--- a/Engine/Source/Editor/MaterialEditor/Private/MaterialEditingLibrary.cpp
+++ b/Engine/Source/Editor/MaterialEditor/Private/MaterialEditingLibrary.cpp
@@ -710,13 +710,13 @@ void UMaterialEditingLibrary::RecompileMaterial(UMaterial* Material)
 			FEditorSupportDelegates::RedrawAllViewports.Broadcast();
 
 			// Force particle components to update their view relevance.
-			for (TObjectIterator<UParticleSystemComponent> It; It; ++It)
+			for (TObjectIterator<UParticleSystemComponent> It(/*AdditionalExclusionFlags = */RF_ClassDefaultObject, /*bIncludeDerivedClasses = */true, /*InInternalExclusionFlags = */EInternalObjectFlags::Garbage); It; ++It)
 			{
 				It->bIsViewRelevanceDirty = true;
 			}
 
 			// Update parameter names on any child material instances
-			for (TObjectIterator<UMaterialInstance> It; It; ++It)
+			for (TObjectIterator<UMaterialInstance> It(/*AdditionalExclusionFlags = */RF_ClassDefaultObject, /*bIncludeDerivedClasses = */true, /*InInternalExclusionFlags = */EInternalObjectFlags::Garbage); It; ++It)
 			{
 				if (It->Parent == Material)
 				{
@@ -1000,7 +1000,7 @@ void UMaterialEditingLibrary::UpdateMaterialFunction(UMaterialFunctionInterface*
 				TRACE_CPUPROFILER_EVENT_SCOPE(UpdateAllMaterialInstances)
 
 				// Go through all function instances in memory and recompile them if they are children
-				for (TObjectIterator<UMaterialFunctionInstance> It; It; ++It)
+				for (TObjectIterator<UMaterialFunctionInstance> It(/*AdditionalExclusionFlags = */RF_ClassDefaultObject, /*bIncludeDerivedClasses = */true, /*InInternalExclusionFlags = */EInternalObjectFlags::Garbage); It; ++It)
 				{
 					UMaterialFunctionInstance* FunctionInstance = *It;
 
diff --git a/Engine/Source/Editor/MaterialEditor/Private/MaterialEditor.cpp b/Engine/Source/Editor/MaterialEditor/Private/MaterialEditor.cpp
--- a/Engine/Source/Editor/MaterialEditor/Private/MaterialEditor.cpp
+++ b/Engine/Source/Editor/MaterialEditor/Private/MaterialEditor.cpp
@@ -5380,7 +5380,7 @@ void FMaterialEditor::SetNumericParameterDefaultOnDependentMaterials(EMaterialPa
 	if (MaterialFunction)
 	{
 		// Find all materials that reference this function
-		for (TObjectIterator<UMaterial> It; It; ++It)
+		for (TObjectIterator<UMaterial> It(/*AdditionalExclusionFlags = */RF_ClassDefaultObject, /*bIncludeDerivedClasses = */true, /*InInternalExclusionFlags = */EInternalObjectFlags::Garbage); It; ++It)
 		{
 			UMaterial* CurrentMaterial = *It;
 
@@ -5419,7 +5419,7 @@ void FMaterialEditor::SetNumericParameterDefaultOnDependentMaterials(EMaterialPa
 	}
 
 	// Update MI's that reference any of the materials affected
-	for (TObjectIterator<UMaterialInstance> It; It; ++It)
+	for (TObjectIterator<UMaterialInstance> It(/*AdditionalExclusionFlags = */RF_ClassDefaultObject, /*bIncludeDerivedClasses = */true, /*InInternalExclusionFlags = */EInternalObjectFlags::Garbage); It; ++It)
 	{
 		UMaterialInstance* CurrentMaterialInstance = *It;
 
@@ -6501,7 +6501,7 @@ void FMaterialEditor::CreateDerivedMaterialInstancesPreviews()
 		const int32 MaxCount = CVarMaterialEdMaxDerivedMaterialInstances.GetValueOnGameThread();
 		// TODO consider a mode where we load all MaterialChildList also
 
-		for (TObjectIterator<UMaterialInstance> It; It; ++It)
+		for (TObjectIterator<UMaterialInstance> It(/*AdditionalExclusionFlags = */RF_ClassDefaultObject, /*bIncludeDerivedClasses = */true, /*InInternalExclusionFlags = */EInternalObjectFlags::Garbage); It; ++It)
 		{
 			UMaterialInstance* Instance = *It;
 			if (!Instance->HasStaticParameters())
diff --git a/Engine/Source/Runtime/Engine/Private/Materials/Material.cpp b/Engine/Source/Runtime/Engine/Private/Materials/Material.cpp
--- a/Engine/Source/Runtime/Engine/Private/Materials/Material.cpp
+++ b/Engine/Source/Runtime/Engine/Private/Materials/Material.cpp
@@ -5514,7 +5514,7 @@ void UMaterial::UpdateMaterialShaders(TArray<const FShaderType*>& ShaderTypesToF
 
 			int32 NumMaterials = 0;
 
-			for( TObjectIterator<UMaterial> It; It; ++It )
+			for (TObjectIterator<UMaterial> It(/*AdditionalExclusionFlags = */RF_ClassDefaultObject, /*bIncludeDerivedClasses = */true, /*InInternalExclusionFlags = */EInternalObjectFlags::Garbage); It; ++It)
 			{
 				NumMaterials++;
 			}
@@ -5525,7 +5525,7 @@ void UMaterial::UpdateMaterialShaders(TArray<const FShaderType*>& ShaderTypesToF
 			int32 MaterialIndex = 0;
 
 			// Reinitialize the material shader maps
-			for( TObjectIterator<UMaterial> It; It; ++It )
+			for (TObjectIterator<UMaterial> It(/*AdditionalExclusionFlags = */RF_ClassDefaultObject, /*bIncludeDerivedClasses = */true, /*InInternalExclusionFlags = */EInternalObjectFlags::Garbage); It; ++It)
 			{
 				UMaterial* BaseMaterial = *It;
 				UpdateContext.AddMaterial(BaseMaterial);
@@ -6604,7 +6604,7 @@ static void ListSceneColorMaterials()
 		FString FeatureLevelName;
 		GetFeatureLevelName(FeatureLevel, FeatureLevelName);
 
-		for (TObjectIterator<UMaterialInterface> It; It; ++It)
+		for (TObjectIterator<UMaterialInterface> It(/*AdditionalExclusionFlags = */RF_ClassDefaultObject, /*bIncludeDerivedClasses = */true, /*InInternalExclusionFlags = */EInternalObjectFlags::Garbage); It; ++It)
 		{
 			UMaterialInterface* Mat = *It;
 			const FMaterial* MatRes = Mat->GetRenderProxy()->GetMaterialNoFallback(FeatureLevel);
diff --git a/Engine/Source/Runtime/Engine/Private/Materials/MaterialExpressions.cpp b/Engine/Source/Runtime/Engine/Private/Materials/MaterialExpressions.cpp
--- a/Engine/Source/Runtime/Engine/Private/Materials/MaterialExpressions.cpp
+++ b/Engine/Source/Runtime/Engine/Private/Materials/MaterialExpressions.cpp
@@ -15110,7 +15110,7 @@ void UMaterialFunctionInterface::ForceRecompileForRendering(FMaterialUpdateConte
 	StateId = FGuid::NewGuid();
 
 	// Go through all materials in memory and recompile them if they use this function
-	for (TObjectIterator<UMaterialInterface> It; It; ++It)
+	for (TObjectIterator<UMaterialInterface> It(/*AdditionalExclusionFlags = */RF_ClassDefaultObject, /*bIncludeDerivedClasses = */true, /*InInternalExclusionFlags = */EInternalObjectFlags::Garbage); It; ++It)
 	{
 		UMaterialInterface* CurrentMaterialInterface = *It;
 		if (CurrentMaterialInterface == InPreviewMaterial)
diff --git a/Engine/Source/Runtime/Engine/Private/Materials/MaterialInstance.cpp b/Engine/Source/Runtime/Engine/Private/Materials/MaterialInstance.cpp
--- a/Engine/Source/Runtime/Engine/Private/Materials/MaterialInstance.cpp
+++ b/Engine/Source/Runtime/Engine/Private/Materials/MaterialInstance.cpp
@@ -130,7 +130,7 @@ void RecacheMaterialInstanceUniformExpressions(const UMaterialInterface* ParentM
 	{
 		UE_LOG(LogMaterial,Verbose,TEXT("Recaching MI Uniform Expressions for parent %s"), *ParentMaterial->GetFullName());
 		TArray<FMICReentranceGuard> ReentranceGuards;
-		for (TObjectIterator<UMaterialInstance> It; It; ++It)
+		for (TObjectIterator<UMaterialInstance> It(/*AdditionalExclusionFlags = */RF_ClassDefaultObject, /*bIncludeDerivedClasses = */true, /*InInternalExclusionFlags = */EInternalObjectFlags::Garbage); It; ++It)
 		{
 			UMaterialInstance* MaterialInstance = *It;
 			do 
diff --git a/Engine/Source/Runtime/Engine/Private/Materials/MaterialInterface.cpp b/Engine/Source/Runtime/Engine/Private/Materials/MaterialInterface.cpp
--- a/Engine/Source/Runtime/Engine/Private/Materials/MaterialInterface.cpp
+++ b/Engine/Source/Runtime/Engine/Private/Materials/MaterialInterface.cpp
@@ -614,10 +614,10 @@ void UMaterialInterface::SetForceMipLevelsToBeResident( bool OverrideForceMiplev
 
 void UMaterialInterface::RecacheAllMaterialUniformExpressions(bool bRecreateUniformBuffer)
 {
-	// For each interface, reacache its uniform parameters
-	for( TObjectIterator<UMaterialInterface> MaterialIt; MaterialIt; ++MaterialIt )
+	// For each interface, recache its uniform parameters
+	for (TObjectIterator<UMaterialInterface> It(/*AdditionalExclusionFlags = */RF_ClassDefaultObject, /*bIncludeDerivedClasses = */true, /*InInternalExclusionFlags = */EInternalObjectFlags::Garbage); It; ++It)
 	{
-		MaterialIt->RecacheUniformExpressions(bRecreateUniformBuffer);
+		It->RecacheUniformExpressions(bRecreateUniformBuffer);
 	}
 }
 
@@ -649,7 +649,7 @@ void UMaterialInterface::SubmitRemainingJobsForWorld(UWorld* World, EMaterialSha
 	}
 
 	// Add UI and PP Materials.
-	for (TObjectIterator<UMaterialInterface> It; It; ++It)
+	for (TObjectIterator<UMaterialInterface> It(/*AdditionalExclusionFlags = */RF_ClassDefaultObject, /*bIncludeDerivedClasses = */true, /*InInternalExclusionFlags = */EInternalObjectFlags::Garbage); It; ++It)
 	{
 		UMaterial* Material = It->GetMaterial();
 		if (Material &&
@@ -661,7 +661,7 @@ void UMaterialInterface::SubmitRemainingJobsForWorld(UWorld* World, EMaterialSha
 	}
 
 	// Add Decal Component Materials.
-	for (TObjectIterator<UDecalComponent> It; It; ++It)
+	for (TObjectIterator<UDecalComponent> It(/*AdditionalExclusionFlags = */RF_ClassDefaultObject, /*bIncludeDerivedClasses = */true, /*InInternalExclusionFlags = */EInternalObjectFlags::Garbage); It; ++It)
 	{
 		if (It)
 		{
diff --git a/Engine/Source/Runtime/Engine/Private/Materials/MaterialShared.cpp b/Engine/Source/Runtime/Engine/Private/Materials/MaterialShared.cpp
--- a/Engine/Source/Runtime/Engine/Private/Materials/MaterialShared.cpp
+++ b/Engine/Source/Runtime/Engine/Private/Materials/MaterialShared.cpp
@@ -144,7 +144,7 @@ static FAutoConsoleCommand GFlushMaterialUniforms(
 	FConsoleCommandDelegate::CreateStatic(
 		[]()
 {
-	for (TObjectIterator<UMaterialInterface> It; It; ++It)
+	for (TObjectIterator<UMaterialInterface> It(/*AdditionalExclusionFlags = */RF_ClassDefaultObject, /*bIncludeDerivedClasses = */true, /*InInternalExclusionFlags = */EInternalObjectFlags::Garbage); It; ++It)
 	{
 		UMaterialInterface* Material = *It;
 		FMaterialRenderProxy* MaterialProxy = Material->GetRenderProxy();
@@ -167,7 +167,7 @@ class FMaterialDumpDebugInfoExecHelper : public FSelfRegisteringExec
 
 			if (RequestedMaterialName.Len() > 0)
 			{
-				for (TObjectIterator<UMaterialInterface> It; It; ++It)
+				for (TObjectIterator<UMaterialInterface> It(/*AdditionalExclusionFlags = */RF_ClassDefaultObject, /*bIncludeDerivedClasses = */true, /*InInternalExclusionFlags = */EInternalObjectFlags::Garbage); It; ++It)
 				{
 					UMaterialInterface* Material = *It;
 					if (Material && Material->GetName() == RequestedMaterialName)
@@ -4689,7 +4689,7 @@ FMaterialUpdateContext::~FMaterialUpdateContext()
 	// Go through all loaded material instances and recompile their static permutation resources if needed
 	// This is necessary since the parent UMaterial stores information about how it should be rendered, (eg bUsesDistortion)
 	// but the child can have its own shader map which may not contain all the shaders that the parent's settings indicate that it should.
-	for (TObjectIterator<UMaterialInstance> It; It; ++It)
+	for (TObjectIterator<UMaterialInstance> It(/*AdditionalExclusionFlags = */RF_ClassDefaultObject, /*bIncludeDerivedClasses = */true, /*InInternalExclusionFlags = */EInternalObjectFlags::Garbage); It; ++It)
 	{
 		UMaterialInstance* CurrentMaterialInstance = *It;
 		UMaterial* BaseMaterial = CurrentMaterialInstance->GetMaterial();
diff --git a/Engine/Source/Runtime/Engine/Private/Materials/ParameterCollection.cpp b/Engine/Source/Runtime/Engine/Private/Materials/ParameterCollection.cpp
--- a/Engine/Source/Runtime/Engine/Private/Materials/ParameterCollection.cpp
+++ b/Engine/Source/Runtime/Engine/Private/Materials/ParameterCollection.cpp
@@ -55,7 +55,7 @@ void UMaterialParameterCollection::PostLoad()
 
 void UMaterialParameterCollection::SetupWorldParameterCollectionInstances()
 {
-	for (TObjectIterator<UWorld> It; It; ++It)
+	for (TObjectIterator<UWorld> It(/*AdditionalExclusionFlags = */RF_ClassDefaultObject, /*bIncludeDerivedClasses = */true, /*InInternalExclusionFlags = */EInternalObjectFlags::Garbage); It; ++It)
 	{
 		UWorld* CurrentWorld = *It;
 		ULevel* Level = CurrentWorld->PersistentLevel;
@@ -269,7 +269,7 @@ void UMaterialParameterCollection::PostEditChangeProperty(FPropertyChangedEvent&
 			FMaterialUpdateContext UpdateContext;
 
 			// Go through all materials in memory and recompile them if they use this material parameter collection
-			for (TObjectIterator<UMaterial> It; It; ++It)
+			for (TObjectIterator<UMaterial> It(/*AdditionalExclusionFlags = */RF_ClassDefaultObject, /*bIncludeDerivedClasses = */true, /*InInternalExclusionFlags = */EInternalObjectFlags::Garbage); It; ++It)
 			{
 				UMaterial* CurrentMaterial = *It;
 
@@ -305,7 +305,7 @@ void UMaterialParameterCollection::PostEditChangeProperty(FPropertyChangedEvent&
 			}
 
 			// Recreate all uniform buffers based off of this collection
-			for (TObjectIterator<UWorld> It; It; ++It)
+			for (TObjectIterator<UWorld> It(/*AdditionalExclusionFlags = */RF_ClassDefaultObject, /*bIncludeDerivedClasses = */true, /*InInternalExclusionFlags = */EInternalObjectFlags::Garbage); It; ++It)
 			{
 				UWorld* CurrentWorld = *It;
 				CurrentWorld->UpdateParameterCollectionInstances(true, true);
@@ -317,7 +317,7 @@ void UMaterialParameterCollection::PostEditChangeProperty(FPropertyChangedEvent&
 	else
 	{
 		// We didn't need to recreate the uniform buffer, just update its contents
-		for (TObjectIterator<UWorld> It; It; ++It)
+		for (TObjectIterator<UWorld> It(/*AdditionalExclusionFlags = */RF_ClassDefaultObject, /*bIncludeDerivedClasses = */true, /*InInternalExclusionFlags = */EInternalObjectFlags::Garbage); It; ++It)
 		{
 			UWorld* CurrentWorld = *It;
 			CurrentWorld->UpdateParameterCollectionInstances(true, false);
