Add support for explicit DDX/DDY driven mip level mode for runtime virtual texture sample.
The derivatives can be provided in world or UV space.
Providing DDX/DDY of a world space node should now be preferred over using the old "Ignore Input WorldPosition" mip level mode. That mode will now no longer be shown for selection (though legacy materials using it will continue to work).

diff --git a/Engine/Source/Runtime/Engine/Private/Materials/MaterialExpressions.cpp b/Engine/Source/Runtime/Engine/Private/Materials/MaterialExpressions.cpp
--- a/Engine/Source/Runtime/Engine/Private/Materials/MaterialExpressions.cpp
+++ b/Engine/Source/Runtime/Engine/Private/Materials/MaterialExpressions.cpp
@@ -3255,13 +3255,62 @@ UObject* UMaterialExpressionRuntimeVirtualTextureSample::GetReferencedTexture()
 
 #if WITH_EDITOR
 
+TArrayView<FExpressionInput*> UMaterialExpressionRuntimeVirtualTextureSample::GetInputsView()
+{
+	CachedInputs.Empty();
+	CachedInputs.Add(&Coordinates);
+	CachedInputs.Add(&WorldPosition);
+	if (MipValueMode == RVTMVM_MipLevel || MipValueMode == RVTMVM_MipBias)
+	{
+		CachedInputs.Add(&MipValue);
+	}
+	if (MipValueMode == RVTMVM_DerivativeUV || MipValueMode == RVTMVM_DerivativeWorld)
+	{
+		CachedInputs.Add(&DDX);
+		CachedInputs.Add(&DDY);
+	}
+	return CachedInputs;
+}
+
 FName UMaterialExpressionRuntimeVirtualTextureSample::GetInputName(int32 InputIndex) const
 {
 	if (CachedInputs[InputIndex] == &WorldPosition)
 	{
 		return GetWorldPositionInputName(WorldPositionOriginType);
 	}
-
+	if (CachedInputs[InputIndex] == &MipValue)
+	{
+		if (MipValueMode == RVTMVM_MipLevel)
+		{
+			return TEXT("Mip Level");
+		}
+		if (MipValueMode == RVTMVM_MipBias)
+		{
+			return TEXT("Mip Bias");
+		}
+	}
+	if (CachedInputs[InputIndex] == &DDX)
+	{
+		if (MipValueMode == RVTMVM_DerivativeUV)
+		{
+			return TEXT("DDX (UV)");
+		}
+		if (MipValueMode == RVTMVM_DerivativeWorld)
+		{
+			return TEXT("DDX (World)");
+		}
+	}
+	if (CachedInputs[InputIndex] == &DDY)
+	{
+		if (MipValueMode == RVTMVM_DerivativeUV)
+		{
+			return TEXT("DDY (UV)");
+		}
+		if (MipValueMode == RVTMVM_DerivativeWorld)
+		{
+			return TEXT("DDY (World)");
+		}
+	}
 	return Super::GetInputName(InputIndex);
 }
 
@@ -3302,7 +3351,7 @@ void UMaterialExpressionRuntimeVirtualTextureSample::PostEditChangeProperty(FPro
 			FEditorSupportDelegates::ForcePropertyWindowRebuild.Broadcast(this);
 		}
 	}
-	else if (PropertyName == GET_MEMBER_NAME_STRING_CHECKED(ThisClass, WorldPositionOriginType))
+	else if (PropertyName == GET_MEMBER_NAME_STRING_CHECKED(ThisClass, WorldPositionOriginType) || PropertyName == GET_MEMBER_NAME_STRING_CHECKED(ThisClass, MipValueMode))
 	{
 		if (GraphNode)
 		{
@@ -3598,9 +3647,37 @@ int32 UMaterialExpressionRuntimeVirtualTextureSample::Compile(class FMaterialCom
 		TextureMipLevelMode = TMVM_MipBias;
 		MipValue0Index = bMipValueExpressionValid ? MipValue.Compile(Compiler) : Compiler->Constant(0);
 	}
+	else if (MipValueMode == RVTMVM_DerivativeUV || MipValueMode == RVTMVM_DerivativeWorld)
+	{
+		if (DDX.GetTracedInput().Expression == nullptr || DDX.GetTracedInput().Expression == nullptr)
+		{
+			Compiler->Errorf(TEXT("Derivative MipValueMode requires connected DDX and DDY pins."));
+		}
+
+		TextureMipLevelMode = TMVM_Derivative;
+		const int32 Ddx = DDX.Compile(Compiler);
+		const int32 Ddy = DDY.Compile(Compiler);
+
+		if (MipValueMode == RVTMVM_DerivativeUV)
+		{
+			MipValue0Index = Ddx;
+			MipValue0Index = Ddy;
+		}
+		else if (MipValueMode == RVTMVM_DerivativeWorld)
+		{
+ 			const int32 UDdx = Compiler->Dot(Ddx, Uniforms[ERuntimeVirtualTextureShaderUniform_WorldToUVTransform1]);
+ 			const int32 VDdx = Compiler->Dot(Ddx, Uniforms[ERuntimeVirtualTextureShaderUniform_WorldToUVTransform2]);
+ 			MipValue0Index = Compiler->AppendVector(UDdx, VDdx);
+
+			const int32 UDdy = Compiler->Dot(Ddy, Uniforms[ERuntimeVirtualTextureShaderUniform_WorldToUVTransform1]);
+			const int32 VDdy = Compiler->Dot(Ddy, Uniforms[ERuntimeVirtualTextureShaderUniform_WorldToUVTransform2]);
+			MipValue1Index = Compiler->AppendVector(UDdy, VDdy);
+		}
+	}
 	else if (MipValueMode == RVTMVM_RecalculateDerivatives)
 	{
 		// Calculate derivatives from world position.
+		// This is legacy/hidden, and is better implemented in the material graph using RVTMVM_DerivativeWorld.
 		TextureMipLevelMode = TMVM_Derivative;
 		const int32 WorldPos = Compiler->WorldPosition(WPT_CameraRelative);
 		const int32 WorldPositionDdx = Compiler->DDX(WorldPos);
diff --git a/Engine/Source/Runtime/Engine/Public/Materials/MaterialExpressionRuntimeVirtualTextureSample.h b/Engine/Source/Runtime/Engine/Public/Materials/MaterialExpressionRuntimeVirtualTextureSample.h
--- a/Engine/Source/Runtime/Engine/Public/Materials/MaterialExpressionRuntimeVirtualTextureSample.h
+++ b/Engine/Source/Runtime/Engine/Public/Materials/MaterialExpressionRuntimeVirtualTextureSample.h
@@ -21,23 +21,34 @@ enum ERuntimeVirtualTextureMipValueMode : int
 	RVTMVM_None UMETA(DisplayName = "Default"),
 
 	/* 
-	 * Use an absolute mip level from the MipValue pin. 
+	 * Use an absolute mip level from the MipLevel pin. 
 	 * 0 is full resolution.
 	 */
 	RVTMVM_MipLevel UMETA(DisplayName = "Mip Level"),
 
 	/* 
-	 * Bias the default computed mip level using the MipValue pin. 
+	 * Bias the default computed mip level using the MipBias pin. 
 	 * Negative values increase resolution.
 	 */
 	RVTMVM_MipBias UMETA(DisplayName = "Mip Bias"),
 
 	/* 
-	 * This is like 'Default' but it ignores the WorldPosition pin when computing the mip level.
-	 * It uses the actual pixel WorldPosition instead.
-	 * This can prevent sampling mip 0 if the WorldPosition pin gives a constant value.
+	* Compute mip level from world position derivatives.
+	* This is intended for cases where the value passed to the WorldPosition pin doesn't give good derivatives.
+	* (For example when using a constant value from primitive world position).
+	* This is deprecated. Use Derivative (World Space) instead, and pass in DDX/DDY of world space.
 	 */
-	RVTMVM_RecalculateDerivatives UMETA(DisplayName = "Ignore Input WorldPosition "),
+	RVTMVM_RecalculateDerivatives UMETA(Hidden, DisplayName = "Ignore Input WorldPosition"),
+
+	/* 
+	 * Compute mip level from explicitly provided DDX and DDY derivatives of the virtual texture UV coordinates.
+	 */
+	RVTMVM_DerivativeUV UMETA(DisplayName = "Derivatives (UV Space)"),
+
+	/*
+	 * Compute mip level from explicitly provided DDX and DDY derivatives of the world position.
+	 */
+	 RVTMVM_DerivativeWorld UMETA(DisplayName = "Derivatives (World Space)"),
 
 	RVTMVM_MAX,
 };
@@ -74,6 +85,14 @@ class UMaterialExpressionRuntimeVirtualTextureSample : public UMaterialExpressio
 	UPROPERTY(meta = (RequiredInput = "false"))
 	FExpressionInput MipValue;
 
+	/** Derivative over the X axis. Enabled only if MipValueMode is one of the derivative modes. Meaning depends on the derivative mode. */
+	UPROPERTY(meta = (RequiredInput = "false"))
+	FExpressionInput DDX;
+	
+	/** Derivative over the Y axis. Enabled only if MipValueMode is one of the derivative modes. Meaning depends on the derivative mode. */
+	UPROPERTY(meta = (RequiredInput = "false"))
+	FExpressionInput DDY;
+
 	/** The virtual texture object to sample. */
 	UPROPERTY(EditAnywhere, Category = VirtualTexture)
 	TObjectPtr<class URuntimeVirtualTexture> VirtualTexture;
@@ -90,25 +109,25 @@ class UMaterialExpressionRuntimeVirtualTextureSample : public UMaterialExpressio
 	UPROPERTY(EditAnywhere, Category = VirtualTexture, meta = (DisplayName = "Enable adaptive page table"))
 	bool bAdaptive = false;
 
-	/** 
-	 * Enable virtual texture feedback. 
-	 * Disabling this can result in the virtual texture not reaching the correct mip level. 
-	 * It should only be used in cases where we don't care about the correct mip level being resident, or some other process is maintaining the correct level.
-	 */
-	UPROPERTY(EditAnywhere, Category = VirtualTexture)
-	bool bEnableFeedback = true;
-
 	/** Defines the reference space for the WorldPosition input. */
-	UPROPERTY(EditAnywhere, Category = TextureSample)
+	UPROPERTY(EditAnywhere, Category = "UV Coordinates")
 	EPositionOrigin WorldPositionOriginType = EPositionOrigin::Absolute;
 
+	/** Defines the texture addressing mode. */
+	UPROPERTY(EditAnywhere, Category = TextureSample)
+	TEnumAsByte<enum ERuntimeVirtualTextureTextureAddressMode> TextureAddressMode = RVTTA_Clamp;
+
 	/** Defines how the mip level is calculated for the virtual texture lookup. */
 	UPROPERTY(EditAnywhere, Category = TextureSample)
 	TEnumAsByte<enum ERuntimeVirtualTextureMipValueMode> MipValueMode = RVTMVM_None;
 
-	/** Defines the texture addressing mode. */
+	/** 
+	 * Enable virtual texture feedback. 
+	 * Disabling this can result in the virtual texture not reaching the correct mip level. 
+	 * It should only be used in cases where we don't care about the correct mip level being resident, or some other process is maintaining the correct level.
+	 */
 	UPROPERTY(EditAnywhere, Category = TextureSample)
-	TEnumAsByte<enum ERuntimeVirtualTextureTextureAddressMode> TextureAddressMode = RVTTA_Clamp;
+	bool bEnableFeedback = true;
 
 	/** Init settings that affect shader compilation and need to match the current VirtualTexture */
 	ENGINE_API bool InitVirtualTextureDependentSettings();
@@ -123,6 +142,7 @@ protected:
 
 #if WITH_EDITOR
 	ENGINE_API virtual void PostLoad() override;
+	ENGINE_API TArrayView<FExpressionInput*> GetInputsView() override;
 	ENGINE_API virtual FName GetInputName(int32 InputIndex) const override;
 	ENGINE_API virtual int32 Compile(class FMaterialCompiler* Compiler, int32 OutputIndex) override;
 	ENGINE_API virtual void GetCaption(TArray<FString>& OutCaptions) const override;
