Fix redundant code typo.
#rnx

diff --git a/Engine/Source/Runtime/Engine/Private/Materials/MaterialExpressions.cpp b/Engine/Source/Runtime/Engine/Private/Materials/MaterialExpressions.cpp
--- a/Engine/Source/Runtime/Engine/Private/Materials/MaterialExpressions.cpp
+++ b/Engine/Source/Runtime/Engine/Private/Materials/MaterialExpressions.cpp
@@ -3649,7 +3649,7 @@ int32 UMaterialExpressionRuntimeVirtualTextureSample::Compile(class FMaterialCom
 	}
 	else if (MipValueMode == RVTMVM_DerivativeUV || MipValueMode == RVTMVM_DerivativeWorld)
 	{
-		if (DDX.GetTracedInput().Expression == nullptr || DDX.GetTracedInput().Expression == nullptr)
+		if (DDX.GetTracedInput().Expression == nullptr || DDY.GetTracedInput().Expression == nullptr)
 		{
 			Compiler->Errorf(TEXT("Derivative MipValueMode requires connected DDX and DDY pins."));
 		}
