MaterialExpressionIf
- Allow bool and staticbool inputs for A and B inputs to be supported

diff --git a/Engine/Source/Runtime/Engine/Private/Materials/HLSLMaterialTranslator.cpp b/Engine/Source/Runtime/Engine/Private/Materials/HLSLMaterialTranslator.cpp
--- a/Engine/Source/Runtime/Engine/Private/Materials/HLSLMaterialTranslator.cpp
+++ b/Engine/Source/Runtime/Engine/Private/Materials/HLSLMaterialTranslator.cpp
@@ -4185,9 +4185,9 @@ bool FHLSLMaterialTranslator::GetTextureCollectionForExpression(int32 Index, int
 // GetArithmeticResultType
 EMaterialValueType FHLSLMaterialTranslator::GetArithmeticResultType(EMaterialValueType TypeA, EMaterialValueType TypeB)
 {
-	if (!IsNumericType(TypeA) || !IsNumericType(TypeB))
+	if (!IsPrimitiveType(TypeA) || !IsPrimitiveType(TypeB))
 	{
-		Errorf(TEXT("Attempting to perform arithmetic on non-numeric types: %s %s"), DescribeType(TypeA),DescribeType(TypeB));
+		Errorf(TEXT("Attempting to perform arithmetic on non-primitive types: %s %s"), DescribeType(TypeA),DescribeType(TypeB));
 		return MCT_Unknown;
 	}
 
diff --git a/Engine/Source/Runtime/Engine/Private/Materials/MaterialExpressions.cpp b/Engine/Source/Runtime/Engine/Private/Materials/MaterialExpressions.cpp
--- a/Engine/Source/Runtime/Engine/Private/Materials/MaterialExpressions.cpp
+++ b/Engine/Source/Runtime/Engine/Private/Materials/MaterialExpressions.cpp
@@ -12418,14 +12418,14 @@ int32 UMaterialExpressionIf::Compile(class FMaterialCompiler* Compiler, int32 Ou
 	int32 CompiledA = A.Compile(Compiler);
 	int32 CompiledB = B.GetTracedInput().Expression ? B.Compile(Compiler) : Compiler->Constant(ConstB);
 
-	if(!IsFloatNumericType(Compiler->GetType(CompiledA)))
+	if(!IsPrimitiveType(Compiler->GetType(CompiledA)))
 	{
-		return Compiler->Errorf(TEXT("If input A must be of type float."));
+		return Compiler->Errorf(TEXT("If input A must be a primitive type."));
 	}
 
-	if(!IsFloatNumericType(Compiler->GetType(CompiledB)))
+	if(!IsPrimitiveType(Compiler->GetType(CompiledB)))
 	{
-		return Compiler->Errorf(TEXT("If input B must be of type float."));
+		return Compiler->Errorf(TEXT("If input B must be a primitive type."));
 	}
 
 	int32 Arg3 = AGreaterThanB.Compile(Compiler);
@@ -12456,20 +12456,7 @@ uint32 UMaterialExpressionIf::GetInputType(int32 InputIndex)
 	// First two inputs are always float
 	if (InputIndex == 0 || InputIndex == 1)
 	{
-		if ((A.GetTracedInput().Expression && A.Expression->IsResultMaterialAttributes(A.OutputIndex)) ||
-			(B.GetTracedInput().Expression && B.Expression->IsResultMaterialAttributes(B.OutputIndex)))
-		{
-			return MCT_MaterialAttributes;
-		}
-		else if ((A.GetTracedInput().Expression && A.Expression->GetOutputType(0) == MCT_ShadingModel) &&
-			(B.GetTracedInput().Expression && B.Expression->GetOutputType(0) == MCT_ShadingModel))
-		{
-			return MCT_ShadingModel;
-		}
-		else
-		{
-			return MCT_Float;
-		}
+		return MCT_MaterialAttributes | MCT_Numeric | MCT_ShadingModel | MCT_StaticBool | MCT_Bool;
 	}
 
 	return MCT_Unknown;
diff --git a/Engine/Source/Runtime/Engine/Public/MaterialShared.h b/Engine/Source/Runtime/Engine/Public/MaterialShared.h
--- a/Engine/Source/Runtime/Engine/Public/MaterialShared.h
+++ b/Engine/Source/Runtime/Engine/Public/MaterialShared.h
@@ -256,6 +256,12 @@ inline bool IsNumericType(EMaterialValueType InType)
 	return IsMaterialValueType(InType, EMaterialValueType(MCT_Float | MCT_LWCType | MCT_UInt | MCT_ShadingModel));
 }
 
+inline bool IsPrimitiveType(EMaterialValueType InType)
+{
+	// 'ShadingModel' is considered an 'int' 
+	return IsMaterialValueType(InType, EMaterialValueType(MCT_Float | MCT_LWCType | MCT_UInt | MCT_ShadingModel | MCT_Bool | MCT_StaticBool));
+}
+
 inline EMaterialValueType MakeNonLWCType(EMaterialValueType Type)
 {
 	switch (Type)
