Catch incorrect FinishCompilation usage.
#rnx

diff --git a/Engine/Source/Runtime/Engine/Private/TextureCompiler.cpp b/Engine/Source/Runtime/Engine/Private/TextureCompiler.cpp
--- a/Engine/Source/Runtime/Engine/Private/TextureCompiler.cpp
+++ b/Engine/Source/Runtime/Engine/Private/TextureCompiler.cpp
@@ -340,6 +340,25 @@ void FTextureCompilingManager::FinishCompilation(TArrayView<UTexture* const> InT
 {
 	TRACE_CPUPROFILER_EVENT_SCOPE(FTextureCompilingManager::FinishCompilation);
 
+	if (InTextures.Num() == 0)
+	{
+		return;
+	}
+
+	if (bIsRoutingPostCompilation)
+	{
+		// This ends up modifying the registered texture buckets which is not allowed
+		// when we are routing PostCompilation. Plus, it doesn't make much sense to 
+		// be calling FinishCompilation while we are in the middle of finishing
+		// compilations!
+		// This is likely because a worker task got scheduled during a wait inside
+		// PostCompilation and it's randomly running during the wait, causing crashes.
+		// Workers that need to interact with textures should do that work in response to
+		// a game tick via e.g. ExecuteOnGameThread
+		UE_LOG(LogTexture, Fatal, TEXT("Calling FinishCompilation is not allowed during PostCompilation. NumTextures = %d, Texture[0] = %s"), InTextures.Num(), *InTextures[0]->GetPathName());
+	}
+
+
 	using namespace TextureCompilingManagerImpl;
 	check(IsInGameThread());
 
