Always batch GPU skin cache dispatches. Removes need to have separate code path, removes unnecessary dispatches when meshes are recreated or created -> removed within the same tick. Also fixed issue with the manual flush behavior that would flush the command list and dispatch without rebinding the shader.

diff --git a/Engine/Source/Runtime/Engine/Private/GPUSkinCache.cpp b/Engine/Source/Runtime/Engine/Private/GPUSkinCache.cpp
--- a/Engine/Source/Runtime/Engine/Private/GPUSkinCache.cpp
+++ b/Engine/Source/Runtime/Engine/Private/GPUSkinCache.cpp
@@ -1483,7 +1483,6 @@ void FGPUSkinCache::DoDispatch(FRHICommandList& RHICmdList)
 				INC_DWORD_STAT_BY(STAT_GPUSkinCache_TotalNumVertices, VertexCountAlign64 * 64);
 				RHICmdList.DispatchComputeShader(VertexCountAlign64, 1, 1);
 
-				IncrementDispatchCounter(RHICmdList);
 				BuffersToTransitionToRead.Add(DispatchData.GetPreviousPositionRWBuffer());
 			}
 
@@ -1504,10 +1503,15 @@ void FGPUSkinCache::DoDispatch(FRHICommandList& RHICmdList)
 				INC_DWORD_STAT_BY(STAT_GPUSkinCache_TotalNumVertices, VertexCountAlign64 * 64);
 				RHICmdList.DispatchComputeShader(VertexCountAlign64, 1, 1);
 
-				IncrementDispatchCounter(RHICmdList);
 				BuffersToTransitionToRead.Add(DispatchData.GetPositionRWBuffer());
 			}
 
+			if (IncrementDispatchCounter(RHICmdList))
+			{
+				// The command list was submitted so rebind the shader on the next loop.
+				LastShaderIndex = -1;
+			}
+
 			BuffersToTransitionToRead.Add(DispatchData.GetTangentRWBuffer());
 			check(DispatchData.PreviousPositionBuffer != DispatchData.PositionBuffer);
 		}
@@ -1619,59 +1623,6 @@ void FGPUSkinCache::DoDispatch(FRHICommandList& RHICmdList)
 	BatchDispatches.Reset();
 }
 
-void FGPUSkinCache::DoDispatch(FRHICommandList& RHICmdList, FGPUSkinCacheEntry* SkinCacheEntry, int32 Section, int32 RevisionNumber)
-{
-	RenderCaptureInterface::FScopedCapture RenderCapture(GNumDispatchesToCapture > 0, &RHICmdList, TEXT("GPUSkinCache"));
-	GNumDispatchesToCapture = FMath::Max(GNumDispatchesToCapture - 1, 0);
-
-	SCOPED_GPU_STAT(RHICmdList, GPUSkinCache);
-
-	INC_DWORD_STAT(STAT_GPUSkinCache_TotalNumChunks);
-
-	TArray<FSkinCacheRWBuffer*> BuffersToTransitionToRead;
-
-	TArray<FSkinCacheRWBuffer*> BuffersToTransitionForSkinning;
-	PrepareUpdateSkinning(SkinCacheEntry, Section, RevisionNumber, &BuffersToTransitionForSkinning);
-	MakeBufferTransitions(RHICmdList, BuffersToTransitionForSkinning, ERHIAccess::UAVCompute);
-
-	TArray<FRHIUnorderedAccessView*> SkinningBuffersToOverlap;
-	GetBufferUAVs(BuffersToTransitionForSkinning, SkinningBuffersToOverlap);
-	RHICmdList.BeginUAVOverlap(SkinningBuffersToOverlap);
-	{
-		DispatchUpdateSkinning(RHICmdList, SkinCacheEntry, Section, RevisionNumber, BuffersToTransitionToRead);
-	}
-	RHICmdList.EndUAVOverlap(SkinningBuffersToOverlap);
-
-	FGPUSkinCacheEntry::FSectionDispatchData& DispatchData = SkinCacheEntry->DispatchData[Section];
-	if (DispatchData.IndexBuffer)
-	{
-		RHICmdList.Transition({
-			DispatchData.GetPositionRWBuffer()->UpdateAccessState(ERHIAccess::SRVCompute),
-			DispatchData.GetActiveTangentRWBuffer()->UpdateAccessState(ERHIAccess::SRVCompute)
-		});
-		if (GRecomputeTangentsParallelDispatch)
-		{
-			RHICmdList.Transition(DispatchData.GetIntermediateAccumulatedTangentBuffer()->UpdateAccessState(ERHIAccess::UAVCompute));
-		}
-		BuffersToTransitionToRead.Add(DispatchData.GetPositionRWBuffer());
-
-		FSkinCacheRWBuffer* StagingBuffer = nullptr;
-		DispatchUpdateSkinTangents(RHICmdList, SkinCacheEntry, Section, StagingBuffer, true);
-		if (GRecomputeTangentsParallelDispatch)
-		{
-			RHICmdList.Transition({
-				DispatchData.GetTangentRWBuffer()->UpdateAccessState(ERHIAccess::UAVCompute),
-				DispatchData.GetIntermediateAccumulatedTangentBuffer()->UpdateAccessState(ERHIAccess::UAVCompute)
-			});
-		}
-		DispatchUpdateSkinTangents(RHICmdList, SkinCacheEntry, Section, StagingBuffer, false);
-	}
-
-	SkinCacheEntry->UpdateVertexFactoryDeclaration(RHICmdList, Section);
-
-	TransitionAllToReadable(RHICmdList, BuffersToTransitionToRead);
-}
-
 bool FGPUSkinCache::ProcessEntry(
 	EGPUSkinCacheEntryMode Mode,
 	FRHICommandList& RHICmdList, 
@@ -1892,7 +1843,6 @@ bool FGPUSkinCache::ProcessEntry(
 		}
 	}
 
-	if (bShouldBatchDispatches)
 	{
 		InOutEntry->bQueuedForDispatch = true;
 
@@ -1911,10 +1861,6 @@ bool FGPUSkinCache::ProcessEntry(
 			BatchDispatches.Add({ InOutEntry, uint32(Section) });
 		}
 	}
-	else
-	{
-		DoDispatch(RHICmdList, InOutEntry, Section, RevisionNumber);
-	}
 
 #if RHI_RAYTRACING
 	if (!Skin->ShouldUseSeparateSkinCacheEntryForRayTracing() || Mode == EGPUSkinCacheEntryMode::RayTracing)
@@ -1958,17 +1904,6 @@ void FGPUSkinCache::ProcessRayTracingGeometryToUpdate(FRHICommandList& RHICmdLis
 
 #endif
 
-void FGPUSkinCache::BeginBatchDispatch()
-{
-	bShouldBatchDispatches = true;
-	DispatchCounter = 0;
-}
-
-void FGPUSkinCache::EndBatchDispatch()
-{
-	bShouldBatchDispatches = false;
-}
-
 void FGPUSkinCache::Release(FGPUSkinCacheEntry*& SkinCacheEntry)
 {
 	if (SkinCacheEntry)
@@ -2338,7 +2273,7 @@ void FGPUSkinCache::CVarSinkFunction()
 
 FAutoConsoleVariableSink FGPUSkinCache::CVarSink(FConsoleCommandDelegate::CreateStatic(&CVarSinkFunction));
 
-void FGPUSkinCache::IncrementDispatchCounter(FRHICommandList& RHICmdList)
+bool FGPUSkinCache::IncrementDispatchCounter(FRHICommandList& RHICmdList)
 {
 	if (GSkinCacheMaxDispatchesPerCmdList > 0)
 	{
@@ -2348,8 +2283,10 @@ void FGPUSkinCache::IncrementDispatchCounter(FRHICommandList& RHICmdList)
 			//UE_LOG(LogSkinCache, Log, TEXT("SubmitCommandsHint issued after %d dispatches"), DispatchCounter);
 			RHICmdList.SubmitCommandsHint();
 			DispatchCounter = 0;
+			return true;
 		}
 	}
+	return false;
 }
 
 uint64 FGPUSkinCache::GetExtraRequiredMemoryAndReset()
diff --git a/Engine/Source/Runtime/Engine/Private/LevelTick.cpp b/Engine/Source/Runtime/Engine/Private/LevelTick.cpp
--- a/Engine/Source/Runtime/Engine/Private/LevelTick.cpp
+++ b/Engine/Source/Runtime/Engine/Private/LevelTick.cpp
@@ -969,29 +969,8 @@ struct FSendAllEndOfFrameUpdates
 	ERHIFeatureLevel::Type FeatureLevel = ERHIFeatureLevel::Num;
 };
 
-void BeginSendEndOfFrameUpdatesDrawEvent(FSendAllEndOfFrameUpdates& SendAllEndOfFrameUpdates)
-{
-	ENQUEUE_RENDER_COMMAND(BeginDrawEventCommand)(UE::RenderCommandPipe::SkeletalMesh, [GPUSkinCache = SendAllEndOfFrameUpdates.GPUSkinCache](FRHICommandList& RHICmdList)
-	{
-		if (GPUSkinCache != nullptr)
-		{
-			GPUSkinCache->BeginBatchDispatch();
-		}
-	});
-}
-
 DECLARE_GPU_STAT(EndOfFrameUpdates);
 DECLARE_GPU_STAT(GPUSkinCacheRayTracingGeometry);
-void EndSendEndOfFrameUpdatesDrawEvent(FSendAllEndOfFrameUpdates& SendAllEndOfFrameUpdates)
-{
-	ENQUEUE_RENDER_COMMAND(EndDrawEventCommand)(UE::RenderCommandPipe::SkeletalMesh, [GPUSkinCache = SendAllEndOfFrameUpdates.GPUSkinCache](FRHICommandList& RHICmdList)
-	{
-		if (GPUSkinCache != nullptr)
-		{
-			GPUSkinCache->EndBatchDispatch();
-		}
-	});
-}
 
 /**
 	* Send all render updates to the rendering thread.
@@ -1055,7 +1034,6 @@ void UWorld::SendAllEndOfFrameUpdates()
 
 	// Issue a GPU event to wrap GPU work done during SendAllEndOfFrameUpdates, like skin cache updates
 	FSendAllEndOfFrameUpdates SendAllEndOfFrameUpdates(Scene);
-	BeginSendEndOfFrameUpdatesDrawEvent(SendAllEndOfFrameUpdates);
 
 	// update all dirty components. 
 	FGuardValue_Bitfield(bPostTickComponentUpdate, true); 
@@ -1166,8 +1144,6 @@ void UWorld::SendAllEndOfFrameUpdates()
 	bMaterialParameterCollectionInstanceNeedsDeferredUpdate = false;
 			
 	LocalComponentsThatNeedEndOfFrameUpdate.Reset();
-
-	EndSendEndOfFrameUpdatesDrawEvent(SendAllEndOfFrameUpdates);
 }
 
 /**
diff --git a/Engine/Source/Runtime/Engine/Public/GPUSkinCache.h b/Engine/Source/Runtime/Engine/Public/GPUSkinCache.h
--- a/Engine/Source/Runtime/Engine/Public/GPUSkinCache.h
+++ b/Engine/Source/Runtime/Engine/Public/GPUSkinCache.h
@@ -203,7 +203,7 @@ public:
 			{
 				PositionBuffers[Index].Buffer.ClassName = PositionsName;
 				PositionBuffers[Index].Buffer.OwnerName = OwnerName;
-				PositionBuffers[Index].Buffer.Initialize(RHICmdList, TEXT("SkinCachePositions"), PosBufferBytesPerElement, NumVertices * 3, PF_R32_FLOAT, BUF_Static);
+				PositionBuffers[Index].Buffer.Initialize(RHICmdList, TEXT("SkinCachePositions"), PosBufferBytesPerElement, NumVertices * 3, PF_R32_FLOAT, ERHIAccess::SRVMask, BUF_Static);
 				PositionBuffers[Index].Buffer.Buffer->SetOwnerName(OwnerName);
 				PositionBuffers[Index].AccessState = ERHIAccess::Unknown;
 			}
@@ -408,8 +408,6 @@ public:
 	void ProcessRayTracingGeometryToUpdate(FRHICommandList& RHICmdList, FGPUSkinCacheEntry* SkinCacheEntry);
 #endif // RHI_RAYTRACING
 
-	void BeginBatchDispatch();
-	void EndBatchDispatch();
 	void ENGINE_API DoDispatch(FRHICommandList& RHICmdList);
 
 	inline ERHIFeatureLevel::Type GetFeatureLevel() const { return FeatureLevel; }
@@ -426,7 +424,6 @@ protected:
 	TArray<FDispatchEntry> BatchDispatches;
 
 	FRWBuffersAllocation* TryAllocBuffer(uint32 NumVertices, bool WithTangnents, bool UseIntermediateTangents, uint32 NumTriangles, FRHICommandList& RHICmdList, const FName& OwnerName);
-	void DoDispatch(FRHICommandList& RHICmdList, FGPUSkinCacheEntry* SkinCacheEntry, int32 Section, int32 RevisionNumber);
 	void DispatchUpdateSkinTangents(FRHICommandList& RHICmdList, FGPUSkinCacheEntry* Entry, int32 SectionIndex, FSkinCacheRWBuffer*& StagingBuffer, bool bTrianglePass);
 
 	void PrepareUpdateSkinning(
@@ -452,7 +449,6 @@ protected:
 	uint64 ExtraRequiredMemory;
 	int32 FlushCounter;
 	bool bRequiresMemoryLimit;
-	bool bShouldBatchDispatches = false;
 
 	// For recompute tangents, holds the data required between compute shaders
 	TArray<FSkinCacheRWBuffer> StagingBuffers;
@@ -464,7 +460,7 @@ protected:
 	static void CVarSinkFunction();
 	static FAutoConsoleVariableSink CVarSink;
 
-	void IncrementDispatchCounter(FRHICommandList& RHICmdList);
+	bool IncrementDispatchCounter(FRHICommandList& RHICmdList);
 	int32 DispatchCounter = 0;
 
 	void PrintMemorySummary() const;
