Force Routing to Union inside of UAudioComponent::AddModulationRouting

- This same fix has already been applied at the FActiveSound level, this is just more consistency in behavior as another edge case was discovered

diff --git a/Engine/Source/Runtime/Engine/Private/Components/AudioComponent.cpp b/Engine/Source/Runtime/Engine/Private/Components/AudioComponent.cpp
--- a/Engine/Source/Runtime/Engine/Private/Components/AudioComponent.cpp
+++ b/Engine/Source/Runtime/Engine/Private/Components/AudioComponent.cpp
@@ -1902,15 +1902,19 @@ void UAudioComponent::AddModulationRouting(const TSet<USoundModulatorBase*>& Mod
 	switch (Destination)
 	{
 		case EModulationDestination::Volume:
+			ModulationRouting.VolumeRouting = EModulationRouting::Union;
 			ModulationRouting.VolumeModulationDestination.Modulators.Append(ConvertedModulators);
 			break;
 		case EModulationDestination::Pitch:
+			ModulationRouting.PitchRouting = EModulationRouting::Union;
 			ModulationRouting.PitchModulationDestination.Modulators.Append(ConvertedModulators);
 			break;
 		case EModulationDestination::Lowpass:
+			ModulationRouting.LowpassRouting = EModulationRouting::Union;
 			ModulationRouting.LowpassModulationDestination.Modulators.Append(ConvertedModulators);
 			break;
 		case EModulationDestination::Highpass:
+			ModulationRouting.HighpassRouting = EModulationRouting::Union;
 			ModulationRouting.HighpassModulationDestination.Modulators.Append(ConvertedModulators);
 			break;
 		default:
