Removed spurious ensure in material code

diff --git a/Engine/Source/Runtime/Engine/Private/Materials/MaterialShader.cpp b/Engine/Source/Runtime/Engine/Private/Materials/MaterialShader.cpp
--- a/Engine/Source/Runtime/Engine/Private/Materials/MaterialShader.cpp
+++ b/Engine/Source/Runtime/Engine/Private/Materials/MaterialShader.cpp
@@ -840,9 +840,7 @@ bool FStaticParameterSet::Equivalent(const FStaticParameterSet& ReferenceSet) co
 	FStaticParameterSet Temp2 = ReferenceSet;
 	Temp1.SortForEquivalent();
 	Temp2.SortForEquivalent();
-	bool bResult = (Temp1 == Temp2);
-	ensure(!bResult || (*this) == ReferenceSet); // if this never fires, then we really didn't need to sort did we?
-	return bResult;
+	return (Temp1 == Temp2);
 }
 
 #if WITH_EDITORONLY_DATA
