ULevel::AddLoadedActors() - When adding multiple actors at once, prevent child actors from being added twice to the Actors array

diff --git a/Engine/Source/Runtime/Engine/Private/Level.cpp b/Engine/Source/Runtime/Engine/Private/Level.cpp
--- a/Engine/Source/Runtime/Engine/Private/Level.cpp
+++ b/Engine/Source/Runtime/Engine/Private/Level.cpp
@@ -879,7 +879,12 @@ void ULevel::AddLoadedActors(const TArray<AActor*>& ActorList, const FTransform*
 
 	for (AActor* Actor : ActorList)
 	{
-		QueueActor(Actor);
+		// Ignore child actors as they are added when their parent is processed.
+		// This garantees a proper ordering in the Actors array.
+		if (!Actor->IsChildActor())
+		{
+			QueueActor(Actor);
+		}
 	}
 
 	FScopedSlowTask SlowTask(ActorsQueue.Num() * 3, LOCTEXT("RegisteringActors", "Registering actors..."));
