[StaticMeshRenderContext]
 - Add support for object cache and IStaticMeshComponent interface
 - Fix potential crash that occured during PostUndo because of the scanning of all components to get their static mesh

#rnx

diff --git a/Engine/Source/Runtime/Engine/Private/StaticMeshResources.cpp b/Engine/Source/Runtime/Engine/Private/StaticMeshResources.cpp
--- a/Engine/Source/Runtime/Engine/Private/StaticMeshResources.cpp
+++ b/Engine/Source/Runtime/Engine/Private/StaticMeshResources.cpp
@@ -10,6 +10,7 @@
 #include "UObject/UObjectIterator.h"
 #include "Engine/SimpleConstructionScript.h"
 
+#include "ObjectCacheContext.h"
 
 FStaticMeshComponentRecreateRenderStateContext::FStaticMeshComponentRecreateRenderStateContext(UStaticMesh* InStaticMesh, bool InUnbuildLighting, bool InRefreshBounds)
 	: FStaticMeshComponentRecreateRenderStateContext(TArray<UStaticMesh*>{ InStaticMesh }, InUnbuildLighting, InRefreshBounds)
@@ -33,44 +34,42 @@ FStaticMeshComponentRecreateRenderStateContext::FStaticMeshComponentRecreateRend
 	{
 		TSet<FSceneInterface*> Scenes;
 
-		const EObjectFlags AdditionalExclusionFlags = RF_ClassDefaultObject;
-		const bool bIncludeDerivedClasses = true;
-		const EInternalObjectFlags InternalExclusionFlags = EInternalObjectFlags::Garbage;
-
-		for (TObjectIterator<UStaticMeshComponent> It(AdditionalExclusionFlags, bIncludeDerivedClasses, InternalExclusionFlags); It; ++It)
+		FObjectCacheContextScope ObjectCacheScope;
+		for (auto It = StaticMeshComponents.CreateIterator(); It; ++It)
 		{
-			UStaticMesh* StaticMesh = It->GetStaticMesh();
-
-			if (StaticMeshComponents.Contains(StaticMesh))
+			for (IStaticMeshComponent* Component : ObjectCacheScope.GetContext().GetStaticMeshComponents((UStaticMesh*)It.Key()))
 			{
-				checkf(!It->IsUnreachable(), TEXT("%s"), *It->GetFullName());
+				IPrimitiveComponent* PrimitiveComponent = Component->GetPrimitiveComponentInterface();
 
-				if (It->bRenderStateCreated)
+				if (PrimitiveComponent->IsRenderStateCreated())
 				{
-					check(It->IsRegistered());
-					It->DestroyRenderState_Concurrent();
-					StaticMeshComponents[StaticMesh].Add(*It);
-					Scenes.Add(It->GetScene());
+					check(PrimitiveComponent->IsRegistered());
+					PrimitiveComponent->DestroyRenderState();
+
+					if (UStaticMeshComponent* StaticMeshComponent = Cast<UStaticMeshComponent>(PrimitiveComponent->GetUObject()))
+					{
+						It.Value().UStaticMeshComponents.Add(StaticMeshComponent);
+					}
+					It.Value().IStaticMeshComponents.Add(Component);
+					Scenes.Add(PrimitiveComponent->GetScene());
 				}
 			}
-			// Recreate dirty render state, if needed, only for components not using the static mesh we currently have released resources for.
-			else if (It->IsRenderStateDirty() && It->IsRegistered() && !It->IsTemplate() && IsValid(*It))
-			{
-				It->DoDeferredRenderUpdates_Concurrent();
-			}
 		}
 
-		UpdateAllPrimitiveSceneInfosForScenes(MoveTemp(Scenes));
+		if (Scenes.Num())
+		{
+			UpdateAllPrimitiveSceneInfosForScenes(MoveTemp(Scenes));
 
-		// Flush the rendering commands generated by the detachments.
-		// The static mesh scene proxies reference the UStaticMesh, and this ensures that they are cleaned up before the UStaticMesh changes.
-		FlushRenderingCommands();
+			// Flush the rendering commands generated by the detachments.
+			// The static mesh scene proxies reference the UStaticMesh, and this ensures that they are cleaned up before the UStaticMesh changes.
+			FlushRenderingCommands();
+		}
 	}
 }
 
 const TArray<UStaticMeshComponent*>& FStaticMeshComponentRecreateRenderStateContext::GetComponentsUsingMesh(UStaticMesh* StaticMesh) const
 {
-	return StaticMeshComponents.FindChecked(StaticMesh);
+	return StaticMeshComponents.FindChecked(StaticMesh).UStaticMeshComponents;
 }
 
 FStaticMeshComponentRecreateRenderStateContext::~FStaticMeshComponentRecreateRenderStateContext()
@@ -81,7 +80,7 @@ FStaticMeshComponentRecreateRenderStateContext::~FStaticMeshComponentRecreateRen
 
 		for (const auto& MeshComponents : StaticMeshComponents)
 		{
-			for (UStaticMeshComponent* Component : MeshComponents.Value)
+			for (UStaticMeshComponent* Component : MeshComponents.Value.UStaticMeshComponents)
 			{
 				if (bUnbuildLighting)
 				{
@@ -95,10 +94,16 @@ FStaticMeshComponentRecreateRenderStateContext::~FStaticMeshComponentRecreateRen
 					Component->UpdateBounds();
 				}
 
-				if (Component->IsRegistered() && !Component->bRenderStateCreated)
+			}
+
+			for (IStaticMeshComponent* Component : MeshComponents.Value.IStaticMeshComponents)
+			{
+				IPrimitiveComponent* PrimitiveComponent = Component->GetPrimitiveComponentInterface();
+
+				if (PrimitiveComponent->IsRegistered() && !PrimitiveComponent->IsRenderStateCreated())
 				{
-					Component->CreateRenderState_Concurrent(nullptr);
-					Scenes.Add(Component->GetScene());
+					PrimitiveComponent->CreateRenderState(nullptr);
+					Scenes.Add(PrimitiveComponent->GetScene());
 				}
 			}
 		}
diff --git a/Engine/Source/Runtime/Engine/Public/StaticMeshResources.h b/Engine/Source/Runtime/Engine/Public/StaticMeshResources.h
--- a/Engine/Source/Runtime/Engine/Public/StaticMeshResources.h
+++ b/Engine/Source/Runtime/Engine/Public/StaticMeshResources.h
@@ -847,8 +847,12 @@ public:
 	ENGINE_API ~FStaticMeshComponentRecreateRenderStateContext();
 
 private:
-
-	TMap<void*, TArray<UStaticMeshComponent*>> StaticMeshComponents;
+	struct FData
+	{
+		TArray<UStaticMeshComponent*> UStaticMeshComponents;
+		TArray<IStaticMeshComponent*> IStaticMeshComponents;
+	};
+	TMap<void*, FData> StaticMeshComponents;
 	bool bUnbuildLighting;
 	bool bRefreshBounds;
 };
