[Material Editor] Deprecated UMaterialExpression::GetInputsView() in favor of making GetInput() the only way to access all inputs in a material expression.

To facilitate iteration, this change introduces FExpressionInputIterator that only provides a convenience layer on top of UMaterialExpression::GetInput()

This change makes code more robust and consistent and enables future further simplifications & optimizations.

diff --git a/Engine/Plugins/Editor/SpeedTreeImporter/Source/SpeedTreeImporter/Private/SpeedTreeImportFactory.cpp b/Engine/Plugins/Editor/SpeedTreeImporter/Source/SpeedTreeImporter/Private/SpeedTreeImportFactory.cpp
--- a/Engine/Plugins/Editor/SpeedTreeImporter/Source/SpeedTreeImporter/Private/SpeedTreeImportFactory.cpp
+++ b/Engine/Plugins/Editor/SpeedTreeImporter/Source/SpeedTreeImporter/Private/SpeedTreeImportFactory.cpp
@@ -575,15 +575,13 @@ void LayoutMaterial(UMaterialInterface* MaterialInterface, bool bOffsetOddColumn
 		{
 			Expression->MaterialExpressionEditorX = FMath::Min(Expression->MaterialExpressionEditorX, -Width);
 
-			TArrayView<FExpressionInput*> Inputs = Expression->GetInputsView();
-			for (int32 InputIndex = 0; InputIndex < Inputs.Num(); ++InputIndex)
+			for (FExpressionInputIterator It{ Expression }; It; ++It)
 			{
-				UMaterialExpression* Input = Inputs[InputIndex]->Expression;
-				if (Input != NULL)
+				if (UMaterialExpression* InputExpression = It->Expression)
 				{
-					if (Input->MaterialExpressionEditorX > Expression->MaterialExpressionEditorX - Width)
+					if (InputExpression->MaterialExpressionEditorX > Expression->MaterialExpressionEditorX - Width)
 					{
-						Input->MaterialExpressionEditorX = Expression->MaterialExpressionEditorX - Width;
+						InputExpression->MaterialExpressionEditorX = Expression->MaterialExpressionEditorX - Width;
 						bContinue = true;
 					}
 				}
@@ -619,10 +617,9 @@ void LayoutMaterial(UMaterialInterface* MaterialInterface, bool bOffsetOddColumn
 				// all the outputs to other expressions
 				for (UMaterialExpression* OtherExpression : Material->GetExpressions())
 				{
-					TArrayView<FExpressionInput*> Inputs = OtherExpression->GetInputsView();
-					for (int32 InputIndex = 0; InputIndex < Inputs.Num(); ++InputIndex)
+					for (FExpressionInputIterator It{ OtherExpression }; It; ++It)
 					{
-						if (Inputs[InputIndex]->Expression == Expression)
+						if (It->Expression == Expression)
 						{
 							Expression->MaterialExpressionEditorY += OtherExpression->MaterialExpressionEditorY;
 							++NumOutputs;
diff --git a/Engine/Plugins/Enterprise/DatasmithImporter/Source/DatasmithImporter/Private/DatasmithMaterialExpressions.cpp b/Engine/Plugins/Enterprise/DatasmithImporter/Source/DatasmithImporter/Private/DatasmithMaterialExpressions.cpp
--- a/Engine/Plugins/Enterprise/DatasmithImporter/Source/DatasmithImporter/Private/DatasmithMaterialExpressions.cpp
+++ b/Engine/Plugins/Enterprise/DatasmithImporter/Source/DatasmithImporter/Private/DatasmithMaterialExpressions.cpp
@@ -276,14 +276,11 @@ void FDatasmithMaterialExpressions::GetSamplersRecursive(UMaterialExpression* Ex
 		TextureSamplers.AddUnique((UMaterialExpressionTextureSample*)Expression);
 	}
 
-	TArrayView<FExpressionInput*> Inputs = Expression->GetInputsView();
-	for (int32 InputIndex = 0; InputIndex < Inputs.Num(); InputIndex++)
+	for (FExpressionInputIterator It{ Expression }; It; ++It)
 	{
-		UMaterialExpression* Input = Inputs[InputIndex]->Expression;
-
-		if (Input != nullptr)
+		if (UMaterialExpression* InputExpression = It->Expression)
 		{
-			GetSamplersRecursive(Input, TextureSamplers);
+			GetSamplersRecursive(InputExpression, TextureSamplers);
 		}
 	}
 }
@@ -3144,10 +3141,15 @@ void FDatasmithMaterialExpressions::CreateUEPbrMaterialGraph(const TSharedPtr< I
 			{
 				for ( int32 ExpressionInput = 0; ExpressionInput < DatasmithExpression->GetInputCount(); ++ExpressionInput )
 				{
-					if ( MaterialOutputExpression->GetInputsView().IsValidIndex( ExpressionInput ) )
+					if (FExpressionInput* MaterialOutputExpressionInput = MaterialOutputExpression->GetInput( ExpressionInput ) )
 					{
-						ConnectExpression( MaterialElement.ToSharedRef(), MaterialExpressions, DatasmithExpression->GetInput( ExpressionInput )->GetExpression(),
-							MaterialOutputExpression->GetInput( ExpressionInput ), DatasmithExpression->GetInput( ExpressionInput )->GetOutputIndex() );
+						ConnectExpression(
+							MaterialElement.ToSharedRef(),
+							MaterialExpressions,
+							DatasmithExpression->GetInput( ExpressionInput )->GetExpression(),
+							MaterialOutputExpressionInput,
+							DatasmithExpression->GetInput( ExpressionInput )->GetOutputIndex()
+						);
 					}
 				}
 			}
diff --git a/Engine/Plugins/Enterprise/GLTFExporter/Source/GLTFExporter/Private/Converters/GLTFMaterialUtilities.cpp b/Engine/Plugins/Enterprise/GLTFExporter/Source/GLTFExporter/Private/Converters/GLTFMaterialUtilities.cpp
--- a/Engine/Plugins/Enterprise/GLTFExporter/Source/GLTFExporter/Private/Converters/GLTFMaterialUtilities.cpp
+++ b/Engine/Plugins/Enterprise/GLTFExporter/Source/GLTFExporter/Private/Converters/GLTFMaterialUtilities.cpp
@@ -719,17 +719,15 @@ FGLTFImportMaterialMatchMakingHelper::FGLTFImportMaterialMatchMakingHelper(FGLTF
 		//If Exported Material is a MaterialInstance then we will use the GetParameter APIs
 		if (!bMaterialInstance)
 		{
-			TArrayView<FExpressionInput*> InputsView = GLTFImportedMaterialFunction->GetInputsView();
-
-			for (FExpressionInput* Input : InputsView)
+			for (FExpressionInputIterator It{ GLTFImportedMaterialFunction}; It; ++It)
 			{
-				if (!Input || !Input->Expression || !Input->IsConnected())
+				if (!It->IsConnected())
 				{
 					continue;
 				}
 
-				FString Name = Input->InputName.ToString();
-				UMaterialExpression* InputExpression = Input->GetTracedInput().Expression;
+				FString Name = It->InputName.ToString();
+				UMaterialExpression* InputExpression = It->GetTracedInput().Expression;
 
 				Inputs.Add(Name, InputExpression);
 			}
diff --git a/Engine/Plugins/Experimental/DynamicMaterial/Source/DynamicMaterialEditor/Private/Components/DMMaterialProperty.cpp b/Engine/Plugins/Experimental/DynamicMaterial/Source/DynamicMaterialEditor/Private/Components/DMMaterialProperty.cpp
--- a/Engine/Plugins/Experimental/DynamicMaterial/Source/DynamicMaterialEditor/Private/Components/DMMaterialProperty.cpp
+++ b/Engine/Plugins/Experimental/DynamicMaterial/Source/DynamicMaterialEditor/Private/Components/DMMaterialProperty.cpp
@@ -588,14 +588,13 @@ void UDMMaterialProperty::AddOutputProcessor(const TSharedRef<FDMMaterialBuildSt
 		UE_DM_NodeComment_Default
 	);
 
-	TArrayView<FExpressionInput*> Inputs = MaterialFunctionCall->GetInputsView();
-
-	if (Inputs.IsEmpty())
+	FExpressionInput* FirstInput = MaterialFunctionCall->GetInput(0);
+	if (!FirstInput)
 	{
 		return;
 	}
 
-	LastPropertyExpression->ConnectExpression(Inputs[0], MaterialPropertyPtr->OutputIndex);
+	LastPropertyExpression->ConnectExpression(FirstInput, MaterialPropertyPtr->OutputIndex);
 	MaterialFunctionCall->ConnectExpression(MaterialPropertyPtr, 0);
 
 	MaterialPropertyPtr->OutputIndex = 0;
diff --git a/Engine/Plugins/Experimental/DynamicMaterial/Source/DynamicMaterialEditor/Private/Components/DMMaterialStageThroughput.cpp b/Engine/Plugins/Experimental/DynamicMaterial/Source/DynamicMaterialEditor/Private/Components/DMMaterialStageThroughput.cpp
--- a/Engine/Plugins/Experimental/DynamicMaterial/Source/DynamicMaterialEditor/Private/Components/DMMaterialStageThroughput.cpp
+++ b/Engine/Plugins/Experimental/DynamicMaterial/Source/DynamicMaterialEditor/Private/Components/DMMaterialStageThroughput.cpp
@@ -266,8 +266,7 @@ void UDMMaterialStageThroughput::ConnectOutputToInput(const TSharedRef<FDMMateri
 	check(!TargetExpressions.IsEmpty());
 
 	UMaterialExpression* TargetExpression = GetExpressionForInput(TargetExpressions, InInputIndex);
-	TConstArrayView<FExpressionInput*> TargetExpressionInputs = TargetExpression->GetInputsView();
-	check(TargetExpressionInputs.IsValidIndex(InInputIndex));
+	check(TargetExpression->GetInput(InInputIndex));
 
 	ConnectOutputToInput_Internal(
 		InBuildState, 
diff --git a/Engine/Plugins/Experimental/DynamicMaterial/Source/DynamicMaterialEditor/Private/Components/MaterialProperties/DMMPNormal.cpp b/Engine/Plugins/Experimental/DynamicMaterial/Source/DynamicMaterialEditor/Private/Components/MaterialProperties/DMMPNormal.cpp
--- a/Engine/Plugins/Experimental/DynamicMaterial/Source/DynamicMaterialEditor/Private/Components/MaterialProperties/DMMPNormal.cpp
+++ b/Engine/Plugins/Experimental/DynamicMaterial/Source/DynamicMaterialEditor/Private/Components/MaterialProperties/DMMPNormal.cpp
@@ -79,14 +79,13 @@ void UDMMaterialPropertyNormal::AddOutputProcessor(const TSharedRef<FDMMaterialB
 		UE_DM_NodeComment_Default
 	);
 
-	TArrayView<FExpressionInput*> Inputs = MaterialFunctionCall->GetInputsView();
-
-	if (Inputs.IsEmpty())
+	FExpressionInput* FirstInput = MaterialFunctionCall->GetInput(0);
+	if (!FirstInput)
 	{
 		return;
 	}
 
-	LastPropertyExpression->ConnectExpression(Inputs[0], MaterialPropertyPtr->OutputIndex);
+	LastPropertyExpression->ConnectExpression(FirstInput, MaterialPropertyPtr->OutputIndex);
 	MaterialFunctionCall->ConnectExpression(MaterialPropertyPtr, 0);
 
 	MaterialPropertyPtr->OutputIndex = 0;
diff --git a/Engine/Plugins/Experimental/DynamicMaterial/Source/DynamicMaterialEditor/Private/Components/MaterialStageInputs/DMMSITextureUV.cpp b/Engine/Plugins/Experimental/DynamicMaterial/Source/DynamicMaterialEditor/Private/Components/MaterialStageInputs/DMMSITextureUV.cpp
--- a/Engine/Plugins/Experimental/DynamicMaterial/Source/DynamicMaterialEditor/Private/Components/MaterialStageInputs/DMMSITextureUV.cpp
+++ b/Engine/Plugins/Experimental/DynamicMaterial/Source/DynamicMaterialEditor/Private/Components/MaterialStageInputs/DMMSITextureUV.cpp
@@ -375,11 +375,9 @@ TArray<UMaterialExpression*> UDMMaterialStageInputTextureUV::CreateTextureUVExpr
 	}
 
 	TMap<FName, int32> NameToInputIndex;
-	TArrayView<FExpressionInput*> FuncInputs = TextureUVFunc->GetInputsView();
-
-	for (int32 InputIdx = 0; InputIdx < FuncInputs.Num(); ++InputIdx)
+	for (FExpressionInputIterator It{ TextureUVFunc }; It; ++It)
 	{
-		NameToInputIndex.Emplace(FuncInputs[InputIdx]->InputName, InputIdx);
+		NameToInputIndex.Emplace(It->InputName, It.Index);
 	}
 
 	// Output nodes
diff --git a/Engine/Plugins/Experimental/DynamicMaterial/Source/DynamicMaterialEditor/Private/Model/DMMaterialNodeArranger.cpp b/Engine/Plugins/Experimental/DynamicMaterial/Source/DynamicMaterialEditor/Private/Model/DMMaterialNodeArranger.cpp
--- a/Engine/Plugins/Experimental/DynamicMaterial/Source/DynamicMaterialEditor/Private/Model/DMMaterialNodeArranger.cpp
+++ b/Engine/Plugins/Experimental/DynamicMaterial/Source/DynamicMaterialEditor/Private/Model/DMMaterialNodeArranger.cpp
@@ -71,29 +71,28 @@ void FDMMaterialNodeArranger::ArrangeNode(TMap<UMaterialExpression*, FIntPoint>&
 		return;
 	}
 
-	TConstArrayView<FExpressionInput*> Inputs = InNode->GetInputsView();
 	InOutNodeSize = {0, 0};
 	const FIntPoint ThisNodeSize = {InNode->GetWidth() * 2, InNode->GetHeight()};
 	FIntPoint ChildOffsetStart = InOffsetStart;
 	ChildOffsetStart.X += ThisNodeSize.X + UE::DynamicMaterialEditor::BuildState::Private::SpaceBetweenNodes;
 
-	for (FExpressionInput* Input : Inputs)
+	for (FExpressionInputIterator It{ InNode }; It; ++It)
 	{
-		if (!Input->Expression)
+		if (!It->IsConnected())
 		{
 			continue;
 		}
 
 		FIntPoint ChildNodeSize;
-		const FIntPoint* NodePosition = InOutNodePositions.Find(Input->Expression);
+		const FIntPoint* NodePosition = InOutNodePositions.Find(It->Expression);
 
 		if (NodePosition && (-NodePosition->X) > InOffsetStart.X)
 		{
-			ChildNodeSize = {Input->Expression->GetWidth(), Input->Expression->GetHeight()};
+			ChildNodeSize = {It->Expression->GetWidth(), It->Expression->GetHeight()};
 		}
 		else
 		{
-			ArrangeNode(InOutNodePositions, ChildOffsetStart, Input->Expression, ChildNodeSize);
+			ArrangeNode(InOutNodePositions, ChildOffsetStart, It->Expression, ChildNodeSize);
 		}
 
 		if (ChildNodeSize.X > 0)
diff --git a/Engine/Plugins/FX/Niagara/Source/NiagaraEditor/Private/ViewModels/Stack/NiagaraStackFunctionInput.cpp b/Engine/Plugins/FX/Niagara/Source/NiagaraEditor/Private/ViewModels/Stack/NiagaraStackFunctionInput.cpp
--- a/Engine/Plugins/FX/Niagara/Source/NiagaraEditor/Private/ViewModels/Stack/NiagaraStackFunctionInput.cpp
+++ b/Engine/Plugins/FX/Niagara/Source/NiagaraEditor/Private/ViewModels/Stack/NiagaraStackFunctionInput.cpp
@@ -1161,8 +1161,10 @@ FString UNiagaraStackFunctionInput::ResolveDisplayNameArgument(const FString& In
 		Material->GetAllExpressionsInMaterialAndFunctionsOfType<UMaterialExpression>(Expressions);
 		for (UMaterialExpression* Expression : Expressions)
 		{
-			TArrayView<FExpressionInput*> ExpressionInputs = Expression->GetInputsView();
-			ExpressionInputsToProcess.Append(ExpressionInputs);
+			for (FExpressionInputIterator It{ Expression}; It; ++It)
+			{
+				ExpressionInputsToProcess.Push(It.Input);
+			}
 		}
 	}
 
diff --git a/Engine/Plugins/Interchange/Runtime/Source/Import/Private/Material/InterchangeMaterialFactory.cpp b/Engine/Plugins/Interchange/Runtime/Source/Import/Private/Material/InterchangeMaterialFactory.cpp
--- a/Engine/Plugins/Interchange/Runtime/Source/Import/Private/Material/InterchangeMaterialFactory.cpp
+++ b/Engine/Plugins/Interchange/Runtime/Source/Import/Private/Material/InterchangeMaterialFactory.cpp
@@ -81,19 +81,14 @@ namespace UE::Interchange::MaterialFactory::Internal
 
 	int32 GetInputIndex(UMaterialExpression& MaterialExpression, const FString& InputName)
 	{
-		int32 ExpressionInputIndex = 0;
-
-		for (const FExpressionInput* ExpressionInput : MaterialExpression.GetInputsView())
+		for (FExpressionInputIterator It{ &MaterialExpression }; It; ++It)
 		{
 			// MaterialFuncCall appends the type to the input name when calling GetInputName
 			// and the InputName in FExpressionInput is optional so we'll check both here to be safe
-			if (MaterialExpression.GetInputName(ExpressionInputIndex) == *InputName ||
-				(ExpressionInput && ExpressionInput->InputName == *InputName))
+			if (MaterialExpression.GetInputName(It.Index) == *InputName || It->InputName == *InputName)
 			{
-				return ExpressionInputIndex;
+				return It.Index;
 			}
-
-			++ExpressionInputIndex;
 		}
 
 		return INDEX_NONE;
diff --git a/Engine/Plugins/Runtime/RenderTrace/Source/Private/Materials/MaterialExpressionPhysicalMaterialOutput.cpp b/Engine/Plugins/Runtime/RenderTrace/Source/Private/Materials/MaterialExpressionPhysicalMaterialOutput.cpp
--- a/Engine/Plugins/Runtime/RenderTrace/Source/Private/Materials/MaterialExpressionPhysicalMaterialOutput.cpp
+++ b/Engine/Plugins/Runtime/RenderTrace/Source/Private/Materials/MaterialExpressionPhysicalMaterialOutput.cpp
@@ -75,7 +75,7 @@ TArrayView<FExpressionInput*> UMaterialExpressionPhysicalMaterialOutput::GetInpu
 
 FExpressionInput* UMaterialExpressionPhysicalMaterialOutput::GetInput(int32 InputIndex)
 {
-	return &Inputs[InputIndex].Input;
+	return Inputs.IsValidIndex(InputIndex) ? &Inputs[InputIndex].Input : nullptr;
 }
 
 FName UMaterialExpressionPhysicalMaterialOutput::GetInputName(int32 InputIndex) const
diff --git a/Engine/Plugins/Tests/InterchangeTests/Source/InterchangeTests/Private/ImportTestFunctions/MaterialXTestFunctions.cpp b/Engine/Plugins/Tests/InterchangeTests/Source/InterchangeTests/Private/ImportTestFunctions/MaterialXTestFunctions.cpp
--- a/Engine/Plugins/Tests/InterchangeTests/Source/InterchangeTests/Private/ImportTestFunctions/MaterialXTestFunctions.cpp
+++ b/Engine/Plugins/Tests/InterchangeTests/Source/InterchangeTests/Private/ImportTestFunctions/MaterialXTestFunctions.cpp
@@ -20,13 +20,13 @@ namespace UE::Interchange::Tests::Private
 		const UMaterialEditorOnlyData* MaterialEditorOnly = Material.GetEditorOnlyData();
 		if (UMaterialExpressionMaterialFunctionCall* FunctionCall = Cast<UMaterialExpressionMaterialFunctionCall>(MaterialEditorOnly->BaseColor.Expression))
 		{
-			TArrayView<FExpressionInput*> Inputs = FunctionCall->GetInputsView();
-			InputNames.Empty(Inputs.Num());
-			for (const FExpressionInput* InputExpression : Inputs)
+			InputNames.Empty();
+
+			for (FExpressionInputIterator It{ FunctionCall }; It; ++It)
 			{
-				if (InputExpression && InputExpression->Expression)
+				if (It->IsConnected())
 				{
-					InputNames.Add(InputExpression->InputName.ToString());
+					InputNames.Add(It->InputName.ToString());
 				}
 			}
 		}
diff --git a/Engine/Source/Editor/GraphEditor/Private/MaterialNodes/SGraphSubstrateMaterial.cpp b/Engine/Source/Editor/GraphEditor/Private/MaterialNodes/SGraphSubstrateMaterial.cpp
--- a/Engine/Source/Editor/GraphEditor/Private/MaterialNodes/SGraphSubstrateMaterial.cpp
+++ b/Engine/Source/Editor/GraphEditor/Private/MaterialNodes/SGraphSubstrateMaterial.cpp
@@ -41,18 +41,16 @@ bool FSubstrateWidget::HasInputSubstrateType(const UEdGraphPin* InPin)
 	}
 	if (UMaterialGraphNode* PinNode = Cast<UMaterialGraphNode>(InPin->GetOwningNode()))
 	{
-		TArrayView<FExpressionInput*> ExpressionInputs = PinNode->MaterialExpression->GetInputsView();
 		FName TargetPinName = PinNode->GetShortenPinName(InPin->PinName);
 
-		for (int32 Index = 0; Index < ExpressionInputs.Num(); ++Index)
+		for (FExpressionInputIterator It{ PinNode->MaterialExpression }; It; ++It)
 		{
-			FExpressionInput* Input = ExpressionInputs[Index];
-			FName InputName = PinNode->MaterialExpression->GetInputName(Index);
+			FName InputName = PinNode->MaterialExpression->GetInputName(It.Index);
 			InputName = PinNode->GetShortenPinName(InputName);
 
 			if (InputName == TargetPinName)
 			{
-				switch (PinNode->MaterialExpression->GetInputType(Index))
+				switch (PinNode->MaterialExpression->GetInputType(It.Index))
 				{
 					case MCT_Substrate:
 						return true;
diff --git a/Engine/Source/Editor/MaterialEditor/Private/MaterialEditingLibrary.cpp b/Engine/Source/Editor/MaterialEditor/Private/MaterialEditingLibrary.cpp
--- a/Engine/Source/Editor/MaterialEditor/Private/MaterialEditingLibrary.cpp
+++ b/Engine/Source/Editor/MaterialEditor/Private/MaterialEditingLibrary.cpp
@@ -45,43 +45,34 @@ static FExpressionInput* GetExpressionInputByName(UMaterialExpression* Expressio
 	check(Expression);
 	FExpressionInput* Result = nullptr;
 
-	TArrayView<FExpressionInput*> Inputs = Expression->GetInputsView();
-
 	// Return first input if no name specified
 	if (InputName.IsNone())
 	{
-		if (Inputs.Num() > 0)
-		{
-			return Inputs[0];
-		}
+		return Expression->GetInput(0);
 	}
-	else
+
+	// Get all inputs. Get name of each input, see if its the one we want
+	for (FExpressionInputIterator It{ Expression }; It; ++It)
 	{
-		// Get all inputs
-		// Get name of each input, see if its the one we want
-		for (int InputIdx = 0; InputIdx < Inputs.Num(); InputIdx++)
+		FName TestName;
+		if (UMaterialExpressionMaterialFunctionCall* FuncCall = Cast<UMaterialExpressionMaterialFunctionCall>(Expression))
 		{
-			FName TestName;
-			if (UMaterialExpressionMaterialFunctionCall* FuncCall = Cast<UMaterialExpressionMaterialFunctionCall>(Expression))
-			{
-				// If a function call, don't want to compare string with type postfix
-				TestName = FuncCall->GetInputNameWithType(InputIdx, false);
-			}
-			else
-			{
-				const FName ExpressionInputName = Expression->GetInputName(InputIdx);
-				TestName = UMaterialGraphNode::GetShortenPinName(ExpressionInputName);
-			}
+			// If a function call, don't want to compare string with type postfix
+			TestName = FuncCall->GetInputNameWithType(It.Index, false);
+		}
+		else
+		{
+			const FName ExpressionInputName = Expression->GetInputName(It.Index);
+			TestName = UMaterialGraphNode::GetShortenPinName(ExpressionInputName);
+		}
 
-			if (TestName == InputName)
-			{
-				Result = Inputs[InputIdx];
-				break;
-			}
+		if (TestName == InputName)
+		{
+			return It.Input;
 		}
 	}
 
-	return Result;
+	return nullptr;
 }
 
 static int32 GetExpressionOutputIndexByName(UMaterialExpression* Expression, const FName OutputName)
@@ -184,9 +175,9 @@ namespace MaterialEditingLibraryImpl
 
 		MaterialExpressionsToLayout.Add( MaterialExpression ) = MoveTemp( LayoutInfo );
 
-		for ( FExpressionInput* ExpressionInput : MaterialExpression->GetInputsView() )
+		for (FExpressionInputIterator It{ MaterialExpression }; It; ++It)
 		{
-			LayoutMaterialExpression( ExpressionInput->Expression, MaterialExpression, MaterialExpressionsToLayout, Row, Depth + 1 );
+			LayoutMaterialExpression( It->Expression, MaterialExpression, MaterialExpressionsToLayout, Row, Depth + 1 );
 		}
 	}
 
@@ -420,12 +411,11 @@ static void BreakLinksToExpression(TConstArrayView<TObjectPtr<UMaterialExpressio
 		// Don't check myself, though that shouldn't really matter...
 		if (TestExp != Expression)
 		{
-			TArrayView<FExpressionInput*> Inputs = TestExp->GetInputsView();
-			for (FExpressionInput* Input : Inputs)
+			for (FExpressionInputIterator It{ TestExp }; It; ++It)
 			{
-				if (Input->Expression == Expression)
+				if (It->Expression == Expression)
 				{
-					Input->Expression = nullptr;
+					It->Expression = nullptr;
 				}
 			}
 		}
@@ -853,23 +843,24 @@ TArray<FString> UMaterialEditingLibrary::GetMaterialExpressionInputNames(UMateri
 {
 	TArray<FString> InputNames;
 
-	TArrayView<FExpressionInput*> Inputs = MaterialExpression->GetInputsView();
-	for (int32 InputIdx = 0; InputIdx < Inputs.Num(); InputIdx++)
+	for (FExpressionInputIterator It{ MaterialExpression }; It; ++It)
 	{
 		FName Name;
 		if (UMaterialExpressionMaterialFunctionCall* FuncCall = Cast<UMaterialExpressionMaterialFunctionCall>(MaterialExpression))
 		{
 			// If a function call, don't want to compare string with type postfix
-			Name = FuncCall->GetInputNameWithType(InputIdx, false);
+			Name = FuncCall->GetInputNameWithType(It.Index, false);
 		}
 		else
 		{
-			const FName ExpressionInputName = MaterialExpression->GetInputName(InputIdx);
+			const FName ExpressionInputName = MaterialExpression->GetInputName(It.Index);
 			Name = UMaterialGraphNode::GetShortenPinName(ExpressionInputName);
 		}
 
+
 		InputNames.Add(Name.ToString());
 	}
+
 	return InputNames;
 }
 
@@ -877,18 +868,16 @@ TArray<int32> UMaterialEditingLibrary::GetMaterialExpressionInputTypes(UMaterial
 {
 	TArray<int32> InputTypes;
 
-	TArrayView<FExpressionInput*> Inputs = MaterialExpression->GetInputsView();
-	for (int32 InputIdx = 0; InputIdx < Inputs.Num(); InputIdx++)
+	for (FExpressionInputIterator It{ MaterialExpression }; It; ++It)
 	{
-		FExpressionInput* Input = Inputs[InputIdx];
-		UMaterialExpression* Expression = Input != nullptr ? Input->Expression : nullptr;
+		UMaterialExpression* Expression = It->Expression;
 		if (Expression != nullptr)
 		{
-			InputTypes.Add(Expression->GetOutputType(Input->OutputIndex));
+			InputTypes.Add(Expression->GetOutputType(It->OutputIndex));
 		}
 		else
 		{
-			InputTypes.Add(MaterialExpression->GetInputType(InputIdx));
+			InputTypes.Add(MaterialExpression->GetInputType(It.Index));
 		}
 	}
 	return InputTypes;
@@ -899,9 +888,9 @@ TArray<UMaterialExpression*> UMaterialEditingLibrary::GetInputsForMaterialExpres
 	TArray<UMaterialExpression*> MaterialExpressions;
 	if (Material)
 	{
-		for (const FExpressionInput* Input : MaterialExpression->GetInputsView())
+		for (FExpressionInputIterator It{ MaterialExpression }; It; ++It)
 		{
-			MaterialExpressions.Add(Input->Expression);
+			MaterialExpressions.Add(It->Expression);
 		}
 	}
 
@@ -911,13 +900,13 @@ TArray<UMaterialExpression*> UMaterialEditingLibrary::GetInputsForMaterialExpres
 bool UMaterialEditingLibrary::GetInputNodeOutputNameForMaterialExpression(UMaterialExpression* MaterialExpression, UMaterialExpression* InputNode, FString& OutputName)
 {
 	OutputName = TEXT("");
-	for (const FExpressionInput* Input : MaterialExpression->GetInputsView())
+	for (FExpressionInputIterator It{ MaterialExpression }; It; ++It)
 	{
-		if (Input->Expression == InputNode)
+		if (It->Expression == InputNode)
 		{
-			if(Input->OutputIndex != INDEX_NONE && Input->OutputIndex < InputNode->Outputs.Num())
+			if(It->OutputIndex != INDEX_NONE && It->OutputIndex < InputNode->Outputs.Num())
 			{
-				FExpressionOutput& Output = InputNode->Outputs[Input->OutputIndex];
+				FExpressionOutput& Output = InputNode->Outputs[It->OutputIndex];
 				OutputName = GetExpressionOutputName(Output);
 				return true;
 			}
diff --git a/Engine/Source/Editor/MaterialEditor/Private/MaterialEditor.cpp b/Engine/Source/Editor/MaterialEditor/Private/MaterialEditor.cpp
--- a/Engine/Source/Editor/MaterialEditor/Private/MaterialEditor.cpp
+++ b/Engine/Source/Editor/MaterialEditor/Private/MaterialEditor.cpp
@@ -4954,18 +4954,16 @@ UClass* FMaterialEditor::GetOnPromoteToParameterClass(const UEdGraphPin* TargetP
 	}
 	else if (OtherPinNode)
 	{
-		TArrayView<FExpressionInput*> ExpressionInputs = OtherPinNode->MaterialExpression->GetInputsView();
 		FName TargetPinName = OtherPinNode->GetShortenPinName(TargetPin->PinName);
 
-		for (int32 Index = 0; Index < ExpressionInputs.Num(); ++Index)
+		for (FExpressionInputIterator It{ OtherPinNode->MaterialExpression }; It; ++It)
 		{
-			FExpressionInput* Input = ExpressionInputs[Index];
-			FName InputName = OtherPinNode->MaterialExpression->GetInputName(Index);
+			FName InputName = OtherPinNode->MaterialExpression->GetInputName(It.Index);
 			InputName = OtherPinNode->GetShortenPinName(InputName);
 
 			if (InputName == TargetPinName)
 			{
-				switch (OtherPinNode->MaterialExpression->GetInputType(Index))
+				switch (OtherPinNode->MaterialExpression->GetInputType(It.Index))
 				{
 					case MCT_Float1:
 					case MCT_Float: return UMaterialExpressionScalarParameter::StaticClass();
@@ -5249,12 +5247,11 @@ void FMaterialEditor::OnCreateSubstrateNodeForPin(const FToolMenuContext& InMenu
 	{
 		// Link manually
 		UMaterialGraphNode* NewNode = Cast<UMaterialGraphNode>(Action.PerformAction(GraphObj, nullptr, NewNodePos));
-		TArrayView<FExpressionInput*> NewNodeExpressionInputs = NewNode->MaterialExpression->GetInputsView();
 
 		// From that direction, the node is never going to be a root node (a root node has no output we can connect from).
 		UMaterialGraphNode* TargetPinNode = Cast<UMaterialGraphNode>(TargetPin->GetOwningNode());
 
-		check(NewNodeExpressionInputs.Num() > 0 && TargetPin->SourceIndex < TargetPinNode->MaterialExpression->GetOutputs().Num());
+		check(NewNode->MaterialExpression->GetInput(0) && TargetPin->SourceIndex < TargetPinNode->MaterialExpression->GetOutputs().Num());
 
 		FName TargetPinName = TargetPinNode->MaterialExpression->GetOutputs()[TargetPin->SourceIndex].OutputName;
 		UMaterialEditingLibrary::ConnectMaterialExpressions(TargetPinNode->MaterialExpression, TargetPinName.ToString(), NewNode->MaterialExpression, FString());
diff --git a/Engine/Source/Editor/MaterialEditor/Private/MaterialEditorUtilities.cpp b/Engine/Source/Editor/MaterialEditor/Private/MaterialEditorUtilities.cpp
--- a/Engine/Source/Editor/MaterialEditor/Private/MaterialEditorUtilities.cpp
+++ b/Engine/Source/Editor/MaterialEditor/Private/MaterialEditorUtilities.cpp
@@ -637,10 +637,9 @@ void FMaterialEditorUtilities::GetVisibleMaterialParametersFromExpression(
 		else
 		{
 			// Retrieve the expression input and then start parsing its children
-			for (int32 i = 0; i < MaterialExpressionKey.Expression->GetInputsView().Num(); i++)
+			for (FExpressionInputIterator It{ MaterialExpressionKey.Expression }; It; ++It)
 			{
-				FExpressionInput* Input = MaterialExpressionKey.Expression->GetInputsView()[i];
-				GetVisibleMaterialParametersFromExpression(FMaterialExpressionKey(Input->Expression, Input->OutputIndex), MaterialInstance, VisibleExpressions, FunctionStack);
+				GetVisibleMaterialParametersFromExpression(FMaterialExpressionKey(It->Expression, It->OutputIndex), MaterialInstance, VisibleExpressions, FunctionStack);
 			}
 
 			TArray<FExpressionExecOutputEntry> ExpressionExecOutputs;
@@ -743,10 +742,9 @@ bool FMaterialEditorUtilities::HasCompatibleConnection(UClass* ExpressionClass,
 		UMaterialExpression* DefaultExpression = CastChecked<UMaterialExpression>(ExpressionClass->GetDefaultObject());
 		if (TestDirection == EGPD_Output)
 		{
-			int32 NumInputs = DefaultExpression->GetInputsView().Num();
-			for (int32 Index = 0; Index < NumInputs; ++Index)
+			for (FExpressionInputIterator It{ DefaultExpression }; It; ++It)
 			{
-				uint32 InputType = DefaultExpression->GetInputType(Index);
+				uint32 InputType = DefaultExpression->GetInputType(It.Index);
 				if (CanConnectMaterialValueTypes(InputType, TestType))
 				{
 					return true;
diff --git a/Engine/Source/Editor/UnrealEd/Private/MaterialGraph.cpp b/Engine/Source/Editor/UnrealEd/Private/MaterialGraph.cpp
--- a/Engine/Source/Editor/UnrealEd/Private/MaterialGraph.cpp
+++ b/Engine/Source/Editor/UnrealEd/Private/MaterialGraph.cpp
@@ -516,8 +516,6 @@ void UMaterialGraph::LinkGraphNodesFromMaterial()
 			continue;
 		}
 
-		TArrayView<FExpressionInput*> ExpressionInputs = Expression->GetInputsView();
-
 		TArray<FExpressionExecOutputEntry> ExecOutputs;
 		Expression->GetExecOutputs(ExecOutputs);
 
@@ -526,19 +524,20 @@ void UMaterialGraph::LinkGraphNodesFromMaterial()
 			if (Pin->Direction == EGPD_Input && Pin->PinType.PinCategory != UMaterialGraphSchema::PC_Exec)
 			{
 				// Implicitly generated property inputs are not returned by GetInputs(), so check index is within valid range.
-				if (ExpressionInputs.IsValidIndex(Pin->SourceIndex) && ExpressionInputs[Pin->SourceIndex]->Expression)
+				FExpressionInput* SourceInput = Expression->GetInput(Pin->SourceIndex);
+				if (SourceInput && SourceInput->Expression)
 				{
 					// Unclear why this is null sometimes outside of composite reroute, but this is safer than crashing
-					if (UMaterialGraphNode* GraphNode = Cast<UMaterialGraphNode>(ExpressionInputs[Pin->SourceIndex]->Expression->GraphNode))
+					if (UMaterialGraphNode* GraphNode = Cast<UMaterialGraphNode>(SourceInput->Expression->GraphNode))
 					{
 						// if GraphNode is a material function call for a missing material function, it may not have any output pins
-						UEdGraphPin* OutputPin = GraphNode->GetOutputPin(GetValidOutputIndex(ExpressionInputs[Pin->SourceIndex]));
+						UEdGraphPin* OutputPin = GraphNode->GetOutputPin(GetValidOutputIndex(Expression->GetInput(Pin->SourceIndex)));
 						if (LIKELY(OutputPin))
 						{
 							Pin->MakeLinkTo(OutputPin);
 						}
 					}
-					else if (UMaterialExpressionReroute* CompositeReroute = Cast<UMaterialExpressionReroute>(ExpressionInputs[Pin->SourceIndex]->Expression))
+					else if (UMaterialExpressionReroute* CompositeReroute = Cast<UMaterialExpressionReroute>(Expression->GetInput(Pin->SourceIndex)->Expression))
 					{
 						// This is an unseen composite reroute expression, find the actual expression output to connect to.
 						UMaterialExpressionComposite* OwningComposite = Cast<UMaterialExpressionComposite>(CompositeReroute->SubgraphExpression);
@@ -657,13 +656,11 @@ void UMaterialGraph::LinkMaterialExpressionsFromGraph()
 
 					for (UEdGraphPin* Pin : GraphNode->Pins)
 					{
-						TArrayView<FExpressionInput*> ExpressionInputs = Expression->GetInputsView();
 						if (Pin->Direction == EGPD_Input && Pin->PinType.PinCategory != UMaterialGraphSchema::PC_Exec)
 						{
 							// Wire up non-execution input pins
 							// Implicitly generated property inputs are not returned by GetInputs(), so check index is within valid range.
-							FExpressionInput* ExpressionInput = ExpressionInputs.IsValidIndex(Pin->SourceIndex) ? ExpressionInputs[Pin->SourceIndex] : nullptr;
-							if (ExpressionInput)
+							if (FExpressionInput* ExpressionInput = Expression->GetInput(Pin->SourceIndex))
 							{
 								if (Pin->LinkedTo.Num() > 0)
 								{
diff --git a/Engine/Source/Editor/UnrealEd/Private/MaterialGraphNode.cpp b/Engine/Source/Editor/UnrealEd/Private/MaterialGraphNode.cpp
--- a/Engine/Source/Editor/UnrealEd/Private/MaterialGraphNode.cpp
+++ b/Engine/Source/Editor/UnrealEd/Private/MaterialGraphNode.cpp
@@ -664,15 +664,14 @@ void UMaterialGraphNode::CreateInputPins()
 		NewPin->PinFriendlyName = SpaceText;
 	}
 
-	TArrayView<FExpressionInput*> ExpressionInputs = MaterialExpression->GetInputsView();
-	const int32 NumExpressionInputs = ExpressionInputs.Num();
-	for (int32 Index = 0; Index < NumExpressionInputs; ++Index)
+	int32 InputIndex = 0;
+	while (FExpressionInput* Input = MaterialExpression->GetInput(InputIndex))
 	{
-		FName InputName = MaterialExpression->GetInputName(Index);
+		FName InputName = MaterialExpression->GetInputName(InputIndex);
 		InputName = GetShortenPinName(InputName);
 
 		FName PinCategory;
-		if (MaterialExpression->IsInputConnectionRequired(Index))
+		if (MaterialExpression->IsInputConnectionRequired(InputIndex))
 		{
 			PinCategory = UMaterialGraphSchema::PC_Required;
 		}
@@ -681,18 +680,20 @@ void UMaterialGraphNode::CreateInputPins()
 			PinCategory = UMaterialGraphSchema::PC_Optional;
 		}
 
-		FName PinSubCategory = MaterialExpression->GetInputPinSubCategory(Index);
-		UObject* PinSubCategoryObject = MaterialExpression->GetInputPinSubCategoryObject(Index);
+		FName PinSubCategory = MaterialExpression->GetInputPinSubCategory(InputIndex);
+		UObject* PinSubCategoryObject = MaterialExpression->GetInputPinSubCategoryObject(InputIndex);
 
 		UEdGraphPin* NewPin = CreatePin(EGPD_Input, PinCategory, PinSubCategory, PinSubCategoryObject, InputName);
-		NewPin->SourceIndex = Index;
-		NewPin->DefaultValue = MaterialExpression->GetInputPinDefaultValue(Index);
+		NewPin->SourceIndex = InputIndex;
+		NewPin->DefaultValue = MaterialExpression->GetInputPinDefaultValue(InputIndex);
 		if (NewPin->PinName.IsNone())
 		{
 			// Makes sure pin has a name for lookup purposes but user will never see it
 			NewPin->PinName = CreateUniquePinName(TEXT("Input"));
 			NewPin->PinFriendlyName = SpaceText;
 		}
+
+		InputIndex += 1;
 	}
 
 	// Next create pins for property inputs
@@ -705,7 +706,7 @@ void UMaterialGraphNode::CreateInputPins()
 		FString InputName = Property->GetDisplayNameText().ToString();
 		FName PinCategory = UMaterialGraphSchema::PC_Optional;
 
-		int32 PinIndex = NumExpressionInputs + i;
+		int32 PinIndex = InputIndex + i;
 		FName PinSubCategory = MaterialExpression->GetInputPinSubCategory(PinIndex);
 		UObject* PinSubCategoryObject = MaterialExpression->GetInputPinSubCategoryObject(PinIndex);
 
diff --git a/Engine/Source/Runtime/Engine/Private/Materials/Material.cpp b/Engine/Source/Runtime/Engine/Private/Materials/Material.cpp
--- a/Engine/Source/Runtime/Engine/Private/Materials/Material.cpp
+++ b/Engine/Source/Runtime/Engine/Private/Materials/Material.cpp
@@ -6206,7 +6206,12 @@ bool UMaterial::RecursiveGetExpressionChain(
 			else
 			{
 				// Follow all properties.
-				Inputs = InExpression->GetInputsView();
+				Inputs.Empty(InExpression->CountInputs());
+				for (FExpressionInputIterator It{ InExpression }; It; ++It)
+				{
+					Inputs.Push(It.Input);
+				}
+
 				InputsFrequency.Init(InShaderFrequency, Inputs.Num());
 			}
 		}
@@ -6224,7 +6229,7 @@ bool UMaterial::RecursiveGetExpressionChain(
 	}
 	else if ((SetMaterialAttributesExp = Cast<UMaterialExpressionSetMaterialAttributes>(InExpression)) != nullptr)
 	{
-		checkf(!SetMaterialAttributesExp->GetInputsView().IsEmpty() && (SetMaterialAttributesExp->GetInputType(0) == MCT_MaterialAttributes), TEXT("There must always be one input at least : the material attribute pin"));
+		checkf(SetMaterialAttributesExp->GetInput(0) && (SetMaterialAttributesExp->GetInputType(0) == MCT_MaterialAttributes), TEXT("There must always be one input at least : the material attribute pin"));
 		// Always add the material attribute input, so that we keep on traversing up the property chain : 
 		Inputs.Add(SetMaterialAttributesExp->GetInput(0));
 		InputsFrequency.Add(InShaderFrequency);
@@ -6254,7 +6259,11 @@ bool UMaterial::RecursiveGetExpressionChain(
 			else
 			{
 				// Follow all properties.
-				Inputs = InExpression->GetInputsView();
+				Inputs.Empty(InExpression->CountInputs());
+				for (FExpressionInputIterator It{ InExpression }; It; ++It)
+				{
+					Inputs.Push(It.Input);
+				}
 				InputsFrequency.Init(InShaderFrequency, Inputs.Num());
 			}
 		}
@@ -6308,16 +6317,22 @@ bool UMaterial::RecursiveGetExpressionChain(
 		}
 
 		// here we assume ALL inputs to the MaterialFunctionCall are active
-		auto ExprInputs = InExpression->GetInputsView();
-		for (int i = 0; i < ExprInputs.Num(); i++)
+		Inputs.Empty(InExpression->CountInputs());
+		for (FExpressionInputIterator It{ InExpression }; It; ++It)
 		{
-			Inputs.Add(ExprInputs[i]);
-			InputsFrequency.Add(InShaderFrequency);
+			Inputs.Push(It.Input);
 		}
+
+		InputsFrequency.Init(InShaderFrequency, Inputs.Num());
 	}
 	else
 	{
-		Inputs = InExpression->GetInputsView();
+		Inputs.Empty(InExpression->CountInputs());
+		for (FExpressionInputIterator It{ InExpression }; It; ++It)
+		{
+			Inputs.Push(It.Input);
+		}
+
 		InputsFrequency.Init(InShaderFrequency, Inputs.Num());
 	}
 
@@ -6411,20 +6426,18 @@ void UMaterial::RecursiveUpdateRealtimePreview( UMaterialExpression* InExpressio
 	}
 
 	// We need to examine our inputs. If any of them need realtime preview, so do we.
-	TArrayView<FExpressionInput*> Inputs = InExpression->GetInputsView();
-	for (int32 InputIdx = 0; InputIdx < Inputs.Num(); InputIdx++)
+	for (FExpressionInputIterator It{ InExpression }; It; ++It)
 	{
-		FExpressionInput* InnerInput = Inputs[InputIdx];
-		if (InnerInput != NULL && InnerInput->Expression != NULL)
+		if (It.Input != NULL && It.Input->Expression != NULL)
 		{
 			// See if we still need to process this expression, and if so do that first.
-			if (InOutExpressionsToProcess.Find(InnerInput->Expression) != INDEX_NONE)
+			if (InOutExpressionsToProcess.Find(It.Input->Expression) != INDEX_NONE)
 			{
-				RecursiveUpdateRealtimePreview(InnerInput->Expression, InOutExpressionsToProcess);
+				RecursiveUpdateRealtimePreview(It.Input->Expression, InOutExpressionsToProcess);
 			}
 
 			// If our input expression needed realtime preview, we do too.
-			if( InnerInput->Expression->bRealtimePreview )
+			if( It.Input->Expression->bRealtimePreview )
 			{
 
 				InExpression->bRealtimePreview = true;
diff --git a/Engine/Source/Runtime/Engine/Private/Materials/MaterialCachedData.cpp b/Engine/Source/Runtime/Engine/Private/Materials/MaterialCachedData.cpp
--- a/Engine/Source/Runtime/Engine/Private/Materials/MaterialCachedData.cpp
+++ b/Engine/Source/Runtime/Engine/Private/Materials/MaterialCachedData.cpp
@@ -787,9 +787,9 @@ struct FMaterialConnectedPropertiesAnalyzer
 			// If otherwise specified, explore all input expressions to this node.
 			if (bVisitAllInputExpressions)
 			{
-				for (FExpressionInput* ExpressionInput : Expression->GetInputsView())
+				for (FExpressionInputIterator It{ Expression }; It; ++It)
 				{
-					PushUnexploredExpression(ExpressionInput->Expression);
+					PushUnexploredExpression(It->Expression);
 				}
 			}
 		}
diff --git a/Engine/Source/Runtime/Engine/Private/Materials/MaterialExpressions.cpp b/Engine/Source/Runtime/Engine/Private/Materials/MaterialExpressions.cpp
--- a/Engine/Source/Runtime/Engine/Private/Materials/MaterialExpressions.cpp
+++ b/Engine/Source/Runtime/Engine/Private/Materials/MaterialExpressions.cpp
@@ -881,22 +881,19 @@ void UMaterialExpression::CopyMaterialExpressions(const TArray<UMaterialExpressi
 	for( int32 NewExpressionIndex = 0 ; NewExpressionIndex < OutNewExpressions.Num() ; ++NewExpressionIndex )
 	{
 		UMaterialExpression* NewExpression = OutNewExpressions[NewExpressionIndex];
-		TArrayView<FExpressionInput*> ExpressionInputs = NewExpression->GetInputsView();
-		for ( int32 ExpressionInputIndex = 0 ; ExpressionInputIndex < ExpressionInputs.Num() ; ++ExpressionInputIndex )
+		for (FExpressionInputIterator It{ NewExpression }; It; ++It)
 		{
-			FExpressionInput* Input = ExpressionInputs[ExpressionInputIndex];
-			UMaterialExpression* InputExpression = Input->Expression;
-			if ( InputExpression )
+			if (UMaterialExpression* InputExpression = It->Expression)
 			{
 				UMaterialExpression** NewInputExpression = SrcToDestMap.Find( InputExpression );
 				if ( NewInputExpression )
 				{
 					check( *NewInputExpression );
-					Input->Expression = *NewInputExpression;
+					It->Expression = *NewInputExpression;
 				}
 				else
 				{
-					Input->Expression = nullptr;
+					It->Expression = nullptr;
 				}
 			}
 		}
@@ -934,11 +931,9 @@ void UMaterialExpression::Serialize(FStructuredArchive::FRecord Record)
 	const int32 RenderVer = Archive.CustomVer(FRenderingObjectVersion::GUID);
 	const int32 UE5Ver = Archive.CustomVer(FUE5MainStreamObjectVersion::GUID);
 
-	TArrayView<FExpressionInput*> Inputs = GetInputsView();
-	for (int32 InputIndex = 0; InputIndex < Inputs.Num(); ++InputIndex)
+	for (FExpressionInputIterator It{ this }; It; ++It)
 	{
-		FExpressionInput* Input = Inputs[InputIndex];
-		DoMaterialAttributeReorder(Input, UEVer, RenderVer, UE5Ver);
+		DoMaterialAttributeReorder(It.Input, UEVer, RenderVer, UE5Ver);
 	}
 #endif // WITH_EDITORONLY_DATA
 }
@@ -1002,13 +997,13 @@ TArray<FProperty*> UMaterialExpression::GetInputPinProperty(int32 PinIndex)
 {
 	// Find all UPROPERTYs associated with this input pin
 	TArray<FProperty*> Properties;
+
 	// Explicit input pins are before property input pins
 	TArray<FProperty*> PropertyInputs = GetPropertyInputs();
-	const int32 NumInputs = GetInputsView().Num();
-	if (PinIndex < NumInputs)
-	{
-		FExpressionInput* Input = GetInput(PinIndex);
 
+	
+	if (FExpressionInput* Input = GetInput(PinIndex))
+	{
 		// Find the UPROPERTYs that have OverridingInputProperty meta data pointing to the expression input.
 		// There can be multiple scalar entries together forming a vector parameter, e.g. DecalMipmapLevel node has FExpressionInput TextureSize <-> float ConstWidth/ConstHeight.
 		static FName OverridingInputPropertyMetaData(TEXT("OverridingInputProperty"));
@@ -1029,15 +1024,19 @@ TArray<FProperty*> UMaterialExpression::GetInputPinProperty(int32 PinIndex)
 			}
 		}
 	}
-	else if (PinIndex < NumInputs + PropertyInputs.Num())
+	else
 	{
-		FName PropertyName = PropertyInputs[PinIndex - NumInputs]->GetFName();
-		for (TFieldIterator<FProperty> InputIt(GetClass(), EFieldIteratorFlags::IncludeSuper, EFieldIteratorFlags::ExcludeDeprecated); InputIt; ++InputIt)
+		int32 NumInputs = CountInputs();
+		if (PinIndex < NumInputs + PropertyInputs.Num())
 		{
-			FProperty* Property = *InputIt;
-			if (PropertyName == Property->GetFName())
+			FName PropertyName = PropertyInputs[PinIndex - NumInputs]->GetFName();
+			for (TFieldIterator<FProperty> InputIt(GetClass(), EFieldIteratorFlags::IncludeSuper, EFieldIteratorFlags::ExcludeDeprecated); InputIt; ++InputIt)
 			{
-				Properties.Add(Property);
+				FProperty* Property = *InputIt;
+				if (PropertyName == Property->GetFName())
+				{
+					Properties.Add(Property);
+				}
 			}
 		}
 	}
@@ -1549,7 +1548,6 @@ void UMaterialExpression::PostEditChangeProperty(FPropertyChangedEvent& Property
 	if (MemberPropertyThatChanged != nullptr && GraphNode)
 	{
 		int32 PinIndex = -1;
-		TArrayView<FExpressionInput*> AllInputs = GetInputsView();
 
 		// Find the expression input this UPROPERTY points to with OverridingInputProperty meta data
 		static FName OverridingInputPropertyMetaData(TEXT("OverridingInputProperty"));
@@ -1560,11 +1558,11 @@ void UMaterialExpression::PostEditChangeProperty(FPropertyChangedEvent& Property
 			if (StructProp)
 			{
 				const FExpressionInput* TargetInput = StructProp->ContainerPtrToValuePtr<FExpressionInput>(this);
-				for (int32 i = 0; i < AllInputs.Num(); ++i)
+				for (FExpressionInputIterator It{ this }; It; ++It)
 				{
-					if (TargetInput == AllInputs[i])
+					if (TargetInput == It.Input)
 					{
-						PinIndex = i;
+						PinIndex = It.Index;
 						break;
 					}
 				}
@@ -1573,12 +1571,13 @@ void UMaterialExpression::PostEditChangeProperty(FPropertyChangedEvent& Property
 		else
 		{
 			// Not found in explicit expression inputs, so search in property inputs.
+			const int32 NumInputs = CountInputs();
 			TArray<FProperty*> PropertyInputs = GetPropertyInputs();
 			for (int32 i = 0; i < PropertyInputs.Num(); ++i)
 			{
 				if (MemberPropertyThatChanged->GetFName() == PropertyInputs[i]->GetFName())
 				{
-					PinIndex = AllInputs.Num() + i;
+					PinIndex = NumInputs + i;
 				}
 			}
 		}
@@ -1745,14 +1744,21 @@ TArrayView<FExpressionInput*> UMaterialExpression::GetInputsView()
 	return CachedInputs;
 }
 
+int32 UMaterialExpression::CountInputs() const
+{
+	int32 Index = 0;
+	while (const FExpressionInput* Input = GetInput(Index))
+	{
+		Index += 1;
+	}
+	return Index;
+}
 
 FExpressionInput* UMaterialExpression::GetInput(int32 InputIndex)
 {
-	TArrayView<FExpressionInput*> Inputs = GetInputsView();
-	return InputIndex < Inputs.Num() ? Inputs[InputIndex] : nullptr;
+	return InputIndex < CachedInputs.Num() ? CachedInputs[InputIndex] : nullptr;
 }
 
-
 FName UMaterialExpression::GetInputName(int32 InputIndex) const
 {
 	int32 Index = 0;
@@ -1766,7 +1772,7 @@ FName UMaterialExpression::GetInputName(int32 InputIndex) const
 				if (Index == InputIndex)
 				{
 					FExpressionInput const* Input = StructProp->ContainerPtrToValuePtr<FExpressionInput>(this, ArrayIndex);
-
+		
 					if (!Input->InputName.IsNone())
 					{
 						return Input->InputName;
@@ -2175,16 +2181,11 @@ bool UMaterialExpression::GetAllInputExpressions(TArray<UMaterialExpression*>& I
 		bool bFoundRepeat = false;
 		InputExpressions.Add(this);
 
-		TArrayView<FExpressionInput*> Inputs = GetInputsView();
-
-		for (int32 Index = 0; Index < Inputs.Num(); Index++)
+		for (FExpressionInputIterator It{ this }; It; ++It)
 		{
-			if (Inputs[Index]->Expression)
+			if (It->Expression && It->Expression->GetAllInputExpressions(InputExpressions))
 			{
-				if (Inputs[Index]->Expression->GetAllInputExpressions(InputExpressions))
-				{
-					bFoundRepeat = true;
-				}
+				bFoundRepeat = true;
 			}
 		}
 
@@ -2301,9 +2302,9 @@ bool UMaterialExpression::ContainsInputLoop(TSet<UMaterialExpression*>& VisitedE
 
 bool UMaterialExpression::ContainsInputLoopInternal(const FContainsInputLoopInternalExpressionStack& ExpressionStack, TSet<UMaterialExpression*>& VisitedExpressions, const bool bStopOnFunctionCall)
 {
-	for (FExpressionInput* Input : GetInputsView())
+	for (FExpressionInputIterator It{ this }; It; ++It)
 	{
-		UMaterialExpression* InputExpression = Input->Expression;
+		UMaterialExpression* InputExpression = It->Expression;
 		if (!InputExpression)
 		{
 			continue;
@@ -2343,47 +2344,6 @@ bool UMaterialExpression::ContainsInputLoopInternal(const FContainsInputLoopInte
 	return false;
 }
 
-// Deprecated
-bool UMaterialExpression::ContainsInputLoopInternal(TArray<FMaterialExpressionKey>& ExpressionStack, TSet<FMaterialExpressionKey>& VisitedExpressions, const bool bStopOnFunctionCall)
-{
-	for (FExpressionInput* Input : GetInputsView())
-	{
-		if (!Input->Expression)
-		{
-			continue;
-		}
-		
-		// ContainsInputLoop primarily used to detect safe traversal path for IsResultMaterialAttributes.
-		// In those cases we can bail on a function as the inputs are strongly typed
-		UMaterialExpressionMaterialFunctionCall* FunctionCall = Cast<UMaterialExpressionMaterialFunctionCall>(Input->Expression);
-		UMaterialExpressionMaterialAttributeLayers* Layers = Cast<UMaterialExpressionMaterialAttributeLayers>(Input->Expression);
-		if (bStopOnFunctionCall && (FunctionCall || Layers))
-		{
-			continue;
-		}
-
-		FMaterialExpressionKey InputExpressionKey(Input->Expression, Input->OutputIndex);
-		if (ExpressionStack.Contains(InputExpressionKey))
-		{
-			return true;
-		}
-
-		// prevent recurring visits to expressions we've already checked
-		if (!VisitedExpressions.Contains(InputExpressionKey))
-		{
-			VisitedExpressions.Add(InputExpressionKey);
-			ExpressionStack.Add(InputExpressionKey);
-			if (Input->Expression->ContainsInputLoopInternal(ExpressionStack, VisitedExpressions, bStopOnFunctionCall))
-			{
-				return true;
-			}
-			ExpressionStack.Pop();
-		}
-	}
-
-	return false;
-}
-
 bool UMaterialExpression::IsUsingNewHLSLGenerator() const
 {
 	if (Material)
@@ -2676,7 +2636,7 @@ FExpressionInput* UMaterialExpressionTextureSample::GetInput(int32 InputIndex)
 		IF_INPUT_RETURN(TextureObject);
 	}
 
-	if(MipValueMode == TMVM_Derivative)
+	if (MipValueMode == TMVM_Derivative)
 	{
 		IF_INPUT_RETURN(CoordinatesDX);
 		IF_INPUT_RETURN(CoordinatesDY);
@@ -3274,7 +3234,7 @@ TArrayView<FExpressionInput*> UMaterialExpressionRuntimeVirtualTextureSample::Ge
 
 FName UMaterialExpressionRuntimeVirtualTextureSample::GetInputName(int32 InputIndex) const
 {
-	if (CachedInputs[InputIndex] == &WorldPosition)
+	if (GetInput(InputIndex) == &WorldPosition)
 	{
 		return GetWorldPositionInputName(WorldPositionOriginType);
 	}
@@ -3859,9 +3819,9 @@ int32 UMaterialExpressionRuntimeVirtualTextureReplace::Compile(class FMaterialCo
 
 bool UMaterialExpressionRuntimeVirtualTextureReplace::IsResultMaterialAttributes(int32 OutputIndex)
 {
-	for (FExpressionInput* ExpressionInput : GetInputsView())
+	for (FExpressionInputIterator It{ this }; It; ++It)
 	{
-		if (ExpressionInput->GetTracedInput().Expression && ExpressionInput->Expression->IsResultMaterialAttributes(ExpressionInput->OutputIndex))
+		if (It->GetTracedInput().Expression && It->Expression->IsResultMaterialAttributes(It->OutputIndex))
 		{
 			return true;
 		}
@@ -3919,9 +3879,9 @@ int32 UMaterialExpressionVirtualTextureFeatureSwitch::Compile(class FMaterialCom
 
 bool UMaterialExpressionVirtualTextureFeatureSwitch::IsResultMaterialAttributes(int32 OutputIndex)
 {
-	for (FExpressionInput* ExpressionInput : GetInputsView())
+	for (FExpressionInputIterator It{ this }; It; ++It)
 	{
-		if (ExpressionInput->GetTracedInput().Expression && ExpressionInput->Expression->IsResultMaterialAttributes(ExpressionInput->OutputIndex))
+		if (It->GetTracedInput().Expression && It->Expression->IsResultMaterialAttributes(It->OutputIndex))
 		{
 			return true;
 		}
@@ -7139,11 +7099,6 @@ UMaterialExpressionBreakMaterialAttributes::UMaterialExpressionBreakMaterialAttr
 	Outputs.Add(FExpressionOutput(TEXT("ShadingModel"), 0, 0, 0, 0, 0));
 	Outputs.Add(FExpressionOutput(TEXT("Displacement"), 1, 1, 0, 0, 0));
 #endif
-
-#if WITH_EDITOR
-	CachedInputs.Empty();
-	CachedInputs.Add(&MaterialAttributes);
-#endif
 }
 
 void UMaterialExpressionBreakMaterialAttributes::Serialize(FStructuredArchive::FRecord Record)
@@ -7668,7 +7623,7 @@ TArrayView<FExpressionInput*> UMaterialExpressionSetMaterialAttributes::GetInput
 
 FExpressionInput* UMaterialExpressionSetMaterialAttributes::GetInput(int32 InputIndex)
 {
-	return &Inputs[InputIndex];
+	return Inputs.IsValidIndex(InputIndex) ? &Inputs[InputIndex] : nullptr;
 }
 
 FName UMaterialExpressionSetMaterialAttributes::GetInputName(int32 InputIndex) const
@@ -7832,17 +7787,20 @@ UMaterialExpressionBlendMaterialAttributes::UMaterialExpressionBlendMaterialAttr
 	Outputs.Reset();
 	Outputs.Add(FExpressionOutput(TEXT(""), 0, 0, 0, 0, 0));
 #endif
+}
 
 #if WITH_EDITOR
-	CachedInputs.Empty();
-	CachedInputs.Reserve(3);
-	CachedInputs.Add(&A);
-	CachedInputs.Add(&B);
-	CachedInputs.Add(&Alpha);
-#endif
+FExpressionInput* UMaterialExpressionBlendMaterialAttributes::GetInput(int32 InputIndex)
+{
+	switch (InputIndex)
+	{
+		case 0: return &A;
+		case 1: return &B;
+		case 2: return &Alpha;
+		default: return nullptr;
+	}
 }
 
-#if WITH_EDITOR
 int32 UMaterialExpressionBlendMaterialAttributes::Compile(class FMaterialCompiler* Compiler, int32 OutputIndex)
 {
 	const FGuid AttributeID = Compiler->GetMaterialAttribute();
@@ -9907,12 +9865,10 @@ bool UMaterialExpressionQualitySwitch::IsInputConnectionRequired(int32 InputInde
 bool UMaterialExpressionQualitySwitch::IsResultMaterialAttributes(int32 OutputIndex)
 {
 	check(OutputIndex == 0);
-	TArrayView<FExpressionInput*> ExpressionInputs = GetInputsView();
-
-	for (FExpressionInput* ExpressionInput : ExpressionInputs)
+	for (FExpressionInputIterator It{ this }; It; ++It)
 	{
 		// If there is a loop anywhere in this expression's inputs then we can't risk checking them
-		if (ExpressionInput->Expression && ExpressionInput->Expression->IsResultMaterialAttributes(ExpressionInput->OutputIndex))
+		if (It->Expression && It->Expression->IsResultMaterialAttributes(It->OutputIndex))
 		{
 			return true;
 		}
@@ -10032,12 +9988,10 @@ bool UMaterialExpressionFeatureLevelSwitch::IsInputConnectionRequired(int32 Inpu
 bool UMaterialExpressionFeatureLevelSwitch::IsResultMaterialAttributes(int32 OutputIndex)
 {
 	check(OutputIndex == 0);
-	TArrayView<FExpressionInput*> ExpressionInputs = GetInputsView();
-
-	for (FExpressionInput* ExpressionInput : ExpressionInputs)
+	for (FExpressionInputIterator It{ this }; It; ++It)
 	{
 		// If there is a loop anywhere in this expression's inputs then we can't risk checking them
-		if (ExpressionInput->GetTracedInput().Expression && ExpressionInput->Expression->IsResultMaterialAttributes(ExpressionInput->OutputIndex))
+		if (It->GetTracedInput().Expression && It->Expression->IsResultMaterialAttributes(It->OutputIndex))
 		{
 			return true;
 		}
@@ -10220,13 +10174,11 @@ FName UMaterialExpressionDataDrivenShaderPlatformInfoSwitch::GetInputName(int32
 bool UMaterialExpressionDataDrivenShaderPlatformInfoSwitch::IsResultMaterialAttributes(int32 OutputIndex)
 {
 	check(OutputIndex == 0);
-	TArrayView<FExpressionInput*> ExpressionInputs = GetInputsView();
-
-	for (FExpressionInput* ExpressionInput : ExpressionInputs)
+	for (FExpressionInputIterator It{ this }; It; ++It)
 	{
 		// If there is a loop anywhere in this expression's inputs then we can't risk checking them
-		TObjectPtr<class UMaterialExpression> Expression = ExpressionInput->Expression;
-		if (Expression && ExpressionInput->Expression->IsResultMaterialAttributes(ExpressionInput->OutputIndex))
+		TObjectPtr<class UMaterialExpression> Expression = It->Expression;
+		if (Expression && It->Expression->IsResultMaterialAttributes(It->OutputIndex))
 		{
 			return true;
 		}
@@ -10301,10 +10253,10 @@ bool UMaterialExpressionBindlessSwitch::IsInputConnectionRequired(int32 InputInd
 bool UMaterialExpressionBindlessSwitch::IsResultMaterialAttributes(int32 OutputIndex)
 {
 	check(OutputIndex == 0);
-	for (FExpressionInput* ExpressionInput : GetInputsView())
+	for (FExpressionInputIterator It{ this }; It; ++It)
 	{
 		// If there is a loop anywhere in this expression's inputs then we can't risk checking them
-		if (ExpressionInput->Expression && ExpressionInput->Expression->IsResultMaterialAttributes(ExpressionInput->OutputIndex))
+		if (It->Expression && It->Expression->IsResultMaterialAttributes(It->OutputIndex))
 		{
 			return true;
 		}
@@ -10573,10 +10525,10 @@ FName UMaterialExpressionRequiredSamplersSwitch::GetInputName(int32 InputIndex)
 bool UMaterialExpressionRequiredSamplersSwitch::IsResultMaterialAttributes(int32 OutputIndex)
 {
 	check(OutputIndex == 0);
-	for (FExpressionInput* ExpressionInput : GetInputsView())
+	for (FExpressionInputIterator It{ this }; It; ++It)
 	{
 		// If there is a loop anywhere in this expression's inputs then we can't risk checking them
-		if (ExpressionInput->Expression && ExpressionInput->Expression->IsResultMaterialAttributes(ExpressionInput->OutputIndex))
+		if (It->Expression && It->Expression->IsResultMaterialAttributes(It->OutputIndex))
 		{
 			return true;
 		}
@@ -10672,10 +10624,10 @@ bool UMaterialExpressionShadingPathSwitch::IsInputConnectionRequired(int32 Input
 bool UMaterialExpressionShadingPathSwitch::IsResultMaterialAttributes(int32 OutputIndex)
 {
 	check(OutputIndex == 0);
-	for (FExpressionInput* ExpressionInput : GetInputsView())
+	for (FExpressionInputIterator It{ this }; It; ++It)
 	{
 		// If there is a loop anywhere in this expression's inputs then we can't risk checking them
-		if (ExpressionInput->Expression && ExpressionInput->Expression->IsResultMaterialAttributes(ExpressionInput->OutputIndex))
+		if (It->Expression && It->Expression->IsResultMaterialAttributes(It->OutputIndex))
 		{
 			return true;
 		}
@@ -12968,7 +12920,7 @@ TArrayView<FExpressionInput*> UMaterialExpressionComposite::GetInputsView()
 
 FExpressionInput* UMaterialExpressionComposite::GetInput(int32 InputIndex)
 {
-	if (InputIndex >= 0 && InputIndex < InputExpressions->ReroutePins.Num())
+	if (InputExpressions && InputExpressions->ReroutePins.IsValidIndex(InputIndex))
 	{
 		if (InputExpressions->ReroutePins[InputIndex].Expression)
 		{
@@ -12981,7 +12933,7 @@ FExpressionInput* UMaterialExpressionComposite::GetInput(int32 InputIndex)
 
 FName UMaterialExpressionComposite::GetInputName(int32 InputIndex) const
 {
-	if (InputIndex >= 0 && InputIndex < InputExpressions->ReroutePins.Num())
+	if (InputExpressions && InputExpressions->ReroutePins.IsValidIndex(InputIndex))
 	{
 		return InputExpressions->ReroutePins[InputIndex].Name;
 	}
@@ -12991,7 +12943,7 @@ FName UMaterialExpressionComposite::GetInputName(int32 InputIndex) const
 
 uint32 UMaterialExpressionComposite::GetInputType(int32 InputIndex)
 {
-	if (InputIndex >= 0 && InputIndex < InputExpressions->ReroutePins.Num())
+	if (InputExpressions && InputExpressions->ReroutePins.IsValidIndex(InputIndex))
 	{
 		return InputExpressions->ReroutePins[InputIndex].Expression->GetInputType(0);
 	}
@@ -13002,7 +12954,7 @@ uint32 UMaterialExpressionComposite::GetInputType(int32 InputIndex)
 
 uint32 UMaterialExpressionComposite::GetOutputType(int32 OutputIndex)
 {
-	if (OutputIndex >= 0 && OutputIndex < OutputExpressions->ReroutePins.Num())
+	if (OutputExpressions && OutputExpressions->ReroutePins.IsValidIndex(OutputIndex))
 	{
 		return OutputExpressions->ReroutePins[OutputIndex].Expression->GetOutputType(0);
 	}
@@ -13013,7 +12965,7 @@ uint32 UMaterialExpressionComposite::GetOutputType(int32 OutputIndex)
 
 bool UMaterialExpressionComposite::IsExpressionConnected(FExpressionInput* Input, int32 OutputIndex)
 {
-	if (Input && OutputIndex >= 0 && OutputIndex < OutputExpressions->ReroutePins.Num())
+	if (Input && OutputExpressions && OutputExpressions->ReroutePins.IsValidIndex(OutputIndex))
 	{
 		return OutputExpressions->ReroutePins[OutputIndex].Expression == Input->Expression;
 	}
@@ -14222,9 +14174,9 @@ int32 UMaterialExpressionDistanceFieldsRenderingSwitch::Compile(class FMaterialC
 
 bool UMaterialExpressionDistanceFieldsRenderingSwitch::IsResultMaterialAttributes(int32 OutputIndex)
 {
-	for (FExpressionInput* ExpressionInput : GetInputsView())
+	for (FExpressionInputIterator It{ this }; It; ++It)
 	{
-		if (ExpressionInput->GetTracedInput().Expression && ExpressionInput->Expression->IsResultMaterialAttributes(ExpressionInput->OutputIndex))
+		if (It->GetTracedInput().Expression && It->Expression->IsResultMaterialAttributes(It->OutputIndex))
 		{
 			return true;
 		}
@@ -14480,11 +14432,7 @@ TArrayView<FExpressionInput*> UMaterialExpressionCustom::GetInputsView()
 
 FExpressionInput* UMaterialExpressionCustom::GetInput(int32 InputIndex)
 {
-	if( InputIndex < Inputs.Num() )
-	{
-		return &Inputs[InputIndex].Input;
-	}
-	return nullptr;
+	return Inputs.IsValidIndex(InputIndex) ? &Inputs[InputIndex].Input : nullptr;
 }
 
 FName UMaterialExpressionCustom::GetInputName(int32 InputIndex) const
@@ -17553,11 +17501,7 @@ TArrayView<FExpressionInput*> UMaterialExpressionMaterialFunctionCall::GetInputs
 
 FExpressionInput* UMaterialExpressionMaterialFunctionCall::GetInput(int32 InputIndex)
 {
-	if (InputIndex < FunctionInputs.Num())
-	{
-		return &FunctionInputs[InputIndex].Input;
-	}
-	return nullptr;
+	return FunctionInputs.IsValidIndex(InputIndex) ? &FunctionInputs[InputIndex].Input : nullptr;
 }
 
 
@@ -17894,52 +17838,46 @@ void UMaterialExpressionMaterialFunctionCall::UpdateFromFunctionResource(bool bR
 	}
 }
 
-/** Goes through the Inputs array and fixes up each input's OutputIndex, or breaks the connection if necessary. */
-static void FixupReferencingInputs(
+/** Fixes CurrentInput's OutputIndex, or breaks the connection if necessary. */
+static void FixupReferencingInput(
+	FExpressionInput* CurrentInput,
 	const TArray<FFunctionExpressionOutput>& NewOutputs,
 	const TArray<FFunctionExpressionOutput>& OriginalOutputs,
-	TArrayView<FExpressionInput*> Inputs, 
 	UMaterialExpressionMaterialFunctionCall* FunctionExpression,
 	bool bMatchByName)
 {
-	for (int32 InputIndex = 0; InputIndex < Inputs.Num(); InputIndex++)
+	if (CurrentInput->Expression == FunctionExpression)
 	{
-		FExpressionInput* CurrentInput = Inputs[InputIndex];
-
-		if (CurrentInput->Expression == FunctionExpression)
+		if (OriginalOutputs.IsValidIndex(CurrentInput->OutputIndex))
 		{
-			if (OriginalOutputs.IsValidIndex(CurrentInput->OutputIndex))
+			if (bMatchByName)
 			{
-				if (bMatchByName)
-				{
-					if (OriginalOutputs[CurrentInput->OutputIndex].ExpressionOutput)
-					{
-						CurrentInput->OutputIndex = FindOutputIndexByName(OriginalOutputs[CurrentInput->OutputIndex].ExpressionOutput->OutputName, NewOutputs);
-					}
-				}
-				else
+				if (OriginalOutputs[CurrentInput->OutputIndex].ExpressionOutput)
 				{
-					const FGuid OutputId = OriginalOutputs[CurrentInput->OutputIndex].ExpressionOutputId;
-					CurrentInput->OutputIndex = FindOutputIndexById(OutputId, NewOutputs);
-				}
-
-				if (CurrentInput->OutputIndex == INDEX_NONE)
-				{
-					// The output that this input was connected to no longer exists, break the connection
-					CurrentInput->Expression = nullptr;
+					CurrentInput->OutputIndex = FindOutputIndexByName(OriginalOutputs[CurrentInput->OutputIndex].ExpressionOutput->OutputName, NewOutputs);
 				}
 			}
 			else
+			{
+				const FGuid OutputId = OriginalOutputs[CurrentInput->OutputIndex].ExpressionOutputId;
+				CurrentInput->OutputIndex = FindOutputIndexById(OutputId, NewOutputs);
+			}
+
+			if (CurrentInput->OutputIndex == INDEX_NONE)
 			{
 				// The output that this input was connected to no longer exists, break the connection
-				CurrentInput->OutputIndex = INDEX_NONE;
 				CurrentInput->Expression = nullptr;
 			}
 		}
+		else
+		{
+			// The output that this input was connected to no longer exists, break the connection
+			CurrentInput->OutputIndex = INDEX_NONE;
+			CurrentInput->Expression = nullptr;
+		}
 	}
 }
 
-
 void UMaterialExpressionMaterialFunctionCall::FixupReferencingExpressions(
 	const TArray<FFunctionExpressionOutput>& NewOutputs,
 	const TArray<FFunctionExpressionOutput>& OriginalOutputs,
@@ -17952,12 +17890,17 @@ void UMaterialExpressionMaterialFunctionCall::FixupReferencingExpressions(
 		UMaterialExpression* CurrentExpression = Expressions[ExpressionIndex];
 		if (CurrentExpression)
 		{
-			TArrayView<FExpressionInput*> Inputs = CurrentExpression->GetInputsView();
-			FixupReferencingInputs(NewOutputs, OriginalOutputs, Inputs, this, bMatchByName);
+			for (FExpressionInputIterator It{ CurrentExpression}; It; ++It)
+			{
+				FixupReferencingInput(It.Input, NewOutputs, OriginalOutputs, this, bMatchByName);
+			}
 		}
 	}
 
-	FixupReferencingInputs(NewOutputs, OriginalOutputs, MaterialInputs, this, bMatchByName);
+	for (FExpressionInput* CurrentInput : MaterialInputs)
+	{
+		FixupReferencingInput(CurrentInput, NewOutputs, OriginalOutputs, this, bMatchByName);
+	}
 }
 
 bool UMaterialExpressionMaterialFunctionCall::MatchesSearchQuery( const TCHAR* SearchQuery )
@@ -19097,9 +19040,9 @@ int32 UMaterialExpressionMaterialProxyReplace::Compile(class FMaterialCompiler*
 
 bool UMaterialExpressionMaterialProxyReplace::IsResultMaterialAttributes(int32 OutputIndex)
 {
-	for (FExpressionInput* ExpressionInput : GetInputsView())
+	for (FExpressionInputIterator It{ this }; It; ++It)
 	{
-		if (ExpressionInput->GetTracedInput().Expression && ExpressionInput->Expression->IsResultMaterialAttributes(ExpressionInput->OutputIndex))
+		if (It->GetTracedInput().Expression && It->Expression->IsResultMaterialAttributes(It->OutputIndex))
 		{
 			return true;
 		}
@@ -20441,7 +20384,7 @@ bool UMaterialExpressionNoise::CanEditChange(const FProperty* InProperty) const
 
 FName UMaterialExpressionNoise::GetInputName(int32 InputIndex) const
 {
-	if (CachedInputs[InputIndex] == &Position)
+	if (GetInput(InputIndex) == &Position)
 	{
 		return GetWorldPositionInputName(WorldPositionOriginType);
 	}
@@ -20552,7 +20495,7 @@ bool UMaterialExpressionVectorNoise::CanEditChange(const FProperty* InProperty)
 
 FName UMaterialExpressionVectorNoise::GetInputName(int32 InputIndex) const
 {
-	if (CachedInputs[InputIndex] == &Position)
+	if (GetInput(InputIndex) == &Position)
 	{
 		return GetWorldPositionInputName(WorldPositionOriginType);
 	}
@@ -20670,7 +20613,7 @@ UMaterialExpressionDistanceToNearestSurface::UMaterialExpressionDistanceToNeares
 #if WITH_EDITOR
 FName UMaterialExpressionDistanceToNearestSurface::GetInputName(int32 InputIndex) const
 {
-	if (CachedInputs[InputIndex] == &Position)
+	if (GetInput(InputIndex) == &Position)
 	{
 		return GetWorldPositionInputName(WorldPositionOriginType);
 	}
@@ -20739,7 +20682,7 @@ UMaterialExpressionDistanceFieldGradient::UMaterialExpressionDistanceFieldGradie
 #if WITH_EDITOR
 FName UMaterialExpressionDistanceFieldGradient::GetInputName(int32 InputIndex) const
 {
-	if (CachedInputs[InputIndex] == &Position)
+	if (GetInput(InputIndex) == &Position)
 	{
 		return GetWorldPositionInputName(WorldPositionOriginType);
 	}
@@ -20813,7 +20756,7 @@ UMaterialExpressionDistanceFieldApproxAO::UMaterialExpressionDistanceFieldApprox
 #if WITH_EDITOR
 FName UMaterialExpressionDistanceFieldApproxAO::GetInputName(int32 InputIndex) const
 {
-	if (CachedInputs[InputIndex] == &Position)
+	if (GetInput(InputIndex) == &Position)
 	{
 		return GetWorldPositionInputName(WorldPositionOriginType);
 	}
@@ -20899,7 +20842,7 @@ UMaterialExpressionSamplePhysicsVectorField::UMaterialExpressionSamplePhysicsVec
 #if WITH_EDITOR
 FName UMaterialExpressionSamplePhysicsVectorField::GetInputName(int32 InputIndex) const
 {
-	if (CachedInputs[InputIndex] == &WorldPosition)
+	if (GetInput(InputIndex) == &WorldPosition)
 	{
 		return GetWorldPositionInputName(WorldPositionOriginType);
 	}
@@ -20971,7 +20914,7 @@ UMaterialExpressionSamplePhysicsScalarField::UMaterialExpressionSamplePhysicsSca
 #if WITH_EDITOR
 FName UMaterialExpressionSamplePhysicsScalarField::GetInputName(int32 InputIndex) const
 {
-	if (CachedInputs[InputIndex] == &WorldPosition)
+	if (GetInput(InputIndex) == &WorldPosition)
 	{
 		return GetWorldPositionInputName(WorldPositionOriginType);
 	}
@@ -21046,7 +20989,7 @@ UMaterialExpressionSamplePhysicsIntegerField::UMaterialExpressionSamplePhysicsIn
 #if WITH_EDITOR
 FName UMaterialExpressionSamplePhysicsIntegerField::GetInputName(int32 InputIndex) const
 {
-	if (CachedInputs[InputIndex] == &WorldPosition)
+	if (GetInput(InputIndex) == &WorldPosition)
 	{
 		return GetWorldPositionInputName(WorldPositionOriginType);
 	}
@@ -22281,7 +22224,7 @@ UMaterialExpressionAtmosphericFogColor::UMaterialExpressionAtmosphericFogColor(c
 #if WITH_EDITOR
 FName UMaterialExpressionAtmosphericFogColor::GetInputName(int32 InputIndex) const
 {
-	if (CachedInputs[InputIndex] == &WorldPosition)
+	if (GetInput(InputIndex) == &WorldPosition)
 	{
 		return GetWorldPositionInputName(WorldPositionOriginType);
 	}
@@ -22602,7 +22545,7 @@ void UMaterialExpressionClearCoatNormalCustomOutput::GetCaption(TArray<FString>&
 
 FExpressionInput* UMaterialExpressionClearCoatNormalCustomOutput::GetInput(int32 InputIndex)
 {
-	return &Input;
+	return InputIndex == 0 ? &Input : nullptr;
 }
 #endif // WITH_EDITOR
 
@@ -22655,7 +22598,7 @@ void UMaterialExpressionBentNormalCustomOutput::GetCaption(TArray<FString>& OutC
 
 FExpressionInput* UMaterialExpressionBentNormalCustomOutput::GetInput(int32 InputIndex)
 {
-	return &Input;
+	return InputIndex == 0 ? &Input : nullptr;
 }
 #endif // WITH_EDITOR
 
@@ -22769,7 +22712,7 @@ void UMaterialExpressionVertexInterpolator::GetCaption(TArray<FString>& OutCapti
 
 FExpressionInput* UMaterialExpressionVertexInterpolator::GetInput(int32 InputIndex)
 {
-	return &Input;
+	return InputIndex == 0 ? &Input : nullptr;
 }
 
 uint32 UMaterialExpressionVertexInterpolator::GetInputType(int32 InputIndex)
@@ -22880,7 +22823,7 @@ UMaterialExpressionSkyAtmosphereLightIlluminance::UMaterialExpressionSkyAtmosphe
 #if WITH_EDITOR
 FName UMaterialExpressionSkyAtmosphereLightIlluminance::GetInputName(int32 InputIndex) const
 {
-	if (CachedInputs[InputIndex] == &WorldPosition)
+	if (GetInput(InputIndex) == &WorldPosition)
 	{
 		return GetWorldPositionInputName(WorldPositionOriginType);
 	}
@@ -23101,7 +23044,7 @@ UMaterialExpressionSkyAtmosphereAerialPerspective::UMaterialExpressionSkyAtmosph
 #if WITH_EDITOR
 FName UMaterialExpressionSkyAtmosphereAerialPerspective::GetInputName(int32 InputIndex) const
 {
-	if (CachedInputs[InputIndex] == &WorldPosition)
+	if (GetInput(InputIndex) == &WorldPosition)
 	{
 		return GetWorldPositionInputName(WorldPositionOriginType);
 	}
@@ -24609,12 +24552,38 @@ UMaterialExpressionSubstrateShadingModels::UMaterialExpressionSubstrateShadingMo
 #if WITH_EDITORONLY_DATA
 	MenuCategories.Add(ConstructorStatics.NAME_Substrate);
 #endif
-#if WITH_EDITOR
-	CachedInputs.Add(&ShadingModel);
-#endif
 }
 
 #if WITH_EDITOR
+FExpressionInput* UMaterialExpressionSubstrateShadingModels::GetInput(int32 InputIndex)
+{
+	switch (InputIndex)
+	{
+		case 0: return &BaseColor;
+		case 1: return &Metallic;
+		case 2: return &Specular;
+		case 3: return &Roughness;
+		case 4: return &Anisotropy;
+		case 5: return &EmissiveColor;
+		case 6: return &Normal;
+		case 7: return &Tangent;
+		case 8: return &SubSurfaceColor;
+		case 9: return &ClearCoat;
+		case 10: return &ClearCoatRoughness;
+		case 11: return &Opacity;
+		case 12: return &TransmittanceColor;
+		case 13: return &WaterScatteringCoefficients;
+		case 14: return &WaterAbsorptionCoefficients;
+		case 15: return &WaterPhaseG;
+		case 16: return &ColorScaleBehindWater;
+		case 17: return &ClearCoatNormal;
+		case 18: return &CustomTangent;
+		case 19: return &ThinTranslucentSurfaceCoverage;
+		case 20: return &ShadingModel;
+		default: return nullptr;
+	}
+}
+
 int32 UMaterialExpressionSubstrateShadingModels::Compile(class FMaterialCompiler* Compiler, int32 OutputIndex)
 {
 	int32 RoughnessCodeChunk = CompileWithDefaultFloat1(Compiler, Roughness, 0.5f);
@@ -27777,16 +27746,6 @@ UMaterialExpressionSubstrateConvertMaterialAttributes::UMaterialExpressionSubstr
 #if WITH_EDITORONLY_DATA
 	MenuCategories.Add(ConstructorStatics.NAME_Strata);
 #endif
-	// Cache some material attribute inputs?
-#if WITH_EDITOR
-	CachedInputs.Empty();
-	CachedInputs.Reserve(1);
-	CachedInputs.Add(&MaterialAttributes);
-	CachedInputs.Add(&WaterScatteringCoefficients); 
-	CachedInputs.Add(&WaterAbsorptionCoefficients);
-	CachedInputs.Add(&WaterPhaseG);
-	CachedInputs.Add(&ColorScaleBehindWater); 
-#endif
 
 #if WITH_EDITORONLY_DATA
 	Outputs.Reset();
@@ -27796,6 +27755,19 @@ UMaterialExpressionSubstrateConvertMaterialAttributes::UMaterialExpressionSubstr
 }
 
 #if WITH_EDITOR
+FExpressionInput* UMaterialExpressionSubstrateConvertMaterialAttributes::GetInput(int32 InputIndex)
+{
+	switch (InputIndex)
+	{
+		case 0: return &MaterialAttributes;
+		case 1: return &WaterScatteringCoefficients;
+		case 2: return &WaterAbsorptionCoefficients;
+		case 3: return &WaterPhaseG;
+		case 4: return &ColorScaleBehindWater;
+		default: return nullptr;
+	}
+}
+
 int32 UMaterialExpressionSubstrateConvertMaterialAttributes::Compile(class FMaterialCompiler* Compiler, int32 OutputIndex)
 {
 	static const FGuid ClearCoatBottomNormalGuid 			= FMaterialAttributeDefinitionMap::GetCustomAttributeID(TEXT("ClearCoatBottomNormal"));
diff --git a/Engine/Source/Runtime/Engine/Private/Materials/MaterialHLSLGenerator.cpp b/Engine/Source/Runtime/Engine/Private/Materials/MaterialHLSLGenerator.cpp
--- a/Engine/Source/Runtime/Engine/Private/Materials/MaterialHLSLGenerator.cpp
+++ b/Engine/Source/Runtime/Engine/Private/Materials/MaterialHLSLGenerator.cpp
@@ -445,14 +445,17 @@ FStringView FMaterialHLSLGenerator::AcquireError()
 
 int32 FMaterialHLSLGenerator::FindInputIndex(const FExpressionInput* Input) const
 {
-	UMaterialExpression* OwnerMaterialExpression = GetCurrentExpression();
-	int32 Index = INDEX_NONE;
-	if (OwnerMaterialExpression)
+	if (UMaterialExpression* OwnerMaterialExpression = GetCurrentExpression())
 	{
-		TArrayView<FExpressionInput*> Inputs = OwnerMaterialExpression->GetInputsView();
-		Index = Inputs.Find(const_cast<FExpressionInput*>(Input));
+		for (FExpressionInputIterator It{ OwnerMaterialExpression }; It; ++It)
+		{
+			if (Input == It.Input)
+			{
+				return It.Index;
+			}
+		}
 	}
-	return Index;
+	return INDEX_NONE;
 }
 
 const UE::HLSLTree::FExpression* FMaterialHLSLGenerator::NewDefaultInputConstant(int32 InputIndex, const UE::Shader::FValue& Value)
diff --git a/Engine/Source/Runtime/Engine/Public/Materials/MaterialExpression.h b/Engine/Source/Runtime/Engine/Public/Materials/MaterialExpression.h
--- a/Engine/Source/Runtime/Engine/Public/Materials/MaterialExpression.h
+++ b/Engine/Source/Runtime/Engine/Public/Materials/MaterialExpression.h
@@ -345,26 +345,30 @@ class UMaterialExpression : public UObject
 	virtual UTextureCollection* GetReferencedTextureCollection() const { return nullptr; }
 
 #if WITH_EDITOR
+
+	/** Counts the number of inputs this expression has. Default implementation has O(n) complexity. */
+	ENGINE_API virtual int32 CountInputs() const;
+
 	/**
-	 *	Get the outputs supported by this expression.
-	 *
-	 *	@param	Outputs		The TArray of outputs to fill in.
+	 * Returns the input at index InputIndex if valid, nullptr otherwise.
+	 * Note: Implementations are required to return nullptr for invalid input indices.
 	 */
-	ENGINE_API virtual TArray<FExpressionOutput>& GetOutputs();
-	
-	/** Get the expression inputs supported by this expression (Note: property inputs NOT included). */
+	ENGINE_API virtual FExpressionInput* GetInput(int32 InputIndex);
+	ENGINE_API const FExpressionInput* GetInput(int32 InputIndex) const { return const_cast<UMaterialExpression*>(this)->GetInput(InputIndex); }
+
+	UE_DEPRECATED(5.5, "Use FExpressionInputIterator instead or GetInput() directly.")
 	ENGINE_API virtual TArrayView<FExpressionInput*> GetInputsView();
-	
-	UE_DEPRECATED(5.3, "Use GetInputsView() instead.")
-	const TArray<FExpressionInput*> GetInputs() { return TArray<FExpressionInput*>{ GetInputsView() }; }
 
-	ENGINE_API virtual FExpressionInput* GetInput(int32 InputIndex);
 	ENGINE_API virtual FName GetInputName(int32 InputIndex) const;
+	
 	ENGINE_API virtual bool IsInputConnectionRequired(int32 InputIndex) const;
-	virtual bool CanUserDeleteExpression() const
-	{
-		return true;
-	};
+		
+	ENGINE_API virtual uint32 GetInputType(int32 InputIndex);
+
+	ENGINE_API virtual TArray<FExpressionOutput>& GetOutputs();
+	ENGINE_API virtual uint32 GetOutputType(int32 OutputIndex);
+
+	virtual bool CanUserDeleteExpression() const { return true; }
 
 	/** Find the property that is associated with the input pin. */
 	ENGINE_API virtual TArray<FProperty*> GetInputPinProperty(int32 PinIndex);
@@ -376,9 +380,6 @@ class UMaterialExpression : public UObject
 	ENGINE_API virtual FString GetInputPinDefaultValue(int32 PinIndex);
 	ENGINE_API virtual TArray<FProperty*> GetPropertyInputs() const;
 
-	ENGINE_API virtual uint32 GetInputType(int32 InputIndex);
-	ENGINE_API virtual uint32 GetOutputType(int32 OutputIndex);
-
 	ENGINE_API virtual void GetExecOutputs(TArray<FExpressionExecOutputEntry>& Outputs);
 
 	ENGINE_API virtual bool HasExecInput();
@@ -587,7 +588,7 @@ class UMaterialExpression : public UObject
 	ENGINE_API bool ContainsInputLoop(TSet<UMaterialExpression*>& VisitedExpressions, const bool bStopOnFunctionCall = true);
 
 protected:
-	/** Caches the list of expression inputs this expression has. */
+	/** Caches the list of fixed expression inputs this expression has. */
 	TArray<FExpressionInput*> CachedInputs;
 
 private:
@@ -606,8 +607,6 @@ private:
 	 */
 	ENGINE_API bool ContainsInputLoopInternal(const FContainsInputLoopInternalExpressionStack& ExpressionStack, TSet<UMaterialExpression*>& VisitedExpressions, const bool bStopOnFunctionCall);
 
-	UE_DEPRECATED(5.3, "Use the other, more efficient ContainsInputLoopInternal() implementation.")
-	bool ContainsInputLoopInternal(TArray<FMaterialExpressionKey>& ExpressionStack, TSet<FMaterialExpressionKey>& VisitedExpressions, const bool bStopOnFunctionCall);
 #endif // WITH_EDITOR
 };
 
@@ -621,3 +620,61 @@ enum class EPositionOrigin : uint8
 	/** Camera relative world position, i.e. translated world space */
 	CameraRelative UMETA(DisplayName="Camera Relative World Position")
 };
+
+
+#if WITH_EDITOR
+
+/**
+ * @brief An iterator for traversing the inputs of a UMaterialExpression.
+ *
+ * This struct provides a way to iterate over the inputs of a given
+ * UMaterialExpression. It starts from the first input and advances
+ * through subsequent inputs until all are iterated through.
+ *
+ * Usage:
+ *
+ * @code
+ * UMaterialExpression* MyExpression = ...;
+ * for (FExpressionInputIterator It{ MyExpression }; It; ++It)
+ * {
+ *     // Process the input, e.g. It->Input->IsConnected();
+ * }
+ * @endcode
+ */
+struct FExpressionInputIterator
+{
+	/** The expression whose inputs to iterate through. */
+    UMaterialExpression* Expression;
+
+	/** Current input. */
+    FExpressionInput* Input;
+
+	/** Current input index. */
+	int Index;
+
+    FExpressionInputIterator(UMaterialExpression* InExpression)
+    : Expression{ InExpression }
+    , Input{ InExpression->GetInput(0) }
+    , Index{ 0 }
+    {
+    }
+
+    operator bool() const
+    {
+        return Input != nullptr;
+    }
+
+    FExpressionInputIterator& operator++()
+    {
+		Index += 1;
+		Input = Expression->GetInput(Index);
+        return *this;
+    }
+
+    FExpressionInput* operator->()
+    {
+        return Input;
+    }
+};
+
+#endif // WITH_EDITOR
diff --git a/Engine/Source/Runtime/Engine/Public/Materials/MaterialExpressionBlendMaterialAttributes.h b/Engine/Source/Runtime/Engine/Public/Materials/MaterialExpressionBlendMaterialAttributes.h
--- a/Engine/Source/Runtime/Engine/Public/Materials/MaterialExpressionBlendMaterialAttributes.h
+++ b/Engine/Source/Runtime/Engine/Public/Materials/MaterialExpressionBlendMaterialAttributes.h
@@ -44,6 +44,7 @@ class UMaterialExpressionBlendMaterialAttributes : public UMaterialExpression
 
 	//~ Begin UMaterialExpression Interface
 #if WITH_EDITOR
+	virtual FExpressionInput* GetInput(int32 InputIndex) override;
 	virtual int32 Compile(class FMaterialCompiler* Compiler, int32 OutputIndex) override;
 	virtual void GetCaption(TArray<FString>& OutCaptions) const override;
 	virtual FName GetInputName(int32 InputIndex) const override;
diff --git a/Engine/Source/Runtime/Engine/Public/Materials/MaterialExpressionBreakMaterialAttributes.h b/Engine/Source/Runtime/Engine/Public/Materials/MaterialExpressionBreakMaterialAttributes.h
--- a/Engine/Source/Runtime/Engine/Public/Materials/MaterialExpressionBreakMaterialAttributes.h
+++ b/Engine/Source/Runtime/Engine/Public/Materials/MaterialExpressionBreakMaterialAttributes.h
@@ -24,6 +24,7 @@ class UMaterialExpressionBreakMaterialAttributes : public UMaterialExpression
  
 	//~ Begin UMaterialExpression Interface
 #if WITH_EDITOR
+	virtual FExpressionInput* GetInput(int32 InputIndex) override { return InputIndex == 0 ? &MaterialAttributes : nullptr; }
 	virtual int32 Compile(class FMaterialCompiler* Compiler, int32 OutputIndex) override;
 	virtual void GetCaption(TArray<FString>& OutCaptions) const override;
 	virtual FName GetInputName(int32 InputIndex) const override;
diff --git a/Engine/Source/Runtime/Engine/Public/Materials/MaterialExpressionSubstrate.h b/Engine/Source/Runtime/Engine/Public/Materials/MaterialExpressionSubstrate.h
--- a/Engine/Source/Runtime/Engine/Public/Materials/MaterialExpressionSubstrate.h
+++ b/Engine/Source/Runtime/Engine/Public/Materials/MaterialExpressionSubstrate.h
@@ -183,6 +183,7 @@ class UMaterialExpressionSubstrateShadingModels : public UMaterialExpressionSubs
 
 	//~ Begin UMaterialExpression Interface
 #if WITH_EDITOR
+	virtual FExpressionInput* GetInput(int32 InputIndex) override;
 	virtual int32 Compile(class FMaterialCompiler* Compiler, int32 OutputIndex) override;
 	virtual void GetCaption(TArray<FString>& OutCaptions) const override;
 	virtual uint32 GetOutputType(int32 OutputIndex) override;
@@ -867,6 +868,7 @@ class UMaterialExpressionSubstrateConvertMaterialAttributes : public UMaterialEx
 
 	//~ Begin UMaterialExpression Interface
 #if WITH_EDITOR
+	virtual FExpressionInput* GetInput(int32 InputIndex) override;
 	virtual int32 Compile(class FMaterialCompiler* Compiler, int32 OutputIndex) override;
 	virtual void GetCaption(TArray<FString>& OutCaptions) const override;
 	virtual uint32 GetOutputType(int32 OutputIndex) override;
diff --git a/Engine/Source/Runtime/Landscape/Private/LandscapeGrass.cpp b/Engine/Source/Runtime/Landscape/Private/LandscapeGrass.cpp
--- a/Engine/Source/Runtime/Landscape/Private/LandscapeGrass.cpp
+++ b/Engine/Source/Runtime/Landscape/Private/LandscapeGrass.cpp
@@ -1154,7 +1154,7 @@ TArrayView<FExpressionInput*> UMaterialExpressionLandscapeGrassOutput::GetInputs
 
 FExpressionInput* UMaterialExpressionLandscapeGrassOutput::GetInput(int32 InputIndex)
 {
-	return &GrassTypes[InputIndex].Input;
+	return GrassTypes.IsValidIndex(InputIndex) ? &GrassTypes[InputIndex].Input : nullptr;
 }
 
 FName UMaterialExpressionLandscapeGrassOutput::GetInputName(int32 InputIndex) const
diff --git a/Engine/Source/Runtime/Landscape/Private/Materials/MaterialExpressionLandscapePhysicalMaterialOutput.cpp b/Engine/Source/Runtime/Landscape/Private/Materials/MaterialExpressionLandscapePhysicalMaterialOutput.cpp
--- a/Engine/Source/Runtime/Landscape/Private/Materials/MaterialExpressionLandscapePhysicalMaterialOutput.cpp
+++ b/Engine/Source/Runtime/Landscape/Private/Materials/MaterialExpressionLandscapePhysicalMaterialOutput.cpp
@@ -57,7 +57,7 @@ TArrayView<FExpressionInput*> UMaterialExpressionLandscapePhysicalMaterialOutput
 
 FExpressionInput* UMaterialExpressionLandscapePhysicalMaterialOutput::GetInput(int32 InputIndex)
 {
-	return &Inputs[InputIndex].Input;
+	return Inputs.IsValidIndex(InputIndex) ? &Inputs[InputIndex].Input : nullptr;
 }
 
 
