Fix crash when setting the Text property in Text Render component to a string longer than 16384 characters.

Switched the index buffer to 32 bits so this is no longer a problem in practice and changed the code to do actual error handling (truncate the text if it's too long) rather than assert.

#rnx

diff --git a/Engine/Source/Runtime/Engine/Private/Components/TextRenderComponent.cpp b/Engine/Source/Runtime/Engine/Private/Components/TextRenderComponent.cpp
--- a/Engine/Source/Runtime/Engine/Private/Components/TextRenderComponent.cpp
+++ b/Engine/Source/Runtime/Engine/Private/Components/TextRenderComponent.cpp
@@ -618,7 +618,7 @@ private:
 	// End FPrimitiveSceneProxy interface
 
 	void ReleaseRenderThreadResources();
-	bool BuildStringMesh( TArray<FDynamicMeshVertex>& OutVertices, TArray<uint16>& OutIndices );
+	bool BuildStringMesh( TArray<FDynamicMeshVertex>& OutVertices, TArray<uint32>& OutIndices );
 
 private:
 	struct FTextBatch
@@ -633,7 +633,7 @@ private:
 	};
 
 	FStaticMeshVertexBuffers VertexBuffers;
-	FDynamicMeshIndexBuffer16 IndexBuffer;
+	FDynamicMeshIndexBuffer32 IndexBuffer;
 	FLocalVertexFactory VertexFactory;
 	TArray<FTextBatch> TextBatches;
 	const FColor TextRenderColor;
@@ -966,7 +966,7 @@ void FTextRenderSceneProxy::GetDynamicRayTracingInstances(FRayTracingMaterialGat
 /**
 * For the given text, constructs a mesh to be used by the vertex factory for rendering.
 */
-bool FTextRenderSceneProxy::BuildStringMesh( TArray<FDynamicMeshVertex>& OutVertices, TArray<uint16>& OutIndices )
+bool FTextRenderSceneProxy::BuildStringMesh( TArray<FDynamicMeshVertex>& OutVertices, TArray<uint32>& OutIndices )
 {
 	TextBatches.Reset();
 
@@ -1013,8 +1013,20 @@ bool FTextRenderSceneProxy::BuildStringMesh( TArray<FDynamicMeshVertex>& OutVert
 		FirstIndiceIndexInTextBatch = OutIndices.Num();
 	};
 
-	FTextIterator It(*Text.ToString());
-	while (It.NextLine())
+	// TArray uses int32 for size so it cannot store more than MAX_int32 elements.
+	const uint64 MaxVerts = MAX_int32 - OutVertices.Num();
+	const uint64 MaxIndices = MAX_int32 - OutIndices.Num();
+	const uint64 MaxQuads = FMath::Min(MaxVerts / 4, MaxIndices / 6);
+
+	// We have one quad per glyph, so presize the index and vertex buffer.
+	FString Str = Text.ToString();
+	const uint64 NumQuads = FMath::Min((uint64)Str.Len(), MaxQuads);
+	OutVertices.Reserve(OutVertices.Num() + NumQuads * 4);
+	OutIndices.Reserve(OutIndices.Num() + NumQuads * 6);
+
+	uint64 QuadIdx = 0;
+	FTextIterator It(*Str);
+	while (It.NextLine() && QuadIdx < NumQuads)
 	{
 		FVector2D LineSize = ComputeTextSize(It, Font, XScale, YScale, HorizSpacingAdjust, VertSpacingAdjust);
 		float StartX = ComputeHorizontalAlignmentOffset(LineSize, HorizontalAlignment);
@@ -1027,7 +1039,7 @@ bool FTextRenderSceneProxy::BuildStringMesh( TArray<FDynamicMeshVertex>& OutVert
 		LineX = 0.f;
 
 		TCHAR Ch = 0;
-		while (It.NextCharacterInLine(Ch))
+		while (It.NextCharacterInLine(Ch) && QuadIdx < NumQuads)
 		{
 			Ch = Font->RemapChar(Ch);
 
@@ -1090,11 +1102,6 @@ bool FTextRenderSceneProxy::BuildStringMesh( TArray<FDynamicMeshVertex>& OutVert
 				const int32 V01 = OutVertices.Add(FDynamicMeshVertex(V2, TangentX, TangentZ, FVector2f(U, V + SizeV), TextRenderColor));
 				const int32 V11 = OutVertices.Add(FDynamicMeshVertex(V3, TangentX, TangentZ, FVector2f(U + SizeU, V + SizeV), TextRenderColor));
 
-				check(V00 < 65536);
-				check(V10 < 65536);
-				check(V01 < 65536);
-				check(V11 < 65536);
-
 				OutIndices.Add(V00);
 				OutIndices.Add(V11);
 				OutIndices.Add(V10);
@@ -1111,6 +1118,8 @@ bool FTextRenderSceneProxy::BuildStringMesh( TArray<FDynamicMeshVertex>& OutVert
 				{
 					LineX += CharIncrement;
 				}
+
+				++QuadIdx;
 			}
 		}
 
