[World Partition RTree Spatial Index Optimizations]
- Since elements are contiguous, there's no need to store an array of indices per leaf, simply store StartIndex and NumElements instead.
- When recursing into the tree, don't go through a virtual function to gather the node box, get the pointer/stride to the box list instead.
- Fixed Hilbert curve index encoding that was always an order below the required one, which makes on par and even better than Morton in some situations.
Also updated unit tests to include no sort tests.

diff --git a/Engine/Source/Runtime/Engine/Private/Tests/WorldPartition/WorldPartitionStaticSpatialIndexTests.cpp b/Engine/Source/Runtime/Engine/Private/Tests/WorldPartition/WorldPartitionStaticSpatialIndexTests.cpp
--- a/Engine/Source/Runtime/Engine/Private/Tests/WorldPartition/WorldPartitionStaticSpatialIndexTests.cpp
+++ b/Engine/Source/Runtime/Engine/Private/Tests/WorldPartition/WorldPartitionStaticSpatialIndexTests.cpp
@@ -35,19 +35,48 @@ namespace WorldPartitionTests
 		Class SpatialIndex;
 		SpatialIndex.Init(Elements);
 
-		Results.Reserve(Tests.Num());
-
 		const double StartTime = FPlatformTime::Seconds();
-
 		for (int32 ListNumTests = 0; ListNumTests < Tests.Num(); ListNumTests++)
 		{
 			SpatialIndex.ForEachIntersectingElement(Tests[ListNumTests], [&Results](const int32& Value) { Results.Add(Value); });
 		}
-
 		const double RunTime = FPlatformTime::Seconds() - StartTime;
+
 		Test->AddInfo(FString::Printf(TEXT("%s(%s): %d tests in %s (%.2f/s, %s)"), Name, GetSpaceString<typename Profile::FBox>(), Tests.Num(), *FPlatformTime::PrettyTime(RunTime), Tests.Num() / RunTime, *FGenericPlatformMemory::PrettyMemory(SpatialIndex.GetAllocatedSize())));
 	}
 
+	template <typename Profile, int32 MaxNumElementsPerNode, int32 MaxNumElementsPerLeaf>
+	FORCENOINLINE void PerformNoSortTest(FWorldPartitionStaticSpatialIndexTest* Test, const TArray<TPair<typename Profile::FBox, int32>>& Elements, const TArray<FSphere>& Tests, const TArray<int32>& ReferenceResults)
+	{
+		TArray<int32> RTreeNoSortResults;
+		RTreeNoSortResults.Reserve(ReferenceResults.Num());
+		FString TestName = FString::Printf(TEXT("TStaticSpatialIndexRTree(%s-NoSort-%d-%d)"), GetSpaceString<typename Profile::FBox>(), MaxNumElementsPerNode, MaxNumElementsPerLeaf);
+		PerformTests<Profile, TStaticSpatialIndexRTree<int32, FStaticSpatialIndex::TNodeSorterNoSort<Profile>, Profile>>(Test, *TestName, Elements, Tests, RTreeNoSortResults);
+		RTreeNoSortResults.Sort();
+		Test->TestTrue(TestName, RTreeNoSortResults == ReferenceResults);
+	}
+
+	template <typename Profile>
+	FORCENOINLINE void PerformNoSortTests(FWorldPartitionStaticSpatialIndexTest* Test, const TArray<TPair<typename Profile::FBox, int32>>& Elements, const TArray<FSphere>& Tests, const TArray<int32>& ReferenceResults)
+	{
+		PerformNoSortTest<Profile, 16, 16>(Test, Elements, Tests, ReferenceResults);
+		PerformNoSortTest<Profile, 16, 64>(Test, Elements, Tests, ReferenceResults);
+		PerformNoSortTest<Profile, 16, 256>(Test, Elements, Tests, ReferenceResults);
+		PerformNoSortTest<Profile, 16, 1024>(Test, Elements, Tests, ReferenceResults);
+		PerformNoSortTest<Profile, 64, 16>(Test, Elements, Tests, ReferenceResults);
+		PerformNoSortTest<Profile, 64, 64>(Test, Elements, Tests, ReferenceResults);
+		PerformNoSortTest<Profile, 64, 256>(Test, Elements, Tests, ReferenceResults);
+		PerformNoSortTest<Profile, 64, 1024>(Test, Elements, Tests, ReferenceResults);
+		PerformNoSortTest<Profile, 256, 16>(Test, Elements, Tests, ReferenceResults);
+		PerformNoSortTest<Profile, 256, 64>(Test, Elements, Tests, ReferenceResults);
+		PerformNoSortTest<Profile, 256, 256>(Test, Elements, Tests, ReferenceResults);
+		PerformNoSortTest<Profile, 256, 1024>(Test, Elements, Tests, ReferenceResults);
+		PerformNoSortTest<Profile, 1024, 16>(Test, Elements, Tests, ReferenceResults);
+		PerformNoSortTest<Profile, 1024, 64>(Test, Elements, Tests, ReferenceResults);
+		PerformNoSortTest<Profile, 1024, 256>(Test, Elements, Tests, ReferenceResults);
+		PerformNoSortTest<Profile, 1024, 1024>(Test, Elements, Tests, ReferenceResults);
+	}
+
 	template <typename Profile, int32 MaxNumElementsPerNode, int32 MaxNumElementsPerLeaf>
 	FORCENOINLINE void PerformMinXTest(FWorldPartitionStaticSpatialIndexTest* Test, const TArray<TPair<typename Profile::FBox, int32>>& Elements, const TArray<FSphere>& Tests, const TArray<int32>& ReferenceResults)
 	{
@@ -151,9 +180,8 @@ namespace WorldPartitionTests
 		PerformTests<Profile, TStaticSpatialIndexList<int32, FStaticSpatialIndex::TNodeSorterNoSort<Profile>, Profile>>(Test, TEXT("TStaticSpatialIndexList"), Elements, Tests, ListResults);
 		ListResults.Sort();
 
-		PerformMinXTests<Profile>(Test, Elements, Tests, ListResults);
-		PerformMinXTests<Profile>(Test, Elements, Tests, ListResults);
-		PerformMinXTests<Profile>(Test, Elements, Tests, ListResults);
+		PerformNoSortTests<Profile>(Test, Elements, Tests, ListResults);
+
 		PerformMinXTests<Profile>(Test, Elements, Tests, ListResults);
 
 		PerformMortonTests<Profile, 4096>(Test, Elements, Tests, ListResults);
diff --git a/Engine/Source/Runtime/Engine/Public/WorldPartition/RuntimeHashSet/StaticSpatialIndex.h b/Engine/Source/Runtime/Engine/Public/WorldPartition/RuntimeHashSet/StaticSpatialIndex.h
--- a/Engine/Source/Runtime/Engine/Public/WorldPartition/RuntimeHashSet/StaticSpatialIndex.h
+++ b/Engine/Source/Runtime/Engine/Public/WorldPartition/RuntimeHashSet/StaticSpatialIndex.h
@@ -43,6 +43,7 @@ struct TStaticSpatialIndexDataInterface
 	virtual ~TStaticSpatialIndexDataInterface() {}
 	virtual int32 GetNumBox() const = 0;
 	virtual const typename Profile::FBox& GetBox(uint32 InIndex) const = 0;
+	virtual const typename Profile::FBox* GetBoxes(uint32 InIndex, uint32& OutStride) const = 0;
 	virtual uint32 GetAllocatedSize() const = 0;
 };
 
@@ -123,6 +124,12 @@ public:
 		return Elements[InIndex].Key;
 	}
 
+	virtual const FBox* GetBoxes(uint32 InIndex, uint32& OutStride) const override
+	{
+		OutStride = Elements.GetTypeSize();
+		return &Elements[InIndex].Key;
+	}
+
 	virtual uint32 GetAllocatedSize() const override
 	{
 		return sizeof(*this) + Elements.GetAllocatedSize() + SpatialIndex.GetAllocatedSize();
@@ -191,11 +198,11 @@ namespace FStaticSpatialIndex
 
 		bool ForEachIntersectingElement(const FBox& InBox, TFunctionRef<bool(uint32 InValueIndex)> InFunc) const
 		{
-			for (int32 ValueIndex = 0; ValueIndex < this->DataInterface.GetNumBox(); ValueIndex++)
+			uint32 BoxStride;
+			const FBox* Box = this->DataInterface.GetBoxes(0, BoxStride);
+			for (int32 ValueIndex = 0; ValueIndex < this->DataInterface.GetNumBox(); ValueIndex++, *(uint8**)&Box += BoxStride)
 			{
-				const FBox& Box = this->DataInterface.GetBox(ValueIndex);
-
-				if (Box.Intersect(InBox))
+				if (Box->Intersect(InBox))
 				{
 					if (!InFunc(ValueIndex))
 					{
@@ -211,11 +218,11 @@ namespace FStaticSpatialIndex
 		{
 			const FSphere::FReal RadiusSquared = FMath::Square(InSphere.W);
 
-			for (int32 ValueIndex = 0; ValueIndex < this->DataInterface.GetNumBox(); ValueIndex++)
+			uint32 BoxStride;
+			const FBox* Box = this->DataInterface.GetBoxes(0, BoxStride);
+			for (int32 ValueIndex = 0; ValueIndex < this->DataInterface.GetNumBox(); ValueIndex++, *(uint8**)&Box += BoxStride)
 			{
-				const FBox& Box = this->DataInterface.GetBox(ValueIndex);
-
-				if (FastSphereAABBIntersection<Profile>(FVector(InSphere.Center), RadiusSquared, Box))
+				if (FastSphereAABBIntersection<Profile>(FVector(InSphere.Center), RadiusSquared, *Box))
 				{
 					if (!InFunc(ValueIndex))
 					{
@@ -337,7 +344,44 @@ namespace FStaticSpatialIndex
 		struct FNode
 		{
 			using FNodeType = TArray<FNode>;
-			using FLeafType = TArray<uint32>;
+			struct FLeafType
+			{
+				uint32 StartIndex;
+				uint32 NumElements;
+
+				inline FLeafType()
+					: StartIndex(0)
+					, NumElements(0)
+				{}
+
+				inline void Add(uint32 InIndex)
+				{
+					if (!NumElements)
+					{
+						StartIndex = InIndex;
+					}
+
+					check((StartIndex + NumElements) == InIndex);
+					NumElements++;					
+				}
+
+				inline uint32 Num() const { return NumElements; }
+				inline uint32 GetAllocatedSize() const { return sizeof(*this); }
+
+				struct TIterator
+				{
+					inline TIterator(uint32 InValue) : Value(InValue) {}
+					inline uint32 operator++() { return Value++; }
+					inline uint32 operator*() const { return Value; }
+					inline bool operator!=(const TIterator& Other) const { return Value != Other.Value; }
+					uint32 Value;
+				};
+				
+				inline TIterator begin() { return TIterator(StartIndex); }
+				inline TIterator begin() const { return TIterator(StartIndex); }
+				inline TIterator end() { return TIterator(StartIndex + NumElements); }
+				inline TIterator end() const { return TIterator(StartIndex + NumElements); }
+			};
 			using FBoxType = typename Profile::FBox;
 
 			FBoxType Box = FBoxType(ForceInit);
@@ -387,17 +431,20 @@ namespace FStaticSpatialIndex
 			}
 			else
 			{
+				uint32 BoxStride;
+				const FBox* Box = this->DataInterface.GetBoxes(InNode->Content.template Get<typename FNode::FLeafType>().StartIndex, BoxStride);
+
 				for (uint32 ValueIndex : InNode->Content.template Get<typename FNode::FLeafType>())
 				{
-					const FBox& Box = this->DataInterface.GetBox(ValueIndex);
-
-					if (Box.Intersect(InBox))
+					if (Box->Intersect(InBox))
 					{
 						if (!InFunc(ValueIndex))
 						{
 							return false;
 						}
 					}
+
+					*(uint8**)&Box += BoxStride;
 				}
 			}
 			return true;
@@ -420,18 +467,21 @@ namespace FStaticSpatialIndex
 			}
 			else
 			{
+				uint32 BoxStride;
+				const FBox* Box = this->DataInterface.GetBoxes(InNode->Content.template Get<typename FNode::FLeafType>().StartIndex, BoxStride);
+
 				for (uint32 ValueIndex : InNode->Content.template Get<typename FNode::FLeafType>())
 				{
-					const FBox& Box = this->DataInterface.GetBox(ValueIndex);
-
-					if (FastSphereAABBIntersection<Profile>(FVector(InSphereCenter), InRadiusSquared, Box))
+					if (FastSphereAABBIntersection<Profile>(FVector(InSphereCenter), InRadiusSquared, *Box))
 					{
 						if (!InFunc(ValueIndex))
 						{
 							return false;
 						}
 					}
-				}
+
+					*(uint8**)&Box += BoxStride;
+				}				
 			}
 			return true;
 		}
@@ -520,10 +570,10 @@ namespace FStaticSpatialIndex
 		void Init(const FBox& SortBox)
 		{
 			const FReal MaxExtent = SortBox.GetExtent().GetMax();
-			const uint32 NumBuckets = FMath::CeilToInt32(MaxExtent / (FReal)BucketSize);			
-			HilbertOrder = 1 + FMath::FloorLog2(NumBuckets);
+			const uint32 NumBuckets = FMath::CeilToInt32(MaxExtent / (FReal)BucketSize);
+			HilbertOrder = 1 + FMath::CeilLogTwo(NumBuckets);
 		}
-		
+
 		bool Sort(const FBox& A, const FBox& B)
 		{
 			const int32 HilbertCodeA = HilbertEncode({ int32(A.GetCenter().X / (FReal)BucketSize), int32(A.GetCenter().Y / (FReal)BucketSize) }, HilbertOrder);
@@ -532,7 +582,7 @@ namespace FStaticSpatialIndex
 		}
 
 	private:
-		uint32 HilbertEncode(const FIntVector2 Point, uint32 Order)
+		uint32 HilbertEncode(const FIntVector2& Point, uint32 Order)
 		{
 			uint32 Result = 0;
 
@@ -557,4 +607,4 @@ class TStaticSpatialIndexList : public TStaticSpatialIndex<ValueType, Profile, F
 
 template <typename ValueType, typename NodeSorter, typename Profile>
 class TStaticSpatialIndexRTree : public TStaticSpatialIndex<ValueType, Profile, FStaticSpatialIndex::TRTreeImpl<Profile>, NodeSorter>
-{};
\ No newline at end of file
+{};
