Fix superfluous copying of input string in the text render component introduced in CL 33967839.

#rnx

diff --git a/Engine/Source/Runtime/Engine/Private/Components/TextRenderComponent.cpp b/Engine/Source/Runtime/Engine/Private/Components/TextRenderComponent.cpp
--- a/Engine/Source/Runtime/Engine/Private/Components/TextRenderComponent.cpp
+++ b/Engine/Source/Runtime/Engine/Private/Components/TextRenderComponent.cpp
@@ -1019,7 +1019,7 @@ bool FTextRenderSceneProxy::BuildStringMesh( TArray<FDynamicMeshVertex>& OutVert
 	const uint64 MaxQuads = FMath::Min(MaxVerts / 4, MaxIndices / 6);
 
 	// We have one quad per glyph, so presize the index and vertex buffer.
-	FString Str = Text.ToString();
+	const FString& Str = Text.ToString();
 	const uint64 NumQuads = FMath::Min((uint64)Str.Len(), MaxQuads);
 	OutVertices.Reserve(OutVertices.Num() + NumQuads * 4);
 	OutIndices.Reserve(OutIndices.Num() + NumQuads * 6);
