[World Partition Runtime HashSet Optimization] Don't go through all streaming datas to compute the overriden loading ranges hash.
Don't affect ship builds but makes test builds less affected.

diff --git a/Engine/Source/Runtime/Engine/Private/WorldPartition/RuntimeHashSet/WorldPartitionRuntimeHashSet.cpp b/Engine/Source/Runtime/Engine/Private/WorldPartition/RuntimeHashSet/WorldPartitionRuntimeHashSet.cpp
--- a/Engine/Source/Runtime/Engine/Private/WorldPartition/RuntimeHashSet/WorldPartitionRuntimeHashSet.cpp
+++ b/Engine/Source/Runtime/Engine/Private/WorldPartition/RuntimeHashSet/WorldPartitionRuntimeHashSet.cpp
@@ -622,17 +622,6 @@ void UWorldPartitionRuntimeHashSet::ForEachStreamingCellsSources(const TArray<FW
 	}
 }
 
-uint32 UWorldPartitionRuntimeHashSet::ComputeUpdateStreamingHash() const
-{
-	FHashBuilder HashBuilder(Super::ComputeUpdateStreamingHash());
-	ForEachStreamingData([&HashBuilder](const FRuntimePartitionStreamingData& StreamingData)
-	{
-		HashBuilder << StreamingData.GetLoadingRange();
-		return true;
-	});
-	return HashBuilder.GetHash();
-}
-
 bool UWorldPartitionRuntimeHashSet::SupportsWorldAssetStreaming(const FName& InTargetGrid)
 {
 	return IsValidGridInternal(InTargetGrid, nullptr);
diff --git a/Engine/Source/Runtime/Engine/Private/WorldPartition/WorldPartitionStreamingPolicy.cpp b/Engine/Source/Runtime/Engine/Private/WorldPartition/WorldPartitionStreamingPolicy.cpp
--- a/Engine/Source/Runtime/Engine/Private/WorldPartition/WorldPartitionStreamingPolicy.cpp
+++ b/Engine/Source/Runtime/Engine/Private/WorldPartition/WorldPartitionStreamingPolicy.cpp
@@ -155,6 +155,9 @@ uint32 UWorldPartitionStreamingPolicy::ComputeUpdateStreamingHash(bool bCanOptim
 		if (WorldPartition->RuntimeHash)
 		{
 			HashBuilder << WorldPartition->RuntimeHash->ComputeUpdateStreamingHash();
+#if !UE_BUILD_SHIPPING
+			HashBuilder << UWorldPartitionSubsystem::GetOverriddenLoadingRangesEpoch();
+#endif
 		}
 		HashBuilder << ComputeServerStreamingEnabledEpoch();
 		HashBuilder << WorldPartition->GetStreamingStateEpoch();
diff --git a/Engine/Source/Runtime/Engine/Private/WorldPartition/WorldPartitionSubsystem.cpp b/Engine/Source/Runtime/Engine/Private/WorldPartition/WorldPartitionSubsystem.cpp
--- a/Engine/Source/Runtime/Engine/Private/WorldPartition/WorldPartitionSubsystem.cpp
+++ b/Engine/Source/Runtime/Engine/Private/WorldPartition/WorldPartitionSubsystem.cpp
@@ -197,6 +197,7 @@ static FAutoConsoleVariableRef CVarUdateStreamingStateTimeLimit(
 
 #if !UE_BUILD_SHIPPING
 TMap<FName, int32> UWorldPartitionSubsystem::OverriddenLoadingRanges;
+uint32 UWorldPartitionSubsystem::OverriddenLoadingRangesEpoch = 0;
 static const TCHAR* GOverrideLoadingRangeCommandName = TEXT("wp.Runtime.OverrideRuntimeLoadingRange");
 static FDelegateHandle OnWorldPartitionSubsystemDeinitializedFDelegateHandle;
 FAutoConsoleCommand UWorldPartitionSubsystem::OverrideLoadingRangeCommand(
@@ -530,6 +531,7 @@ void UWorldPartitionSubsystem::SetOverrideLoadingRange(FName Name, int32 Loading
 	{
 		OverriddenLoadingRanges.Remove(Name);
 	}
+	OverriddenLoadingRangesEpoch++;
 }
 
 bool UWorldPartitionSubsystem::GetOverrideLoadingRange(FName Name, int32& LoadingRange)
@@ -541,6 +543,11 @@ bool UWorldPartitionSubsystem::GetOverrideLoadingRange(FName Name, int32& Loadin
 	}
 	return false;
 }
+
+uint32 UWorldPartitionSubsystem::GetOverriddenLoadingRangesEpoch()
+{
+	return OverriddenLoadingRangesEpoch;
+}
 #endif
 
 void UWorldPartitionSubsystem::OnWorldPartitionInitialized(UWorldPartition* InWorldPartition)
diff --git a/Engine/Source/Runtime/Engine/Public/WorldPartition/RuntimeHashSet/WorldPartitionRuntimeHashSet.h b/Engine/Source/Runtime/Engine/Public/WorldPartition/RuntimeHashSet/WorldPartitionRuntimeHashSet.h
--- a/Engine/Source/Runtime/Engine/Public/WorldPartition/RuntimeHashSet/WorldPartitionRuntimeHashSet.h
+++ b/Engine/Source/Runtime/Engine/Public/WorldPartition/RuntimeHashSet/WorldPartitionRuntimeHashSet.h
@@ -191,7 +191,6 @@ public:
 	ENGINE_API virtual void ForEachStreamingCells(TFunctionRef<bool(const UWorldPartitionRuntimeCell*)> Func) const override;
 	ENGINE_API virtual void ForEachStreamingCellsQuery(const FWorldPartitionStreamingQuerySource& QuerySource, TFunctionRef<bool(const UWorldPartitionRuntimeCell*)> Func, FWorldPartitionQueryCache* QueryCache) const override;
 	ENGINE_API virtual void ForEachStreamingCellsSources(const TArray<FWorldPartitionStreamingSource>& Sources, TFunctionRef<bool(const UWorldPartitionRuntimeCell*, EStreamingSourceTargetState)> Func) const override;
-	ENGINE_API virtual uint32 ComputeUpdateStreamingHash() const override;
 
 protected:
 	ENGINE_API virtual bool SupportsWorldAssetStreaming(const FName& InTargetGrid) override;
diff --git a/Engine/Source/Runtime/Engine/Public/WorldPartition/WorldPartitionSubsystem.h b/Engine/Source/Runtime/Engine/Public/WorldPartition/WorldPartitionSubsystem.h
--- a/Engine/Source/Runtime/Engine/Public/WorldPartition/WorldPartitionSubsystem.h
+++ b/Engine/Source/Runtime/Engine/Public/WorldPartition/WorldPartitionSubsystem.h
@@ -104,6 +104,7 @@ public:
 #if !UE_BUILD_SHIPPING
 	static ENGINE_API void SetOverrideLoadingRange(FName Name, int32 LoadingRange);
 	static ENGINE_API bool GetOverrideLoadingRange(FName Name, int32& LoadingRange);
+	static ENGINE_API uint32 GetOverriddenLoadingRangesEpoch();
 #endif
 
 #if WITH_EDITOR
@@ -134,6 +135,7 @@ protected:
 private:
 #if !UE_BUILD_SHIPPING
 	static TMap<FName, int32> OverriddenLoadingRanges;
+	static uint32 OverriddenLoadingRangesEpoch;
 	static class FAutoConsoleCommand OverrideLoadingRangeCommand;
 #endif
 
