NetDriver Handshake
* The NetDriver of client connections will now automatically change the ReplicationSystem from non-Iris to Iris (or vice-versa) when they receive an Upgrade control message with a NetworkRuntimeFeature incompatibility
* Works in both the statelesshandshakecomponent and regular netconnection handshake

diff --git a/Engine/Plugins/Online/OnlineSubsystemUtils/Source/OnlineSubsystemUtils/Private/OnlineBeaconClient.cpp b/Engine/Plugins/Online/OnlineSubsystemUtils/Source/OnlineSubsystemUtils/Private/OnlineBeaconClient.cpp
--- a/Engine/Plugins/Online/OnlineSubsystemUtils/Source/OnlineSubsystemUtils/Private/OnlineBeaconClient.cpp
+++ b/Engine/Plugins/Online/OnlineSubsystemUtils/Source/OnlineSubsystemUtils/Private/OnlineBeaconClient.cpp
@@ -511,12 +511,22 @@ void AOnlineBeaconClient::NotifyControlMessage(UNetConnection* Connection, uint8
 		case NMT_Upgrade:
 			{
 				// Report mismatch.
-				uint32 RemoteNetworkVersion;
+				uint32 RemoteNetworkVersion = 0;
 				EEngineNetworkRuntimeFeatures RemoteNetworkFeatures = EEngineNetworkRuntimeFeatures::None;
 
 				if (FNetControlMessage<NMT_Upgrade>::Receive(Bunch, RemoteNetworkVersion, RemoteNetworkFeatures))
 				{
-					Connection->HandleReceiveNetUpgrade(RemoteNetworkVersion, RemoteNetworkFeatures);
+					const bool bUpgradeSuccess = Connection->HandleReceiveNetUpgrade(RemoteNetworkVersion, RemoteNetworkFeatures);
+
+					if (bUpgradeSuccess)
+					{
+						ensureMsgf(false, TEXT("Beacons don't support any NetworkFeature trait yet. Upgrades shouldn't happen."));
+						//todo: retrigger an handshake via Join message
+					}
+					else
+					{
+						// Should the beacon disconnect here ?
+					}
 				}
 
 				break;
diff --git a/Engine/Source/Runtime/Engine/Classes/Engine/Channel.h b/Engine/Source/Runtime/Engine/Classes/Engine/Channel.h
--- a/Engine/Source/Runtime/Engine/Classes/Engine/Channel.h
+++ b/Engine/Source/Runtime/Engine/Classes/Engine/Channel.h
@@ -95,6 +95,9 @@ public:
 	/** Initialize this channel for the given connection and index. */
 	ENGINE_API virtual void Init(UNetConnection* InConnection, int32 InChIndex, EChannelCreateFlags CreateFlags);
 
+    /** Reinitialize the channel with the existing NetConnection */
+	ENGINE_API virtual void ReInit() {}
+
 	/** Set the closing flag. */
 	ENGINE_API virtual void SetClosingFlag();
 
diff --git a/Engine/Source/Runtime/Engine/Classes/Engine/NetConnection.h b/Engine/Source/Runtime/Engine/Classes/Engine/NetConnection.h
--- a/Engine/Source/Runtime/Engine/Classes/Engine/NetConnection.h
+++ b/Engine/Source/Runtime/Engine/Classes/Engine/NetConnection.h
@@ -56,6 +56,7 @@ namespace UE::Net
 {
 	class FNetPing;
 	class FNetConnectionFaultRecovery;
+	struct FStatelessHandshakeFailureInfo;
 
 } // end namespace UE::Net
 
@@ -1583,6 +1584,9 @@ private:
 
 private:
 
+	/** Called when the stateless handshake component is enabled and we received an handshake failure */
+	void OnStatelessHandshakeFailure(UE::Net::FStatelessHandshakeFailureInfo HandshakeFailureInfo);
+
 	/** Called by PlayerController to tell connection about client level visibility change */
 	void UpdateLevelVisibilityInternal(const struct FUpdateLevelVisibilityLevelInfo& LevelVisibility);
 
@@ -1816,8 +1820,9 @@ public:
 	 * @param RemoteNetworkVersion		The net version of the remote side
 	 * @param RemoteNetworkFeatures		The net runtime features of the remote side
 	 * @param NetUpgradeSource			The source of the net upgrade message
+	 * @return Return true if the connection could be upgraded and we can restart the handshake process. Return false if the upgrade was impossible and the connection must disconnect.
 	 */
-	ENGINE_API void HandleReceiveNetUpgrade(uint32 RemoteNetworkVersion, EEngineNetworkRuntimeFeatures RemoteNetworkFeatures,
+	ENGINE_API bool HandleReceiveNetUpgrade(uint32 RemoteNetworkVersion, EEngineNetworkRuntimeFeatures RemoteNetworkFeatures,
 											UE::Net::ENetUpgradeSource NetUpgradeSource=UE::Net::ENetUpgradeSource::ControlChannel);
 
 private:
diff --git a/Engine/Source/Runtime/Engine/Classes/Engine/NetDriver.h b/Engine/Source/Runtime/Engine/Classes/Engine/NetDriver.h
--- a/Engine/Source/Runtime/Engine/Classes/Engine/NetDriver.h
+++ b/Engine/Source/Runtime/Engine/Classes/Engine/NetDriver.h
@@ -858,6 +858,9 @@ public:
 	/** Override the configured server tick rate. Value is in ticks per second. */
 	ENGINE_API void SetNetServerMaxTickRate(int32 InServerMaxTickRate);
 
+	/** Change the netdriver to have runtime features compatible with the remote connection. */
+	void TryUpgradeNetworkFeatures(EEngineNetworkRuntimeFeatures RemoteFeatures);
+
 	/** 
 	* Delegate triggered when SetNetServerMaxTickRate is called and causes a change to the current max tick rate.
 	* @param UNetDriver The netdriver that changed max tick rate.
@@ -1861,7 +1864,7 @@ public:
 
 	/** Returns the bitflag telling which network features are activated for this NetDriver. */
 	ENGINE_API EEngineNetworkRuntimeFeatures GetNetworkRuntimeFeatures() const;
-	
+
 #if UE_WITH_IRIS
 	/**
 	 * Destroy and recreate the iris replication system for an active netdriver.
diff --git a/Engine/Source/Runtime/Engine/Private/Net/Experimental/Iris/DataStreamChannel.cpp b/Engine/Source/Runtime/Engine/Private/Net/Experimental/Iris/DataStreamChannel.cpp
--- a/Engine/Source/Runtime/Engine/Private/Net/Experimental/Iris/DataStreamChannel.cpp
+++ b/Engine/Source/Runtime/Engine/Private/Net/Experimental/Iris/DataStreamChannel.cpp
@@ -79,6 +79,17 @@ void UDataStreamChannel::Init(UNetConnection* InConnection, int32 InChIndex, ECh
 #endif // UE_WITH_IRIS
 }
 
+void UDataStreamChannel::ReInit()
+{
+	check(Connection);
+
+	if (UReplicationSystem* ReplicationSystem = Connection->Driver->GetReplicationSystem())
+	{
+		bIsReadyToHandshake = 1U;
+		ReplicationSystem->InitDataStreams(Connection->GetConnectionId(), DataStreamManager);
+	}
+}
+
 bool UDataStreamChannel::CleanUp(const bool bForDestroy, EChannelCloseReason CloseReason)
 {
 #if UE_WITH_IRIS
diff --git a/Engine/Source/Runtime/Engine/Private/Net/Experimental/Iris/DataStreamChannel.h b/Engine/Source/Runtime/Engine/Private/Net/Experimental/Iris/DataStreamChannel.h
--- a/Engine/Source/Runtime/Engine/Private/Net/Experimental/Iris/DataStreamChannel.h
+++ b/Engine/Source/Runtime/Engine/Private/Net/Experimental/Iris/DataStreamChannel.h
@@ -33,7 +33,7 @@ private:
 
 	// UChannel interface
 	ENGINE_API virtual void Init(UNetConnection* InConnection, int32 InChIndex, EChannelCreateFlags CreateFlags) override;
-
+	ENGINE_API virtual void ReInit() override;
 	ENGINE_API virtual bool CleanUp(const bool bForDestroy, EChannelCloseReason CloseReason) override;
 
 	/**
diff --git a/Engine/Source/Runtime/Engine/Private/NetConnection.cpp b/Engine/Source/Runtime/Engine/Private/NetConnection.cpp
--- a/Engine/Source/Runtime/Engine/Private/NetConnection.cpp
+++ b/Engine/Source/Runtime/Engine/Private/NetConnection.cpp
@@ -721,15 +721,7 @@ void UNetConnection::InitHandler()
 				StatelessConnectHandlerComponent* CurComponent = StatelessConnectComponent.Pin().Get();
 				
 				CurComponent->SetDriver(Driver);
-
-				CurComponent->SetHandshakeFailureCallback([this](FStatelessHandshakeFailureInfo HandshakeFailureInfo)
-					{
-						if (HandshakeFailureInfo.FailureReason == EHandshakeFailureReason::WrongVersion)
-						{
-							this->HandleReceiveNetUpgrade(HandshakeFailureInfo.RemoteNetworkVersion, HandshakeFailureInfo.RemoteNetworkFeatures,
-															ENetUpgradeSource::StatelessHandshake);
-						}
-					});
+				CurComponent->SetHandshakeFailureCallback([this](FStatelessHandshakeFailureInfo FailureInfo){this->OnStatelessHandshakeFailure(FailureInfo);});
 			}
 
 
@@ -780,6 +772,22 @@ void UNetConnection::InitSequence(int32 IncomingSequence, int32 OutgoingSequence
 	}
 }
 
+void UNetConnection::OnStatelessHandshakeFailure(UE::Net::FStatelessHandshakeFailureInfo HandshakeFailureInfo)
+{
+	using namespace UE::Net;
+
+	if (HandshakeFailureInfo.FailureReason == EHandshakeFailureReason::WrongVersion)
+	{
+		const bool bUpgradeSuccess = HandleReceiveNetUpgrade(HandshakeFailureInfo.RemoteNetworkVersion, HandshakeFailureInfo.RemoteNetworkFeatures, ENetUpgradeSource::StatelessHandshake);
+
+		if (bUpgradeSuccess)
+		{
+			// Restart the handshake since the network driver is now compatible
+			StatelessConnectComponent.Pin().Get()->NotifyHandshakeBegin();
+		}
+	}
+}
+
 void UNetConnection::NotifyAnalyticsProvider()
 {
 	if (Handler.IsValid())
@@ -1209,34 +1217,47 @@ void UNetConnection::HandleReceiveCloseReason(const FString& CloseReasonList)
 	}
 }
 
-void UNetConnection::HandleReceiveNetUpgrade(uint32 RemoteNetworkVersion, EEngineNetworkRuntimeFeatures RemoteNetworkFeatures,
-												UE::Net::ENetUpgradeSource NetUpgradeSource/*=UE::Net::ENetUpgradeSource::ControlChannel*/)
+bool UNetConnection::HandleReceiveNetUpgrade(uint32 RemoteNetworkVersion, EEngineNetworkRuntimeFeatures RemoteNetworkFeatures, UE::Net::ENetUpgradeSource NetUpgradeSource)
 {
-	TStringBuilder<128> RemoteFeaturesDescription;
-	TStringBuilder<128> LocalFeaturesDescription;
+	bool bUpgradeSucceeded = false;
+	const uint32 LocalNetworkVersion = FNetworkVersion::GetLocalNetworkVersion();
 
-	FNetworkVersion::DescribeNetworkRuntimeFeaturesBitset(RemoteNetworkFeatures, RemoteFeaturesDescription);
+	const bool bIsNetCLValid = FNetworkVersion::IsNetworkCompatible(LocalNetworkVersion, RemoteNetworkVersion);
 
-	if (Driver != nullptr)
+	// If the NetCL is not identical, we cannot upgrade and must disconnect
+	if (!bIsNetCLValid || Driver == nullptr)
+	{
+		bUpgradeSucceeded = false;
+	}
+	else
 	{
-		FNetworkVersion::DescribeNetworkRuntimeFeaturesBitset(Driver->GetNetworkRuntimeFeatures(), LocalFeaturesDescription);
+		Driver->TryUpgradeNetworkFeatures(RemoteNetworkFeatures);
+		bUpgradeSucceeded = Driver->GetNetworkRuntimeFeatures() == RemoteNetworkFeatures;
 	}
 
-	UE_LOG(LogNet, Error, TEXT("Server is incompatible with the local version of the game: RemoteNetworkVersion=%u, ")
-			TEXT("RemoteNetworkFeatures=%s vs LocalNetworkVersion=%u, LocalNetworkFeatures=%s"), 
-			RemoteNetworkVersion, RemoteFeaturesDescription.ToString(), FNetworkVersion::GetLocalNetworkVersion(),
-			LocalFeaturesDescription.ToString());
+	if (!bUpgradeSucceeded)
+	{
+		const EEngineNetworkRuntimeFeatures LocalNetworkFeatures = Driver ? Driver->GetNetworkRuntimeFeatures() : EEngineNetworkRuntimeFeatures::None;
 
+		TStringBuilder<128> LocalFeaturesDescription;
+		TStringBuilder<128> RemoteFeaturesDescription;
+		FNetworkVersion::DescribeNetworkRuntimeFeaturesBitset(LocalNetworkFeatures, LocalFeaturesDescription);
+		FNetworkVersion::DescribeNetworkRuntimeFeaturesBitset(RemoteNetworkFeatures, RemoteFeaturesDescription);
 
-	const FString ConnectionError = NSLOCTEXT("Engine", "ClientOutdated",
-		"The match you are trying to join is running an incompatible version of the game.  Please try upgrading your game version.").ToString();
+		UE_LOG(LogNet, Error, TEXT("Server is incompatible with the local version of the game: RemoteNetworkVersion=%u vs LocalNetworkVersion=%u, RemoteNetworkFeatures=%s vs LocalNetworkFeatures=%s"),
+			RemoteNetworkVersion, LocalNetworkVersion,
+			RemoteFeaturesDescription.ToString(), LocalFeaturesDescription.ToString());
 
-	GEngine->BroadcastNetworkFailure(GetWorld(), Driver, ENetworkFailure::OutdatedClient, ConnectionError);
+		const FString ConnectionError = NSLOCTEXT("Engine", "ClientOutdated", "The match you are trying to join is running an incompatible version of the game.  Please try upgrading your game version.").ToString();
+		GEngine->BroadcastNetworkFailure(GetWorld(), Driver, ENetworkFailure::OutdatedClient, ConnectionError);
 
-	if (NetUpgradeSource == UE::Net::ENetUpgradeSource::StatelessHandshake)
-	{
-		Close(ENetCloseResult::OutdatedClient);
+		if (NetUpgradeSource == UE::Net::ENetUpgradeSource::StatelessHandshake)
+		{
+			Close(ENetCloseResult::OutdatedClient);
+		}
 	}
+
+	return bUpgradeSucceeded;
 }
 
 FString UNetConnection::Describe()
diff --git a/Engine/Source/Runtime/Engine/Private/NetDriver.cpp b/Engine/Source/Runtime/Engine/Private/NetDriver.cpp
--- a/Engine/Source/Runtime/Engine/Private/NetDriver.cpp
+++ b/Engine/Source/Runtime/Engine/Private/NetDriver.cpp
@@ -2229,10 +2229,67 @@ EEngineNetworkRuntimeFeatures UNetDriver::GetNetworkRuntimeFeatures() const
 	{
 		EnumAddFlags(NetDriverFeatures, EEngineNetworkRuntimeFeatures::IrisEnabled);
 	}
-	
+
 	return NetDriverFeatures;
 }
 
+void UNetDriver::TryUpgradeNetworkFeatures(EEngineNetworkRuntimeFeatures RemoteFeatures)
+{
+#if UE_WITH_IRIS
+	const bool bIsRemoteUsingIris = EnumHasAnyFlags(RemoteFeatures, EEngineNetworkRuntimeFeatures::IrisEnabled);
+
+	// If our netdriver needs to enable Iris
+	if (bIsRemoteUsingIris && !IsUsingIrisReplication())
+	{
+		SetReplicationDriver(nullptr);
+
+		constexpr bool bInitWithIris = true;
+		PostCreation(bInitWithIris);
+
+		RecreateIrisSystem();
+		
+		ReinitBase();
+
+		if (ensureMsgf(ReplicationSystem, TEXT("NetworkFeature Upgraded failed to create repSystem for %s"), *GetDescription()))
+		{
+			checkf(ServerConnection, TEXT("NMT_Upgrade control message received with no server connection on %s"), *GetDescription());
+			
+			// Tell Iris and the Datastream about the current netconnection
+			ReplicationSystem->AddConnection(ServerConnection->GetConnectionId());
+			ReplicationSystem->SetConnectionUserData(ServerConnection->GetConnectionId(), ServerConnection);
+			for (UChannel* Channel : ServerConnection->Channels)
+			{
+				if (Channel)
+				{
+					Channel->ReInit();
+				}
+			}
+		}
+
+	}
+	// If the remote was not using Iris but we did
+	else if(!bIsRemoteUsingIris && IsUsingIrisReplication())
+	{
+		DestroyIrisSystem();
+
+		constexpr bool bInitWithoutIris = false;
+		PostCreation(bInitWithoutIris);
+
+		ReinitBase();
+
+		for (UChannel* Channel : ServerConnection->Channels)
+		{
+			if (Channel)
+			{
+				Channel->ReInit();
+			}
+		}
+	}
+
+	ensureMsgf(GetNetworkRuntimeFeatures() == RemoteFeatures, TEXT("UNetDriver::TryUpgradeNetworkFeatures failed for %s"), *GetDescription());
+#endif
+}
+
 // ----------------------------------------------------------------------------------------
 //	RPC Tracking (receiving) via CSV profiler
 // ----------------------------------------------------------------------------------------
diff --git a/Engine/Source/Runtime/Engine/Private/PacketHandlers/StatelessConnectHandlerComponent.cpp b/Engine/Source/Runtime/Engine/Private/PacketHandlers/StatelessConnectHandlerComponent.cpp
--- a/Engine/Source/Runtime/Engine/Private/PacketHandlers/StatelessConnectHandlerComponent.cpp
+++ b/Engine/Source/Runtime/Engine/Private/PacketHandlers/StatelessConnectHandlerComponent.cpp
@@ -1304,147 +1304,157 @@ void StatelessConnectHandlerComponent::IncomingConnectionless(FIncomingPacketRef
 		// No ClientID validation until connected
 	}
 
-	bool bHandshakePacket = !!Packet.ReadBit() && !Packet.IsError();
+	const bool bHandshakePacket = !!Packet.ReadBit() && !Packet.IsError();
 
 	LastChallengeSuccessAddress = nullptr;
 
-	if (bHandshakePacket)
+	if (!bHandshakePacket)
 	{
-		FParsedHandshakeData HandshakeData;
-
-		bHandshakePacket = ParseHandshakePacket(Packet, HandshakeData);
-
-		if (bHandshakePacket)
+		if (Packet.IsError())
 		{
-			EHandshakeVersion TargetVersion = EHandshakeVersion::Latest;
-			const bool bValidVersion = CheckVersion(HandshakeData, TargetVersion);
-			const bool bInitialConnect = HandshakeData.HandshakePacketType == EHandshakePacketType::InitialPacket &&
-												HandshakeData.Timestamp == 0.0;
-
-			if (Handler->Mode == UE::Handler::Mode::Server && bValidVersion && (bHasValidSessionID || bInitialConnect))
+#if !UE_BUILD_SHIPPING
+			UE_LOG(LogHandshake, Log, TEXT("IncomingConnectionless: Error reading session id, connection id and handshake bit from packet."));
+#endif
+		}
+		else if (bHasValidSessionID)
+		{
+			// Late packets from recently disconnected clients may incorrectly trigger this code path, so detect and exclude those packets
+			if (!Packet.IsError() && !PacketRef.Traits.bFromRecentlyDisconnected)
 			{
-				if (bInitialConnect)
-				{
-					SendConnectChallenge(FCommonSendToClientParams(Address, TargetVersion, ClientID), HandshakeData.RemoteSentHandshakePacketCount);
-				}
-				// Challenge response
-				else if (Driver != nullptr)
-				{
-					// NOTE: Allow CookieDelta to be 0.0, as it is possible for a server to send a challenge and receive a response,
-					//			during the same tick
-					bool bChallengeSuccess = false;
-					const double CookieDelta = Driver->GetElapsedTime() - HandshakeData.Timestamp;
-					const double SecretDelta = HandshakeData.Timestamp - LastSecretUpdateTimestamp;
-					const bool bValidCookieLifetime = CookieDelta >= 0.0 && (MAX_COOKIE_LIFETIME - CookieDelta) > 0.0;
-					const bool bValidSecretIdTimestamp = (HandshakeData.SecretId == ActiveSecret) ? (SecretDelta >= 0.0) : (SecretDelta <= 0.0);
-
-					if (bValidCookieLifetime && bValidSecretIdTimestamp)
-					{
-						// Regenerate the cookie from the packet info, and see if the received cookie matches the regenerated one
-						uint8 RegenCookie[COOKIE_BYTE_SIZE];
+				// The packet was fine but not a handshake packet - an existing client might suddenly be communicating on a different address.
+				// If we get them to resend their cookie, we can update the connection's info with their new address.
+				SendRestartHandshakeRequest(FCommonSendToClientParams(Address, static_cast<EHandshakeVersion>(MinSupportedHandshakeVersion), ClientID));
+			}
+		}
 
-						GenerateCookie(Address, HandshakeData.SecretId, HandshakeData.Timestamp, RegenCookie);
+		return;
+	}
 
-						bChallengeSuccess = FMemory::Memcmp(HandshakeData.Cookie, RegenCookie, COOKIE_BYTE_SIZE) == 0;
+	FParsedHandshakeData HandshakeData;
+	const bool bValidHandshakePacket = ParseHandshakePacket(Packet, HandshakeData);
 
-						if (bChallengeSuccess)
-						{
-							if (HandshakeData.bRestartHandshake)
-							{
-								FMemory::Memcpy(AuthorisedCookie, HandshakeData.OrigCookie, UE_ARRAY_COUNT(AuthorisedCookie));
-							}
-							else
-							{
-								int16* CurSequence = (int16*)HandshakeData.Cookie;
+	if (UNLIKELY(!bValidHandshakePacket))
+	{
+		Packet.SetError();
 
-								LastServerSequence = *CurSequence & (MAX_PACKETID - 1);
-								LastClientSequence = *(CurSequence + 1) & (MAX_PACKETID - 1);
+		FDDoSDetection* DDoS = Handler->GetDDoS();
 
-								FMemory::Memcpy(AuthorisedCookie, HandshakeData.Cookie, UE_ARRAY_COUNT(AuthorisedCookie));
-							}
+		if (DDoS != nullptr)
+		{
+			DDoS->IncBadPacketCounter();
+		}
 
-							bRestartedHandshake = HandshakeData.bRestartHandshake;
-							LastChallengeSuccessAddress = Address->Clone();
-							LastRemoteHandshakeVersion = TargetVersion;
-							CachedClientID = ClientID;
+#if !UE_BUILD_SHIPPING
+		UE_CLOG(DDoS == nullptr || !DDoS->CheckLogRestrictions(), LogHandshake, Log,
+			TEXT("IncomingConnectionless: Error reading handshake packet."));
+#endif
+		return;
+	}
+	
+	const bool bIsServer = Handler->Mode == UE::Handler::Mode::Server;
+	if (UNLIKELY(!bIsServer))
+	{
+		// Only server can negotiate handshake requests here
+		return;
+	}
 
-							if (TargetVersion < MinClientHandshakeVersion && static_cast<uint8>(TargetVersion) >= MinSupportedHandshakeVersion)
-							{
-								MinClientHandshakeVersion = TargetVersion;
-							}
+	EHandshakeVersion TargetVersion = EHandshakeVersion::Latest;
+	const bool bValidVersion = CheckVersion(HandshakeData, TargetVersion);
+	const bool bInitialConnect = HandshakeData.HandshakePacketType == EHandshakePacketType::InitialPacket && HandshakeData.Timestamp == 0.0;
+	const double ElapsedTime = Driver ? Driver->GetElapsedTime() : 0.0;
 
+	const bool bIsValidRequest = bValidVersion && (bHasValidSessionID || bInitialConnect);
 
-							// Now ack the challenge response - the cookie is stored in AuthorisedCookie, to enable retries
-							SendChallengeAck(FCommonSendToClientParams(Address, TargetVersion, ClientID),
-												HandshakeData.RemoteSentHandshakePacketCount, AuthorisedCookie);
-						}
-					}
-				}
-			}
-			else if (Handler->Mode == UE::Handler::Mode::Server && bValidVersion && !bHasValidSessionID)
-			{
+	// Handle invalid requests
+	if (!bIsValidRequest)
+	{
+		// Connection has an invalid Session ID
+		if (bValidVersion && !bHasValidSessionID)
+		{
 #if !UE_BUILD_SHIPPING
-				UE_LOG(LogHandshake, Log, TEXT("IncomingConnectionless: Rejecting packet with invalid session id: %u vs %u."),
-						SessionID, CachedGlobalNetTravelCount);
+			UE_LOG(LogHandshake, Log, TEXT("IncomingConnectionless: Rejecting packet with invalid session id: %u vs %u."), SessionID, CachedGlobalNetTravelCount);
 #endif
+			Packet.SetError();
+		}
+		// Connection is using an incompatible version, tell it to upgrade
+		else if (!bValidVersion && bInitialConnect && HandshakeData.RemoteCurVersion >= EHandshakeVersion::NetCLUpgradeMessage)
+		{
+			// Limit of 512 upgrade message packets, over 5 minutes
+			const uint32 NumUpgradeMessagesPerPeriod = 512;
+			const double UpgradeMessagePeriod = 300;
 
-				Packet.SetError();
+			if (ElapsedTime - LastUpgradeMessagePeriodStart >= UpgradeMessagePeriod)
+			{
+				LastUpgradeMessagePeriodStart = ElapsedTime;
+				UpgradeMessageCounter = 0;
+			}
+			else
+			{
+				UpgradeMessageCounter++;
 			}
-			else if (Handler->Mode == UE::Handler::Mode::Server && !bValidVersion && bInitialConnect &&
-						HandshakeData.RemoteCurVersion >= EHandshakeVersion::NetCLUpgradeMessage)
+
+			if (UpgradeMessageCounter < NumUpgradeMessagesPerPeriod)
 			{
-				// Limit of 512 upgrade message packets, over 5 minutes
-				const uint32 NumUpgradeMessagesPerPeriod = 512;
-				const double UpgradeMessagePeriod = 300;
-				const double ElapsedTime = Driver->GetElapsedTime();
+				SendVersionUpgradeMessage(FCommonSendToClientParams(Address, TargetVersion, ClientID));
+			}
+		}
+
+		return;
+	}
 
-				if (ElapsedTime - LastUpgradeMessagePeriodStart >= UpgradeMessagePeriod)
+	if (bInitialConnect)
+	{
+		SendConnectChallenge(FCommonSendToClientParams(Address, TargetVersion, ClientID), HandshakeData.RemoteSentHandshakePacketCount);
+	}
+	else
+	{
+		// Challenge response
+		// NOTE: Allow CookieDelta to be 0.0, as it is possible for a server to send a challenge and receive a response,
+		//			during the same tick
+		bool bChallengeSuccess = false;
+		const double CookieDelta = ElapsedTime - HandshakeData.Timestamp;
+		const double SecretDelta = HandshakeData.Timestamp - LastSecretUpdateTimestamp;
+		const bool bValidCookieLifetime = CookieDelta >= 0.0 && (MAX_COOKIE_LIFETIME - CookieDelta) > 0.0;
+		const bool bValidSecretIdTimestamp = (HandshakeData.SecretId == ActiveSecret) ? (SecretDelta >= 0.0) : (SecretDelta <= 0.0);
+
+		if (bValidCookieLifetime && bValidSecretIdTimestamp)
+		{
+			// Regenerate the cookie from the packet info, and see if the received cookie matches the regenerated one
+			uint8 RegenCookie[COOKIE_BYTE_SIZE];
+
+			GenerateCookie(Address, HandshakeData.SecretId, HandshakeData.Timestamp, RegenCookie);
+
+			bChallengeSuccess = FMemory::Memcmp(HandshakeData.Cookie, RegenCookie, COOKIE_BYTE_SIZE) == 0;
+
+			if (bChallengeSuccess)
+			{
+				if (HandshakeData.bRestartHandshake)
 				{
-					LastUpgradeMessagePeriodStart = ElapsedTime;
-					UpgradeMessageCounter = 0;
+					FMemory::Memcpy(AuthorisedCookie, HandshakeData.OrigCookie, UE_ARRAY_COUNT(AuthorisedCookie));
 				}
 				else
 				{
-					UpgradeMessageCounter++;
+					int16* CurSequence = (int16*)HandshakeData.Cookie;
+
+					LastServerSequence = *CurSequence & (MAX_PACKETID - 1);
+					LastClientSequence = *(CurSequence + 1) & (MAX_PACKETID - 1);
+
+					FMemory::Memcpy(AuthorisedCookie, HandshakeData.Cookie, UE_ARRAY_COUNT(AuthorisedCookie));
 				}
 
-				if (UpgradeMessageCounter < NumUpgradeMessagesPerPeriod)
+				bRestartedHandshake = HandshakeData.bRestartHandshake;
+				LastChallengeSuccessAddress = Address->Clone();
+				LastRemoteHandshakeVersion = TargetVersion;
+				CachedClientID = ClientID;
+
+				if (TargetVersion < MinClientHandshakeVersion && static_cast<uint8>(TargetVersion) >= MinSupportedHandshakeVersion)
 				{
-					SendVersionUpgradeMessage(FCommonSendToClientParams(Address, TargetVersion, ClientID));
+					MinClientHandshakeVersion = TargetVersion;
 				}
-			}
-		}
-		else
-		{
-			Packet.SetError();
 
-			FDDoSDetection* DDoS = Handler->GetDDoS();
-
-			if (DDoS != nullptr)
-			{
-				DDoS->IncBadPacketCounter();
+				// Now ack the challenge response - the cookie is stored in AuthorisedCookie, to enable retries
+				SendChallengeAck(FCommonSendToClientParams(Address, TargetVersion, ClientID), HandshakeData.RemoteSentHandshakePacketCount, AuthorisedCookie);
 			}
-
-#if !UE_BUILD_SHIPPING
-			UE_CLOG(DDoS == nullptr || !DDoS->CheckLogRestrictions(), LogHandshake, Log,
-					TEXT("IncomingConnectionless: Error reading handshake packet."));
-#endif
-		}
-	}
-#if !UE_BUILD_SHIPPING
-	else if (Packet.IsError())
-	{
-		UE_LOG(LogHandshake, Log, TEXT("IncomingConnectionless: Error reading session id, connection id and handshake bit from packet."));
-	}
-#endif
-	else if (bHasValidSessionID)
-	{
-		// Late packets from recently disconnected clients may incorrectly trigger this code path, so detect and exclude those packets
-		if (!Packet.IsError() && !PacketRef.Traits.bFromRecentlyDisconnected)
-		{
-			// The packet was fine but not a handshake packet - an existing client might suddenly be communicating on a different address.
-			// If we get them to resend their cookie, we can update the connection's info with their new address.
-			SendRestartHandshakeRequest(FCommonSendToClientParams(Address, static_cast<EHandshakeVersion>(MinSupportedHandshakeVersion), ClientID));
 		}
 	}
 }
@@ -1674,29 +1684,26 @@ bool StatelessConnectHandlerComponent::CheckVersion(const FParsedHandshakeData&
 		bValidHandshakeVersion = true;
 	}
 
+	const EEngineNetworkRuntimeFeatures LocalNetworkFeatures = Driver ? Driver->GetNetworkRuntimeFeatures() : EEngineNetworkRuntimeFeatures::None;
+
 	const bool bCheckNetVersion = !UE_BUILD_SHIPPING || GHandshakeEnforceNetworkCLVersion;
-	bool bValidNetVersion = true;
-	EEngineNetworkRuntimeFeatures LocalNetworkFeatures = EEngineNetworkRuntimeFeatures::None;
+	bool bValidNetCLVersion = true;
+	bool bIsNetFeatureCompatible = true;
 
-	if (bCheckNetVersion && HandshakeData.RemoteCurVersion >= EHandshakeVersion::NetCLVersion)
+	if (HandshakeData.RemoteCurVersion >= EHandshakeVersion::NetCLVersion)
 	{
-		if (Driver != nullptr)
+		if (bCheckNetVersion)
 		{
-			LocalNetworkFeatures = Driver->GetNetworkRuntimeFeatures();
+			const uint32 LocalNetworkVersion = FNetworkVersion::GetLocalNetworkVersion();
+			bValidNetCLVersion = FNetworkVersion::IsNetworkCompatible(LocalNetworkVersion, HandshakeData.RemoteNetworkVersion);
 		}
 
-		const uint32 LocalNetworkVersion = FNetworkVersion::GetLocalNetworkVersion();
-		const bool bIsCompatible = FNetworkVersion::IsNetworkCompatible(LocalNetworkVersion, HandshakeData.RemoteNetworkVersion) &&
-									FNetworkVersion::AreNetworkRuntimeFeaturesCompatible(LocalNetworkFeatures, HandshakeData.RemoteNetworkFeatures);
-
-		if (!bIsCompatible)
-		{
-			bValidNetVersion = false;
-		}
+		// Network Runtime Features are always validated since the client can try to upgrade his NetDriver and become compatible
+		bIsNetFeatureCompatible = FNetworkVersion::AreNetworkRuntimeFeaturesCompatible(LocalNetworkFeatures, HandshakeData.RemoteNetworkFeatures);
 	}
 
 #if !UE_BUILD_SHIPPING
-	if (!bValidHandshakeVersion || !bValidNetVersion)
+	if (!bValidHandshakeVersion || !bValidNetCLVersion || !bIsNetFeatureCompatible)
 	{
 		FDDoSDetection* DDoS = Handler->GetDDoS();
 		const uint32 LocalNetworkVersion = FNetworkVersion::GetLocalNetworkVersion();
@@ -1711,16 +1718,16 @@ bool StatelessConnectHandlerComponent::CheckVersion(const FParsedHandshakeData&
 				TEXT("GHandshakeEnforceNetworkCLVersion: %i, RemoteMinVersion: %u, RemoteCurVersion: %u, MinSupportedHandshakeVersion: %i, ")
 				TEXT("CurrentHandshakeVersion: %i, RemoteNetworkVersion: %u, LocalNetworkVersion: %u, RemoteNetworkFeatures: %s, ")
 				TEXT("LocalNetworkFeatures: %s"),
-				(int32)bValidHandshakeVersion, (int32)bValidNetVersion, GHandshakeEnforceNetworkCLVersion, RemoteMinVersionUint8,
+				(int32)bValidHandshakeVersion, (int32)bValidNetCLVersion, GHandshakeEnforceNetworkCLVersion, RemoteMinVersionUint8,
 				static_cast<uint8>(HandshakeData.RemoteCurVersion), MinSupportedHandshakeVersion, CurrentHandshakeVersion,
-				HandshakeData.RemoteNetworkVersion, LocalNetworkVersion, ToCStr(RemoteNetFeaturesDescription.ToString()),
-				ToCStr(LocalNetFeaturesDescription.ToString()));
+				HandshakeData.RemoteNetworkVersion, LocalNetworkVersion, RemoteNetFeaturesDescription.ToString(), 
+				LocalNetFeaturesDescription.ToString());
 	}
 #endif
 
-	const bool bPassedNetVersionConditions = bValidNetVersion || !GHandshakeEnforceNetworkCLVersion;
+	const bool bPassedNetVersionConditions = bValidNetCLVersion || !GHandshakeEnforceNetworkCLVersion;
 
-	return bValidHandshakeVersion && bPassedNetVersionConditions;
+	return bValidHandshakeVersion && bPassedNetVersionConditions && bIsNetFeatureCompatible;
 }
 
 void StatelessConnectHandlerComponent::GenerateCookie(const TSharedPtr<const FInternetAddr>& ClientAddress, uint8 SecretId, double Timestamp,
diff --git a/Engine/Source/Runtime/Engine/Private/PendingNetGame.cpp b/Engine/Source/Runtime/Engine/Private/PendingNetGame.cpp
--- a/Engine/Source/Runtime/Engine/Private/PendingNetGame.cpp
+++ b/Engine/Source/Runtime/Engine/Private/PendingNetGame.cpp
@@ -227,7 +227,7 @@ void UPendingNetGame::NotifyControlMessage(UNetConnection* Connection, uint8 Mes
 	check(Connection == NetDriver->ServerConnection);
 
 #if !UE_BUILD_SHIPPING
-	UE_LOG(LogNet, Verbose, TEXT("PendingLevel received: %s"), FNetControlMessageInfo::GetName(MessageType));
+	UE_LOG(LogNet, Verbose, TEXT("NotifyControlMessage: PendingLevel received: %s"), FNetControlMessageInfo::GetName(MessageType));
 #endif
 
 	// This client got a response from the server.
@@ -236,16 +236,22 @@ void UPendingNetGame::NotifyControlMessage(UNetConnection* Connection, uint8 Mes
 		case NMT_Upgrade:
 		{
 			// Report mismatch.
-			uint32 RemoteNetworkVersion;
-
+			uint32 RemoteNetworkVersion = 0;
 			EEngineNetworkRuntimeFeatures RemoteNetworkFeatures = EEngineNetworkRuntimeFeatures::None;
 
 			if (FNetControlMessage<NMT_Upgrade>::Receive(Bunch, RemoteNetworkVersion, RemoteNetworkFeatures))
 			{
-				// Upgrade
-				ConnectionError = NSLOCTEXT("Engine", "ClientOutdated", "The match you are trying to join is running an incompatible version of the game.  Please try upgrading your game version.").ToString();
-
-				Connection->HandleReceiveNetUpgrade(RemoteNetworkVersion, RemoteNetworkFeatures);
+				const bool bUpgradeSuccess = Connection->HandleReceiveNetUpgrade(RemoteNetworkVersion, RemoteNetworkFeatures);
+				if (bUpgradeSuccess)
+				{
+					// Restart the handshake process now that we are compatible
+					SendInitialJoin();
+				}
+				else
+				{
+					// Upgrade failed, trigger a disconnect via the connect error
+					ConnectionError = NSLOCTEXT("Engine", "ClientOutdated", "The match you are trying to join is running an incompatible version of the game.  Please try upgrading your game version.").ToString();
+				}
 			}
 
 			break;
diff --git a/Engine/Source/Runtime/Engine/Private/World.cpp b/Engine/Source/Runtime/Engine/Private/World.cpp
--- a/Engine/Source/Runtime/Engine/Private/World.cpp
+++ b/Engine/Source/Runtime/Engine/Private/World.cpp
@@ -6494,12 +6494,13 @@ void UWorld::NotifyControlMessage(UNetConnection* Connection, uint8 MessageType,
 
 				if (FNetControlMessage<NMT_Hello>::Receive(Bunch, IsLittleEndian, RemoteNetworkVersion, EncryptionToken, RemoteNetworkFeatures))
 				{
-					const bool bIsCompatible = FNetworkVersion::IsNetworkCompatible(LocalNetworkVersion, RemoteNetworkVersion) && FNetworkVersion::AreNetworkRuntimeFeaturesCompatible(LocalNetworkFeatures, RemoteNetworkFeatures);
-					if (!bIsCompatible)
+					const bool bIsNetCLCompatible = FNetworkVersion::IsNetworkCompatible(LocalNetworkVersion, RemoteNetworkVersion);
+					const bool bAreNetFeaturesCompatible = FNetworkVersion::AreNetworkRuntimeFeaturesCompatible(LocalNetworkFeatures, RemoteNetworkFeatures);
+
+					if (!bIsNetCLCompatible || !bAreNetFeaturesCompatible)
 					{
 						TStringBuilder<128> LocalNetFeaturesDescription;
 						TStringBuilder<128> RemoteNetFeaturesDescription;
-
 						FNetworkVersion::DescribeNetworkRuntimeFeaturesBitset(LocalNetworkFeatures, LocalNetFeaturesDescription);
 						FNetworkVersion::DescribeNetworkRuntimeFeaturesBitset(RemoteNetworkFeatures, RemoteNetFeaturesDescription);
 
@@ -6511,10 +6512,15 @@ void UWorld::NotifyControlMessage(UNetConnection* Connection, uint8 MessageType,
 
 						FNetControlMessage<NMT_Upgrade>::Send(Connection, LocalNetworkVersion, LocalNetworkFeatures);
 						Connection->FlushNet(true);
-						Connection->Close(ENetCloseResult::Upgrade);
+						
+						// If the NetCL is not compatible, disconnect the client immediately. 
+						if (!bIsNetCLCompatible)
+						{
+							Connection->Close(ENetCloseResult::Upgrade);
 #if USE_SERVER_PERF_COUNTERS
-						PerfCountersIncrement(TEXT("ClosedConnectionsDueToIncompatibleVersion"));
+							PerfCountersIncrement(TEXT("ClosedConnectionsDueToIncompatibleVersion"));
 #endif
+						}
 					}
 					else
 					{
diff --git a/Engine/Source/Runtime/Engine/Public/PacketHandlers/StatelessConnectHandlerComponent.h b/Engine/Source/Runtime/Engine/Public/PacketHandlers/StatelessConnectHandlerComponent.h
--- a/Engine/Source/Runtime/Engine/Public/PacketHandlers/StatelessConnectHandlerComponent.h
+++ b/Engine/Source/Runtime/Engine/Public/PacketHandlers/StatelessConnectHandlerComponent.h
@@ -173,7 +173,7 @@ private:
 	 *
 	 * @param HandshakeVersion	The handshake format version to use when sending
 	 */
-	ENGINE_API void SendInitialPacket(EHandshakeVersion HandshakeVersion);
+	void SendInitialPacket(EHandshakeVersion HandshakeVersion);
 
 	/**
 	 * Constructs and sends the server response to the initial connect packet, from the server to the client.
@@ -181,7 +181,7 @@ private:
 	 * @param CommonParams						Common parameters for 'Send*' functions to clients
 	 * @param ClientSentHandshakePacketCount	The number of handshake packets the client has sent (for debugging/packet-analysis)
 	 */
-	ENGINE_API void SendConnectChallenge(FCommonSendToClientParams CommonParams, uint8 ClientSentHandshakePacketCount);
+	void SendConnectChallenge(FCommonSendToClientParams CommonParams, uint8 ClientSentHandshakePacketCount);
 
 	/**
 	 * Constructs and sends the handshake challenge response packet, from the client to the server
@@ -191,7 +191,7 @@ private:
 	 * @param InTimestamp		The timestamp value to send
 	 * @param InCookie			The cookie value to send
 	 */
-	ENGINE_API void SendChallengeResponse(EHandshakeVersion HandshakeVersion, uint8 InSecretId, double InTimestamp, uint8 InCookie[COOKIE_BYTE_SIZE]);
+	void SendChallengeResponse(EHandshakeVersion HandshakeVersion, uint8 InSecretId, double InTimestamp, uint8 InCookie[COOKIE_BYTE_SIZE]);
 
 	/**
 	 * Constructs and sends the server ack to a successful challenge response, from the server to the client.
@@ -200,21 +200,21 @@ private:
 	 * @param ClientSentHandshakePacketCount	The number of handshake packets the client has sent (for debugging/packet-analysis)
 	 * @param InCookie							The cookie value to send
 	 */
-	ENGINE_API void SendChallengeAck(FCommonSendToClientParams CommonParams, uint8 ClientSentHandshakePacketCount, uint8 InCookie[COOKIE_BYTE_SIZE]);
+	void SendChallengeAck(FCommonSendToClientParams CommonParams, uint8 ClientSentHandshakePacketCount, uint8 InCookie[COOKIE_BYTE_SIZE]);
 
 	/**
 	 * Constructs and sends a request to resend the cookie, from the server to the client.
 	 *
 	 * @param CommonParams	Common parameters for 'Send*' functions to clients
 	 */
-	ENGINE_API void SendRestartHandshakeRequest(FCommonSendToClientParams CommonParams);
+	void SendRestartHandshakeRequest(FCommonSendToClientParams CommonParams);
 
 	/**
 	 * Constructs and sends a stateless handshake level NMT_Upgrade message
 	 *
 	 * @param CommonParams	Common parameters for 'Send*' functions to clients
 	 */
-	ENGINE_API void SendVersionUpgradeMessage(FCommonSendToClientParams CommonParams);
+	void SendVersionUpgradeMessage(FCommonSendToClientParams CommonParams);
 
 
 	/**
@@ -227,7 +227,7 @@ private:
 	 * @param ClientID									The client-specified connection id.
 	 * @param HandshakePacketModifier					Flags/modifiers affecting how the packet will be handled.
 	 */
-	ENGINE_API void BeginHandshakePacket(FBitWriter& HandshakePacket, EHandshakePacketType HandshakePacketType, EHandshakeVersion HandshakeVersion,
+	void BeginHandshakePacket(FBitWriter& HandshakePacket, EHandshakePacketType HandshakePacketType, EHandshakeVersion HandshakeVersion,
 								uint8 SentHandshakePacketCount_LocalOrRemote, uint32 ClientID,
 								EHandshakePacketModifier HandshakePacketModifier=EHandshakePacketModifier::None);
 
@@ -239,7 +239,7 @@ private:
 	 * @param PacketType			The type of handshake packet.
 	 * @param Packet				The handshake packet being written.
 	 */
-	ENGINE_API void SendToServer(EHandshakeVersion HandshakeVersion, EHandshakePacketType PacketType, FBitWriter& Packet);
+	void SendToServer(EHandshakeVersion HandshakeVersion, EHandshakePacketType PacketType, FBitWriter& Packet);
 
 	/**
 	 * Sends a packet from the server to a client.
@@ -248,7 +248,7 @@ private:
 	 * @param PacketType			The type of handshake packet.
 	 * @param Packet				The handshake packet being written.
 	 */
-	ENGINE_API void SendToClient(FCommonSendToClientParams CommonParams, EHandshakePacketType PacketType, FBitWriter& Packet);
+	void SendToClient(FCommonSendToClientParams CommonParams, EHandshakePacketType PacketType, FBitWriter& Packet);
 
 
 	/**
@@ -257,7 +257,7 @@ private:
 	 * @param HandshakePacket	The handshake packet to be aligned.
 	 * @param HandshakeVersion	The handshake format version to use when sending
 	 */
-	ENGINE_API void CapHandshakePacket(FBitWriter& HandshakePacket, EHandshakeVersion HandshakeVersion);
+	void CapHandshakePacket(FBitWriter& HandshakePacket, EHandshakeVersion HandshakeVersion);
 
 public:
 	/**
@@ -389,9 +389,9 @@ private:
 	 * @param OutResult		The parsed handshake data, if successful
 	 * @return				Whether or not the handshake packet was parsed successfully
 	 */
-	ENGINE_API bool ParseHandshakePacket(FBitReader& Packet, FParsedHandshakeData& OutResult) const;
+	bool ParseHandshakePacket(FBitReader& Packet, FParsedHandshakeData& OutResult) const;
 
-	ENGINE_API bool ParseHandshakePacketOriginal(FBitReader& Packet, FParsedHandshakeData& OutResult) const;
+	bool ParseHandshakePacketOriginal(FBitReader& Packet, FParsedHandshakeData& OutResult) const;
 
 	/**
 	 * Checks the handshake protocol version and Network CL version for incoming handshake packets, serverside
@@ -400,7 +400,7 @@ private:
 	 * @param OutTargetVersion	The target handshake protocol version to use, for the new connection
 	 * @return					Whether or not the checks successfully validated that the version is correct, and that communication can continue
 	 */
-	ENGINE_API bool CheckVersion(const FParsedHandshakeData& HandshakeData, EHandshakeVersion& OutTargetVersion) const;
+	bool CheckVersion(const FParsedHandshakeData& HandshakeData, EHandshakeVersion& OutTargetVersion) const;
 
 	/**
 	 * Takes the client address plus server timestamp, and outputs a deterministic cookie value
@@ -410,17 +410,17 @@ private:
 	 * @param TimeStamp			The serverside timestamp
 	 * @param OutCookie			Outputs the generated cookie value.
 	 */
-	ENGINE_API void GenerateCookie(const TSharedPtr<const FInternetAddr>& ClientAddress, uint8 SecretId, double Timestamp, uint8 (&OutCookie)[COOKIE_BYTE_SIZE]) const;
+	void GenerateCookie(const TSharedPtr<const FInternetAddr>& ClientAddress, uint8 SecretId, double Timestamp, uint8 (&OutCookie)[COOKIE_BYTE_SIZE]) const;
 
 	/**
 	 * Generates a new HandshakeSecret value
 	 */
-	ENGINE_API void UpdateSecret();
+	void UpdateSecret();
 
 	/**
 	 * Calculate the adjusted predefined packet size, based on MagicHeader
 	 */
-	ENGINE_API int32 GetAdjustedSizeBits(int32 InSizeBits, EHandshakeVersion HandshakeVersion) const;
+	int32 GetAdjustedSizeBits(int32 InSizeBits, EHandshakeVersion HandshakeVersion) const;
 
 #if !UE_BUILD_SHIPPING
 	/**
@@ -428,7 +428,7 @@ private:
 	 *
 	 * @return		Whether or not logging is currently allowed
 	 */
-	ENGINE_API bool TrackValidationLogs();
+	bool TrackValidationLogs();
 #endif
 
 
