NetDriver
* Allow GameInstance's to override the Iris ReplicationSystem config

diff --git a/Engine/Source/Runtime/Engine/Classes/Engine/GameInstance.h b/Engine/Source/Runtime/Engine/Classes/Engine/GameInstance.h
--- a/Engine/Source/Runtime/Engine/Classes/Engine/GameInstance.h
+++ b/Engine/Source/Runtime/Engine/Classes/Engine/GameInstance.h
@@ -26,6 +26,7 @@ class ULocalPlayer;
 class UOnlineSession;
 enum EPlayNetMode : int;
 struct FLatentActionManager;
+struct FNetDriverReplicationSystemConfig;
 class ULevelEditorPlaySettings;
 class IAnalyticsProvider;
 namespace ERHIFeatureLevel { enum Type : int; }
@@ -664,6 +665,9 @@ public:
 	/** Allows a GameInstance to override the default configured replication system for it's specific conditions. Return Default to use the configured replication system. */
 	ENGINE_API virtual EReplicationSystem GetDesiredReplicationSystem(FName InNetDriverDefinition) const;
 
+	/** Allows a GameInstance to override the default configuration of the Iris replication system. For instance if specific game modes require different settings. */
+	ENGINE_API virtual void OverrideIrisReplicationSystemConfig(FNetDriverReplicationSystemConfig& OutConfig, bool bIsServer) const {}
+
 protected:
 	/** Non-virtual dispatch for OnStart, also calls the associated global OnStartGameInstance. */
 	ENGINE_API void BroadcastOnStart();
diff --git a/Engine/Source/Runtime/Engine/Private/NetDriver.cpp b/Engine/Source/Runtime/Engine/Private/NetDriver.cpp
--- a/Engine/Source/Runtime/Engine/Private/NetDriver.cpp
+++ b/Engine/Source/Runtime/Engine/Private/NetDriver.cpp
@@ -7087,21 +7087,23 @@ void UNetDriver::CreateReplicationSystem(bool bInitAsClient)
 	{
 		ReplicationBridge->SetNetDriver(this);
 
+		FNetDriverReplicationSystemConfig& RepSystemConfig = bInitAsClient ? ReplicationSystemConfigClient : ReplicationSystemConfigServer;
+		// Ask the GameInstance if it wants to override any configs
+		if (const UWorld* LocalWorld = GetWorld())
+		{
+			if (UGameInstance* GameInstance = LocalWorld->GetGameInstance())
+			{
+				GameInstance->OverrideIrisReplicationSystemConfig(RepSystemConfig, !bInitAsClient);
+			}
+		}
+
 		// Create ReplicationSystem
 		UReplicationSystem::FReplicationSystemParams Params;
 		Params.ReplicationBridge = ReplicationBridge;
 		Params.bIsServer = !bInitAsClient;
 		Params.bAllowObjectReplication = !bInitAsClient;
 		Params.ForwardNetRPCCallDelegate.BindUObject(this, &UNetDriver::ForwardRemoteFunction);
-
-		if (bInitAsClient)
-		{
-			UE::Net::Private::ApplyReplicationSystemConfig(ReplicationSystemConfigClient, Params);
-		}
-		else
-		{
-			UE::Net::Private::ApplyReplicationSystemConfig(ReplicationSystemConfigServer, Params);
-		}
+		UE::Net::Private::ApplyReplicationSystemConfig(RepSystemConfig, Params);
 
 		SetReplicationSystem(UE::Net::FReplicationSystemFactory::CreateReplicationSystem(Params));
 	}
