Update UMaterialExpressionRuntimeVirtualTextureSample to no longer use the deprecated GetInputViews()
Also fix a bug in that class with DDX (UV) being duplicated into DDY (UV).
#rnx

diff --git a/Engine/Source/Runtime/Engine/Private/Materials/MaterialExpressions.cpp b/Engine/Source/Runtime/Engine/Private/Materials/MaterialExpressions.cpp
--- a/Engine/Source/Runtime/Engine/Private/Materials/MaterialExpressions.cpp
+++ b/Engine/Source/Runtime/Engine/Private/Materials/MaterialExpressions.cpp
@@ -3215,62 +3215,71 @@ UObject* UMaterialExpressionRuntimeVirtualTextureSample::GetReferencedTexture()
 
 #if WITH_EDITOR
 
-TArrayView<FExpressionInput*> UMaterialExpressionRuntimeVirtualTextureSample::GetInputsView()
+FExpressionInput* UMaterialExpressionRuntimeVirtualTextureSample::GetInput(int32 InputIndex)
 {
-	CachedInputs.Empty();
-	CachedInputs.Add(&Coordinates);
-	CachedInputs.Add(&WorldPosition);
-	if (MipValueMode == RVTMVM_MipLevel || MipValueMode == RVTMVM_MipBias)
-	{
-		CachedInputs.Add(&MipValue);
-	}
-	if (MipValueMode == RVTMVM_DerivativeUV || MipValueMode == RVTMVM_DerivativeWorld)
+	switch (InputIndex)
 	{
-		CachedInputs.Add(&DDX);
-		CachedInputs.Add(&DDY);
+	case 0: 
+		return &Coordinates;
+	case 1: 
+		return &WorldPosition;
+	case 2:
+		if (MipValueMode == RVTMVM_MipLevel || MipValueMode == RVTMVM_MipBias)
+		{
+			return &MipValue;
+		}
+		else if (MipValueMode == RVTMVM_DerivativeUV || MipValueMode == RVTMVM_DerivativeWorld)
+		{
+			return &DDX;
+		}
+		break;
+	case 3:
+		if (MipValueMode == RVTMVM_DerivativeUV || MipValueMode == RVTMVM_DerivativeWorld)
+		{
+			return &DDY;
+		}
+		break;
 	}
-	return CachedInputs;
+
+	return nullptr;
 }
 
 FName UMaterialExpressionRuntimeVirtualTextureSample::GetInputName(int32 InputIndex) const
 {
-	if (GetInput(InputIndex) == &WorldPosition)
+	switch (InputIndex)
 	{
+	case 1:
 		return GetWorldPositionInputName(WorldPositionOriginType);
-	}
-	if (CachedInputs[InputIndex] == &MipValue)
-	{
+	case 2:
 		if (MipValueMode == RVTMVM_MipLevel)
 		{
 			return TEXT("Mip Level");
 		}
 		if (MipValueMode == RVTMVM_MipBias)
 		{
-			return TEXT("Mip Bias");
+			return TEXT("Mip Level");
 		}
-	}
-	if (CachedInputs[InputIndex] == &DDX)
-	{
-		if (MipValueMode == RVTMVM_DerivativeUV)
+		else if (MipValueMode == RVTMVM_DerivativeUV)
 		{
 			return TEXT("DDX (UV)");
 		}
-		if (MipValueMode == RVTMVM_DerivativeWorld)
+		else if (MipValueMode == RVTMVM_DerivativeWorld)
 		{
 			return TEXT("DDX (World)");
 		}
-	}
-	if (CachedInputs[InputIndex] == &DDY)
-	{
+		break;
+	case 3:
 		if (MipValueMode == RVTMVM_DerivativeUV)
 		{
 			return TEXT("DDY (UV)");
 		}
-		if (MipValueMode == RVTMVM_DerivativeWorld)
+		else if (MipValueMode == RVTMVM_DerivativeWorld)
 		{
 			return TEXT("DDY (World)");
 		}
+		break;
 	}
+
 	return Super::GetInputName(InputIndex);
 }
 
@@ -3621,7 +3630,7 @@ int32 UMaterialExpressionRuntimeVirtualTextureSample::Compile(class FMaterialCom
 		if (MipValueMode == RVTMVM_DerivativeUV)
 		{
 			MipValue0Index = Ddx;
-			MipValue0Index = Ddy;
+			MipValue1Index = Ddy;
 		}
 		else if (MipValueMode == RVTMVM_DerivativeWorld)
 		{
diff --git a/Engine/Source/Runtime/Engine/Public/Materials/MaterialExpressionRuntimeVirtualTextureSample.h b/Engine/Source/Runtime/Engine/Public/Materials/MaterialExpressionRuntimeVirtualTextureSample.h
--- a/Engine/Source/Runtime/Engine/Public/Materials/MaterialExpressionRuntimeVirtualTextureSample.h
+++ b/Engine/Source/Runtime/Engine/Public/Materials/MaterialExpressionRuntimeVirtualTextureSample.h
@@ -142,7 +142,7 @@ protected:
 
 #if WITH_EDITOR
 	ENGINE_API virtual void PostLoad() override;
-	ENGINE_API TArrayView<FExpressionInput*> GetInputsView() override;
+	FExpressionInput* GetInput(int32 InputIndex) override;
 	ENGINE_API virtual FName GetInputName(int32 InputIndex) const override;
 	ENGINE_API virtual int32 Compile(class FMaterialCompiler* Compiler, int32 OutputIndex) override;
 	ENGINE_API virtual void GetCaption(TArray<FString>& OutCaptions) const override;
