Temp fix to make PackedLevelActors properly update when making changes to the source Level.
(Proper fix is described in https://jira.it.epicgames.com/browse/UE-216035).

diff --git a/Engine/Source/Runtime/Engine/Private/PackedLevelActor/PackedLevelActor.cpp b/Engine/Source/Runtime/Engine/Private/PackedLevelActor/PackedLevelActor.cpp
--- a/Engine/Source/Runtime/Engine/Private/PackedLevelActor/PackedLevelActor.cpp
+++ b/Engine/Source/Runtime/Engine/Private/PackedLevelActor/PackedLevelActor.cpp
@@ -81,9 +81,18 @@ void APackedLevelActor::RerunConstructionScripts()
 				*GetClass()->GetDefaultObject<APackedLevelActor>()->PackedVersion.ToString());
 		}
 	}
-	
+
 	if(bShouldRerunConstructionScript)
 	{
+		// Set bEditableWhenInherited to false to disable editing of properties on components.
+		// This was enabled in 5.4, but is not properly handled, so for now we are disabling it.
+		// @todo_ow: See https://jira.it.epicgames.com/browse/UE-216035
+		TArray<UActorComponent*> PackedComponents;
+		GetPackedComponents(PackedComponents);
+		for (UActorComponent* PackedComponent : PackedComponents)
+		{
+			PackedComponent->bEditableWhenInherited = false;
+		}
 		Super::RerunConstructionScripts();
 		PackedVersion = GetClass()->GetDefaultObject<APackedLevelActor>()->PackedVersion;
 	}
diff --git a/Engine/Source/Runtime/Engine/Private/PackedLevelActor/PackedLevelActorBuilder.cpp b/Engine/Source/Runtime/Engine/Private/PackedLevelActor/PackedLevelActorBuilder.cpp
--- a/Engine/Source/Runtime/Engine/Private/PackedLevelActor/PackedLevelActorBuilder.cpp
+++ b/Engine/Source/Runtime/Engine/Private/PackedLevelActor/PackedLevelActorBuilder.cpp
@@ -481,6 +481,10 @@ bool FPackedLevelActorBuilder::CreateOrUpdateBlueprintFromPacked(APackedLevelAct
 		{
 			ISMComponent->bInheritPerInstanceData = true;
 		}
+		// Set bEditableWhenInherited to false to disable editing of properties on components.
+		// This was enabled in 5.4, but is not properly handled, so for now we are disabling it.
+		// @todo_ow: See https://jira.it.epicgames.com/browse/UE-216035
+		PackedComponent->bEditableWhenInherited = false;
 	}
 
 	FKismetEditorUtilities::AddComponentsToBlueprint(BP, PackedComponents, AddCompToBPParams);
