MaterialExpressions: Cook Indeterminism: Compute ExpressionGuids deterministically when cooking, so that if any MaterialExpression was saved without calculating one, it will be calculated deterministically before saving the cooked package.
#rnx

diff --git a/Engine/Source/Runtime/Engine/Private/Materials/MaterialExpressions.cpp b/Engine/Source/Runtime/Engine/Private/Materials/MaterialExpressions.cpp
--- a/Engine/Source/Runtime/Engine/Private/Materials/MaterialExpressions.cpp
+++ b/Engine/Source/Runtime/Engine/Private/Materials/MaterialExpressions.cpp
@@ -2063,7 +2063,16 @@ void UMaterialExpression::UpdateMaterialExpressionGuid(bool bForceGeneration, bo
 
 		if (bForceGeneration || !Guid.IsValid())
 		{
-			Guid = FGuid::NewGuid();
+			if (IsRunningCookCommandlet())
+			{
+				// Prevent indeterminism in the cook by constructing the guid deterministically
+				constexpr int64 MaterialExpressionGuidDeterminismSeed = 0;
+				Guid = FGuid::NewDeterministicGuid(GetPathName(), MaterialExpressionGuidDeterminismSeed);
+			}
+			else
+			{
+				Guid = FGuid::NewGuid();
+			}
 
 			if (bAllowMarkingPackageDirty)
 			{
@@ -2085,7 +2094,16 @@ void UMaterialExpression::UpdateParameterGuid(bool bForceGeneration, bool bAllow
 
 			if (bForceGeneration || !Guid.IsValid())
 			{
-				Guid = FGuid::NewGuid();
+				if (IsRunningCookCommandlet())
+				{
+					// Prevent indeterminism in the cook by constructing the guid deterministically
+					constexpr int64 ParameterGuidDeterminismSeed = 1;
+					Guid = FGuid::NewDeterministicGuid(GetPathName(), ParameterGuidDeterminismSeed);
+				}
+				else
+				{
+					Guid = FGuid::NewGuid();
+				}
 
 				if (bAllowMarkingPackageDirty)
 				{
