Avoid overflow when estimating the raw animation sequence size

diff --git a/Engine/Source/Runtime/Engine/Classes/Animation/AnimSequence.h b/Engine/Source/Runtime/Engine/Classes/Animation/AnimSequence.h
--- a/Engine/Source/Runtime/Engine/Classes/Animation/AnimSequence.h
+++ b/Engine/Source/Runtime/Engine/Classes/Animation/AnimSequence.h
@@ -564,7 +564,7 @@ public:
 	/** @return	estimate uncompressed raw size. This is *not* the real raw size. 
 				Here we estimate what it would be with no trivial compression. */
 #if WITH_EDITOR
-	ENGINE_API int32 GetUncompressedRawSize() const;
+	ENGINE_API int64 GetUncompressedRawSize() const;
 
 	/**
 	 * @return		The approximate size of raw animation data.
diff --git a/Engine/Source/Runtime/Engine/Private/Animation/AnimSequence.cpp b/Engine/Source/Runtime/Engine/Private/Animation/AnimSequence.cpp
--- a/Engine/Source/Runtime/Engine/Private/Animation/AnimSequence.cpp
+++ b/Engine/Source/Runtime/Engine/Private/Animation/AnimSequence.cpp
@@ -431,7 +431,7 @@ void UAnimSequence::GetAssetRegistryTags(FAssetRegistryTagsContext Context) cons
 
 	if (DataModelInterface.GetObject() && DataModelInterface->HasBeenPopulated())
 	{
-		Context.AddTag(FAssetRegistryTag(TEXT("Compression Ratio"), FString::Printf(TEXT("%.03f"), (float)GetApproxCompressedSize() / (float)GetUncompressedRawSize()), FAssetRegistryTag::TT_Numerical));
+		Context.AddTag(FAssetRegistryTag(TEXT("Compression Ratio"), FString::Printf(TEXT("%.03f"), (double)GetApproxCompressedSize() / (double)GetUncompressedRawSize()), FAssetRegistryTag::TT_Numerical));
 		Context.AddTag(FAssetRegistryTag(TEXT("Source Frame Rate"), FString::Printf(TEXT("%.2f"), DataModelInterface->GetFrameRate().AsDecimal()), FAssetRegistryTag::TT_Numerical));
 		Context.AddTag(FAssetRegistryTag(TEXT("Number of Frames"), FString::Printf(TEXT("%.i"), DataModelInterface->GetNumberOfFrames()), FAssetRegistryTag::TT_Numerical));
 		Context.AddTag(FAssetRegistryTag(TEXT("Number of Keys"), FString::Printf(TEXT("%i"), DataModelInterface->GetNumberOfKeys()), FAssetRegistryTag::TT_Numerical));
@@ -502,17 +502,17 @@ void UAnimSequence::ClearAllCachedCookedPlatformData()
 	PRAGMA_ENABLE_DEPRECATION_WARNINGS
 }
 
-int32 UAnimSequence::GetUncompressedRawSize() const
+int64 UAnimSequence::GetUncompressedRawSize() const
 {
 	if (DataModelInterface.GetObject())
 	{
-		const int32 BoneRawSize = FRawAnimSequenceTrack::SingleKeySize * DataModelInterface->GetNumBoneTracks() * DataModelInterface->GetNumberOfKeys();
-		int32 CurveRawSize = 0;
+		const int64 BoneRawSize = FRawAnimSequenceTrack::SingleKeySize * int64(DataModelInterface->GetNumBoneTracks()) * int64(DataModelInterface->GetNumberOfKeys());
+		int64 CurveRawSize = 0;
 
 		for (const FFloatCurve& Curve : DataModelInterface->GetCurveData().FloatCurves)
 		{
 			CurveRawSize += sizeof(FFloatCurve);
-			CurveRawSize += sizeof(FRichCurveKey) * Curve.FloatCurve.Keys.Num();
+			CurveRawSize += sizeof(FRichCurveKey) * int64(Curve.FloatCurve.Keys.Num());
 		}
 		return BoneRawSize + CurveRawSize;
 	}
@@ -524,18 +524,18 @@ int64 UAnimSequence::GetApproxRawSize() const
 	if (ShouldDataModelBeValid())
 	{
 		ValidateModel();
-		int32 Total = sizeof(FRawAnimSequenceTrack) * DataModelInterface->GetNumBoneTracks();
 
-		const int32 NumberOfTracks = DataModelInterface->GetNumBoneTracks();
+		const int64 NumberOfTracks = DataModelInterface->GetNumBoneTracks();
 
-		Total += DataModelInterface->GetNumberOfKeys() * sizeof(FVector3f) * NumberOfTracks;
-		Total += DataModelInterface->GetNumberOfKeys() * sizeof(FQuat4f) * NumberOfTracks;
-		Total += DataModelInterface->GetNumberOfKeys() * sizeof(FVector3f) * NumberOfTracks;
+		int64 Total = sizeof(FRawAnimSequenceTrack) * NumberOfTracks;
+		Total += int64(DataModelInterface->GetNumberOfKeys()) * sizeof(FVector3f) * NumberOfTracks;
+		Total += int64(DataModelInterface->GetNumberOfKeys()) * sizeof(FQuat4f) * NumberOfTracks;
+		Total += int64(DataModelInterface->GetNumberOfKeys()) * sizeof(FVector3f) * NumberOfTracks;
 
 		for (const FFloatCurve& Curve : DataModelInterface->GetCurveData().FloatCurves)
 		{
 			Total += sizeof(FFloatCurve);
-			Total += sizeof(FRichCurveKey) * Curve.FloatCurve.Keys.Num();
+			Total += sizeof(FRichCurveKey) * int64(Curve.FloatCurve.Keys.Num());
 		}
 		return Total;
 	}
diff --git a/Engine/Source/Runtime/Engine/Public/Animation/AnimCompressionTypes.h b/Engine/Source/Runtime/Engine/Public/Animation/AnimCompressionTypes.h
--- a/Engine/Source/Runtime/Engine/Public/Animation/AnimCompressionTypes.h
+++ b/Engine/Source/Runtime/Engine/Public/Animation/AnimCompressionTypes.h
@@ -260,38 +260,38 @@ public:
 
 	const ITargetPlatform* TargetPlatform = nullptr;
 
-	static int32 GetApproxRawDataArraySize(const TArray<FRawAnimSequenceTrack>& AnimData)
+	static int64 GetApproxRawDataArraySize(const TArray<FRawAnimSequenceTrack>& AnimData)
 	{
-		int32 Total = sizeof(FRawAnimSequenceTrack) * AnimData.Num();
+		int64 Total = sizeof(FRawAnimSequenceTrack) * int64(AnimData.Num());
 		for (int32 i = 0; i < AnimData.Num(); ++i)
 		{
 			const FRawAnimSequenceTrack& RawTrack = AnimData[i];
 			Total +=
-				sizeof(FVector) * RawTrack.PosKeys.Num() +
-				sizeof(FQuat) * RawTrack.RotKeys.Num() +
-				sizeof(FVector) * RawTrack.ScaleKeys.Num();
+				sizeof(FVector) * int64(RawTrack.PosKeys.Num()) +
+				sizeof(FQuat) * int64(RawTrack.RotKeys.Num()) +
+				sizeof(FVector) * int64(RawTrack.ScaleKeys.Num());
 		}
 
 		return Total;
 	}
 
-	int32 GetApproxRawBoneSize() const
+	int64 GetApproxRawBoneSize() const
 	{
 		return GetApproxRawDataArraySize(RawAnimationData);
 	}
 
-	int32 GetApproxRawCurveSize() const
+	int64 GetApproxRawCurveSize() const
 	{
-		int32 Total = 0;
+		int64 Total = 0;
 		for (const FFloatCurve& Curve : RawFloatCurves)
 		{
 			Total += sizeof(FFloatCurve);
-			Total += sizeof(FRichCurveKey) * Curve.FloatCurve.Keys.Num();
+			Total += sizeof(FRichCurveKey) * int64(Curve.FloatCurve.Keys.Num());
 		}
 		return Total;
 	}
 
-	int32 GetApproxRawSize() const
+	int64 GetApproxRawSize() const
 	{
 		return GetApproxRawBoneSize() + GetApproxRawCurveSize();
 	}
