[Material Editor] Removed redundant and incorrect comment.

diff --git a/Engine/Source/Runtime/Engine/Private/Materials/MaterialExpressions.cpp b/Engine/Source/Runtime/Engine/Private/Materials/MaterialExpressions.cpp
--- a/Engine/Source/Runtime/Engine/Private/Materials/MaterialExpressions.cpp
+++ b/Engine/Source/Runtime/Engine/Private/Materials/MaterialExpressions.cpp
@@ -2634,7 +2634,6 @@ void UMaterialExpressionTextureSample::PostLoad()
 TArrayView<FExpressionInput*> UMaterialExpressionTextureSample::GetInputsView()
 {
 	CachedInputs.Empty();
-	// todo: we should remove GetInputs() and make this the common code for all expressions
 	uint32 InputIndex = 0;
 	while (FExpressionInput* Ptr = GetInput(InputIndex++))
 	{
