[World Partition Optimization] Build an actor to external streaming objects mapping to avoid going through all external streaming objects when looking up soft object paths.

diff --git a/Engine/Source/Runtime/Engine/Private/WorldPartition/WorldPartitionLevelStreamingPolicy.cpp b/Engine/Source/Runtime/Engine/Private/WorldPartition/WorldPartitionLevelStreamingPolicy.cpp
--- a/Engine/Source/Runtime/Engine/Private/WorldPartition/WorldPartitionLevelStreamingPolicy.cpp
+++ b/Engine/Source/Runtime/Engine/Private/WorldPartition/WorldPartitionLevelStreamingPolicy.cpp
@@ -301,6 +301,12 @@ bool UWorldPartitionLevelStreamingPolicy::InjectExternalStreamingObject(URuntime
 	if (Super::InjectExternalStreamingObject(ExternalStreamingObject))
 	{
 		ExternalStreamingObjects.Add(ExternalStreamingObject);
+
+		for (const TPair<FName, FName> Pair : ExternalStreamingObject->SubObjectsToCellRemapping)
+		{
+			SubObjectsToExternalStreamingObjectsRemapping.Add(Pair.Key, ExternalStreamingObject);
+		}
+
 		return true;
 	}
 
@@ -311,6 +317,11 @@ bool UWorldPartitionLevelStreamingPolicy::RemoveExternalStreamingObject(URuntime
 {
 	bool bSuccess = Super::RemoveExternalStreamingObject(ExternalStreamingObject);
 
+	for (const TPair<FName, FName> Pair : ExternalStreamingObject->SubObjectsToCellRemapping)
+	{
+		verify(SubObjectsToExternalStreamingObjectsRemapping.FindAndRemoveChecked(Pair.Key) == ExternalStreamingObject);
+	}
+
 	ExternalStreamingObjects.RemoveSwap(ExternalStreamingObject);
 
 	return bSuccess;
@@ -407,16 +418,13 @@ const FName* UWorldPartitionLevelStreamingPolicy::FindCellNameForSubObject(const
 		return CellName;
 	}
 
-	for (const TWeakObjectPtr<URuntimeHashExternalStreamingObjectBase>& ExternalStreamingObject : ExternalStreamingObjects)
+	if (const TWeakObjectPtr<URuntimeHashExternalStreamingObjectBase>* ExternalStreamingObject = SubObjectsToExternalStreamingObjectsRemapping.Find(*SubObjectContext); ExternalStreamingObject && ExternalStreamingObject->IsValid())
 	{
-		if (ExternalStreamingObject.IsValid())
+		const URuntimeHashExternalStreamingObjectBase* ExternalStreamingObjectPtr = ExternalStreamingObject->Get();
+		if (const FName* CellName = ExternalStreamingObjectPtr->SubObjectsToCellRemapping.Find(*SubObjectContext))
 		{
-			const URuntimeHashExternalStreamingObjectBase* ExternalStreamingObjectPtr = ExternalStreamingObject.Get();
-			if (const FName* CellName = ExternalStreamingObjectPtr->SubObjectsToCellRemapping.Find(*SubObjectContext))
-			{
-				OutLevelMountPointContext = ExternalStreamingObjectPtr->GetLevelMountPointContextObject();
-				return CellName;
-			}
+			OutLevelMountPointContext = ExternalStreamingObjectPtr->GetLevelMountPointContextObject();
+			return CellName;
 		}
 	}
 	
@@ -455,20 +463,15 @@ const UWorldPartitionRuntimeLevelStreamingCell* UWorldPartitionLevelStreamingPol
 
 	if (!CellPackage)
 	{
-		for (const TWeakObjectPtr<URuntimeHashExternalStreamingObjectBase>& ExternalStreamingObject : ExternalStreamingObjects)
+		if (const TWeakObjectPtr<URuntimeHashExternalStreamingObjectBase>* ExternalStreamingObject = SubObjectsToExternalStreamingObjectsRemapping.Find(SubObjectName); ExternalStreamingObject && ExternalStreamingObject->IsValid())
 		{
-			if (ExternalStreamingObject.IsValid())
+			if (const FName* FoundCell = ExternalStreamingObject->Get()->SubObjectsToCellRemapping.Find(SubObjectName))
 			{
-				if (const FName* FoundCell = ExternalStreamingObject->SubObjectsToCellRemapping.Find(SubObjectName))
-				{
-					CellName = *FoundCell;
-					CellPackage = ExternalStreamingObject.Get();
-					break;
-				}
+				CellName = *FoundCell;
+				CellPackage = ExternalStreamingObject->Get();
 			}
 		}
 	}
-	
 
 	if (CellPackage)
 	{
diff --git a/Engine/Source/Runtime/Engine/Public/WorldPartition/WorldPartitionLevelStreamingPolicy.h b/Engine/Source/Runtime/Engine/Public/WorldPartition/WorldPartitionLevelStreamingPolicy.h
--- a/Engine/Source/Runtime/Engine/Public/WorldPartition/WorldPartitionLevelStreamingPolicy.h
+++ b/Engine/Source/Runtime/Engine/Public/WorldPartition/WorldPartitionLevelStreamingPolicy.h
@@ -64,4 +64,7 @@ private:
 
 	UPROPERTY(Transient)
 	TArray<TWeakObjectPtr<URuntimeHashExternalStreamingObjectBase>> ExternalStreamingObjects;
+
+	UPROPERTY(Transient)
+	TMap<FName, TWeakObjectPtr<URuntimeHashExternalStreamingObjectBase>> SubObjectsToExternalStreamingObjectsRemapping;
 };
\ No newline at end of file
