Track memory allocations from game callbacks and Iris pre-allocated buffers in Iris.

diff --git a/Engine/Source/Runtime/Engine/Private/Net/Iris/ReplicationSystem/ActorReplicationBridge.cpp b/Engine/Source/Runtime/Engine/Private/Net/Iris/ReplicationSystem/ActorReplicationBridge.cpp
--- a/Engine/Source/Runtime/Engine/Private/Net/Iris/ReplicationSystem/ActorReplicationBridge.cpp
+++ b/Engine/Source/Runtime/Engine/Private/Net/Iris/ReplicationSystem/ActorReplicationBridge.cpp
@@ -11,6 +11,7 @@
 #include "Iris/Core/IrisLog.h"
 #include "Iris/Core/IrisProfiler.h"
 #include "Iris/Core/NetObjectReference.h"
+#include "Iris/Core/IrisMemoryTracker.h"
 #include "Iris/ReplicationSystem/NetToken.h"
 #include "Iris/ReplicationSystem/ObjectReplicationBridgeConfig.h"
 #include "Iris/ReplicationSystem/StringTokenStore.h"
@@ -618,6 +619,8 @@ TUniquePtr<UObjectReplicationBridge::FCreationHeader> UActorReplicationBridge::R
 
 FObjectReplicationBridgeInstantiateResult UActorReplicationBridge::BeginInstantiateFromRemote(FNetRefHandle RootObjectOfSubObject, const UE::Net::FNetObjectResolveContext& ResolveContext, const UObjectReplicationBridge::FCreationHeader* InHeader)
 {
+	LLM_SCOPE(ELLMTag::EngineMisc);
+
 	using namespace UE::Net::Private;
 
 	IRIS_PROFILER_SCOPE(ActorReplicationBridge_OnBeginInstantiateFromRemote);
@@ -812,6 +815,8 @@ FObjectReplicationBridgeInstantiateResult UActorReplicationBridge::BeginInstanti
 
 bool UActorReplicationBridge::OnInstantiatedFromRemote(UObject* Instance, const UObjectReplicationBridge::FCreationHeader* InHeader, uint32 ConnectionId) const
 {
+	LLM_SCOPE(ELLMTag::EngineMisc);
+
 	using namespace UE::Net::Private;
 
 	const FActorReplicationBridgeCreationHeader* BridgeHeader = static_cast<const FActorReplicationBridgeCreationHeader*>(InHeader);
diff --git a/Engine/Source/Runtime/Experimental/Iris/Core/Private/Iris/Core/IrisMemoryTracker.cpp b/Engine/Source/Runtime/Experimental/Iris/Core/Private/Iris/Core/IrisMemoryTracker.cpp
--- a/Engine/Source/Runtime/Experimental/Iris/Core/Private/Iris/Core/IrisMemoryTracker.cpp
+++ b/Engine/Source/Runtime/Experimental/Iris/Core/Private/Iris/Core/IrisMemoryTracker.cpp
@@ -5,9 +5,11 @@
 
 DECLARE_LLM_MEMORY_STAT(TEXT("Iris"), STAT_IrisLLM, STATGROUP_LLMFULL);
 DECLARE_LLM_MEMORY_STAT(TEXT("IrisState"), STAT_IrisStateLLM, STATGROUP_LLMFULL);
+DECLARE_LLM_MEMORY_STAT(TEXT("IrisInitialization"), STAT_IrisInitializationLLM, STATGROUP_LLMFULL);
 DECLARE_LLM_MEMORY_STAT(TEXT("IrisConnection"), STAT_IrisConnectionLLM, STATGROUP_LLMFULL);
 
 // Tag name, display name, parent tag name, stat name, summary stat name
 LLM_DEFINE_TAG(Iris, NAME_None, NAME_None, GET_STATFNAME(STAT_IrisLLM), GET_STATFNAME(STAT_NetworkingSummaryLLM));
 LLM_DEFINE_TAG(IrisState, "State", "Iris", GET_STATFNAME(STAT_IrisStateLLM), GET_STATFNAME(STAT_NetworkingSummaryLLM));
+LLM_DEFINE_TAG(IrisInitialization, "Initialization", "Iris", GET_STATFNAME(STAT_IrisInitializationLLM), GET_STATFNAME(STAT_NetworkingSummaryLLM));
 LLM_DEFINE_TAG(IrisConnection, "Connection", "Iris", GET_STATFNAME(STAT_IrisConnectionLLM), GET_STATFNAME(STAT_NetworkingSummaryLLM));
diff --git a/Engine/Source/Runtime/Experimental/Iris/Core/Private/Iris/ReplicationSystem/ObjectReplicationBridge.cpp b/Engine/Source/Runtime/Experimental/Iris/Core/Private/Iris/ReplicationSystem/ObjectReplicationBridge.cpp
--- a/Engine/Source/Runtime/Experimental/Iris/Core/Private/Iris/ReplicationSystem/ObjectReplicationBridge.cpp
+++ b/Engine/Source/Runtime/Experimental/Iris/Core/Private/Iris/ReplicationSystem/ObjectReplicationBridge.cpp
@@ -163,8 +163,6 @@ bool UObjectReplicationBridge::ShouldUseVerboseCsvStats() const
 
 void UObjectReplicationBridge::Initialize(UReplicationSystem* InReplicationSystem)
 {
-	LLM_SCOPE_BYTAG(Iris);
-
 	ensureMsgf((GEnsureNetRefHandleError == -1 || GEnsureNetRefHandleError < (int32)UE::Net::ENetRefHandleError::Max), TEXT("GEnsureNetRefHandleError is set to an invalid value: %u"), GEnsureNetRefHandleError);
 
 	Super::Initialize(InReplicationSystem);
diff --git a/Engine/Source/Runtime/Experimental/Iris/Core/Private/Iris/ReplicationSystem/ReplicationSystem.cpp b/Engine/Source/Runtime/Experimental/Iris/Core/Private/Iris/ReplicationSystem/ReplicationSystem.cpp
--- a/Engine/Source/Runtime/Experimental/Iris/Core/Private/Iris/ReplicationSystem/ReplicationSystem.cpp
+++ b/Engine/Source/Runtime/Experimental/Iris/Core/Private/Iris/ReplicationSystem/ReplicationSystem.cpp
@@ -95,8 +95,6 @@ public:
 
 	void Init(const UReplicationSystem::FReplicationSystemParams& Params)
 	{
-		LLM_SCOPE_BYTAG(Iris);
-
 #if !UE_BUILD_SHIPPING
 		IrisDebugHelperDummy = UE::Net::IrisDebugHelper::Init();
 #endif
@@ -658,8 +656,6 @@ UReplicationSystem::UReplicationSystem()
 
 void UReplicationSystem::Init(uint32 InId, const FReplicationSystemParams& Params)
 {
-	LLM_SCOPE_BYTAG(Iris);
-
 	Id = InId;
 	bIsServer = Params.bIsServer;
 	bAllowObjectReplication = Params.bAllowObjectReplication;
@@ -1755,6 +1751,8 @@ uint32 FReplicationSystemFactory::MaxReplicationSystemId = 0;
 
 UReplicationSystem* FReplicationSystemFactory::CreateReplicationSystem(const UReplicationSystem::FReplicationSystemParams& Params)
 {
+	LLM_SCOPE_BYTAG(IrisInitialization);
+
 	if (!Params.ReplicationBridge)
 	{
 		UE_LOG(LogIris, Error, TEXT("Cannot create ReplicationSystem without a ReplicationBridge"));
diff --git a/Engine/Source/Runtime/Experimental/Iris/Core/Public/Iris/Core/IrisMemoryTracker.h b/Engine/Source/Runtime/Experimental/Iris/Core/Public/Iris/Core/IrisMemoryTracker.h
--- a/Engine/Source/Runtime/Experimental/Iris/Core/Public/Iris/Core/IrisMemoryTracker.h
+++ b/Engine/Source/Runtime/Experimental/Iris/Core/Public/Iris/Core/IrisMemoryTracker.h
@@ -6,4 +6,5 @@
 
 LLM_DECLARE_TAG_API(Iris, IRISCORE_API);
 LLM_DECLARE_TAG_API(IrisState, IRISCORE_API);
+LLM_DECLARE_TAG_API(IrisInitialization, IRISCORE_API);
 LLM_DECLARE_TAG_API(IrisConnection, IRISCORE_API);
