[Backout] - CL33830717
Original CL Desc
-----------------------------------------------------------------
Bring back CL33705721 with fixes
Original CL Desc
-----------------------------------------------------------------
[Content Bundle] Added an optional flag to allow marking content bundles as "required", which will make them unremovable and bake their actors into the main map to reduce IO stress.

diff --git a/Engine/Source/Runtime/Engine/Private/WorldPartition/ContentBundle/ContentBundleDescriptor.cpp b/Engine/Source/Runtime/Engine/Private/WorldPartition/ContentBundle/ContentBundleDescriptor.cpp
--- a/Engine/Source/Runtime/Engine/Private/WorldPartition/ContentBundle/ContentBundleDescriptor.cpp
+++ b/Engine/Source/Runtime/Engine/Private/WorldPartition/ContentBundle/ContentBundleDescriptor.cpp
@@ -8,9 +8,10 @@
 #include UE_INLINE_GENERATED_CPP_BY_NAME(ContentBundleDescriptor)
 
 UContentBundleDescriptor::UContentBundleDescriptor(const FObjectInitializer& ObjectInitializer)
-	: bRequired(false)
-	, DebugColor(FColor::Black)
-{}
+	: DebugColor(FColor::Black)
+{
+
+}
 
 FString UContentBundleDescriptor::GetPackageRoot() const
 {
diff --git a/Engine/Source/Runtime/Engine/Private/WorldPartition/ContentBundle/ContentBundleEditor.cpp b/Engine/Source/Runtime/Engine/Private/WorldPartition/ContentBundle/ContentBundleEditor.cpp
--- a/Engine/Source/Runtime/Engine/Private/WorldPartition/ContentBundle/ContentBundleEditor.cpp
+++ b/Engine/Source/Runtime/Engine/Private/WorldPartition/ContentBundle/ContentBundleEditor.cpp
@@ -93,8 +93,6 @@ void FContentBundleEditor::DoInjectContent()
 
 			check(GetDescriptor()->GetGuid().IsValid());
 
-			ActorDescContainerInstance->bIsContentBundleRequired = GetDescriptor()->GetRequired();
-
 			if (!ActorDescContainerInstance->GetContainer()->IsEmpty())
 			{
 				WorldDataLayersActorReference = FWorldDataLayersReference(ActorDescContainerInstance.Get(), BuildWorlDataLayersName());
@@ -348,11 +346,6 @@ void FContentBundleEditor::GenerateStreaming(TArray<FString>* OutPackageToGenera
 		return;
 	}
 
-	if (GetDescriptor()->GetRequired())
-	{
-		return;
-	}
-
 	FActorDescContainerInstanceCollection Collection({ TObjectPtr<UActorDescContainerInstance>(ActorDescContainerInstance.Get()) });
 	UWorldPartition::FGenerateStreamingParams Params = UWorldPartition::FGenerateStreamingParams()
 		.SetContainerInstanceCollection(Collection, FStreamingGenerationContainerInstanceCollection::ECollectionType::BaseAsContentBundle);
@@ -395,22 +388,16 @@ void FContentBundleEditor::GenerateStreaming(TArray<FString>* OutPackageToGenera
 
 void FContentBundleEditor::OnBeginCook(IWorldPartitionCookPackageContext& CookContext)
 {
-	if (!GetDescriptor()->GetRequired())
-	{
-		check(!bIsInCook);
-		CookContext.RegisterPackageCookPackageGenerator(this);
-		bIsInCook = true;
-	}
+	check(!bIsInCook);
+	CookContext.RegisterPackageCookPackageGenerator(this);
+	bIsInCook = true;
 }
 
 void FContentBundleEditor::OnEndCook(IWorldPartitionCookPackageContext& CookContext)
 {
-	if (!GetDescriptor()->GetRequired())
-	{
-		check(bIsInCook);
-		CookContext.UnregisterPackageCookPackageGenerator(this);
-		bIsInCook = false;
-	}
+	check(bIsInCook);
+	CookContext.UnregisterPackageCookPackageGenerator(this);
+	bIsInCook = false;
 }
 
 bool FContentBundleEditor::GatherPackagesToCook(class IWorldPartitionCookPackageContext& CookContext)
diff --git a/Engine/Source/Runtime/Engine/Private/WorldPartition/RuntimeSpatialHash/RuntimeSpatialHashHLOD.cpp b/Engine/Source/Runtime/Engine/Private/WorldPartition/RuntimeSpatialHash/RuntimeSpatialHashHLOD.cpp
--- a/Engine/Source/Runtime/Engine/Private/WorldPartition/RuntimeSpatialHash/RuntimeSpatialHashHLOD.cpp
+++ b/Engine/Source/Runtime/Engine/Private/WorldPartition/RuntimeSpatialHash/RuntimeSpatialHashHLOD.cpp
@@ -396,7 +396,7 @@ static void UpdateHLODGridsActors(UWorld* World, const UActorDescContainerInstan
 		ASpatialHashRuntimeGridInfo* GridActor = ExistingGridActors.FindRef(GridName);
 		if (!GridActor)
 		{
-			FContentBundleActivationScope ContentBundleScope(ContainerInstance->IsContentBundleRequired() ? FGuid() : ContainerInstance->GetContentBundleGuid());
+			FContentBundleActivationScope ContentBndleScope(ContainerInstance->GetContentBundleGuid());
 			FScopedOverrideSpawningLevelMountPointObject EDLScope(ContainerInstance->GetExternalDataLayerAsset());
 
 			FActorSpawnParameters SpawnParams;
@@ -405,7 +405,7 @@ static void UpdateHLODGridsActors(UWorld* World, const UActorDescContainerInstan
 			GridActor->Tags.Add(HLODGridTag);
 			bDirty = true;
 
-			check(GridActor->GetContentBundleGuid() == (ContainerInstance->IsContentBundleRequired() ? FGuid() : ContainerInstance->GetContentBundleGuid()));
+			check(GridActor->GetContentBundleGuid() == ContainerInstance->GetContentBundleGuid());
 			check(GridActor->GetExternalDataLayerAsset() == ContainerInstance->GetExternalDataLayerAsset());
 		}
 
diff --git a/Engine/Source/Runtime/Engine/Private/WorldPartition/WorldPartitionStreamingGeneration.cpp b/Engine/Source/Runtime/Engine/Private/WorldPartition/WorldPartitionStreamingGeneration.cpp
--- a/Engine/Source/Runtime/Engine/Private/WorldPartition/WorldPartitionStreamingGeneration.cpp
+++ b/Engine/Source/Runtime/Engine/Private/WorldPartition/WorldPartitionStreamingGeneration.cpp
@@ -2073,56 +2073,42 @@ URuntimeHashExternalStreamingObjectBase* UWorldPartition::FlushStreamingToExtern
 	return ExternalStreamingObject;
 }
 
-void UWorldPartition::SetupHLODActors(const FSetupHLODActorsParams& Params)
+static void ExtractContentBundleContainerInstances(const FActorDescContainerInstanceCollection* InContainerInstanceCollection, TArray<const UActorDescContainerInstance*>& OutContentBundleContainerInstances, TArray<const UActorDescContainerInstance*>& OutNonContentBundleContainerInstances)
 {
-	TArray<const UActorDescContainerInstance*> ContentBundleContainerInstances;
-	TArray<const UActorDescContainerInstance*> BaseAndEDLContainerInstances;
-
-	ForEachActorDescContainerInstance([&ContentBundleContainerInstances, &BaseAndEDLContainerInstances](const UActorDescContainerInstance* InActorDescContainerInstance)
+	InContainerInstanceCollection->ForEachActorDescContainerInstance([&OutContentBundleContainerInstances, &OutNonContentBundleContainerInstances](const UActorDescContainerInstance* InActorDescContainerInstance)
 	{
-		// If the content bundle is required, we need to add it to both lists, and we'll filter out all actors so the HLOd generation process can delete HLOD actor files.
 		if (InActorDescContainerInstance->GetContentBundleGuid().IsValid())
 		{
-			ContentBundleContainerInstances.Add(InActorDescContainerInstance);
+			OutContentBundleContainerInstances.Add(InActorDescContainerInstance);
 		}
-		
-		if (!InActorDescContainerInstance->GetContentBundleGuid().IsValid() || InActorDescContainerInstance->IsContentBundleRequired())
+		else
 		{
-			BaseAndEDLContainerInstances.Add(InActorDescContainerInstance);
+			OutNonContentBundleContainerInstances.Add(InActorDescContainerInstance);
 		}
 	});
+}
+
+void UWorldPartition::SetupHLODActors(const FSetupHLODActorsParams& Params)
+{
+	TArray<const UActorDescContainerInstance*> ContentBundleContainerInstances;
+	TArray<const UActorDescContainerInstance*> BaseAndEDLContainerInstances;
+	ExtractContentBundleContainerInstances(this, ContentBundleContainerInstances, BaseAndEDLContainerInstances);
 
-	auto SetupHLODActorsForCollection = [this, &Params](const FStreamingGenerationContainerInstanceCollection& InContainerInstanceCollection, bool bFilterAllActors)
+	auto SetupHLODActorsForCollection = [this, &Params](const FStreamingGenerationContainerInstanceCollection& InContainerInstanceCollection)
 	{
 		TErrorHandlerSelector<FStreamingGenerationLogErrorHandler> ErrorHandlerSelector;
 		FWorldPartitionStreamingGenerator::FWorldPartitionStreamingGeneratorParams StreamingGeneratorParams = FWorldPartitionStreamingGenerator::FWorldPartitionStreamingGeneratorParams()
 			.SetWorldPartitionContext(this)
 			.SetErrorHandler(ErrorHandlerSelector.Get())
 			.SetEnableStreaming(IsStreamingEnabled())
-			.SetFilteredClasses({ bFilterAllActors ? AActor::StaticClass() : AWorldPartitionHLOD::StaticClass() })
+			.SetFilteredClasses({ AWorldPartitionHLOD::StaticClass() })
 			.SetIsValidGrid([this](FName GridName, const UClass* ActorClass) { return RuntimeHash->IsValidGrid(GridName, ActorClass); })
 			.SetIsValidHLODLayer([this](FName GridName, const FSoftObjectPath& HLODLayerPath) { return RuntimeHash->IsValidHLODLayer(GridName, HLODLayerPath); });
 
 		FWorldPartitionStreamingGenerator StreamingGenerator(StreamingGeneratorParams);
 		StreamingGenerator.PreparationPhase(InContainerInstanceCollection);
 
-		const FString ContainerPackageName = InContainerInstanceCollection.GetBaseContainerInstancePackageName().ToString();
-		FString ContainerShortName = FPackageName::GetShortName(ContainerPackageName);
-		if (!ContainerPackageName.StartsWith(TEXT("/Game/")))
-		{
-			TArray<FString> SplitContainerPath;
-			if (ContainerPackageName.ParseIntoArray(SplitContainerPath, TEXT("/")))
-			{
-				ContainerShortName += TEXT(".");
-				ContainerShortName += SplitContainerPath[0];
-			}
-		}
-
-		TStringBuilder<256> StateLogSuffix;
-		StateLogSuffix += TEXT("HLOD_");
-		StateLogSuffix += ContainerShortName;
-
-		TUniquePtr<FArchive> LogFileAr = FWorldPartitionStreamingGenerator::CreateDumpStateLogArchive(*StateLogSuffix);
+		TUniquePtr<FArchive> LogFileAr = FWorldPartitionStreamingGenerator::CreateDumpStateLogArchive(TEXT("HLOD"));
 		if (LogFileAr.IsValid())
 		{
 			FHierarchicalLogArchive HierarchicalLogAr(*LogFileAr);
@@ -2136,14 +2122,14 @@ void UWorldPartition::SetupHLODActors(const FSetupHLODActorsParams& Params)
 	for (const UActorDescContainerInstance* ContentBundleContainerInstance : ContentBundleContainerInstances)
 	{
 		FStreamingGenerationContainerInstanceCollection ContentBundleCollection({ ContentBundleContainerInstance }, FStreamingGenerationContainerInstanceCollection::ECollectionType::BaseAsContentBundle);
-		SetupHLODActorsForCollection(ContentBundleCollection, ContentBundleContainerInstance->IsContentBundleRequired());
+		SetupHLODActorsForCollection(ContentBundleCollection);
 	}
 
 	// Single pass for base and EDL container instances
 	if (!BaseAndEDLContainerInstances.IsEmpty())
 	{
 		FStreamingGenerationContainerInstanceCollection Collection(BaseAndEDLContainerInstances, FStreamingGenerationContainerInstanceCollection::ECollectionType::BaseAndEDLs);
-		SetupHLODActorsForCollection(Collection, false);
+		SetupHLODActorsForCollection(Collection);
 	}
 }
 
@@ -2246,15 +2232,8 @@ void FStreamingGenerationContainerInstanceCollection::InitializeCollection()
 
 	if (CollectionType == ECollectionType::BaseAsContentBundle)
 	{
-#if DO_CHECK
 		check(GetActorDescContainerCount() == 1);
 		check(GetContentBundleGuid().IsValid());
-		Algo::ForEach(ActorDescContainerInstanceCollection, [](const UActorDescContainerInstance* ActorDescContainerInstance)
-		{
-			check(ActorDescContainerInstance->GetContentBundleGuid().IsValid());
-			check(!ActorDescContainerInstance->IsContentBundleRequired());
-		});
-#endif
 		ContentBundleStartIdx = 0;
 		return;
 	}
@@ -2267,10 +2246,7 @@ void FStreamingGenerationContainerInstanceCollection::InitializeCollection()
 		// When type is set to BaseAndEDL, we remove ContentBundle containers from the collection.
 		// BaseAndEDL type assumes ContentBundle containers are generated separately one at a time.
 		check(!ShouldRegisterDelegates());
-		ActorDescContainerInstanceCollection.SetNum(Algo::RemoveIf(ActorDescContainerInstanceCollection, [](const UActorDescContainerInstance* ActorDescContainerInstance)
-		{
-			return ActorDescContainerInstance->GetContentBundleGuid().IsValid() && !ActorDescContainerInstance->IsContentBundleRequired();
-		}));
+		ActorDescContainerInstanceCollection.SetNum(Algo::RemoveIf(ActorDescContainerInstanceCollection, [](const UActorDescContainerInstance* ActorDescContainerInstance) { return ActorDescContainerInstance->GetContentBundleGuid().IsValid(); }));
 	}
 
 	int32 BaseContainerCount = Algo::CountIf(ActorDescContainerInstanceCollection, [](const UActorDescContainerInstance* ActorDescContainerInstance) { return !ActorDescContainerInstance->HasExternalContent(); });
@@ -2367,18 +2343,7 @@ void UWorldPartition::CheckForErrors(const FCheckForErrorsParams& InParams)
 	// @todo_ow : Once content bundles are remove, we will only do this validation in 1 pass
 	TArray<const UActorDescContainerInstance*> ContentBundleContainerInstances;
 	TArray<const UActorDescContainerInstance*> BaseAndEDLContainerInstances;
-
-	InParams.ActorDescContainerInstanceCollection->ForEachActorDescContainerInstance([&ContentBundleContainerInstances, &BaseAndEDLContainerInstances](const UActorDescContainerInstance* InActorDescContainerInstance)
-	{
-		if (InActorDescContainerInstance->GetContentBundleGuid().IsValid() && !InActorDescContainerInstance->IsContentBundleRequired())
-		{
-			ContentBundleContainerInstances.Add(InActorDescContainerInstance);
-		}
-		else
-		{
-			BaseAndEDLContainerInstances.Add(InActorDescContainerInstance);
-		}
-	});
+	ExtractContentBundleContainerInstances(InParams.ActorDescContainerInstanceCollection, ContentBundleContainerInstances, BaseAndEDLContainerInstances);
 
 	if (!BaseAndEDLContainerInstances.IsEmpty())
 	{
diff --git a/Engine/Source/Runtime/Engine/Public/WorldPartition/ActorDescContainerInstance.h b/Engine/Source/Runtime/Engine/Public/WorldPartition/ActorDescContainerInstance.h
--- a/Engine/Source/Runtime/Engine/Public/WorldPartition/ActorDescContainerInstance.h
+++ b/Engine/Source/Runtime/Engine/Public/WorldPartition/ActorDescContainerInstance.h
@@ -14,7 +14,6 @@ class UActorDescContainer;
 class FWorldPartitionActorDesc;
 class UWorldPartition;
 class UExternalDataLayerAsset;
-class FContentBundleEditor;
 struct FWorldPartitionRuntimeCellPropertyOverride;
 
 class FActorDescInstanceList : public TActorDescList<FWorldPartitionActorDescInstance> { };
@@ -30,14 +29,11 @@ protected:
 		: bIsInitialized(false) 
 		, bRegisteredDelegates(false)
 		, bCreateChildContainerHierarchy(false)
-		, bIsContentBundleRequired(false)
 #endif
 	{}
 
 #if WITH_EDITOR
-	friend class FStreamingGenerationContainerInstanceCollection;
 	friend FWorldPartitionActorDescInstance;
-	friend FContentBundleEditor;
 	friend UWorldPartition;
 		
 public:
@@ -152,9 +148,6 @@ public:
 	ENGINE_API const TMap<FGuid, TObjectPtr<UActorDescContainerInstance>>& GetChildContainerInstances() const { return ChildContainerInstances; }
 		
 	const UActorDescContainerInstance* GetParentContainerInstance() const { return ParentContainerInstance.Get(); }
-
-	bool IsContentBundleRequired() const { return bIsContentBundleRequired; }
-
 protected:
 	virtual void RegisterContainer(const FInitializeParams& InParams);
 	virtual void UnregisterContainer();
@@ -211,6 +204,5 @@ protected:
 	bool														bIsInitialized;
 	bool														bRegisteredDelegates;
 	bool														bCreateChildContainerHierarchy;
-	bool														bIsContentBundleRequired;
 #endif
 };
\ No newline at end of file
diff --git a/Engine/Source/Runtime/Engine/Public/WorldPartition/ContentBundle/ContentBundleDescriptor.h b/Engine/Source/Runtime/Engine/Public/WorldPartition/ContentBundle/ContentBundleDescriptor.h
--- a/Engine/Source/Runtime/Engine/Public/WorldPartition/ContentBundle/ContentBundleDescriptor.h
+++ b/Engine/Source/Runtime/Engine/Public/WorldPartition/ContentBundle/ContentBundleDescriptor.h
@@ -18,7 +18,6 @@ class UContentBundleDescriptor : public UObject
 
 public:
 	const FString& GetDisplayName() const { return DisplayName; }
-	bool GetRequired() const { return bRequired; }
 	const FColor& GetDebugColor() const { return DebugColor; }
 	ENGINE_API FString GetPackageRoot() const;
 	const FGuid& GetGuid() const { return Guid; }
@@ -43,9 +42,6 @@ private:
 	UPROPERTY(EditAnywhere, Category = BaseInformation)
 	FString DisplayName;
 
-	UPROPERTY(EditAnywhere, Category = BaseInformation)
-	bool bRequired;
-
 	UPROPERTY(EditAnywhere, DuplicateTransient, Category = BaseInformation)
 	FColor DebugColor;
 
diff --git a/Engine/Source/Runtime/Engine/Public/WorldPartition/WorldPartitionStreamingGeneration.h b/Engine/Source/Runtime/Engine/Public/WorldPartition/WorldPartitionStreamingGeneration.h
--- a/Engine/Source/Runtime/Engine/Public/WorldPartition/WorldPartitionStreamingGeneration.h
+++ b/Engine/Source/Runtime/Engine/Public/WorldPartition/WorldPartitionStreamingGeneration.h
@@ -243,10 +243,9 @@ public:
 	ENGINE_API FName GetBaseContainerInstancePackageName() const;
 	ENGINE_API FGuid GetContentBundleGuid() const; // @todo_ow : Remove once conversion to ExternalDataLayer is complete. 
 
-	ECollectionType GetCollectionType() const { return CollectionType; }
-
 private:
-	void SetCollectionType(const ECollectionType& InCollectionType) { check(CollectionType == ECollectionType::Invalid); CollectionType = InCollectionType; }	
+	void SetCollectionType(const ECollectionType& InCollectionType) { check(CollectionType == ECollectionType::Invalid); CollectionType = InCollectionType; }
+	ECollectionType GetCollectionType() const { return CollectionType; }
 	void InitializeCollection();
 	FName GetBaseContainerPackageName() const;
 	TArrayView<const UActorDescContainerInstance* const> GetExternalDataLayerContainerInstances() const;
