[Niagara] Move CSVStats for Particle systems into the FParticlePerfStats struct and lazy init when we want to use them to reduce overall FName creation

diff --git a/Engine/Plugins/FX/Niagara/Source/Niagara/Private/NiagaraScalabilityManager.cpp b/Engine/Plugins/FX/Niagara/Source/Niagara/Private/NiagaraScalabilityManager.cpp
--- a/Engine/Plugins/FX/Niagara/Source/Niagara/Private/NiagaraScalabilityManager.cpp
+++ b/Engine/Plugins/FX/Niagara/Source/Niagara/Private/NiagaraScalabilityManager.cpp
@@ -776,7 +776,8 @@ void FNiagaraScalabilityManager::CSVProfilerUpdate(FCsvProfiler* CSVProfiler)
 		if(CompState.bCulled)
 		{
 			UNiagaraSystem* System = Comp->GetAsset();
-			CSVProfiler->RecordCustomStat(System->CSVStat_Culled, CSV_CATEGORY_INDEX(Particles), 1, ECsvCustomStatOp::Accumulate);
+			FParticlePerfStats* Stats = FParticlePerfStatsManager::GetSystemPerfStats(System);
+			CSVProfiler->RecordCustomStat(Stats->CSVStat_Culled, CSV_CATEGORY_INDEX(Particles), 1, ECsvCustomStatOp::Accumulate);
 		}
 	}
 	CSVProfiler->RecordCustomStat(Total, CSV_CATEGORY_INDEX(Particles), NumCulled, ECsvCustomStatOp::Accumulate);
diff --git a/Engine/Source/Runtime/Engine/Classes/Particles/ParticlePerfStats.h b/Engine/Source/Runtime/Engine/Classes/Particles/ParticlePerfStats.h
--- a/Engine/Source/Runtime/Engine/Classes/Particles/ParticlePerfStats.h
+++ b/Engine/Source/Runtime/Engine/Classes/Particles/ParticlePerfStats.h
@@ -270,6 +270,25 @@ struct FParticlePerfStats
 		return GPUStats;
 	}
 
+	//Cached CSV Stat names for this system.
+#if WITH_PARTICLE_PERF_CSV_STATS
+	FName CSVStat_Count = NAME_None;
+	FName CSVStat_Total = NAME_None;
+	FName CSVStat_GTOnly = NAME_None;
+	FName CSVStat_InstAvgGT = NAME_None;
+	FName CSVStat_RT = NAME_None;
+	FName CSVStat_InstAvgRT = NAME_None;
+	FName CSVStat_GPU = NAME_None;
+	FName CSVStat_InstAvgGPU = NAME_None;
+	FName CSVStat_Activation = NAME_None;
+	FName CSVStat_Waits = NAME_None;
+	FName CSVStat_Culled = NAME_None;
+	FName CSVStat_MemoryKB = NAME_None;
+
+	void PopulateStatNames(const FName InName);
+	void ResetStatNames();
+#endif
+
 private:
 	static ENGINE_API FParticlePerfStats* GetWorldPerfStats(const UWorld* World);
 	static ENGINE_API FParticlePerfStats* GetSystemPerfStats(const UFXSystemAsset* FXAsset);
diff --git a/Engine/Source/Runtime/Engine/Classes/Particles/ParticleSystem.h b/Engine/Source/Runtime/Engine/Classes/Particles/ParticleSystem.h
--- a/Engine/Source/Runtime/Engine/Classes/Particles/ParticleSystem.h
+++ b/Engine/Source/Runtime/Engine/Classes/Particles/ParticleSystem.h
@@ -120,8 +120,6 @@ class UFXSystemAsset : public UObject
 public:
 	UFXSystemAsset() {}
 
-	ENGINE_API virtual void PostInitProperties() override;
-
 	/** Max number of components of this system to keep resident in the world component pool. */
 	UPROPERTY(EditAnywhere, Category = Performance, AdvancedDisplay)
 	uint32 MaxPoolSize;
@@ -137,22 +135,6 @@ public:
 
 #if WITH_PER_SYSTEM_PARTICLE_PERF_STATS
 	mutable FParticlePerfStats* ParticlePerfStats = nullptr;
-
-	//Cached CSV Stat names for this system.
-#if WITH_PARTICLE_PERF_CSV_STATS
-	FName CSVStat_Count = NAME_None;
-	FName CSVStat_Total = NAME_None;
-	FName CSVStat_GTOnly = NAME_None;
-	FName CSVStat_InstAvgGT = NAME_None;
-	FName CSVStat_RT = NAME_None;
-	FName CSVStat_InstAvgRT = NAME_None;
-	FName CSVStat_GPU = NAME_None;
-	FName CSVStat_InstAvgGPU = NAME_None;
-	FName CSVStat_Activation = NAME_None;
-	FName CSVStat_Waits = NAME_None;
-	FName CSVStat_Culled = NAME_None;
-	FName CSVStat_MemoryKB = NAME_None;
-#endif
 #endif
 
 	const FGraphEventRef& GetPrecachePSOsEvent() const { return PrecachePSOsEvent; }
diff --git a/Engine/Source/Runtime/Engine/Private/Particles/ParticleComponents.cpp b/Engine/Source/Runtime/Engine/Private/Particles/ParticleComponents.cpp
--- a/Engine/Source/Runtime/Engine/Private/Particles/ParticleComponents.cpp
+++ b/Engine/Source/Runtime/Engine/Private/Particles/ParticleComponents.cpp
@@ -128,26 +128,6 @@ FCustomVersionRegistration GRegisterParticleSystemCustomVersion(FParticleSystemC
 
 //////////////////////////////////////////////////////////////////////////
 
-void UFXSystemAsset::PostInitProperties()
-{
-	Super::PostInitProperties();
-
-#if WITH_PARTICLE_PERF_CSV_STATS
-	CSVStat_Total = *FString::Printf(TEXT("Total/%s"), *GetFName().ToString());
-	CSVStat_GTOnly = *FString::Printf(TEXT("GTOnly/%s"), *GetFName().ToString());
-	CSVStat_InstAvgGT = *FString::Printf(TEXT("InstAvgGT/%s"), *GetFName().ToString());
-	CSVStat_RT = *FString::Printf(TEXT("RT/%s"), *GetFName().ToString());
-	CSVStat_InstAvgRT = *FString::Printf(TEXT("InstAvgRT/%s"), *GetFName().ToString());
-	CSVStat_GPU = *FString::Printf(TEXT("GPU/%s"), *GetFName().ToString());
-	CSVStat_InstAvgGPU = *FString::Printf(TEXT("InstAvgGPU/%s"), *GetFName().ToString());
-	CSVStat_Count = *FString::Printf(TEXT("Count/%s"), *GetFName().ToString());
-	CSVStat_Activation = *FString::Printf(TEXT("Activation/%s"), *GetFName().ToString());
-	CSVStat_Waits = *FString::Printf(TEXT("Waits/%s"), *GetFName().ToString());
-	CSVStat_Culled = *FString::Printf(TEXT("Culled/%s"), *GetFName().ToString());
-	CSVStat_MemoryKB = *FString::Printf(TEXT("MemoryKB/%s"), *GetFName().ToString());
-#endif
-}
-
 void UFXSystemAsset::LaunchPSOPrecaching(const FMaterialInterfacePSOPrecacheParamsList& PSOPrecacheParamsList)
 {
 	FGraphEventArray PrecachePSOsEvents;
diff --git a/Engine/Source/Runtime/Engine/Private/Particles/ParticlePerfStats.cpp b/Engine/Source/Runtime/Engine/Private/Particles/ParticlePerfStats.cpp
--- a/Engine/Source/Runtime/Engine/Private/Particles/ParticlePerfStats.cpp
+++ b/Engine/Source/Runtime/Engine/Private/Particles/ParticlePerfStats.cpp
@@ -318,6 +318,9 @@ void FParticlePerfStatsManager::Tick()
 				}
 				else
 				{
+#if WITH_PARTICLE_PERF_CSV_STATS
+					it.Value()->ResetStatNames();
+#endif
 					FreeSystemStatsPool.Emplace(it.Value().Release());
 					for (FParticlePerfStatsListenerPtr& Listener : Listeners)
 					{
@@ -366,6 +369,44 @@ void FParticlePerfStatsManager::Tick()
 	}
 }
 
+#if WITH_PARTICLE_PERF_CSV_STATS
+void FParticlePerfStats::PopulateStatNames(const FName InName)
+{
+	if (CSVStat_Total == NAME_None)
+	{
+		const FString StatName = InName.ToString();
+		CSVStat_Total = *FString::Printf(TEXT("Total/%s"), *StatName);
+		CSVStat_GTOnly = *FString::Printf(TEXT("GTOnly/%s"), *StatName);
+		CSVStat_InstAvgGT = *FString::Printf(TEXT("InstAvgGT/%s"), *StatName);
+		CSVStat_RT = *FString::Printf(TEXT("RT/%s"), *StatName);
+		CSVStat_InstAvgRT = *FString::Printf(TEXT("InstAvgRT/%s"), *StatName);
+		CSVStat_GPU = *FString::Printf(TEXT("GPU/%s"), *StatName);
+		CSVStat_InstAvgGPU = *FString::Printf(TEXT("InstAvgGPU/%s"), *StatName);
+		CSVStat_Count = *FString::Printf(TEXT("Count/%s"), *StatName);
+		CSVStat_Activation = *FString::Printf(TEXT("Activation/%s"), *StatName);
+		CSVStat_Waits = *FString::Printf(TEXT("Waits/%s"), *StatName);
+		CSVStat_Culled = *FString::Printf(TEXT("Culled/%s"), *StatName);
+		CSVStat_MemoryKB = *FString::Printf(TEXT("MemoryKB/%s"), *StatName);
+	}
+}
+
+void FParticlePerfStats::ResetStatNames()
+{
+	CSVStat_Total = NAME_None;
+	CSVStat_GTOnly = NAME_None;
+	CSVStat_InstAvgGT = NAME_None;
+	CSVStat_RT = NAME_None;
+	CSVStat_InstAvgRT = NAME_None;
+	CSVStat_GPU = NAME_None;
+	CSVStat_InstAvgGPU = NAME_None;
+	CSVStat_Count = NAME_None;
+	CSVStat_Activation = NAME_None;
+	CSVStat_Waits = NAME_None;
+	CSVStat_Culled = NAME_None;
+	CSVStat_MemoryKB = NAME_None;
+}
+#endif
+
 FParticlePerfStats* FParticlePerfStatsManager::GetWorldPerfStats(const UWorld* World)
 {
 	checkSlow(World && FParticlePerfStats::GetGatherWorldStats() && FParticlePerfStats::GetStatsEnabled());
@@ -413,9 +454,12 @@ FParticlePerfStats* FParticlePerfStatsManager::GetSystemPerfStats(const UFXSyste
 			{
 				PerfStats.Reset(new FParticlePerfStats());
 			}
-
 		}
+
 		FXAsset->ParticlePerfStats = PerfStats.Get();
+#if WITH_PARTICLE_PERF_CSV_STATS
+		FXAsset->ParticlePerfStats->PopulateStatNames(FXAsset->GetFName());
+#endif
 
 		for (auto& Listener : Listeners)
 		{
@@ -1195,7 +1239,8 @@ bool FParticlePerfStatsListener_CSVProfiler::Tick()
 					continue;
 				}
 
-				if (FXAsset->CSVStat_MemoryKB.IsNone())
+				FParticlePerfStats* Stats = FParticlePerfStatsManager::GetSystemPerfStats(FXAsset);
+				if (Stats->CSVStat_MemoryKB.IsNone())
 				{
 					continue;
 				}
@@ -1223,7 +1268,8 @@ bool FParticlePerfStatsListener_CSVProfiler::Tick()
 			for (auto OutputIt=MemoryUsage.CreateConstIterator(); OutputIt; ++OutputIt)
 			{
 				const int32 MemoryKB = (int32)FMath::DivideAndRoundUp(OutputIt.Value(), 1024ull);
-				CSVProfiler->RecordCustomStat(OutputIt.Key()->CSVStat_MemoryKB, CSV_CATEGORY_INDEX(Particles), MemoryKB, ECsvCustomStatOp::Set);
+				FParticlePerfStats* Stats = FParticlePerfStatsManager::GetSystemPerfStats(OutputIt.Key());
+				CSVProfiler->RecordCustomStat(Stats->CSVStat_MemoryKB, CSV_CATEGORY_INDEX(Particles), MemoryKB, ECsvCustomStatOp::Set);
 			}
 		}
 
@@ -1241,13 +1287,15 @@ bool FParticlePerfStatsListener_CSVProfiler::Tick()
 						const float Activation = FPlatformTime::ToMilliseconds64(Stats->GetGameThreadStats().ActivationCycles) * 1000.0f;
 						const float Wait = FPlatformTime::ToMilliseconds64(Stats->GetGameThreadStats().WaitCycles) * 1000.0f;
 
-						CSVProfiler->RecordCustomStat(System->CSVStat_Total, CSV_CATEGORY_INDEX(Particles), TotalTime, ECsvCustomStatOp::Set);
-						CSVProfiler->RecordCustomStat(System->CSVStat_GTOnly, CSV_CATEGORY_INDEX(Particles), GTTime, ECsvCustomStatOp::Set);
-						CSVProfiler->RecordCustomStat(System->CSVStat_InstAvgGT, CSV_CATEGORY_INDEX(Particles), AvgTime, ECsvCustomStatOp::Set);
-						CSVProfiler->RecordCustomStat(System->CSVStat_Count, CSV_CATEGORY_INDEX(Particles), Count, ECsvCustomStatOp::Set);
+						Stats->PopulateStatNames(System->GetFName());
 
-						CSVProfiler->RecordCustomStat(System->CSVStat_Activation, CSV_CATEGORY_INDEX(Particles), Activation, ECsvCustomStatOp::Set);
-						CSVProfiler->RecordCustomStat(System->CSVStat_Waits, CSV_CATEGORY_INDEX(Particles), Wait, ECsvCustomStatOp::Set);
+						CSVProfiler->RecordCustomStat(Stats->CSVStat_Total, CSV_CATEGORY_INDEX(Particles), TotalTime, ECsvCustomStatOp::Set);
+						CSVProfiler->RecordCustomStat(Stats->CSVStat_GTOnly, CSV_CATEGORY_INDEX(Particles), GTTime, ECsvCustomStatOp::Set);
+						CSVProfiler->RecordCustomStat(Stats->CSVStat_InstAvgGT, CSV_CATEGORY_INDEX(Particles), AvgTime, ECsvCustomStatOp::Set);
+						CSVProfiler->RecordCustomStat(Stats->CSVStat_Count, CSV_CATEGORY_INDEX(Particles), Count, ECsvCustomStatOp::Set);
+
+						CSVProfiler->RecordCustomStat(Stats->CSVStat_Activation, CSV_CATEGORY_INDEX(Particles), Activation, ECsvCustomStatOp::Set);
+						CSVProfiler->RecordCustomStat(Stats->CSVStat_Waits, CSV_CATEGORY_INDEX(Particles), Wait, ECsvCustomStatOp::Set);
 					}
 				}
 			}
@@ -1270,15 +1318,17 @@ void FParticlePerfStatsListener_CSVProfiler::TickRT()
 			{
 				if (const UFXSystemAsset* System = WeakSystem.Get())
 				{
+					Stats->PopulateStatNames(System->GetFName());
+
 					const float RTTime = FPlatformTime::ToMilliseconds64(Stats->GetRenderThreadStats().GetTotalCycles()) * 1000.0f;
 					const float RTAvgTime = FPlatformTime::ToMilliseconds64(Stats->GetRenderThreadStats().GetPerInstanceAvgCycles()) * 1000.0f;
-					CSVProfiler->RecordCustomStat(System->CSVStat_RT, CSV_CATEGORY_INDEX(Particles), RTTime, ECsvCustomStatOp::Set);
-					CSVProfiler->RecordCustomStat(System->CSVStat_InstAvgRT, CSV_CATEGORY_INDEX(Particles), RTAvgTime, ECsvCustomStatOp::Set);
+					CSVProfiler->RecordCustomStat(Stats->CSVStat_RT, CSV_CATEGORY_INDEX(Particles), RTTime, ECsvCustomStatOp::Set);
+					CSVProfiler->RecordCustomStat(Stats->CSVStat_InstAvgRT, CSV_CATEGORY_INDEX(Particles), RTAvgTime, ECsvCustomStatOp::Set);
 
 					const float GpuTime = float(Stats->GetGPUStats().GetTotalMicroseconds());
 					const float GpuAvgTime = float(Stats->GetGPUStats().GetPerInstanceAvgMicroseconds());
-					CSVProfiler->RecordCustomStat(System->CSVStat_GPU, CSV_CATEGORY_INDEX(Particles), GpuTime, ECsvCustomStatOp::Set);
-					CSVProfiler->RecordCustomStat(System->CSVStat_InstAvgGPU, CSV_CATEGORY_INDEX(Particles), GpuAvgTime, ECsvCustomStatOp::Set);
+					CSVProfiler->RecordCustomStat(Stats->CSVStat_GPU, CSV_CATEGORY_INDEX(Particles), GpuTime, ECsvCustomStatOp::Set);
+					CSVProfiler->RecordCustomStat(Stats->CSVStat_InstAvgGPU, CSV_CATEGORY_INDEX(Particles), GpuAvgTime, ECsvCustomStatOp::Set);
 				}
 			}
 		);
