Fix: Can't save materials while ODSC server is open

Close material file handles after remote recompile

diff --git a/Engine/Source/Runtime/Engine/Private/ShaderCompiler/ShaderCompiler.cpp b/Engine/Source/Runtime/Engine/Private/ShaderCompiler/ShaderCompiler.cpp
--- a/Engine/Source/Runtime/Engine/Private/ShaderCompiler/ShaderCompiler.cpp
+++ b/Engine/Source/Runtime/Engine/Private/ShaderCompiler/ShaderCompiler.cpp
@@ -57,6 +57,8 @@
 #include "StereoRenderUtils.h"
 #include "Tasks/Task.h"
 #include "UObject/DevObjectVersion.h"
+#include "UObject/Linker.h"
+#include "UObject/Package.h"
 #include "UObject/UObjectIterator.h"
 #include "Math/UnitConversion.h"
 #include "UnrealEngine.h"
@@ -10272,6 +10274,7 @@ void RecompileShadersForRemote(
 					TStrongObjectPtr<UMaterialInterface> MaterialInterface = TStrongObjectPtr<UMaterialInterface>(FindObject<UMaterialInterface>(nullptr, *Iter.Key));
 					if (MaterialInterface)
 					{
+						ResetLoaders(MaterialInterface->GetPackage());
 						Args.LoadedMaterialsToRecompile->Add(MaterialInterface);
 					}
 				}
@@ -10356,6 +10359,15 @@ void RecompileShadersForRemote(
 		}
 	}
 
+	for (UMaterialInterface* MaterialInterface : MaterialsToCompile)
+	{
+		if (MaterialInterface)
+		{
+			ResetLoaders(MaterialInterface->GetPackage());
+		}
+	}
+	GEngine->ForceGarbageCollection(true);
+
 	UE_LOG(LogShaders, Display, TEXT(""));
 	UE_LOG(LogShaders, Display, TEXT("Compiled %u shaders in %.2f seconds."), GShaderCompilerStats->GetTotalShadersCompiled() - StartTotalShadersCompiled, FPlatformTime::Seconds() - StartTime);
 
