autortfm:
- when running under autortfm, we now conditionally take the SerializeAndPostLoadCritical only when the async loading thread is enabled
- we still always take the lock when not running under autortfm (!AutoRTFM::IsClosed()) to preseve the existing behavior in uninstrumented code
- this is expected to work for us for now because on the server, the async loading thread is disabled, so we can avoid taking this lock

diff --git a/Engine/Source/Runtime/Core/Public/Misc/ScopeLock.h b/Engine/Source/Runtime/Core/Public/Misc/ScopeLock.h
--- a/Engine/Source/Runtime/Core/Public/Misc/ScopeLock.h
+++ b/Engine/Source/Runtime/Core/Public/Misc/ScopeLock.h
@@ -176,6 +176,42 @@ namespace UE
 		MutexType* Mutex;
 	};
 
+	// RAII-style scope locking of a synchronisation primitive. Same
+	// as TScopeLock except taking the lock is conditional
+	template<typename MutexType>
+	class TConditionalScopeLock
+	{
+	public:
+		UE_NONCOPYABLE(TConditionalScopeLock);
+
+		UE_NODISCARD_CTOR TConditionalScopeLock(MutexType& InMutex, bool bShouldLock)
+			: Mutex(bShouldLock ? &InMutex : nullptr)
+		{
+			if (bShouldLock)
+			{
+				check(Mutex);
+				Mutex->Lock();
+			}
+		}
+
+		~TConditionalScopeLock()
+		{
+			Unlock();
+		}
+
+		void Unlock()
+		{
+			if (Mutex)
+			{
+				Mutex->Unlock();
+				Mutex = nullptr;
+			}
+		}
+
+	private:
+		MutexType* Mutex;
+	};
+
 	// RAII-style scope unlocking of a synchronisation primitive
 	// `MutexType` is required to implement `Lock` and `Unlock` methods
 	// Example:
diff --git a/Engine/Source/Runtime/Engine/Private/BlueprintGeneratedClass.cpp b/Engine/Source/Runtime/Engine/Private/BlueprintGeneratedClass.cpp
--- a/Engine/Source/Runtime/Engine/Private/BlueprintGeneratedClass.cpp
+++ b/Engine/Source/Runtime/Engine/Private/BlueprintGeneratedClass.cpp
@@ -31,6 +31,8 @@
 #include "UObject/SparseClassDataUtils.h"
 #include "UObject/UE5MainStreamObjectVersion.h"
 #include "GenericPlatform/GenericPlatformCrashContext.h"
+#include "AutoRTFM/AutoRTFM.h"
+#include "Serialization/AsyncPackageLoader.h"
 
 #include UE_INLINE_GENERATED_CPP_BY_NAME(BlueprintGeneratedClass)
 
@@ -1220,7 +1222,13 @@ void UBlueprintGeneratedClass::SetupObjectInitializer(FObjectInitializer& Object
 
 void UBlueprintGeneratedClass::InitPropertiesFromCustomList(uint8* DataPtr, const uint8* DefaultDataPtr)
 {
-	FScopeLock SerializeAndPostLoadLock(&SerializeAndPostLoadCritical);
+	// autortfm: we've introduced a conditional lock here because while running under autortfm, we
+	// only want to take this lock if necessary (if FAsyncLoadingThreadSettings says we're multithreaded)
+	// if not running under autortfm, IsClosed is false and we will always lock without querying whether
+	// we are running multithreaded.
+	UE::TConditionalScopeLock SerializeAndPostLoadLock(
+		SerializeAndPostLoadCritical, 
+		(!AutoRTFM::IsClosed()) || FAsyncLoadingThreadSettings::Get().bAsyncLoadingThreadEnabled);
 
 	if (GBlueprintNativePropertyInitFastPathDisabled
 		|| !ensureMsgf(bCustomPropertyListForPostConstructionInitialized, TEXT("Custom Property List Not Initialized for %s"), *GetPathNameSafe(this))) // Something went wrong, probably a race condition
