Add new Kismet function to allow external linking

#rnx
#Tests Ran game in a single client with debug node set in BP to call the new LaunchExternalURL function. Clicked the button and confirmed the link opened to the appropriate page.

diff --git a/Engine/Source/Runtime/Engine/Classes/Kismet/KismetSystemLibrary.h b/Engine/Source/Runtime/Engine/Classes/Kismet/KismetSystemLibrary.h
--- a/Engine/Source/Runtime/Engine/Classes/Kismet/KismetSystemLibrary.h
+++ b/Engine/Source/Runtime/Engine/Classes/Kismet/KismetSystemLibrary.h
@@ -1677,6 +1677,10 @@ class UKismetSystemLibrary : public UBlueprintFunctionLibrary
 	UFUNCTION(BlueprintCallable, Category = "Utilities|Platform")
 	static ENGINE_API void LaunchURL(const FString& URL);
 
+	// Opens an external URL in the platform's web browser of choice if it meets the allowlist of passed in domains
+	UFUNCTION(BlueprintCallable, Category = "Utilities|Platform")
+	static ENGINE_API void LaunchExternalUrl(const TArray<FString>& InDomainStrings, const FString& URL);
+
 	UFUNCTION(BlueprintCallable, Category = "Utilities|Platform")
 	static ENGINE_API bool CanLaunchURL(const FString& URL);
 
diff --git a/Engine/Source/Runtime/Engine/Private/KismetSystemLibrary.cpp b/Engine/Source/Runtime/Engine/Private/KismetSystemLibrary.cpp
--- a/Engine/Source/Runtime/Engine/Private/KismetSystemLibrary.cpp
+++ b/Engine/Source/Runtime/Engine/Private/KismetSystemLibrary.cpp
@@ -2770,6 +2770,18 @@ void UKismetSystemLibrary::LaunchURL(const FString& URL)
 	}
 }
 
+void UKismetSystemLibrary::LaunchExternalUrl(const TArray<FString>& DomainStrings, const FString& URL)
+{
+	if (!URL.IsEmpty())
+	{
+		UE::Core::FURLRequestFilter::FRequestMap AllowedDomains;
+		AllowedDomains.Add(TEXT("http"), DomainStrings);
+		UE::Core::FURLRequestFilter Filter(AllowedDomains);
+
+		FPlatformProcess::LaunchURLFiltered(*URL, nullptr, nullptr, Filter);
+	}
+}
+
 bool UKismetSystemLibrary::CanLaunchURL(const FString& URL)
 {
 	if (!URL.IsEmpty())
