Adding missing object references in FActiveSound to fix possible GC issues.

[RN] Fix for missing GC references in FActiveSound

diff --git a/Engine/Source/Runtime/Engine/Private/ActiveSound.cpp b/Engine/Source/Runtime/Engine/Private/ActiveSound.cpp
--- a/Engine/Source/Runtime/Engine/Private/ActiveSound.cpp
+++ b/Engine/Source/Runtime/Engine/Private/ActiveSound.cpp
@@ -245,6 +245,40 @@ void FActiveSound::AddReferencedObjects(FReferenceCollector& Collector)
 			Collector.AddReferencedObject(const_cast<TObjectPtr<UObject>&>(*Object));
 		}
 	}
+
+	// Adding ref'd objects in AttenuationSettings
+	Collector.AddReferencedObject(AttenuationSettings.AudioLinkSettingsOverride);
+	Collector.AddReferencedObjects(AttenuationSettings.PluginSettings.OcclusionPluginSettingsArray);
+	Collector.AddReferencedObjects(AttenuationSettings.PluginSettings.ReverbPluginSettingsArray);
+	Collector.AddReferencedObjects(AttenuationSettings.PluginSettings.SourceDataOverridePluginSettingsArray);
+	Collector.AddReferencedObjects(AttenuationSettings.PluginSettings.SpatializationPluginSettingsArray);
+
+	for (FAudioVolumeSubmixSendSettings& SendSettings : AudioVolumeSubmixSendSettings)
+	{
+		for (FSoundSubmixSendInfo& SendInfo : SendSettings.SubmixSends)
+		{
+			Collector.AddReferencedObject(SendInfo.SoundSubmix);
+		}
+	}
+
+	for (FAudioVolumeSubmixSendSettings& SendSettings : PreviousAudioVolumeSubmixSendSettings)
+	{
+		for (FSoundSubmixSendInfo& SendInfo : SendSettings.SubmixSends)
+		{
+			Collector.AddReferencedObject(SendInfo.SoundSubmix);
+		}
+	}
+
+	for (TTuple<EBusSendType, FSoundSourceBusSendInfo>& BusSendTuple : NewBusSends)
+	{
+		Collector.AddReferencedObject(BusSendTuple.Value.AudioBus);
+		Collector.AddReferencedObject(BusSendTuple.Value.SoundSourceBus);
+	}
+
+	Collector.AddReferencedObjects(ModulationRouting.VolumeModulationDestination.Modulators);
+	Collector.AddReferencedObjects(ModulationRouting.PitchModulationDestination.Modulators);
+	Collector.AddReferencedObjects(ModulationRouting.HighpassModulationDestination.Modulators);
+	Collector.AddReferencedObjects(ModulationRouting.LowpassModulationDestination.Modulators);
 }
 
 int32 FActiveSound::GetPlayCount() const
