Initialize FStaticMeshVertexFactories using FStaticMeshVertexBuffers instead of FStaticMeshLODResources.
- required to support ray tracing proxies which will use FStaticMeshVertexBuffers but not FStaticMeshLODResources.

diff --git a/Engine/Source/Runtime/Engine/Private/StaticMesh.cpp b/Engine/Source/Runtime/Engine/Private/StaticMesh.cpp
--- a/Engine/Source/Runtime/Engine/Private/StaticMesh.cpp
+++ b/Engine/Source/Runtime/Engine/Private/StaticMesh.cpp
@@ -897,19 +897,19 @@ int32 FStaticMeshLODResources::GetNumTexCoords() const
 }
 
 void FStaticMeshVertexFactories::InitVertexFactory(
-	const FStaticMeshLODResources& LodResources,
+	const FStaticMeshVertexBuffers& VertexBuffers,
 	FLocalVertexFactory& InOutVertexFactory,
 	uint32 LODIndex,
 	const UStaticMesh* InParentMesh,
 	bool bInOverrideColorVertexBuffer
 	)
 {
-	check( InParentMesh != nullptr);
+	check(InParentMesh != nullptr);
 
 	struct InitStaticMeshVertexFactoryParams
 	{
 		FLocalVertexFactory* VertexFactory;
-		const FStaticMeshLODResources* LODResources;
+		const FStaticMeshVertexBuffers* VertexBuffers;
 	#if WITH_EDITORONLY_DATA
 		const UStaticMesh* StaticMesh;
 	#endif
@@ -922,10 +922,10 @@ void FStaticMeshVertexFactories::InitVertexFactory(
 	} Params;
 
 	uint32 LightMapCoordinateIndex = (uint32)InParentMesh->GetLightMapCoordinateIndex();
-	LightMapCoordinateIndex = LightMapCoordinateIndex < LodResources.VertexBuffers.StaticMeshVertexBuffer.GetNumTexCoords() ? LightMapCoordinateIndex : LodResources.VertexBuffers.StaticMeshVertexBuffer.GetNumTexCoords() - 1;
+	LightMapCoordinateIndex = LightMapCoordinateIndex < VertexBuffers.StaticMeshVertexBuffer.GetNumTexCoords() ? LightMapCoordinateIndex : VertexBuffers.StaticMeshVertexBuffer.GetNumTexCoords() - 1;
 
 	Params.VertexFactory				= &InOutVertexFactory;
-	Params.LODResources					= &LodResources;
+	Params.VertexBuffers				= &VertexBuffers;
 	Params.bOverrideColorVertexBuffer	= bInOverrideColorVertexBuffer;
 	Params.LightMapCoordinateIndex		= LightMapCoordinateIndex;
 	Params.LODIndex						= LODIndex;
@@ -940,10 +940,10 @@ void FStaticMeshVertexFactories::InitVertexFactory(
 		{
 			FLocalVertexFactory::FDataType Data;
 
-			Params.LODResources->VertexBuffers.PositionVertexBuffer.BindPositionVertexBuffer(Params.VertexFactory, Data);
-			Params.LODResources->VertexBuffers.StaticMeshVertexBuffer.BindTangentVertexBuffer(Params.VertexFactory, Data);
-			Params.LODResources->VertexBuffers.StaticMeshVertexBuffer.BindPackedTexCoordVertexBuffer(Params.VertexFactory, Data);
-			Params.LODResources->VertexBuffers.StaticMeshVertexBuffer.BindLightMapVertexBuffer(Params.VertexFactory, Data, Params.LightMapCoordinateIndex);
+			Params.VertexBuffers->PositionVertexBuffer.BindPositionVertexBuffer(Params.VertexFactory, Data);
+			Params.VertexBuffers->StaticMeshVertexBuffer.BindTangentVertexBuffer(Params.VertexFactory, Data);
+			Params.VertexBuffers->StaticMeshVertexBuffer.BindPackedTexCoordVertexBuffer(Params.VertexFactory, Data);
+			Params.VertexBuffers->StaticMeshVertexBuffer.BindLightMapVertexBuffer(Params.VertexFactory, Data, Params.LightMapCoordinateIndex);
 
 			// bOverrideColorVertexBuffer means we intend to override the color later.  We must construct the vertexfactory such that it believes a proper stride (not 0) is set for
 			// the color stream so that the real stream works later.
@@ -954,7 +954,7 @@ void FStaticMeshVertexFactories::InitVertexFactory(
 			//otherwise just bind the incoming buffer directly.
 			else
 			{
-				Params.LODResources->VertexBuffers.ColorVertexBuffer.BindColorVertexBuffer(Params.VertexFactory, Data);
+				Params.VertexBuffers->ColorVertexBuffer.BindColorVertexBuffer(Params.VertexFactory, Data);
 			}
 
 			Data.LODLightmapDataIndex	= Params.LODIndex;
@@ -967,12 +967,12 @@ void FStaticMeshVertexFactories::InitVertexFactory(
 		});
 }
 
-void FStaticMeshVertexFactories::InitResources(const FStaticMeshLODResources& LodResources, uint32 LODIndex, const UStaticMesh* Parent)
+void FStaticMeshVertexFactories::InitResources(const FStaticMeshVertexBuffers& VertexBuffers, uint32 LODIndex, const UStaticMesh* Parent)
 {
-	InitVertexFactory(LodResources, VertexFactory, LODIndex, Parent, false);
+	InitVertexFactory(VertexBuffers, VertexFactory, LODIndex, Parent, false);
 	BeginInitResource(&VertexFactory);
 
-	InitVertexFactory(LodResources, VertexFactoryOverrideColorVertexBuffer, LODIndex, Parent, true);
+	InitVertexFactory(VertexBuffers, VertexFactoryOverrideColorVertexBuffer, LODIndex, Parent, true);
 	BeginInitResource(&VertexFactoryOverrideColorVertexBuffer);
 }
 
@@ -1937,7 +1937,7 @@ void FStaticMeshRenderData::InitResources(ERHIFeatureLevel::Type InFeatureLevel,
 		if (LODResources[LODIndex].VertexBuffers.StaticMeshVertexBuffer.GetNumVertices() > 0)
 		{
 			LODResources[LODIndex].InitResources(Owner, LODIndex);
-			LODVertexFactories[LODIndex].InitResources(LODResources[LODIndex], LODIndex, Owner);
+			LODVertexFactories[LODIndex].InitResources(LODResources[LODIndex].VertexBuffers, LODIndex, Owner);
 		}
 	}
 
diff --git a/Engine/Source/Runtime/Engine/Public/StaticMeshResources.h b/Engine/Source/Runtime/Engine/Public/StaticMeshResources.h
--- a/Engine/Source/Runtime/Engine/Public/StaticMeshResources.h
+++ b/Engine/Source/Runtime/Engine/Public/StaticMeshResources.h
@@ -654,10 +654,22 @@ struct FStaticMeshVertexFactories
 	* @param	InParentMesh					Parent static mesh
 	* @param	bInOverrideColorVertexBuffer	If true, make a vertex factory ready for per-instance colors
 	*/
-	ENGINE_API void InitVertexFactory(const FStaticMeshLODResources& LodResources, FLocalVertexFactory& InOutVertexFactory, uint32 LODIndex, const UStaticMesh* InParentMesh, bool bInOverrideColorVertexBuffer);
+	ENGINE_API void InitVertexFactory(const FStaticMeshVertexBuffers& VertexBuffers, FLocalVertexFactory& InOutVertexFactory, uint32 LODIndex, const UStaticMesh* InParentMesh, bool bInOverrideColorVertexBuffer);
+
+	UE_DEPRECATED(5.5, "Provide FStaticMeshVertexBuffers instead of FStaticMeshLODResources")
+	ENGINE_API void InitVertexFactory(const FStaticMeshLODResources& LodResources, FLocalVertexFactory& InOutVertexFactory, uint32 LODIndex, const UStaticMesh* InParentMesh, bool bInOverrideColorVertexBuffer)
+	{
+		InitVertexFactory(LodResources.VertexBuffers, InOutVertexFactory, LODIndex, InParentMesh, bInOverrideColorVertexBuffer);
+	}
 
 	/** Initializes all rendering resources. */
-	ENGINE_API void InitResources(const FStaticMeshLODResources& LodResources, uint32 LODIndex, const UStaticMesh* Parent);
+	ENGINE_API void InitResources(const FStaticMeshVertexBuffers& VertexBuffers, uint32 LODIndex, const UStaticMesh* Parent);
+
+	UE_DEPRECATED(5.5, "Provide FStaticMeshVertexBuffers instead of FStaticMeshLODResources")
+	ENGINE_API void InitResources(const FStaticMeshLODResources& LodResources, uint32 LODIndex, const UStaticMesh* Parent)
+	{
+		InitResources(LodResources.VertexBuffers, LODIndex, Parent);
+	}
 
 	/** Releases all rendering resources. */
 	ENGINE_API void ReleaseResources();
