PR #11908: ExposeOnSpawn Class properties do not respect "none" as the value on "Construct object from class" node

diff --git a/Engine/Source/Editor/KismetCompiler/Private/KismetCompilerMisc.cpp b/Engine/Source/Editor/KismetCompiler/Private/KismetCompilerMisc.cpp
--- a/Engine/Source/Editor/KismetCompiler/Private/KismetCompilerMisc.cpp
+++ b/Engine/Source/Editor/KismetCompiler/Private/KismetCompilerMisc.cpp
@@ -777,9 +777,8 @@ UEdGraphPin* FKismetCompilerUtilities::GenerateAssignmentNodes(class FKismetComp
 	{
 		// Only create 'set param by name' node if this pin is linked to something
 		UEdGraphPin* OrgPin = SpawnNode->Pins[PinIdx];
-		const bool bHasDefaultValue = !OrgPin->DefaultValue.IsEmpty() || !OrgPin->DefaultTextValue.IsEmpty() || OrgPin->DefaultObject;
 		if (NULL == CallBeginSpawnNode->FindPin(OrgPin->PinName) &&
-			(OrgPin->LinkedTo.Num() > 0 || bHasDefaultValue))
+			(OrgPin->LinkedTo.Num() > 0 || !OrgPin->DoesDefaultValueMatchAutogenerated()))
 		{
 			FProperty* Property = FindFProperty<FProperty>(ForClass, OrgPin->PinName);
 			// NULL property indicates that this pin was part of the original node, not the 
diff --git a/Engine/Source/Runtime/Engine/Private/KismetSystemLibrary.cpp b/Engine/Source/Runtime/Engine/Private/KismetSystemLibrary.cpp
--- a/Engine/Source/Runtime/Engine/Private/KismetSystemLibrary.cpp
+++ b/Engine/Source/Runtime/Engine/Private/KismetSystemLibrary.cpp
@@ -1163,24 +1163,28 @@ void UKismetSystemLibrary::SetBoolPropertyByName(UObject* Object, FName Property
 
 void UKismetSystemLibrary::SetObjectPropertyByName(UObject* Object, FName PropertyName, UObject* Value)
 {
-	if(Object != NULL && Value != NULL)
+	if (Object)
 	{
-		FObjectPropertyBase* ObjectProp = FindFProperty<FObjectPropertyBase>(Object->GetClass(), PropertyName);
-		if(ObjectProp != NULL && Value->IsA(ObjectProp->PropertyClass)) // check it's the right type
+		if (FObjectPropertyBase* ObjectProp = FindFProperty<FObjectPropertyBase>(Object->GetClass(), PropertyName))
 		{
-			ObjectProp->SetObjectPropertyValue_InContainer(Object, Value);
-		}		
+			if (!Value || Value->IsA(ObjectProp->PropertyClass)) // check it's the right type
+			{
+				ObjectProp->SetObjectPropertyValue_InContainer(Object, Value);
+			}
+		}
 	}
 }
 
 void UKismetSystemLibrary::SetClassPropertyByName(UObject* Object, FName PropertyName, TSubclassOf<UObject> Value)
 {
-	if (Object && *Value)
+	if (Object)
 	{
-		FClassProperty* ClassProp = FindFProperty<FClassProperty>(Object->GetClass(), PropertyName);
-		if (ClassProp != NULL && Value->IsChildOf(ClassProp->MetaClass)) // check it's the right type
+		if (FClassProperty* ClassProp = FindFProperty<FClassProperty>(Object->GetClass(), PropertyName))
 		{
-			ClassProp->SetObjectPropertyValue_InContainer(Object, *Value);
+			if (!*Value || Value->IsChildOf(ClassProp->MetaClass)) // check it's the right type
+			{
+				ClassProp->SetObjectPropertyValue_InContainer(Object, *Value);
+			}
 		}
 	}
 }
