ALODActor should always be considered HLOD relevant

diff --git a/Engine/Source/Runtime/Engine/Private/Components/PrimitiveComponent.cpp b/Engine/Source/Runtime/Engine/Private/Components/PrimitiveComponent.cpp
--- a/Engine/Source/Runtime/Engine/Private/Components/PrimitiveComponent.cpp
+++ b/Engine/Source/Runtime/Engine/Private/Components/PrimitiveComponent.cpp
@@ -4890,6 +4890,11 @@ FPrimitiveMaterialPropertyDescriptor UPrimitiveComponent::GetUsedMaterialPropert
 #if WITH_EDITOR
 const bool UPrimitiveComponent::ShouldGenerateAutoLOD(const int32 HierarchicalLevelIndex) const
 {	
+	if (GetOwner() && GetOwner()->IsA<ALODActor>())
+	{
+		return true;
+	}
+
 	if (!IsHLODRelevant())
 	{
 		return false;
