New mode for ODSC recompileshaders, driven by ODSC.recompilemode (default=0, ie we're not enabling it by default currently)

This allows partial shader maps on materials, and will only compile the permutations actually used by rendering. This disables PSO precaching, otherwise we would request too many shaders.

The aggregation of the shaders happens on the server. Because of that, we need to keep in memory on the server the materials we want, because otherwise we will lose the requests that happened in the past. When a client disconnects, we release the references to these materials.

This mode is still disabled by default because we need to re-create MDCs whenever we change the shadermap. This can create some sever hitches when a lot of material instances are changed. Some throttling needs to be added

diff --git a/Engine/Source/Runtime/Engine/Private/Materials/MaterialShader.cpp b/Engine/Source/Runtime/Engine/Private/Materials/MaterialShader.cpp
--- a/Engine/Source/Runtime/Engine/Private/Materials/MaterialShader.cpp
+++ b/Engine/Source/Runtime/Engine/Private/Materials/MaterialShader.cpp
@@ -36,6 +36,10 @@
 #include "PSOPrecacheValidation.h"
 #include "ProfilingDebugging/CsvProfiler.h"
 
+#if WITH_ODSC
+#include "ODSC/ODSCManager.h"
+#endif
+
 #if WITH_EDITOR
 #include "Algo/Sort.h"
 #include "Containers/AnsiString.h"
@@ -3385,6 +3389,11 @@ void FMaterialShaderMap::RegisterForODSC(EShaderPlatform InShaderPlatform)
 				bRegistered = true;
 			}
 		}
+
+#if WITH_ODSC
+		bIsFromODSC = true;
+		FODSCManager::RegisterMaterialShaderMap(*this);
+#endif
 	}
 }
 
@@ -3455,6 +3464,9 @@ FMaterialShaderMap::FMaterialShaderMap() :
 	bCompilationFinalized(true),
 	bCompiledSuccessfully(true),
 	bIsPersistent(false)
+#if WITH_ODSC
+	, bIsFromODSC(false)
+#endif
 {
 	checkSlow(IsInGameThread() || IsAsyncLoading());
 #if ALLOW_SHADERMAP_DEBUG_DATA
diff --git a/Engine/Source/Runtime/Engine/Private/Materials/MaterialShared.cpp b/Engine/Source/Runtime/Engine/Private/Materials/MaterialShared.cpp
--- a/Engine/Source/Runtime/Engine/Private/Materials/MaterialShared.cpp
+++ b/Engine/Source/Runtime/Engine/Private/Materials/MaterialShared.cpp
@@ -3676,6 +3676,12 @@ bool FMaterial::TryGetShaders(const FMaterialShaderTypes& InTypes, const FVertex
 		return false;
 	}
 
+#if WITH_ODSC
+	bool bShouldForceRecompile = FODSCManager::ShouldForceRecompile(ShaderMap, this) && !IsDefaultMaterial() && !IsUsedWithLandscape();
+#else
+	constexpr bool bShouldForceRecompile = false;
+#endif
+
 	OutShaders.ShaderMap = ShaderMap;
 	const EShaderPlatform ShaderPlatform = ShaderMap->GetShaderPlatform();
 	const EShaderPermutationFlags PermutationFlags = ShaderMap->GetPermutationFlags();
@@ -3731,11 +3737,14 @@ bool FMaterial::TryGetShaders(const FMaterialShaderTypes& InTypes, const FVertex
 				}
 			}
 		}
-		else
+        // we don't do 'else' here because when bShouldForceRecompile is true, we still want to use the current 
+        // pipeline until we have a new one ready. The ODSC server might fail to find the right shader, and this might
+        // skew results when doing some A/B comparisons		
+		if (Pipeline == nullptr || bShouldForceRecompile)
 		{
 			if (InTypes.PipelineType->ShouldOptimizeUnusedOutputs(ShaderPlatform))
 			{
-				bMissingShader = true;
+				bMissingShader = (Pipeline == nullptr);
 
 #if WITH_EDITOR || WITH_ODSC
 				for (const FShaderType* ShaderType : InTypes.PipelineType->GetStages())
@@ -3758,12 +3767,15 @@ bool FMaterial::TryGetShaders(const FMaterialShaderTypes& InTypes, const FVertex
 						const FString VFTypeName(InVertexFactoryType ? InVertexFactoryType->GetName() : TEXT(""));
 						const FString PipelineName(InTypes.PipelineType->GetName());
 						TArray<FString> ShaderStageNamesToCompile;
+						TArray<FShaderId> RequestShaderIds;
 						for (auto* ShaderType : InTypes.PipelineType->GetStages())
 						{
 							ShaderStageNamesToCompile.Add(ShaderType->GetName());
+							RequestShaderIds.Add(FShaderId(ShaderType, ShaderMap->GetShaderMapId().CookedShaderMapIdHash, InTypes.PipelineType->GetHashedName(), InVertexFactoryType, kUniqueShaderPermutationId, ShaderPlatform));
 						}
 
-						GODSCManager->AddThreadedShaderPipelineRequest(ShaderPlatform, GetFeatureLevel(), GetQualityLevel(), MaterialName, VFTypeName, PipelineName, ShaderStageNamesToCompile, kUniqueShaderPermutationId);
+						GODSCManager->AddThreadedShaderPipelineRequest(ShaderPlatform, GetFeatureLevel(), GetQualityLevel(), MaterialName, VFTypeName, PipelineName, 
+						                                               ShaderStageNamesToCompile, kUniqueShaderPermutationId, RequestShaderIds);
 					}
 				}
 				else
@@ -3829,9 +3841,12 @@ bool FMaterial::TryGetShaders(const FMaterialShaderTypes& InTypes, const FVertex
 				{
 					OutShaders.Shaders[FrequencyIndex] = Shader;
 				}
-				else
+		        // we don't do 'else' here because when bShouldForceRecompile is true, we still want to use the current 
+		        // shader until we have a new one ready. The ODSC server might fail to find the right shader, and this might
+		        // skew results when doing some A/B comparisons		
+				if (Shader == nullptr || bShouldForceRecompile)
 				{
-					bMissingShader = true;
+					bMissingShader = (Shader == nullptr);
 
 #if WITH_EDITOR || WITH_ODSC
 					if (!ShouldCacheShaderType(ShaderType, InVertexFactoryType, PermutationId))
@@ -3851,7 +3866,11 @@ bool FMaterial::TryGetShaders(const FMaterialShaderTypes& InTypes, const FVertex
 							TArray<FString> ShaderStageNamesToCompile;
 							ShaderStageNamesToCompile.Add(ShaderType->GetName());
 
-							GODSCManager->AddThreadedShaderPipelineRequest(ShaderPlatform, GetFeatureLevel(), GetQualityLevel(), MaterialName, VFTypeName, PipelineName, ShaderStageNamesToCompile, PermutationId);
+							TArray<FShaderId> RequestShaderIds;
+							RequestShaderIds.Add(FShaderId(ShaderType, ShaderMap->GetShaderMapId().CookedShaderMapIdHash, FHashedName(), InVertexFactoryType, PermutationId, ShaderPlatform));
+
+							GODSCManager->AddThreadedShaderPipelineRequest(ShaderPlatform, GetFeatureLevel(), GetQualityLevel(), MaterialName, VFTypeName, PipelineName, 
+							                                               ShaderStageNamesToCompile, PermutationId, RequestShaderIds);
 						}
 					}
 					else
diff --git a/Engine/Source/Runtime/Engine/Private/ODSC/ODSCManager.cpp b/Engine/Source/Runtime/Engine/Private/ODSC/ODSCManager.cpp
--- a/Engine/Source/Runtime/Engine/Private/ODSC/ODSCManager.cpp
+++ b/Engine/Source/Runtime/Engine/Private/ODSC/ODSCManager.cpp
@@ -6,14 +6,31 @@
 #include "ODSCLog.h"
 #include "ODSCThread.h"
 #include "Containers/BackgroundableTicker.h"
+#include "MaterialShared.h"
 #include "Materials/MaterialInstance.h"
+#include "Materials/Material.h"
 
 DEFINE_LOG_CATEGORY(LogODSC);
 
 // FODSCManager
 
+static TAutoConsoleVariable<int32> CVarODSCRecompileMode(
+	TEXT("ODSC.recompilemode"),
+	0,
+	TEXT("Highly experimental - Changes how recompileshaders behaves in cooked build\n")
+	TEXT("0 (default): Gathers all visible materials in a single frame and compiles all permutations for them\n")
+	TEXT("1: Compile only the permutations that are requested by the renderer. Faster iteration but more prone to hitching because of MDC recaching\n")
+	);
+
 FODSCManager* GODSCManager = nullptr;
 
+namespace ODSCManagerPrivate
+{
+#if WITH_ODSC
+	static thread_local int32 ODSCSuspendForceRecompileCount=0;
+#endif
+}
+
 FODSCManager::FODSCManager()
 	: FTSTickerObjectBase(0.0f, FTSBackgroundableTicker::GetCoreTicker())
 {
@@ -105,7 +122,21 @@ void FODSCManager::AddThreadedRequest(
 {
 	if (IsHandlingRequests())
 	{
-		Thread->AddRequest(MaterialsToCompile, ShaderTypesToLoad, ShaderPlatform, FeatureLevel, QualityLevel, RecompileCommandType);
+		if ((RecompileCommandType == ODSCRecompileCommand::Material || RecompileCommandType == ODSCRecompileCommand::Changed)
+			&& (CVarODSCRecompileMode.GetValueOnAnyThread() > 0))
+		{
+			Thread->ResetMaterialsODSCData(FeatureLevel);
+			
+			// when we ask for "changed", we want both materials and global
+			if (RecompileCommandType == ODSCRecompileCommand::Changed)
+			{
+				Thread->AddRequest(TArray<FString>(), FString(), ShaderPlatform, FeatureLevel, QualityLevel, ODSCRecompileCommand::Changed);
+			}
+		}
+		else
+		{
+			Thread->AddRequest(MaterialsToCompile, ShaderTypesToLoad, ShaderPlatform, FeatureLevel, QualityLevel, RecompileCommandType);
+		}
 	}
 }
 
@@ -117,20 +148,15 @@ void FODSCManager::AddThreadedShaderPipelineRequest(
 	const FString& VertexFactoryName,
 	const FString& PipelineName,
 	const TArray<FString>& ShaderTypeNames,
-	int32 PermutationId
-)
+	int32 PermutationId,
+	const TArray<FShaderId>& RequestShaderIds)
 {
 	if (IsHandlingRequests())
 	{
-		Thread->AddShaderPipelineRequest(ShaderPlatform, FeatureLevel, QualityLevel, MaterialName, VertexFactoryName, PipelineName, ShaderTypeNames, PermutationId);
+		Thread->AddShaderPipelineRequest(ShaderPlatform, FeatureLevel, QualityLevel, MaterialName, VertexFactoryName, PipelineName, ShaderTypeNames, PermutationId, RequestShaderIds);
 	}
 }
 
-static inline bool IsODSCActive()
-{
-	return GODSCManager && GODSCManager->IsHandlingRequests();
-}
-
 void FODSCManager::RegisterMaterialInstance(const UMaterialInstance* MaterialInstance)
 {
 	if (IsODSCActive() && MaterialInstance->HasAnyInternalFlags(EInternalObjectFlags::AsyncLoading))
@@ -170,3 +196,47 @@ bool FODSCManager::HasAsyncLoadingInstances()
 
 	return bHasAsyncLoadingInstances;
 }
+
+
+void FODSCManager::SuspendODSCForceRecompile()
+{
+#if WITH_ODSC
+	check(ODSCManagerPrivate::ODSCSuspendForceRecompileCount >= 0);
+	++ODSCManagerPrivate::ODSCSuspendForceRecompileCount;
+#endif
+}
+
+void FODSCManager::ResumeODSCForceRecompile()
+{
+#if WITH_ODSC
+	--ODSCManagerPrivate::ODSCSuspendForceRecompileCount;
+	check(ODSCManagerPrivate::ODSCSuspendForceRecompileCount >= 0);
+#endif
+}
+
+bool FODSCManager::ShouldForceRecompileInternal(const FMaterialShaderMap* MaterialShaderMap, const FMaterial* Material)
+{
+#if WITH_ODSC
+	if (!FPlatformProperties::RequiresCookedData() || CVarODSCRecompileMode.GetValueOnAnyThread() == 0  || ODSCManagerPrivate::ODSCSuspendForceRecompileCount > 0)
+	{
+		return false;
+	}
+
+	if (MaterialShaderMap->IsFromODSC())
+	{
+		return false;
+	}
+
+	return true;
+#else
+	return false;
+#endif
+}
+
+void FODSCManager::RegisterMaterialShaderMap(const FMaterialShaderMap& MaterialShaderMap)
+{
+	if (IsODSCActive())
+	{
+		GODSCManager->Thread->RegisterMaterialShaderMap(MaterialShaderMap);
+	}
+}
diff --git a/Engine/Source/Runtime/Engine/Private/ODSC/ODSCThread.cpp b/Engine/Source/Runtime/Engine/Private/ODSC/ODSCThread.cpp
--- a/Engine/Source/Runtime/Engine/Private/ODSC/ODSCThread.cpp
+++ b/Engine/Source/Runtime/Engine/Private/ODSC/ODSCThread.cpp
@@ -5,6 +5,9 @@
 #include "ODSCLog.h"
 #include "HAL/FileManager.h"
 #include "Modules/ModuleManager.h"
+#include "MaterialShared.h"
+#include "Materials/MaterialInterface.h"
+#include "UObject/UObjectIterator.h"
 
 FODSCRequestPayload::FODSCRequestPayload(
 	EShaderPlatform InShaderPlatform,
@@ -156,6 +159,35 @@ void FODSCThread::Tick()
 	Process();
 }
 
+void FODSCThread::ResetMaterialsODSCData(ERHIFeatureLevel::Type FeatureLevel)
+{
+#if WITH_ODSC
+	FlushRenderingCommands();
+
+	{
+		FScopeLock Lock(&RequestHashCriticalSection);
+
+		// this will stop the rendering thread, and reattach components, in the destructor
+		FMaterialUpdateContext UpdateContext(FMaterialUpdateContext::EOptions::Default);
+		RequestHashes.Empty();
+
+		for (TObjectIterator<UMaterialInterface> It; It; ++It)
+		{ 
+			UMaterialInterface* Material = *It;
+			if (Material)
+			{
+				const FMaterialResource* MaterialResource = Material->GetMaterialResource((ERHIFeatureLevel::Type)FeatureLevel);
+				if (MaterialResource && MaterialResource->GetGameThreadShaderMap())
+				{
+					MaterialResource->GetGameThreadShaderMap()->SetIsFromODSC(false);
+				}
+				UpdateContext.AddMaterialInterface(Material);
+			}
+		}
+	}
+#endif
+}
+
 void FODSCThread::AddRequest(const TArray<FString>& MaterialsToCompile, const FString& ShaderTypesToLoad, EShaderPlatform ShaderPlatform, ERHIFeatureLevel::Type FeatureLevel, EMaterialQualityLevel::Type QualityLevel, ODSCRecompileCommand RecompileCommandType)
 {
 	PendingMaterialThreadedRequests.Enqueue(new FODSCMessageHandler(MaterialsToCompile, ShaderTypesToLoad, ShaderPlatform, FeatureLevel, QualityLevel, RecompileCommandType));
@@ -169,7 +201,8 @@ void FODSCThread::AddShaderPipelineRequest(
 	const FString& VertexFactoryName,
 	const FString& PipelineName,
 	const TArray<FString>& ShaderTypeNames,
-	int32 PermutationId
+	int32 PermutationId,
+	const TArray<FShaderId>& RequestShaderIds
 )
 {
 	// TODO: Requests for individual permutations come in here, but a single coalesced payload is submitted to the server since 
@@ -184,13 +217,45 @@ void FODSCThread::AddShaderPipelineRequest(
 	const FString RequestHash = FMD5::HashAnsiString(*RequestString);
 
 	FScopeLock Lock(&RequestHashCriticalSection);
-	if (!RequestHashes.Contains(RequestHash))
+
+	bool bShouldAddRequest = false;
+
+	for (const FShaderId& ShaderId : RequestShaderIds)
+	{
+		if (!RequestHashes.Contains(ShaderId))
+		{
+			RequestHashes.Add(ShaderId);
+			bShouldAddRequest = true;
+		}
+	}
+
+	if (bShouldAddRequest)
 	{
 		PendingMeshMaterialThreadedRequests.Enqueue(FODSCRequestPayload(ShaderPlatform, FeatureLevel, QualityLevel, MaterialName, VertexFactoryName, PipelineName, ShaderTypeNames, PermutationId, RequestHash));
-		RequestHashes.Add(RequestHash);
 	}
 }
 
+void FODSCThread::RegisterMaterialShaderMap(const FMaterialShaderMap& MaterialShaderMap)
+{
+	FScopeLock Lock(&RequestHashCriticalSection);
+
+	TMap<FShaderId, TShaderRef<FShader>> ShadersInMap;
+	MaterialShaderMap.GetShaderList(ShadersInMap);
+	FSHAHash CookedShaderMapIdHash = MaterialShaderMap.GetShaderMapId().CookedShaderMapIdHash;
+	for (auto Iter : ShadersInMap)
+	{
+		// GetShaderList doesn't use the Cookedshadermap id
+		const FShaderId& ShaderIdSrc = Iter.Key;
+		FShaderId ShaderIdCopy(ShaderIdSrc.Type, CookedShaderMapIdHash, ShaderIdSrc.ShaderPipelineName, ShaderIdSrc.VFType, ShaderIdSrc.PermutationId, (EShaderPlatform)ShaderIdSrc.Platform);
+
+		// The shadermap we receive contains all the requests the client sent until now, so it's possible they got removed from RequestHashes already
+		if (RequestHashes.Find(ShaderIdCopy))
+		{
+			RequestHashes.Remove(ShaderIdCopy);
+		}
+    }
+}
+
 void FODSCThread::GetCompletedRequests(TArray<FODSCMessageHandler*>& OutCompletedRequests)
 {
 	check(IsInGameThread());
@@ -244,11 +309,6 @@ void FODSCThread::Process()
 		while (PendingMeshMaterialThreadedRequests.Dequeue(Payload))
 		{
 			PayloadsToAggregate.Add(Payload);
-			int FoundIndex = INDEX_NONE;
-			if (RequestHashes.Find(Payload.RequestHash, FoundIndex))
-			{
-				RequestHashes.RemoveAt(FoundIndex);
-			}
 		}
 	}
 
diff --git a/Engine/Source/Runtime/Engine/Private/ODSC/ODSCThread.h b/Engine/Source/Runtime/Engine/Private/ODSC/ODSCThread.h
--- a/Engine/Source/Runtime/Engine/Private/ODSC/ODSCThread.h
+++ b/Engine/Source/Runtime/Engine/Private/ODSC/ODSCThread.h
@@ -12,6 +12,7 @@
 
 class FEvent;
 class FRunnableThread;
+class FMaterialShaderMap;
 
 namespace UE
 {
@@ -133,7 +134,8 @@ public:
 		const FString& VertexFactoryName,
 		const FString& PipelineName,
 		const TArray<FString>& ShaderTypeNames,
-		int32 PermutationId
+		int32 PermutationId,
+		const TArray<FShaderId>& RequestShaderIds
 	);
 
 	/**
@@ -150,6 +152,9 @@ public:
 
 	bool HasPendingRequests() const { return bHasPendingRequests; }
 
+	void RegisterMaterialShaderMap(const FMaterialShaderMap& MaterialShaderMap);
+
+	void ResetMaterialsODSCData(ERHIFeatureLevel::Type FeatureLevel);
 protected:
 
 	//~ Begin FRunnable Interface
@@ -194,7 +199,7 @@ private:
 	FCriticalSection RequestHashCriticalSection;
 
 	/** Hashes for all Pending or Completed requests.  This is so we avoid making the same request multiple times. */
-	TArray<FString> RequestHashes;
+	TSet<FShaderId> RequestHashes;
 
 	/** Pointer to Runnable Thread */
 	FRunnableThread* Thread = nullptr;
diff --git a/Engine/Source/Runtime/Engine/Private/PSOPrecacheMaterial.cpp b/Engine/Source/Runtime/Engine/Private/PSOPrecacheMaterial.cpp
--- a/Engine/Source/Runtime/Engine/Private/PSOPrecacheMaterial.cpp
+++ b/Engine/Source/Runtime/Engine/Private/PSOPrecacheMaterial.cpp
@@ -14,6 +14,7 @@
 #include "MaterialShared.h"
 #include "VertexFactory.h"
 #include "SceneInterface.h"
+#include "ODSC/ODSCManager.h"
 
 int32 GPSOUseBackgroundThreadForCollection = 1;
 static FAutoConsoleVariableRef CVarPSOUseBackgroundThreadForCollection(
@@ -468,6 +469,10 @@ void FMaterialPSOPrecacheCollectionTask::DoTask(ENamedThreads::Type CurrentThrea
 {
 	TRACE_CPUPROFILER_EVENT_SCOPE(FMaterialPSOPrecacheCollectionTask);
 
+#if WITH_ODSC
+	FODSCSuspendForceRecompileScope ODSCSuspendForceRecompileScope;
+#endif
+
 	// Make sure task is still relevant
 	if (RequestLifecycleID != GMaterialPSORequestManager.GetLifecycleID())
 	{
diff --git a/Engine/Source/Runtime/Engine/Public/MaterialShared.h b/Engine/Source/Runtime/Engine/Public/MaterialShared.h
--- a/Engine/Source/Runtime/Engine/Public/MaterialShared.h
+++ b/Engine/Source/Runtime/Engine/Public/MaterialShared.h
@@ -1545,6 +1545,11 @@ public:
 	void RemoveCompilingDependency(FMaterial* Material);
 #endif // WITH_EDITOR
 
+#if WITH_ODSC
+	bool IsFromODSC() const { return bIsFromODSC; }
+	void SetIsFromODSC(bool bInIsFromODSC) { bIsFromODSC = bInIsFromODSC; }
+#endif
+
 #if WITH_EDITORONLY_DATA
 	const TCHAR* GetFriendlyName() const { return *GetContent()->FriendlyName; }
 	const TCHAR* GetDebugDescription() const { return *GetContent()->DebugDescription; }
@@ -1705,6 +1710,9 @@ private:
 	/** Indicates whether the shader map should be stored in the shader cache. */
 	uint32 bIsPersistent : 1;
 
+	/** Indicates whether the shader map is from ODSC */
+	uint32 bIsFromODSC : 1;
+
 #if WITH_EDITOR
 	FShader* ProcessCompilationResultsForSingleJob(class FShaderCompileJob* SingleJob, const FShaderPipelineType* ShaderPipeline, const FSHAHash& MaterialShaderMapHash);
 #endif
diff --git a/Engine/Source/Runtime/Engine/Public/ODSC/ODSCManager.h b/Engine/Source/Runtime/Engine/Public/ODSC/ODSCManager.h
--- a/Engine/Source/Runtime/Engine/Public/ODSC/ODSCManager.h
+++ b/Engine/Source/Runtime/Engine/Public/ODSC/ODSCManager.h
@@ -10,6 +10,7 @@
 
 class FODSCThread;
 class UMaterialInstance;
+class FMaterialShaderMap;
 
 /**
  * Responsible for processing shader compile responses from the ODSC Thread.
@@ -76,21 +77,42 @@ public:
 		const FString& VertexFactoryName,
 		const FString& PipelineName,
 		const TArray<FString>& ShaderTypeNames,
-		int32 PermutationId
+		int32 PermutationId,
+		const TArray<FShaderId>& RequestShaderIds
 	);
 
+	UE_DEPRECATED(5.5, "RequestShaderIds is needed for AddThreadedShaderPipelineRequest and need to match the ones from FMaterialShaderMap::GetShaderList")
+	ENGINE_API void AddThreadedShaderPipelineRequest(
+		EShaderPlatform ShaderPlatform,
+		ERHIFeatureLevel::Type FeatureLevel,
+		EMaterialQualityLevel::Type QualityLevel,
+		const FString& MaterialName,
+		const FString& VertexFactoryName,
+		const FString& PipelineName,
+		const TArray<FString>& ShaderTypeNames,
+		int32 PermutationId
+	) {}
+
 	/** Returns true if we would actually add a request when calling AddThreadedShaderPipelineRequest. */
 	inline bool IsHandlingRequests() const { return Thread != nullptr; }
+	static void RegisterMaterialShaderMap(const FMaterialShaderMap& MaterialShaderMap);
 
 	static void RegisterMaterialInstance(const UMaterialInstance* MI);
 	static void UnregisterMaterialInstance(const UMaterialInstance* MI);
 
+	static inline bool IsODSCActive();
+	static inline bool ShouldForceRecompile(const FMaterialShaderMap* MaterialShaderMap, const FMaterial* Material);
+
+	static void SuspendODSCForceRecompile();
+	static void ResumeODSCForceRecompile();
+
 private:
 
 	ENGINE_API void OnEnginePreExit();
 	ENGINE_API void StopThread();
 
 	bool HasAsyncLoadingInstances();
+	bool ShouldForceRecompileInternal(const FMaterialShaderMap* MaterialShaderMap, const FMaterial* Material);
 
 	/** Handles communicating directly with the cook on the fly server. */
 	FODSCThread* Thread = nullptr;
@@ -98,8 +120,30 @@ private:
 	FDelegateHandle OnScreenMessagesHandle;
 	FCriticalSection MaterialInstancesCachedUniformExpressionsCS;
 	TMap<const void*, TWeakObjectPtr<const UMaterialInstance> > MaterialInstancesCachedUniformExpressions;
+};
 
+struct FODSCSuspendForceRecompileScope
+{
+	FODSCSuspendForceRecompileScope()
+	{
+		FODSCManager::SuspendODSCForceRecompile();
+	}
+
+	~FODSCSuspendForceRecompileScope()
+	{
+		FODSCManager::ResumeODSCForceRecompile();
+	}
 };
 
 /** The global shader ODSC manager. */
 extern ENGINE_API FODSCManager* GODSCManager;
+
+inline bool FODSCManager::IsODSCActive()
+{
+	return GODSCManager && GODSCManager->IsHandlingRequests();
+}
+
+inline bool FODSCManager::ShouldForceRecompile(const FMaterialShaderMap* MaterialShaderMap, const FMaterial* Material)
+{
+	return FODSCManager::IsODSCActive() && GODSCManager->ShouldForceRecompileInternal(MaterialShaderMap, Material);
+}
