Mitigation for race condition leading to crash in body setup cooking. If a race is detected during the generation of the cooked physics data we return the error condition to the caller if they elected to handle this case. This allows us to handle the race when checking if cached data is ready and defer cached data requests until after the async build has completed. As this is still an error condition we also report the path of each body setup that hits a race condition for future investigation into why it occurs.

diff --git a/Engine/Source/Runtime/Engine/Classes/PhysicsEngine/BodySetup.h b/Engine/Source/Runtime/Engine/Classes/PhysicsEngine/BodySetup.h
--- a/Engine/Source/Runtime/Engine/Classes/PhysicsEngine/BodySetup.h
+++ b/Engine/Source/Runtime/Engine/Classes/PhysicsEngine/BodySetup.h
@@ -290,9 +290,38 @@ public:
 	/** Aborts an async cook that hasn't begun. See CreatePhysicsMeshesAsync.  (Useful for cases where frequent updates at runtime would otherwise cause a backlog) */
 	ENGINE_API void AbortPhysicsMeshAsyncCreation();
 
+	enum class EBodySetupCookResult
+	{
+		NoError,
+		FailedToLock,
+		FailedToCook
+	};
+
+	struct FBodySetupTryCookResult
+	{
+		FByteBulkData* Data;
+		EBodySetupCookResult CookResult;
+	};
+
 private:
 	FByteBulkData* GetCookedFormatData();
 
+	/**
+	* Given a format name returns its cooked data.
+	*
+	* @param Format Physics format name.
+	* @return Cooked data or NULL of the data was not found.
+	*/
+	FByteBulkData* GetCookedData(FName Format);
+
+	/**
+	 * Attempt to get cooked physics data, with extended information about
+	 * potential failure cases during the cooking process
+	 * @param Format Physics format name
+	 * @return Cooked data, nullptr if no data or failure occured, plus a failure code
+	 */
+	FBodySetupTryCookResult TryGetCookedData(FName Format);
+
 	// #TODO MRMesh for some reason needs to be able to call this - that case needs fixed to correctly use the create meshes flow
 	friend class UMRMeshComponent;
 	bool ProcessFormatData_Chaos(FByteBulkData* FormatData);
@@ -312,14 +341,6 @@ private:
 	 */
 	void FinishCreatePhysicsMeshesAsync(FAsyncCookHelper* AsyncPhysicsCookHelper, FOnAsyncPhysicsCookFinished OnAsyncPhysicsCookFinished);
 
-	/**
-	* Given a format name returns its cooked data.
-	*
-	* @param Format Physics format name.
-	* @return Cooked data or NULL of the data was not found.
-	*/
-	FByteBulkData* GetCookedData(FName Format);
-
 public:
 
 	/**
diff --git a/Engine/Source/Runtime/Engine/Private/PhysicsEngine/BodySetup.cpp b/Engine/Source/Runtime/Engine/Private/PhysicsEngine/BodySetup.cpp
--- a/Engine/Source/Runtime/Engine/Private/PhysicsEngine/BodySetup.cpp
+++ b/Engine/Source/Runtime/Engine/Private/PhysicsEngine/BodySetup.cpp
@@ -1206,7 +1206,14 @@ bool UBodySetup::IsCachedCookedPlatformDataLoaded(const ITargetPlatform* TargetP
 			return false;
 		}
 	}
-	GetCookedData(TargetPlatform->GetPhysicsFormat(this));
+	FBodySetupTryCookResult CookResult = TryGetCookedData(TargetPlatform->GetPhysicsFormat(this));
+
+	if(CookResult.CookResult == EBodySetupCookResult::FailedToLock)
+	{
+		UE_LOG(LogPhysics, Warning, TEXT("Multiple writes to BodySetup bulk data detected during cook for setup: %s"), *GetPathName());
+		return false;
+	}
+
 	return true;
 }
 
@@ -1313,7 +1320,7 @@ UE_TRACE_EVENT_BEGIN(Cpu, BodySetupDDCFetch, NoSync)
 UE_TRACE_EVENT_END()
 
 template<typename DDCBuilderType>
-void GetDDCBuiltData(FByteBulkData* OutResult, DDCBuilderType& InBuilder, UBodySetup* InSetup, bool bInIsRuntime)
+UBodySetup::EBodySetupCookResult TryGetDDCBuiltData(FByteBulkData* OutResult, DDCBuilderType& InBuilder, UBodySetup* InSetup, bool bInIsRuntime)
 {
 	TArray<uint8> OutData;
 
@@ -1349,6 +1356,11 @@ void GetDDCBuiltData(FByteBulkData* OutResult, DDCBuilderType& InBuilder, UBodyS
 
 	if(OutData.Num())
 	{
+		if(OutResult->IsLocked())
+		{
+			return UBodySetup::EBodySetupCookResult::FailedToLock;
+		}
+
 		OutResult->Lock(LOCK_READ_WRITE);
 		FMemory::Memcpy(OutResult->Realloc(OutData.Num()), OutData.GetData(), OutData.Num());
 		OutResult->Unlock();
@@ -1356,31 +1368,55 @@ void GetDDCBuiltData(FByteBulkData* OutResult, DDCBuilderType& InBuilder, UBodyS
 	else if(!bInIsRuntime)	//only want to warn if DDC cooking failed - if it's really trying to use runtime and we can't, the runtime cooker code will catch it
 	{
 		UE_LOG(LogPhysics, Warning, TEXT("Attempt to build physics data for %s when we are unable to."), *InSetup->GetPathName());
+		return UBodySetup::EBodySetupCookResult::FailedToCook;
 	}
+
+	return UBodySetup::EBodySetupCookResult::NoError;
+}
+
+template<typename DDCBuilderType>
+void GetDDCBuiltData(FByteBulkData* OutResult, DDCBuilderType& InBuilder, UBodySetup* InSetup, bool bInIsRuntime)
+{
+	UBodySetup::EBodySetupCookResult Result = TryGetDDCBuiltData(OutResult, InBuilder, InSetup, bInIsRuntime);
+	check(Result != UBodySetup::EBodySetupCookResult::FailedToLock);
 }
 
 #endif //#if WITH_EDITOR
 
 FByteBulkData* UBodySetup::GetCookedData(FName Format)
 {
-	if (IsTemplate())
+	UBodySetup::FBodySetupTryCookResult CookResult = TryGetCookedData(Format);
+
+	// Failing to cook means we'll have no physics data which is not ideal but works and the cook will have warned the user.
+	// Failing to lock means we have multiple paths attempting to write into the bulk data buffer which is more serious. Separating these
+	// case to help trace any race conditions arising during the physics cook
+	check(CookResult.CookResult != EBodySetupCookResult::FailedToLock);
+
+	return CookResult.Data;
+}
+
+UBodySetup::FBodySetupTryCookResult UBodySetup::TryGetCookedData(FName Format)
+{
+	EBodySetupCookResult CookDataResult = EBodySetupCookResult::NoError;
+
+	if(IsTemplate())
 	{
-		return nullptr;
+		return { nullptr, CookDataResult };
 	}
 
 	// Geometry should never have collision data, cooked data will never be present
-	if (bNeverNeedsCookedCollisionData)
+	if(bNeverNeedsCookedCollisionData)
 	{
-		return nullptr;
+		return { nullptr, CookDataResult };
 	}
 
 	IInterface_CollisionDataProvider* CDP = Cast<IInterface_CollisionDataProvider>(GetOuter());
 
 	// If there is nothing to cook or if we are reading data from a cooked package for an asset with no collision, 
 	// we want to return here
-	if ((AggGeom.ConvexElems.Num() == 0 && CDP == nullptr) || !bHasCookedCollisionData)
+	if((AggGeom.ConvexElems.Num() == 0 && CDP == nullptr) || !bHasCookedCollisionData)
 	{
-		return nullptr;
+		return { nullptr, CookDataResult };
 	}
 
 #if WITH_EDITOR
@@ -1395,15 +1431,15 @@ FByteBulkData* UBodySetup::GetCookedData(FName Format)
 	bool bIsRuntime = IsRuntime(this);
 
 #if WITH_EDITOR
-	if (!bContainedData)
+	if(!bContainedData)
 	{
 		SCOPE_CYCLE_COUNTER(STAT_PhysXCooking);
 
 		// Note: Check ContainsPhysicsTriMeshData before looking at the number of convex elems, to ensure the side effects of ContainsPhysicsTriMeshData happen
 		// (specifically, for static mesh this will ensure the mesh render data is already built)
-		if ((CDP == nullptr || CDP->ContainsPhysicsTriMeshData(bMeshCollideAll) == false) && AggGeom.ConvexElems.Num() == 0)
+		if((CDP == nullptr || CDP->ContainsPhysicsTriMeshData(bMeshCollideAll) == false) && AggGeom.ConvexElems.Num() == 0)
 		{
-			return nullptr;
+			return { nullptr, CookDataResult };
 		}
 
 		// We do not want a FGCObject to be created to prevent garbage collection of our own UBodySetup*
@@ -1413,7 +1449,7 @@ FByteBulkData* UBodySetup::GetCookedData(FName Format)
 		FChaosDerivedDataCooker* PhysicsDerivedCooker = new FChaosDerivedDataCooker(this, Format, bUseRefHolder);
 
 		Result = &UseCookedData->GetFormat(Format);
-		GetDDCBuiltData(Result, *PhysicsDerivedCooker, this, bIsRuntime);
+		CookDataResult = TryGetDDCBuiltData(Result, *PhysicsDerivedCooker, this, bIsRuntime);
 	}
 	else
 #endif // #if WITH_EDITOR
@@ -1422,7 +1458,10 @@ FByteBulkData* UBodySetup::GetCookedData(FName Format)
 	}
 
 	check(Result);
-	return Result->GetBulkDataSize() > 0 ? Result : nullptr; // we don't return empty bulk data...but we save it to avoid thrashing the DDC
+
+	// we don't return empty bulk data...but we save it to avoid thrashing the DDC
+	Result = Result->GetBulkDataSize() > 0 ? Result : nullptr;
+	return {Result, CookDataResult};
 }
 
 void UBodySetup::GetGeometryDDCKey(FString& OutString) const
