[Materials] Use lerp intrinsic and ternary operator instead of switch-case statement for dynamic branches when possible.

This generates an HLSL lerp() intrinsic for static switches that are enabled as dynamic branches to avoid unnecessary dynamic flow control.
The lerp() intrinsic is only applied to float values. Other numerical values use the ternary operator and the old switch-case statement is used as fallback.

This translates a simple Float3 static switch from this:

    float3 Static7;
    [branch]
    switch (int(UnpackUniform_bool(asuint(Material_PreshaderBuffer[0][0]), 0)))
    {
        default: Static7 = Local5.rgb; break;
        case 0:  Static7 = Local2.rgb; break;
    }
    float3 Local8 = Static7;

 ... to this:

    float3 Local7 = lerp(Local5.rgb, Local2.rgb, UnpackUniform_bool(asuint(Material_PreshaderBuffer[0][0]), 0));

#rnx

diff --git a/Engine/Source/Runtime/Engine/Private/Materials/HLSLMaterialTranslator.cpp b/Engine/Source/Runtime/Engine/Private/Materials/HLSLMaterialTranslator.cpp
--- a/Engine/Source/Runtime/Engine/Private/Materials/HLSLMaterialTranslator.cpp
+++ b/Engine/Source/Runtime/Engine/Private/Materials/HLSLMaterialTranslator.cpp
@@ -6437,7 +6437,6 @@ int32 FHLSLMaterialTranslator::DynamicBranch(int32 Condition, int32 A, int32 B)
 
 	A = ForceCast(A, ResultType, MFCF_ReplicateValue);
 	B = ForceCast(B, ResultType, MFCF_ReplicateValue);
-	FString SymbolName = CreateSymbolName(TEXT("Static"));
 
 	checkf(Condition >= 0 && Condition < CurrentScopeChunks->Num(), TEXT("Index %d/%d, Platform=%d"), Condition, CurrentScopeChunks->Num(), (int)Platform);
 	const FShaderCodeChunk& CodeChunk = (*CurrentScopeChunks)[Condition];
@@ -6448,9 +6447,24 @@ int32 FHLSLMaterialTranslator::DynamicBranch(int32 Condition, int32 A, int32 B)
 		AddCodeChunk(MCT_VoidStatement, TEXT("//%s"), *StaticBoolParameter->GetParameterName().ToString());
 	}
 
-	AddCodeChunk(MCT_VoidStatement, TEXT("%s %s;"), HLSLTypeString(ResultType), *SymbolName);
-	AddCodeChunk(MCT_VoidStatement, TEXT("[branch] switch (int(%s)){ default: %s = %s; break; case 0: %s = %s; break;}"), *GetParameterCode(Condition), *SymbolName, *GetParameterCode(A), *SymbolName, *GetParameterCode(B));
-	return AddCodeChunk(ResultType, *SymbolName);
+	if ((ResultType & MCT_Float) != 0)
+	{
+		// Use lerp() intrinsic for floating-point values to avoid dynamic branching
+		return AddCodeChunk(ResultType, TEXT("lerp(%s, %s, %s)"), *GetParameterCode(A), *GetParameterCode(B), *GetParameterCode(Condition));
+	}
+	else if ((ResultType & (MCT_Numeric | MCT_UInt)) != 0)
+	{
+		// Use ternary-operator to simplify output for dynamic branch of two numerical values
+		return AddCodeChunk(ResultType, TEXT("%s ? %s : %s"), *GetParameterCode(Condition), *GetParameterCode(A), *GetParameterCode(B));
+	}
+	else
+	{
+		// Fallback to switch-case statement for then/else branches for all other types
+		FString SymbolName = CreateSymbolName(TEXT("Static"));
+		AddCodeChunk(MCT_VoidStatement, TEXT("%s %s;"), HLSLTypeString(ResultType), *SymbolName);
+		AddCodeChunk(MCT_VoidStatement, TEXT("[branch] switch (int(%s)){ default: %s = %s; break; case 0: %s = %s; break;}"), *GetParameterCode(Condition), *SymbolName, *GetParameterCode(A), *SymbolName, *GetParameterCode(B));
+		return AddCodeChunk(ResultType, *SymbolName);
+	}
 }
 
 // Compare two inputs and return true if they are either the same, or if they evaluate to equal constant expressions.
