Updated FInstancedStaticMeshSceneProxy to get ray tracing geometry from RayTracingProxy LODs instead of rendering LODs.

diff --git a/Engine/Source/Runtime/Engine/Private/InstancedStaticMesh.cpp b/Engine/Source/Runtime/Engine/Private/InstancedStaticMesh.cpp
--- a/Engine/Source/Runtime/Engine/Private/InstancedStaticMesh.cpp
+++ b/Engine/Source/Runtime/Engine/Private/InstancedStaticMesh.cpp
@@ -1633,23 +1633,27 @@ void FInstancedStaticMeshSceneProxy::GetDynamicRayTracingInstances(struct FRayTr
 
 	const int32 MinAllowedLODIndex = FMath::Clamp<int32>(CVarRayTracingInstancedStaticMeshesMinLOD.GetValueOnRenderThread(), 0, RenderData->LODResources.Num() - 1);
 
-	int32 LODIndex = FMath::Max<int32>(MinAllowedLODIndex, GetCurrentFirstLODIdx_RenderThread());
+	const int32 RayTracingMinLOD = RenderData->RayTracingProxy->bUsingRenderingLODs ? FMath::Max(MinAllowedLODIndex, (int32)GetCurrentFirstLODIdx_RenderThread()) : 0;
+
+	int32 LODIndex = RayTracingMinLOD;
+
+	FStaticMeshRayTracingProxyLODArray& RayTracingLODs = RenderData->RayTracingProxy->LODs;
 
 	FRayTracingGeometry* RayTracingGeometry = nullptr;
 
 	// Select first LOD with valid ray tracing geometry
 	// TODO: Should only do this if any instance uses static geometry
-	for (; LODIndex < RenderData->LODResources.Num(); ++LODIndex)
+	for (; LODIndex < RayTracingLODs.Num(); ++LODIndex)
 	{
-		FStaticMeshLODResources& CurrentLODResources = RenderData->LODResources[LODIndex];
+		FRayTracingGeometry& CurrentRayTracingGeometry = *RayTracingLODs[LODIndex].RayTracingGeometry;
 
-		if (CurrentLODResources.RayTracingGeometry->HasPendingBuildRequest())
+		if (CurrentRayTracingGeometry.HasPendingBuildRequest())
 		{
-			CurrentLODResources.RayTracingGeometry->BoostBuildPriority();
+			CurrentRayTracingGeometry.BoostBuildPriority();
 		}
-		else if (CurrentLODResources.RayTracingGeometry->IsValid() && !CurrentLODResources.RayTracingGeometry->IsEvicted())
+		else if (CurrentRayTracingGeometry.IsValid() && !CurrentRayTracingGeometry.IsEvicted())
 		{
-			RayTracingGeometry = CurrentLODResources.RayTracingGeometry;
+			RayTracingGeometry = &CurrentRayTracingGeometry;
 			break;
 		}
 	}
@@ -1711,7 +1715,7 @@ void FInstancedStaticMeshSceneProxy::GetDynamicRayTracingInstances(struct FRayTr
 			{
 				FMeshBatch DynamicMeshBatch;
 
-				if (!GetMeshElement(LODIndex, 0, SectionIdx, 0, false, false, DynamicMeshBatch))
+				if (!GetMeshElement(LODIndex, 0, SectionIdx, 0, false, false, DynamicMeshBatch)) // todo: RayTracingLOD vertex factory
 				{
 					DynamicMeshBatch.MaterialRenderProxy = UMaterial::GetDefaultMaterial(MD_Surface)->GetRenderProxy();
 					DynamicMeshBatch.SegmentIndex = SectionIdx;
@@ -1983,13 +1987,13 @@ void FInstancedStaticMeshSceneProxy::SetupRayTracingDynamicInstances(int32 NumDy
 
 	// create new geometries
 
-	const FStaticMeshLODResources& LODModel = RenderData->LODResources[LODIndex];
+	FStaticMeshRayTracingProxyLOD& RayTracingLOD = RenderData->RayTracingProxy->LODs[LODIndex];
 
 	for (int32 Item = 0; Item < NumDynamicInstances; Item++)
 	{
 		FRayTracingDynamicData& DynamicData = RayTracingDynamicData.AddDefaulted_GetRef();
 
-		FRayTracingGeometryInitializer Initializer = LODModel.RayTracingGeometry->Initializer;
+		FRayTracingGeometryInitializer Initializer = RayTracingLOD.RayTracingGeometry->Initializer;
 		for (FRayTracingGeometrySegment& Segment : Initializer.Segments)
 		{
 			Segment.VertexBuffer = nullptr;
