Set ECVF_RenderThreadSafe in variables that are used in render threads

- r.GPUSkin.UnlimitedBoneInfluences and r.GPUSkin.AlwaysUseDeformerForUnlimitedBoneInfluences are both read in GPUSkinVertexFactory when recompiling materials that can be used in skeletal meshes

diff --git a/Engine/Source/Runtime/Engine/Private/GPUSkinVertexFactory.cpp b/Engine/Source/Runtime/Engine/Private/GPUSkinVertexFactory.cpp
--- a/Engine/Source/Runtime/Engine/Private/GPUSkinVertexFactory.cpp
+++ b/Engine/Source/Runtime/Engine/Private/GPUSkinVertexFactory.cpp
@@ -53,7 +53,7 @@ static FAutoConsoleVariableRef CVarUnlimitedBoneInfluences(
 	TEXT("r.GPUSkin.UnlimitedBoneInfluences"),
 	GCVarUnlimitedBoneInfluences,
 	TEXT("Whether to use unlimited bone influences instead of default 4/8 for GPU skinning. Cannot be changed at runtime."),
-	ECVF_ReadOnly);
+	ECVF_ReadOnly | ECVF_RenderThreadSafe);
 
 static int32 GCVarUnlimitedBoneInfluencesThreshold = EXTRA_BONE_INFLUENCES;
 static FAutoConsoleVariableRef CVarUnlimitedBoneInfluencesThreshold(
@@ -67,7 +67,7 @@ static FAutoConsoleVariableRef CVarAlwaysUseDeformerForUnlimitedBoneInfluences(
 	TEXT("r.GPUSkin.AlwaysUseDeformerForUnlimitedBoneInfluences"),
 	GCVarAlwaysUseDeformerForUnlimitedBoneInfluences,
 	TEXT("Any meshes using Unlimited Bone Influences will always be rendered with a Mesh Deformer. This reduces the number of shader permutations needed for skeletal mesh materials, saving memory at the cost of performance. Has no effect if either Unlimited Bone Influences or Deformer Graph is disabled. Cannot be changed at runtime."),
-	ECVF_ReadOnly);
+	ECVF_ReadOnly | ECVF_RenderThreadSafe);
 
 static TAutoConsoleVariable<bool> CVarMobileEnableCloth(
 	TEXT("r.Mobile.EnableCloth"),
