World Partition - HLOD: When performing HLOD setup, remove unused ASpatialHashRuntimeGridInfo actors that were left being when switching to a new hash

diff --git a/Engine/Source/Runtime/Engine/Private/WorldPartition/RuntimeSpatialHash/RuntimeSpatialHashHLOD.cpp b/Engine/Source/Runtime/Engine/Private/WorldPartition/RuntimeSpatialHash/RuntimeSpatialHashHLOD.cpp
--- a/Engine/Source/Runtime/Engine/Private/WorldPartition/RuntimeSpatialHash/RuntimeSpatialHashHLOD.cpp
+++ b/Engine/Source/Runtime/Engine/Private/WorldPartition/RuntimeSpatialHash/RuntimeSpatialHashHLOD.cpp
@@ -5,6 +5,7 @@
 
 #if WITH_EDITOR
 #include "HAL/PlatformFile.h"
+#include "WorldPartition/ActorDescContainerInstanceCollection.h"
 #include "WorldPartition/RuntimeSpatialHash/RuntimeSpatialHashGridHelper.h"
 #include "WorldPartition/HLOD/HLODLayer.h"
 #include "WorldPartition/HLOD/HLODActor.h"
@@ -355,32 +356,23 @@ static TMap<FName, FSpatialHashRuntimeGrid> CreateHLODGrids(TMap<UHLODLayer*, in
 	return HLODGrids;
 }
 
+// Actor tag used to mark ASpatialHashRuntimeGridInfo actors as being valid for the current grid setup.
+static const FName HLODGridActorInUseTag = TEXT("HLODGridActorInUse");
+
 // Create/destroy HLOD grid actors
 static void UpdateHLODGridsActors(UWorld* World, const UActorDescContainerInstance* ContainerInstance, const TMap<FName, FSpatialHashRuntimeGrid>& HLODGrids, ISourceControlHelper* SourceControlHelper)
 {
-	static const FName HLODGridTag = TEXT("HLOD");
-	static const uint32 HLODGridTagLen = HLODGridTag.ToString().Len();
-
-	// Gather all existing HLOD grid actors, see if some are unused and needs to be deleted
+	// Gather all existing HLOD grid actors
 	TMap<FName, ASpatialHashRuntimeGridInfo*> ExistingGridActors;
-	for (UActorDescContainerInstance::TConstIterator<> Iterator(ContainerInstance); Iterator; ++Iterator)
+	for (UActorDescContainerInstance::TConstIterator<ASpatialHashRuntimeGridInfo> Iterator(ContainerInstance); Iterator; ++Iterator)
 	{
-		if(Iterator->GetActorNativeClass()->IsChildOf<ASpatialHashRuntimeGridInfo>())
+		if (ASpatialHashRuntimeGridInfo* GridActor = CastChecked<ASpatialHashRuntimeGridInfo>(Iterator->GetActor()))
 		{
-			ASpatialHashRuntimeGridInfo* GridActor = CastChecked<ASpatialHashRuntimeGridInfo>(Iterator->GetActor());
-			if (GridActor->ActorHasTag(HLODGridTag))
+			const FSpatialHashRuntimeGrid* HLODGrid = HLODGrids.Find(GridActor->GridSettings.GridName);
+			if (HLODGrid && GridActor->GridSettings.Priority && GridActor->GridSettings.HLODLayer)
 			{
-				const FSpatialHashRuntimeGrid* HLODGrid = HLODGrids.Find(GridActor->GridSettings.GridName);
-				if (HLODGrid && GridActor->GridSettings.Priority && GridActor->GridSettings.HLODLayer)
-				{
-					check(GridActor->GetContentBundleGuid() == ContainerInstance->GetContentBundleGuid());
-					ExistingGridActors.Emplace(GridActor->GridSettings.GridName, GridActor);
-				}
-				else
-				{
-					World->DestroyActor(GridActor);
-					DeletePackage(GridActor->GetPackage(), SourceControlHelper);
-				}
+				check(GridActor->GetContentBundleGuid() == ContainerInstance->GetContentBundleGuid());
+				ExistingGridActors.Emplace(GridActor->GridSettings.GridName, GridActor);
 			}
 		}
 	}
@@ -402,13 +394,15 @@ static void UpdateHLODGridsActors(UWorld* World, const UActorDescContainerInstan
 			FActorSpawnParameters SpawnParams;
 			SpawnParams.bCreateActorPackage = true;
 			GridActor = World->SpawnActor<ASpatialHashRuntimeGridInfo>(SpawnParams);
-			GridActor->Tags.Add(HLODGridTag);
 			bDirty = true;
 
 			check(GridActor->GetContentBundleGuid() == ContainerInstance->GetContentBundleGuid());
 			check(GridActor->GetExternalDataLayerAsset() == ContainerInstance->GetExternalDataLayerAsset());
 		}
 
+		// Flag this ASpatialHashRuntimeGridInfo actor as being in-use
+		GridActor->Tags.Add(HLODGridActorInUseTag);
+
 		const FString ActorLabel = GridSettings.GridName.ToString();
 		if (GridActor->GetActorLabel() != ActorLabel)
 		{
@@ -610,7 +604,7 @@ bool UWorldPartitionRuntimeSpatialHash::SetupHLODActors(const IStreamingGenerati
 	{
 		BaseContainerInstanceCollection->ForEachActorDescContainerInstance([this, &HLODGrids, &SourceControlHelper](const UActorDescContainerInstance* ActorDescContainerInstance)
 		{
-		// Create/destroy HLOD grid actors
+			// Create/destroy HLOD grid actors
 			UpdateHLODGridsActors(GetWorld(), ActorDescContainerInstance, HLODGrids, SourceControlHelper);
 		});
 	}
@@ -633,4 +627,46 @@ bool UWorldPartitionRuntimeSpatialHash::SetupHLODActors(const IStreamingGenerati
 	return true;
 }
 
+void UWorldPartitionRuntimeSpatialHash::PreSetupHLODActors(const UWorldPartition* InWorldPartition, const UWorldPartition::FSetupHLODActorsParams& InParams) const
+{
+	if (!InParams.bReportOnly)
+	{
+		// Clear "in-use" tag on all ASpatialHashRuntimeGridInfo actors
+		for (FActorDescContainerInstanceCollection::TConstIterator<ASpatialHashRuntimeGridInfo> Iterator(InWorldPartition); Iterator; ++Iterator)
+		{
+			if (ASpatialHashRuntimeGridInfo* SpatialHashRuntimeGridInfo = Cast<ASpatialHashRuntimeGridInfo>(Iterator->GetActor()))
+			{
+				SpatialHashRuntimeGridInfo->Tags.Remove(HLODGridActorInUseTag);
+			}
+		}
+	}
+}
+
+void UWorldPartitionRuntimeSpatialHash::PostSetupHLODActors(const UWorldPartition* InWorldPartition, const UWorldPartition::FSetupHLODActorsParams& InParams) const
+{
+	if (!InParams.bReportOnly)
+	{
+		// Delete all ASpatialHashRuntimeGridInfo actors that don't have the "in-use" tag
+		for (FActorDescContainerInstanceCollection::TConstIterator<ASpatialHashRuntimeGridInfo> Iterator(InWorldPartition); Iterator; ++Iterator)
+		{
+			if (ASpatialHashRuntimeGridInfo* SpatialHashRuntimeGridInfo = Cast<ASpatialHashRuntimeGridInfo>(Iterator->GetActor()))
+			{
+				if (!SpatialHashRuntimeGridInfo->ActorHasTag(HLODGridActorInUseTag))
+				{
+					FString PackageName = Iterator->GetActorPackage().ToString();
+					bool bDeleted = InParams.SourceControlHelper->Delete(SpatialHashRuntimeGridInfo->GetPackage());
+					if (bDeleted)
+					{
+						UE_LOG(LogWorldPartitionRuntimeSpatialHashHLOD, Display, TEXT("Deleting %s..."), *PackageName);
+					}
+					else
+					{
+						UE_LOG(LogWorldPartitionRuntimeSpatialHashHLOD, Error, TEXT("Failed to delete %s"), *PackageName);
+					}
+				}
+			}
+		}
+	}
+}
+
 #endif // #if WITH_EDITOR
diff --git a/Engine/Source/Runtime/Engine/Private/WorldPartition/WorldPartitionRuntimeHash.cpp b/Engine/Source/Runtime/Engine/Private/WorldPartition/WorldPartitionRuntimeHash.cpp
--- a/Engine/Source/Runtime/Engine/Private/WorldPartition/WorldPartitionRuntimeHash.cpp
+++ b/Engine/Source/Runtime/Engine/Private/WorldPartition/WorldPartitionRuntimeHash.cpp
@@ -15,6 +15,7 @@
 #include "WorldPartition/Cook/WorldPartitionCookPackage.h"
 #include "WorldPartition/Cook/WorldPartitionCookPackageContextInterface.h"
 #include "UObject/ObjectSaveContext.h"
+#include "UObject/UObjectIterator.h"
 #endif
 
 #include UE_INLINE_GENERATED_CPP_BY_NAME(WorldPartitionRuntimeHash)
@@ -631,6 +632,31 @@ UWorldPartitionRuntimeHash* UWorldPartitionRuntimeHash::ConvertWorldPartitionHas
 	NewHash->SetDefaultValues();
 	return NewHash;
 }
+
+void UWorldPartitionRuntimeHash::ExecutePreSetupHLODActors(const UWorldPartition* InWorldPartition, const UWorldPartition::FSetupHLODActorsParams& InParams)
+{
+	// Iterate over all hash types and call PreSetupHLODActors() on each of them
+	for (TObjectIterator<UClass> ClassIterator; ClassIterator; ++ClassIterator)
+	{
+		if (ClassIterator->IsChildOf(UWorldPartitionRuntimeHash::StaticClass()) && !ClassIterator->HasAnyClassFlags(CLASS_Abstract))
+		{
+			CastChecked<UWorldPartitionRuntimeHash>(ClassIterator->GetDefaultObject())->PreSetupHLODActors(InWorldPartition, InParams);
+		}
+	}
+}
+
+void UWorldPartitionRuntimeHash::ExecutePostSetupHLODActors(const UWorldPartition* InWorldPartition, const UWorldPartition::FSetupHLODActorsParams& InParams)
+{
+	// Iterate over all hash types and call PostSetupHLODActors() on each of them
+	for (TObjectIterator<UClass> ClassIterator; ClassIterator; ++ClassIterator)
+	{
+		if (ClassIterator->IsChildOf(UWorldPartitionRuntimeHash::StaticClass()) && !ClassIterator->HasAnyClassFlags(CLASS_Abstract))
+		{
+			CastChecked<UWorldPartitionRuntimeHash>(ClassIterator->GetDefaultObject())->PostSetupHLODActors(InWorldPartition, InParams);
+		}
+	}
+}
+
 #endif
 
 void UWorldPartitionRuntimeHash::FStreamingSourceCells::AddCell(const UWorldPartitionRuntimeCell* Cell, const FWorldPartitionStreamingSource& Source, const FSphericalSector& SourceShape)
diff --git a/Engine/Source/Runtime/Engine/Private/WorldPartition/WorldPartitionStreamingGeneration.cpp b/Engine/Source/Runtime/Engine/Private/WorldPartition/WorldPartitionStreamingGeneration.cpp
--- a/Engine/Source/Runtime/Engine/Private/WorldPartition/WorldPartitionStreamingGeneration.cpp
+++ b/Engine/Source/Runtime/Engine/Private/WorldPartition/WorldPartitionStreamingGeneration.cpp
@@ -2118,6 +2118,8 @@ void UWorldPartition::SetupHLODActors(const FSetupHLODActorsParams& Params)
 		RuntimeHash->SetupHLODActors(StreamingGenerator.GetStreamingGenerationContext(InContainerInstanceCollection), Params);
 	};
 
+	UWorldPartitionRuntimeHash::ExecutePreSetupHLODActors(this, Params);
+
 	// Process all Content Bundle container instances
 	for (const UActorDescContainerInstance* ContentBundleContainerInstance : ContentBundleContainerInstances)
 	{
@@ -2131,6 +2133,8 @@ void UWorldPartition::SetupHLODActors(const FSetupHLODActorsParams& Params)
 		FStreamingGenerationContainerInstanceCollection Collection(BaseAndEDLContainerInstances, FStreamingGenerationContainerInstanceCollection::ECollectionType::BaseAndEDLs);
 		SetupHLODActorsForCollection(Collection);
 	}
+
+	UWorldPartitionRuntimeHash::ExecutePostSetupHLODActors(this, Params);
 }
 
 FStreamingGenerationContainerInstanceCollection::FStreamingGenerationContainerInstanceCollection(std::initializer_list<TObjectPtr<const UActorDescContainerInstance>> ActorDescContainerInstanceArray, const ECollectionType& InCollectionType)
diff --git a/Engine/Source/Runtime/Engine/Public/WorldPartition/WorldPartitionRuntimeHash.h b/Engine/Source/Runtime/Engine/Public/WorldPartition/WorldPartitionRuntimeHash.h
--- a/Engine/Source/Runtime/Engine/Public/WorldPartition/WorldPartitionRuntimeHash.h
+++ b/Engine/Source/Runtime/Engine/Public/WorldPartition/WorldPartitionRuntimeHash.h
@@ -285,5 +285,10 @@ private:
 public:
 	static ENGINE_API void RegisterWorldPartitionRuntimeHashConverter(const UClass* InSrcClass, const UClass* InDstClass, FRuntimeHashConvertFunc&& InConverter);
 	static ENGINE_API UWorldPartitionRuntimeHash* ConvertWorldPartitionHash(const UWorldPartitionRuntimeHash* InSrcHash, const UClass* InDstClass);
+
+	static ENGINE_API void ExecutePreSetupHLODActors(const UWorldPartition* InWorldPartition, const UWorldPartition::FSetupHLODActorsParams& InParams);
+	static ENGINE_API void ExecutePostSetupHLODActors(const UWorldPartition* InWorldPartition, const UWorldPartition::FSetupHLODActorsParams& InParams);
+	ENGINE_API virtual void PreSetupHLODActors(const UWorldPartition* InWorldPartition, const UWorldPartition::FSetupHLODActorsParams& InParams) const {}
+	ENGINE_API virtual void PostSetupHLODActors(const UWorldPartition* InWorldPartition, const UWorldPartition::FSetupHLODActorsParams& InParams) const {}
 #endif
 };
\ No newline at end of file
diff --git a/Engine/Source/Runtime/Engine/Public/WorldPartition/WorldPartitionRuntimeSpatialHash.h b/Engine/Source/Runtime/Engine/Public/WorldPartition/WorldPartitionRuntimeSpatialHash.h
--- a/Engine/Source/Runtime/Engine/Public/WorldPartition/WorldPartitionRuntimeSpatialHash.h
+++ b/Engine/Source/Runtime/Engine/Public/WorldPartition/WorldPartitionRuntimeSpatialHash.h
@@ -321,6 +321,8 @@ public:
 
 	ENGINE_API virtual void SetDefaultValues() override;
 	virtual bool SupportsHLODs() const override { return true; }
+	ENGINE_API virtual void PreSetupHLODActors(const UWorldPartition* InWorldPartition, const UWorldPartition::FSetupHLODActorsParams& InParams) const override;
+	ENGINE_API virtual void PostSetupHLODActors(const UWorldPartition* InWorldPartition, const UWorldPartition::FSetupHLODActorsParams& InParams) const override;
 	ENGINE_API virtual bool SetupHLODActors(const IStreamingGenerationContext* StreamingGenerationContext, const UWorldPartition::FSetupHLODActorsParams& Params) const override;
 	ENGINE_API virtual bool IsValidGrid(FName GridName, const UClass* ActorClass) const override;
 	ENGINE_API virtual bool IsValidHLODLayer(FName GridName, const FSoftObjectPath& HLODLayerPath) const override { return true; }
