ASpatialHashRuntimeGridInfo - Remove Sprite component

diff --git a/Engine/Source/Runtime/Engine/Private/WorldPartition/WorldPartitionRuntimeSpatialHash.cpp b/Engine/Source/Runtime/Engine/Private/WorldPartition/WorldPartitionRuntimeSpatialHash.cpp
--- a/Engine/Source/Runtime/Engine/Private/WorldPartition/WorldPartitionRuntimeSpatialHash.cpp
+++ b/Engine/Source/Runtime/Engine/Private/WorldPartition/WorldPartitionRuntimeSpatialHash.cpp
@@ -1149,7 +1149,7 @@ bool FSpatialHashRuntimeGrid::operator != (const FSpatialHashRuntimeGrid& Other)
 // ------------------------------------------------------------------------------------------------
 
 ASpatialHashRuntimeGridInfo::ASpatialHashRuntimeGridInfo(const FObjectInitializer& ObjectInitializer)
-	: Super(ObjectInitializer)
+	: Super(ObjectInitializer.DoNotCreateDefaultSubobject(TEXT("Sprite")))
 {
 #if WITH_EDITORONLY_DATA
 	bListedInSceneOutliner = false;
