[World Partition HashSet] Dump cells content and cell bounds into the streaming generation logs.

diff --git a/Engine/Source/Runtime/Engine/Private/WorldPartition/RuntimeHashSet/WorldPartitionRuntimeCellDataHashSet.cpp b/Engine/Source/Runtime/Engine/Private/WorldPartition/RuntimeHashSet/WorldPartitionRuntimeCellDataHashSet.cpp
new file mode 100644
--- /dev/null
+++ b/Engine/Source/Runtime/Engine/Private/WorldPartition/RuntimeHashSet/WorldPartitionRuntimeCellDataHashSet.cpp
@@ -0,0 +1,12 @@
+// Copyright Epic Games, Inc. All Rights Reserved.
+
+#include "WorldPartition/RuntimeHashSet/WorldPartitionRuntimeCellDataHashSet.h"
+
+#if WITH_EDITOR
+void UWorldPartitionRuntimeCellDataHashSet::DumpStateLog(FHierarchicalLogArchive& Ar) const
+{
+	Super::DumpStateLog(Ar);
+
+	Ar.Printf(TEXT("Is 2D: %s"), bIs2D ? TEXT("True") : TEXT("False"));
+}
+#endif
diff --git a/Engine/Source/Runtime/Engine/Private/WorldPartition/WorldPartitionRuntimeCell.cpp b/Engine/Source/Runtime/Engine/Private/WorldPartition/WorldPartitionRuntimeCell.cpp
--- a/Engine/Source/Runtime/Engine/Private/WorldPartition/WorldPartitionRuntimeCell.cpp
+++ b/Engine/Source/Runtime/Engine/Private/WorldPartition/WorldPartitionRuntimeCell.cpp
@@ -125,6 +125,8 @@ void UWorldPartitionRuntimeCell::SetDataLayers(const TArray<const UDataLayerInst
 void UWorldPartitionRuntimeCell::DumpStateLog(FHierarchicalLogArchive& Ar) const
 {
 	Ar.Printf(TEXT("Actor Count: %d"), GetActorCount());
+
+	RuntimeCellData->DumpStateLog(Ar);
 }
 #endif
 
diff --git a/Engine/Source/Runtime/Engine/Private/WorldPartition/WorldPartitionRuntimeCellData.cpp b/Engine/Source/Runtime/Engine/Private/WorldPartition/WorldPartitionRuntimeCellData.cpp
--- a/Engine/Source/Runtime/Engine/Private/WorldPartition/WorldPartitionRuntimeCellData.cpp
+++ b/Engine/Source/Runtime/Engine/Private/WorldPartition/WorldPartitionRuntimeCellData.cpp
@@ -1,6 +1,7 @@
 // Copyright Epic Games, Inc. All Rights Reserved.
 
 #include "WorldPartition/WorldPartitionRuntimeCellData.h"
+#include "Misc/HierarchicalLogArchive.h"
 
 int32 UWorldPartitionRuntimeCellData::StreamingSourceCacheEpoch = 0;
 
@@ -23,6 +24,18 @@ void UWorldPartitionRuntimeCellData::Serialize(FArchive& Ar)
 	Ar << DebugName;
 }
 
+#if WITH_EDITOR
+void UWorldPartitionRuntimeCellData::DumpStateLog(FHierarchicalLogArchive& Ar) const
+{
+	Ar.Printf(TEXT("Content Bounds: %s"), *ContentBounds.ToString());
+
+	if (CellBounds.IsSet())
+	{
+		Ar.Printf(TEXT("Cell Bounds: %s"), *CellBounds.GetValue().ToString());
+	}
+}
+#endif
+
 void UWorldPartitionRuntimeCellData::ResetStreamingSourceInfo() const
 {
 	CachedMinSourcePriority = MAX_uint8;
diff --git a/Engine/Source/Runtime/Engine/Public/WorldPartition/RuntimeHashSet/WorldPartitionRuntimeCellDataHashSet.h b/Engine/Source/Runtime/Engine/Public/WorldPartition/RuntimeHashSet/WorldPartitionRuntimeCellDataHashSet.h
--- a/Engine/Source/Runtime/Engine/Public/WorldPartition/RuntimeHashSet/WorldPartitionRuntimeCellDataHashSet.h
+++ b/Engine/Source/Runtime/Engine/Public/WorldPartition/RuntimeHashSet/WorldPartitionRuntimeCellDataHashSet.h
@@ -11,6 +11,10 @@ class UWorldPartitionRuntimeCellDataHashSet : public UWorldPartitionRuntimeCellD
 	GENERATED_BODY()
 
 public:
+#if WITH_EDITOR
+	ENGINE_API virtual void DumpStateLog(FHierarchicalLogArchive& Ar) const;
+#endif
+
 	UPROPERTY()
 	bool bIs2D;
 };
diff --git a/Engine/Source/Runtime/Engine/Public/WorldPartition/WorldPartitionRuntimeCellData.h b/Engine/Source/Runtime/Engine/Public/WorldPartition/WorldPartitionRuntimeCellData.h
--- a/Engine/Source/Runtime/Engine/Public/WorldPartition/WorldPartitionRuntimeCellData.h
+++ b/Engine/Source/Runtime/Engine/Public/WorldPartition/WorldPartitionRuntimeCellData.h
@@ -3,6 +3,7 @@
 
 #include "CoreMinimal.h"
 #include "StringDev.h"
+#include "Misc/HierarchicalLogArchive.h"
 #include "WorldPartition/WorldPartitionStreamingSource.h"
 #include "WorldPartitionRuntimeCellData.generated.h"
 
@@ -20,6 +21,10 @@ class UWorldPartitionRuntimeCellData : public UObject
 	ENGINE_API void Serialize(FArchive& Ar);
 	//~End UObject Interface
 
+#if WITH_EDITOR
+	ENGINE_API virtual void DumpStateLog(FHierarchicalLogArchive& Ar) const;
+#endif
+
 	ENGINE_API virtual void ResetStreamingSourceInfo() const;
 	ENGINE_API virtual void AppendStreamingSourceInfo(const FWorldPartitionStreamingSource& Source, const FSphericalSector& SourceShape) const;
 	ENGINE_API virtual void MergeStreamingSourceInfo() const {}
