[World Partition HashSet] Fix content bounds not properly computed for level instances.

diff --git a/Engine/Source/Runtime/Engine/Private/WorldPartition/RuntimeHashSet/RuntimePartitionLHGrid.cpp b/Engine/Source/Runtime/Engine/Private/WorldPartition/RuntimeHashSet/RuntimePartitionLHGrid.cpp
--- a/Engine/Source/Runtime/Engine/Private/WorldPartition/RuntimeHashSet/RuntimePartitionLHGrid.cpp
+++ b/Engine/Source/Runtime/Engine/Private/WorldPartition/RuntimeHashSet/RuntimePartitionLHGrid.cpp
@@ -200,6 +200,8 @@ bool URuntimePartitionLHGrid::GenerateStreaming(const FGenerateStreamingParams&
 		if (bIsSpatiallyLoaded)
 		{
 			CellDesc.CellBounds = FCellCoord::GetCellBounds(CellCoord, CellSize, Origin * SpaceMask);
+			CellDesc.CellBounds.GetValue().Min *= SpaceMask;
+			CellDesc.CellBounds.GetValue().Max *= SpaceMask;
 			CellDesc.bIs2D = bIs2D;
 		}
 	}
diff --git a/Engine/Source/Runtime/Engine/Private/WorldPartition/RuntimeHashSet/WorldPartitionRuntimeHashSetStreamingGeneration.cpp b/Engine/Source/Runtime/Engine/Private/WorldPartition/RuntimeHashSet/WorldPartitionRuntimeHashSetStreamingGeneration.cpp
--- a/Engine/Source/Runtime/Engine/Private/WorldPartition/RuntimeHashSet/WorldPartitionRuntimeHashSetStreamingGeneration.cpp
+++ b/Engine/Source/Runtime/Engine/Private/WorldPartition/RuntimeHashSet/WorldPartitionRuntimeHashSetStreamingGeneration.cpp
@@ -187,6 +187,12 @@ bool UWorldPartitionRuntimeHashSet::GenerateStreaming(UWorldPartitionStreamingPo
 				RuntimeCell->SetIsAlwaysLoaded(bIsCellAlwaysLoaded);
 				PopulateRuntimeCell(RuntimeCell, CellActorInstances, OutPackagesToGenerate);
 
+				UWorldPartitionRuntimeCellDataHashSet* RuntimeCellData = CastChecked<UWorldPartitionRuntimeCellDataHashSet>(RuntimeCell->RuntimeCellData);
+				const FVector SpaceMask(1, 1, RuntimeCellData->bIs2D ? 0 : 1);
+
+				RuntimeCell->RuntimeCellData->ContentBounds.Min *= SpaceMask;
+				RuntimeCell->RuntimeCellData->ContentBounds.Max *= SpaceMask;
+
 				if (CellDescInstance.CellBounds.IsSet())
 				{
 					switch (CellDescInstance.SourcePartition->BoundsMethod)
@@ -198,6 +204,7 @@ bool UWorldPartitionRuntimeHashSet::GenerateStreaming(UWorldPartitionStreamingPo
 						if (RuntimeCell->RuntimeCellData->ContentBounds.IsValid)
 						{
 							RuntimeCell->RuntimeCellData->ContentBounds = RuntimeCell->RuntimeCellData->ContentBounds.Overlap(CellDescInstance.CellBounds.GetValue());
+							check(CellDescInstance.CellBounds.GetValue().IsValid);
 						}
 						break;
 					}
diff --git a/Engine/Source/Runtime/Engine/Private/WorldPartition/WorldPartitionRuntimeHash.cpp b/Engine/Source/Runtime/Engine/Private/WorldPartition/WorldPartitionRuntimeHash.cpp
--- a/Engine/Source/Runtime/Engine/Private/WorldPartition/WorldPartitionRuntimeHash.cpp
+++ b/Engine/Source/Runtime/Engine/Private/WorldPartition/WorldPartitionRuntimeHash.cpp
@@ -336,12 +336,9 @@ void UWorldPartitionRuntimeHash::PopulateRuntimeCell(UWorldPartitionRuntimeCell*
 	{
 		const FStreamingGenerationActorDescView& ActorDescView = ActorInstance.GetActorDescView();
 		RuntimeCell->AddActorToCell(ActorDescView);
-		const FBox RuntimeBounds = ActorDescView.GetRuntimeBounds();
-		if (RuntimeBounds.IsValid)
-		{
-			CellContentBounds += RuntimeBounds.TransformBy(ActorInstance.GetTransform());
-		}
-					
+
+		CellContentBounds += ActorInstance.GetBounds();
+
 		if (ActorInstance.GetContainerID().IsMainContainer() && RuntimeCell->UnsavedActorsContainer)
 		{
 			if (AActor* Actor = FindObject<AActor>(nullptr, *ActorDescView.GetActorSoftPath().ToString()))
