[World Partition HashSet] Recreate LHGrid world grid previewer when changing the bIs2D flag to reflect the change.

diff --git a/Engine/Source/Runtime/Engine/Private/WorldGridPreviewer.cpp b/Engine/Source/Runtime/Engine/Private/WorldGridPreviewer.cpp
--- a/Engine/Source/Runtime/Engine/Private/WorldGridPreviewer.cpp
+++ b/Engine/Source/Runtime/Engine/Private/WorldGridPreviewer.cpp
@@ -74,7 +74,7 @@ void FWorldGridPreviewer::Update()
 		MaterialInstance->SetScalarParameterValueEditorOnly(TEXT("Grid_CellSize"), (float)CellSize);
 		MaterialInstance->SetScalarParameterValueEditorOnly(TEXT("Grid_LoadingRange"), (float)LoadingRange);
 		MaterialInstance->SetVectorParameterValueEditorOnly(TEXT("Grid_Color"), GridColor);
-		// Grid material has cell's orinig at center, and our grid are using top-left
+		// Grid material has cell's origin at center, and our grid are using top-left
 		MaterialInstance->SetVectorParameterValueEditorOnly(TEXT("Grid_Offset"), FLinearColor(GridOffset + CellSize * 0.5f));
 	}
 }
diff --git a/Engine/Source/Runtime/Engine/Private/WorldPartition/RuntimeHashSet/RuntimePartitionLHGrid.cpp b/Engine/Source/Runtime/Engine/Private/WorldPartition/RuntimeHashSet/RuntimePartitionLHGrid.cpp
--- a/Engine/Source/Runtime/Engine/Private/WorldPartition/RuntimeHashSet/RuntimePartitionLHGrid.cpp
+++ b/Engine/Source/Runtime/Engine/Private/WorldPartition/RuntimeHashSet/RuntimePartitionLHGrid.cpp
@@ -108,12 +108,12 @@ void URuntimePartitionLHGrid::PostEditChangeProperty(FPropertyChangedEvent& InPr
 	{
 		CellSize = FMath::Max<int32>(CellSize, 1600);
 	}
-	else if (PropertyName == GET_MEMBER_NAME_CHECKED(URuntimePartitionLHGrid, bShowGridPreview))
+	else if ((PropertyName == GET_MEMBER_NAME_CHECKED(URuntimePartitionLHGrid, bShowGridPreview)) || 
+			 (PropertyName == GET_MEMBER_NAME_CHECKED(URuntimePartitionLHGrid, bIs2D)))
 	{
 		if (bShowGridPreview)
 		{
-			check(!WorldGridPreviewer);
-			WorldGridPreviewer = MakeUnique<FWorldGridPreviewer>(GetTypedOuter<UWorld>(), false);
+			WorldGridPreviewer = MakeUnique<FWorldGridPreviewer>(GetTypedOuter<UWorld>(), bIs2D);
 		}
 		else
 		{
