Fix Clang Static Analysis warning about storage aligned mismatch.

diff --git a/Engine/Source/Runtime/Engine/Private/Tests/WorldPartition/WorldPartitionHandleTests.cpp b/Engine/Source/Runtime/Engine/Private/Tests/WorldPartition/WorldPartitionHandleTests.cpp
--- a/Engine/Source/Runtime/Engine/Private/Tests/WorldPartition/WorldPartitionHandleTests.cpp
+++ b/Engine/Source/Runtime/Engine/Private/Tests/WorldPartition/WorldPartitionHandleTests.cpp
@@ -131,13 +131,13 @@ namespace WorldPartitionTests
 
 		// inplace new test
 		{
-			uint8 Buffer[sizeof(FWorldPartitionHandle)];
-			FWorldPartitionHandle* HandlePtr = new (Buffer) FWorldPartitionHandle(Reference.ToHandle());
+			TOptional<FWorldPartitionHandle> HandlePtr;
+			HandlePtr.Emplace(Reference.ToHandle());
 
 			TestTrue(TEXT("Handle array soft refcount"), FWorldPartitionActorDescUnitTestAcccessor::GetSoftRefCount(*Reference) == 1);
 			TestTrue(TEXT("Handle array hard refcount"), FWorldPartitionActorDescUnitTestAcccessor::GetHardRefCount(*Reference) == 1);
 
-			HandlePtr->~FWorldPartitionHandle();
+			HandlePtr.Reset();
 
 			TestTrue(TEXT("Handle array soft refcount"), FWorldPartitionActorDescUnitTestAcccessor::GetSoftRefCount(*Reference) == 0);
 			TestTrue(TEXT("Handle array hard refcount"), FWorldPartitionActorDescUnitTestAcccessor::GetHardRefCount(*Reference) == 1);
