[Material Editor] New Material Translator advances.

* Added constants compilation
* Added binary operators
* Added conversion instructions (casting)
* General improvements and refactorings

diff --git a/Engine/Source/Editor/UnrealEd/Private/MaterialGraph.cpp b/Engine/Source/Editor/UnrealEd/Private/MaterialGraph.cpp
--- a/Engine/Source/Editor/UnrealEd/Private/MaterialGraph.cpp
+++ b/Engine/Source/Editor/UnrealEd/Private/MaterialGraph.cpp
@@ -153,7 +153,7 @@ void UMaterialGraph::UpdatePinTypes()
 {
 	using namespace UE::HLSLTree;
 
-	if (!Material->IsUsingNewHLSLGenerator())
+	if (!Material->IsUsingNewHLSLGenerator() || Material->IsUsingNewTranslatorPrototype())
 	{
 		return;
 	}
diff --git a/Engine/Source/Runtime/Engine/Private/Materials/Material.cpp b/Engine/Source/Runtime/Engine/Private/Materials/Material.cpp
--- a/Engine/Source/Runtime/Engine/Private/Materials/Material.cpp
+++ b/Engine/Source/Runtime/Engine/Private/Materials/Material.cpp
@@ -84,6 +84,7 @@
 #include "Materials/MaterialExpressionConstant.h"
 #include "Materials/MaterialExpressionConstant3Vector.h"
 #include "Materials/MaterialExpressionBreakMaterialAttributes.h"
+#include "Materials/MaterialIRModuleBuilder.h"
 #include "MaterialCachedData.h"
 #include "Misc/OutputDeviceArchiveWrapper.h"
 #include "HAL/FileManager.h"
@@ -2084,9 +2085,19 @@ void UMaterial::UpdateCachedExpressionData()
 		// Relinks function call inputs. Otherwise, we can get invalid inputs and they will cause errors when generating the syntax tree
 		UpdateTransientExpressionData();
 
-		LocalCachedTree = new FMaterialCachedHLSLTree();
-		LocalCachedTree->GenerateTree(this, nullptr, nullptr);
-		LocalCachedExpressionData->UpdateForCachedHLSLTree(*LocalCachedTree, nullptr, this);
+		if (bIsUsingNewMaterialTranslatorPrototype)
+		{
+			FMaterialIRModuleBuilder Builder;
+			Builder.Build(this, &IRModule);
+
+			LocalCachedExpressionData->AnalyzeMaterial(*this); // Temporary
+		}
+		else
+		{
+			LocalCachedTree = new FMaterialCachedHLSLTree();
+			LocalCachedTree->GenerateTree(this, nullptr, nullptr);
+			LocalCachedExpressionData->UpdateForCachedHLSLTree(*LocalCachedTree, nullptr, this);
+		}
 	}
 	else
 	{
diff --git a/Engine/Source/Runtime/Engine/Private/Materials/MaterialExpressionsIR.cpp b/Engine/Source/Runtime/Engine/Private/Materials/MaterialExpressionsIR.cpp
new file mode 100644
--- /dev/null
+++ b/Engine/Source/Runtime/Engine/Private/Materials/MaterialExpressionsIR.cpp
@@ -0,0 +1,350 @@
+// Copyright Epic Games, Inc. All Rights Reserved.
+
+#if WITH_EDITOR
+
+#include "Curves/CurveLinearColorAtlas.h"
+#include "Engine/Engine.h"
+#include "Engine/Font.h"
+#include "MaterialDomain.h"
+#include "MaterialExpressionIO.h"
+#include "MaterialHLSLGenerator.h"
+#include "MaterialHLSLTree.h"
+#include "Materials/Material.h"
+#include "Materials/MaterialAttributeDefinitionMap.h"
+#include "Materials/MaterialExpression.h"
+#include "Materials/MaterialExpressionAbs.h"
+#include "Materials/MaterialExpressionAbsorptionMediumMaterialOutput.h"
+#include "Materials/MaterialExpressionActorPositionWS.h"
+#include "Materials/MaterialExpressionAdd.h"
+#include "Materials/MaterialExpressionAntialiasedTextureMask.h"
+#include "Materials/MaterialExpressionAppendVector.h"
+#include "Materials/MaterialExpressionArccosine.h"
+#include "Materials/MaterialExpressionArccosineFast.h"
+#include "Materials/MaterialExpressionArcsine.h"
+#include "Materials/MaterialExpressionArcsineFast.h"
+#include "Materials/MaterialExpressionArctangent.h"
+#include "Materials/MaterialExpressionArctangent2.h"
+#include "Materials/MaterialExpressionArctangent2Fast.h"
+#include "Materials/MaterialExpressionArctangentFast.h"
+#include "Materials/MaterialExpressionAtmosphericLightColor.h"
+#include "Materials/MaterialExpressionAtmosphericLightVector.h"
+#include "Materials/MaterialExpressionBentNormalCustomOutput.h"
+#include "Materials/MaterialExpressionBinaryOp.h"
+#include "Materials/MaterialExpressionBlackBody.h"
+#include "Materials/MaterialExpressionBlendMaterialAttributes.h"
+#include "Materials/MaterialExpressionBreakMaterialAttributes.h"
+#include "Materials/MaterialExpressionBumpOffset.h"
+#include "Materials/MaterialExpressionCameraPositionWS.h"
+#include "Materials/MaterialExpressionCameraVectorWS.h"
+#include "Materials/MaterialExpressionCeil.h"
+#include "Materials/MaterialExpressionChannelMaskParameter.h"
+#include "Materials/MaterialExpressionClamp.h"
+#include "Materials/MaterialExpressionClearCoatNormalCustomOutput.h"
+#include "Materials/MaterialExpressionCloudLayer.h"
+#include "Materials/MaterialExpressionCollectionParameter.h"
+#include "Materials/MaterialExpressionComponentMask.h"
+#include "Materials/MaterialExpressionConstant.h"
+#include "Materials/MaterialExpressionConstant2Vector.h"
+#include "Materials/MaterialExpressionConstant3Vector.h"
+#include "Materials/MaterialExpressionConstant4Vector.h"
+#include "Materials/MaterialExpressionConstantBiasScale.h"
+#include "Materials/MaterialExpressionCosine.h"
+#include "Materials/MaterialExpressionCrossProduct.h"
+#include "Materials/MaterialExpressionCurveAtlasRowParameter.h"
+#include "Materials/MaterialExpressionCustom.h"
+#include "Materials/MaterialExpressionDataDrivenShaderPlatformInfoSwitch.h"
+#include "Materials/MaterialExpressionDBufferTexture.h"
+#include "Materials/MaterialExpressionDDX.h"
+#include "Materials/MaterialExpressionDDY.h"
+#include "Materials/MaterialExpressionDecalColor.h"
+#include "Materials/MaterialExpressionDecalDerivative.h"
+#include "Materials/MaterialExpressionDecalLifetimeOpacity.h"
+#include "Materials/MaterialExpressionDecalMipmapLevel.h"
+#include "Materials/MaterialExpressionDeltaTime.h"
+#include "Materials/MaterialExpressionDepthFade.h"
+#include "Materials/MaterialExpressionDepthOfFieldFunction.h"
+#include "Materials/MaterialExpressionDeriveNormalZ.h"
+#include "Materials/MaterialExpressionDesaturation.h"
+#include "Materials/MaterialExpressionDistance.h"
+#include "Materials/MaterialExpressionDistanceCullFade.h"
+#include "Materials/MaterialExpressionDistanceFieldApproxAO.h"
+#include "Materials/MaterialExpressionDistanceFieldGradient.h"
+#include "Materials/MaterialExpressionDistanceFieldsRenderingSwitch.h"
+#include "Materials/MaterialExpressionDistanceToNearestSurface.h"
+#include "Materials/MaterialExpressionDivide.h"
+#include "Materials/MaterialExpressionDotProduct.h"
+#include "Materials/MaterialExpressionDoubleVectorParameter.h"
+#include "Materials/MaterialExpressionDynamicParameter.h"
+#include "Materials/MaterialExpressionExecBegin.h"
+#include "Materials/MaterialExpressionExecEnd.h"
+#include "Materials/MaterialExpressionExponential.h"
+#include "Materials/MaterialExpressionExponential2.h"
+#include "Materials/MaterialExpressionEyeAdaptation.h"
+#include "Materials/MaterialExpressionEyeAdaptationInverse.h"
+#include "Materials/MaterialExpressionFeatureLevelSwitch.h"
+#include "Materials/MaterialExpressionFloor.h"
+#include "Materials/MaterialExpressionFmod.h"
+#include "Materials/MaterialExpressionAtmosphericFogColor.h"
+#include "Materials/MaterialExpressionFontSample.h"
+#include "Materials/MaterialExpressionForLoop.h"
+#include "Materials/MaterialExpressionFrac.h"
+#include "Materials/MaterialExpressionFresnel.h"
+#include "Materials/MaterialExpressionFunctionInput.h"
+#include "Materials/MaterialExpressionFunctionOutput.h"
+#include "Materials/MaterialExpressionGenericConstant.h"
+#include "Materials/MaterialExpressionGetLocal.h"
+#include "Materials/MaterialExpressionGetMaterialAttributes.h"
+#include "Materials/MaterialExpressionGIReplace.h"
+#include "Materials/MaterialExpressionHairAttributes.h"
+#include "Materials/MaterialExpressionHairColor.h"
+#include "Materials/MaterialExpressionHsvToRgb.h"
+#include "Materials/MaterialExpressionIf.h"
+#include "Materials/MaterialExpressionIfThenElse.h"
+#include "Materials/MaterialExpressionInverseLinearInterpolate.h"
+#include "Materials/MaterialExpressionIsOrthographic.h"
+#include "Materials/MaterialExpressionLength.h"
+#include "Materials/MaterialExpressionLightmapUVs.h"
+#include "Materials/MaterialExpressionLightVector.h"
+#include "Materials/MaterialExpressionLinearInterpolate.h"
+#include "Materials/MaterialExpressionLightmassReplace.h"
+#include "Materials/MaterialExpressionLocalPosition.h"
+#include "Materials/MaterialExpressionLogarithm.h"
+#include "Materials/MaterialExpressionLogarithm10.h"
+#include "Materials/MaterialExpressionLogarithm2.h"
+#include "Materials/MaterialExpressionMakeMaterialAttributes.h"
+#include "Materials/MaterialExpressionMapARPassthroughCameraUV.h"
+#include "Materials/MaterialExpressionMaterialAttributeLayers.h"
+#include "Materials/MaterialExpressionMaterialFunctionCall.h"
+#include "Materials/MaterialExpressionMaterialProxyReplace.h"
+#include "Materials/MaterialExpressionMax.h"
+#include "Materials/MaterialExpressionMin.h"
+#include "Materials/MaterialExpressionModulo.h"
+#include "Materials/MaterialExpressionMultiply.h"
+#include "Materials/MaterialExpressionNamedReroute.h"
+#include "Materials/MaterialExpressionNaniteReplace.h"
+#include "Materials/MaterialExpressionNoise.h"
+#include "Materials/MaterialExpressionNormalize.h"
+#include "Materials/MaterialExpressionNeuralPostProcessNode.h"
+#include "Materials/MaterialExpressionObjectBounds.h"
+#include "Materials/MaterialExpressionObjectLocalBounds.h"
+#include "Materials/MaterialExpressionBounds.h"
+#include "Materials/MaterialExpressionObjectOrientation.h"
+#include "Materials/MaterialExpressionObjectPositionWS.h"
+#include "Materials/MaterialExpressionObjectRadius.h"
+#include "Materials/MaterialExpressionOneMinus.h"
+#include "Materials/MaterialExpressionPanner.h"
+#include "Materials/MaterialExpressionParameter.h"
+#include "Materials/MaterialExpressionParticleColor.h"
+#include "Materials/MaterialExpressionParticleDirection.h"
+#include "Materials/MaterialExpressionParticleMacroUV.h"
+#include "Materials/MaterialExpressionParticleMotionBlurFade.h"
+#include "Materials/MaterialExpressionParticlePositionWS.h"
+#include "Materials/MaterialExpressionParticleRadius.h"
+#include "Materials/MaterialExpressionParticleRandom.h"
+#include "Materials/MaterialExpressionParticleRelativeTime.h"
+#include "Materials/MaterialExpressionParticleSize.h"
+#include "Materials/MaterialExpressionParticleSpeed.h"
+#include "Materials/MaterialExpressionParticleSubUVProperties.h"
+#include "Materials/MaterialExpressionPathTracingBufferTexture.h"
+#include "Materials/MaterialExpressionPathTracingQualitySwitch.h"
+#include "Materials/MaterialExpressionPathTracingRayTypeSwitch.h"
+#include "Materials/MaterialExpressionPerInstanceCustomData.h"
+#include "Materials/MaterialExpressionPerInstanceFadeAmount.h"
+#include "Materials/MaterialExpressionPerInstanceRandom.h"
+#include "Materials/MaterialExpressionPixelDepth.h"
+#include "Materials/MaterialExpressionPixelNormalWS.h"
+#include "Materials/MaterialExpressionPower.h"
+#include "Materials/MaterialExpressionPrecomputedAOMask.h"
+#include "Materials/MaterialExpressionPreSkinnedLocalBounds.h"
+#include "Materials/MaterialExpressionPreSkinnedNormal.h"
+#include "Materials/MaterialExpressionPreSkinnedPosition.h"
+#include "Materials/MaterialExpressionPreviousFrameSwitch.h"
+#include "Materials/MaterialExpressionSamplePhysicsField.h"
+#include "Materials/MaterialExpressionSphericalParticleOpacity.h"
+#include "Materials/MaterialExpressionQualitySwitch.h"
+#include "Materials/MaterialExpressionRayTracingQualitySwitch.h"
+#include "Materials/MaterialExpressionReflectionCapturePassSwitch.h"
+#include "Materials/MaterialExpressionReflectionVectorWS.h"
+#include "Materials/MaterialExpressionReroute.h"
+#include "Materials/MaterialExpressionRgbToHsv.h"
+#include "Materials/MaterialExpressionRotateAboutAxis.h"
+#include "Materials/MaterialExpressionRotator.h"
+#include "Materials/MaterialExpressionRound.h"
+#include "Materials/MaterialExpressionRuntimeVirtualTextureOutput.h"
+#include "Materials/MaterialExpressionRuntimeVirtualTextureReplace.h"
+#include "Materials/MaterialExpressionRuntimeVirtualTextureSample.h"
+#include "Materials/MaterialExpressionRuntimeVirtualTextureSampleParameter.h"
+#include "Materials/MaterialExpressionSaturate.h"
+#include "Materials/MaterialExpressionScalarParameter.h"
+#include "Materials/MaterialExpressionSceneColor.h"
+#include "Materials/MaterialExpressionSceneDepth.h"
+#include "Materials/MaterialExpressionSceneDepthWithoutWater.h"
+#include "Materials/MaterialExpressionSceneTexelSize.h"
+#include "Materials/MaterialExpressionSceneTexture.h"
+#include "Materials/MaterialExpressionScreenPosition.h"
+#include "Materials/MaterialExpressionSetLocal.h"
+#include "Materials/MaterialExpressionSetMaterialAttributes.h"
+#include "Materials/MaterialExpressionShaderStageSwitch.h"
+#include "Materials/MaterialExpressionShadingModel.h"
+#include "Materials/MaterialExpressionShadingPathSwitch.h"
+#include "Materials/MaterialExpressionShadowReplace.h"
+#include "Materials/MaterialExpressionSingleLayerWaterMaterialOutput.h"
+#include "Materials/MaterialExpressionSign.h"
+#include "Materials/MaterialExpressionSine.h"
+#include "Materials/MaterialExpressionSkyAtmosphereLightDirection.h"
+#include "Materials/MaterialExpressionSkyAtmosphereLightIlluminance.h"
+#include "Materials/MaterialExpressionSkyAtmosphereViewLuminance.h"
+#include "Materials/MaterialExpressionSkyLightEnvMapSample.h"
+#include "Materials/MaterialExpressionSmoothStep.h"
+#include "Materials/MaterialExpressionSobol.h"
+#include "Materials/MaterialExpressionSpeedTree.h"
+#include "Materials/MaterialExpressionSphereMask.h"
+#include "Materials/MaterialExpressionSquareRoot.h"
+#include "Materials/MaterialExpressionSRGBColorToWorkingColorSpace.h"
+#include "Materials/MaterialExpressionStaticBool.h"
+#include "Materials/MaterialExpressionStaticComponentMaskParameter.h"
+#include "Materials/MaterialExpressionStaticSwitch.h"
+#include "Materials/MaterialExpressionStaticSwitchParameter.h"
+#include "Materials/MaterialExpressionStep.h"
+#include "Materials/MaterialExpressionSubtract.h"
+#include "Materials/MaterialExpressionSwitch.h"
+#include "Materials/MaterialExpressionTangent.h"
+#include "Materials/MaterialExpressionTangentOutput.h"
+#include "Materials/MaterialExpressionTemporalSobol.h"
+#include "Materials/MaterialExpressionTextureCollection.h"
+#include "Materials/MaterialExpressionTextureCollectionParameter.h"
+#include "Materials/MaterialExpressionTextureCoordinate.h"
+#include "Materials/MaterialExpressionTextureObject.h"
+#include "Materials/MaterialExpressionTextureObjectFromCollection.h"
+#include "Materials/MaterialExpressionTextureObjectParameter.h"
+#include "Materials/MaterialExpressionTextureProperty.h"
+#include "Materials/MaterialExpressionTextureSample.h"
+#include "Materials/MaterialExpressionTextureSampleParameter.h"
+#include "Materials/MaterialExpressionThinTranslucentMaterialOutput.h"
+#include "Materials/MaterialExpressionTime.h"
+#include "Materials/MaterialExpressionTransform.h"
+#include "Materials/MaterialExpressionTransformPosition.h"
+#include "Materials/MaterialExpressionTruncate.h"
+#include "Materials/MaterialExpressionTruncateLWC.h"
+#include "Materials/MaterialExpressionTwoSidedSign.h"
+#include "Materials/MaterialExpressionVectorNoise.h"
+#include "Materials/MaterialExpressionVectorParameter.h"
+#include "Materials/MaterialExpressionVertexColor.h"
+#include "Materials/MaterialExpressionVertexInterpolator.h"
+#include "Materials/MaterialExpressionVertexNormalWS.h"
+#include "Materials/MaterialExpressionVertexTangentWS.h"
+#include "Materials/MaterialExpressionViewProperty.h"
+#include "Materials/MaterialExpressionViewSize.h"
+#include "Materials/MaterialExpressionVirtualTextureFeatureSwitch.h"
+#include "Materials/MaterialExpressionVolumetricAdvancedMaterialInput.h"
+#include "Materials/MaterialExpressionVolumetricAdvancedMaterialOutput.h"
+#include "Materials/MaterialExpressionWhileLoop.h"
+#include "Materials/MaterialExpressionWorldPosition.h"
+#include "Materials/MaterialFunctionInterface.h"
+#include "Materials/MaterialParameterCollection.h"
+#include "MaterialShared.h"
+#include "Misc/MemStackUtility.h"
+#include "RenderUtils.h"
+#include "ColorManagement/ColorSpace.h"
+
+#include "Materials/MaterialIREmitter.h"
+
+namespace IR = UE::MIR;
+
+/* Constants */
+
+void UMaterialExpression::Build(IR::FEmitter& Emitter)
+{
+	Emitter.Error(TEXT("Node does not support expressions"));
+}
+
+void UMaterialExpressionConstant::Build(IR::FEmitter& Emitter)
+{
+	IR::FValuePtr Value = Emitter.EmitConstantFloat1(R);
+	Emitter.Put(GetOutput(0), Value);
+}
+
+void UMaterialExpressionConstant2Vector::Build(IR::FEmitter& Emitter)
+{
+	IR::FValuePtr Value = Emitter.EmitConstantFloat2({ R, G });
+	Emitter.Put(GetOutput(0), Value);
+}
+
+void UMaterialExpressionConstant3Vector::Build(IR::FEmitter& Emitter)
+{
+	IR::FValuePtr Value = Emitter.EmitConstantFloat3({ Constant.R, Constant.G, Constant.B });
+	Emitter.Put(GetOutput(0), Value);
+}
+
+void UMaterialExpressionConstant4Vector::Build(IR::FEmitter& Emitter)
+{
+	IR::FValuePtr Value = Emitter.EmitConstantFloat4(Constant);
+	Emitter.Put(GetOutput(0), Value);
+}
+
+/* Mathematical Operations */
+
+static void BuildBinaryArithmeticOperator(
+	IR::FEmitter& Emitter,
+	IR::EBinaryOperator Op,
+	FExpressionInput* LhsInput,
+	float LhsConst,
+	FExpressionInput* RhsInput,
+	float RhsConst,
+	FExpressionOutput* Output)
+{
+	// Get values flowing into inputs and check that they're arithmetic.
+	IR::FValuePtr LhsValue = Emitter.GetAndCheckInputTypeKindIs(LhsInput, IR::TK_Arithmetic);
+	IR::FValuePtr RhsValue = Emitter.GetAndCheckInputTypeKindIs(RhsInput, IR::TK_Arithmetic);
+	if (Emitter.IsInvalid()) { return; }
+
+	// Determine operation input/output type by looking at the first connected input and picking float1 otherwise.
+	IR::FTypePtr ResultType = LhsValue ? LhsValue->Type
+		: RhsValue ? RhsValue->Type
+		: IR::FArithmeticType::GetScalar(IR::SK_Float);
+
+	// If lhs operand is disconnected, create a value for its constant property instead.
+	if (!LhsValue)
+	{
+		LhsValue = Emitter.EmitConstantFloat1(LhsConst);
+	}
+
+	// If rhs operand is disconnected, create a value for its constant property instead.
+	if (!RhsValue)
+	{
+		RhsValue = Emitter.EmitConstantFloat1(RhsConst);
+	}
+
+	// Convert operand values to determined result type. 
+	LhsValue = Emitter.TryEmitConvert(LhsValue, ResultType);
+	RhsValue = Emitter.TryEmitConvert(RhsValue, ResultType);
+	if (Emitter.IsInvalid()) { return; }
+
+	// Finally emit the binary operator.
+	IR::FValuePtr Value = Emitter.EmitBinaryOperator(Op, LhsValue, RhsValue);
+
+	// And flow it out from our only output.
+	Emitter.Put(Output, Value);
+}
+
+void UMaterialExpressionAdd::Build(IR::FEmitter& Emitter)
+{ 
+	BuildBinaryArithmeticOperator(Emitter, IR::BO_Add, &A, ConstA, &B, ConstB, GetOutput(0));
+}
+
+void UMaterialExpressionSubtract::Build(IR::FEmitter& Emitter)
+{ 
+	BuildBinaryArithmeticOperator(Emitter, IR::BO_Subtract, &A, ConstA, &B, ConstB, GetOutput(0));
+}
+
+void UMaterialExpressionMultiply::Build(IR::FEmitter& Emitter)
+{ 
+	BuildBinaryArithmeticOperator(Emitter, IR::BO_Multiply, &A, ConstA, &B, ConstB, GetOutput(0));
+}
+
+void UMaterialExpressionDivide::Build(IR::FEmitter& Emitter)
+{ 
+	BuildBinaryArithmeticOperator(Emitter, IR::BO_Divide, &A, ConstA, &B, ConstB, GetOutput(0));
+}
+
+#endif // WITH_EDITOR
diff --git a/Engine/Source/Runtime/Engine/Private/Materials/MaterialIR.cpp b/Engine/Source/Runtime/Engine/Private/Materials/MaterialIR.cpp
--- a/Engine/Source/Runtime/Engine/Private/Materials/MaterialIR.cpp
+++ b/Engine/Source/Runtime/Engine/Private/Materials/MaterialIR.cpp
@@ -5,27 +5,57 @@
 
 #if WITH_EDITOR
 
-namespace MaterialIR
+namespace UE::MIR
 {
 
-struct FValue;
-using FValuePtr = const FValue*;
+FInstructionPtr FValue::AsInstruction() const
+{
+	return (Kind >= VK_InstructionBegin && Kind < VK_InstructionEnd) ? static_cast<FInstructionPtr>(this) : nullptr;
+}
+
+bool FValue::Equals(FValuePtr Other) const
+{
+	// If kinds are different the two values are surely different.
+	if (Kind != Other->Kind)
+	{
+		return false;
+	}
+
+	// Get the size of this value in bytes. It should match that of Other, since the value kinds are the same.
+	uint32 SizeInBytes = GetSizeInBytes();
+	check(SizeInBytes == Other->GetSizeInBytes());
+
+	// Values are PODs by design, therefore simply comparing bytes is sufficient.
+	return FMemory::Memcmp(this, Other, SizeInBytes) == 0;
+}
+
+uint32 FValue::GetSizeInBytes() const
+{
+	switch (Kind)
+	{
+		case VK_ScalarConstant: return sizeof(FScalarConstant);
+		case VK_SetMaterialOutput: return sizeof(FScalarConstant);
+		case VK_BinaryOperator: return sizeof(FBinaryOperator);
+		case VK_Dimensional: return sizeof(FDimensional) + sizeof(FValuePtr) * static_cast<const FDimensional*>(this)->GetComponents().Num();
+		default: UE_MIR_UNREACHABLE();
+	}
+}
 
-TArrayView<const FValuePtr> FVectorValue::GetComponents() const
+TArrayView<const FValuePtr> FDimensional::GetComponents() const
 {
-	TArrayView<FValuePtr> Components = const_cast<FVectorValue*>(this)->GetMutableComponents();
+	TArrayView<FValuePtr> Components = const_cast<FDimensional*>(this)->GetMutableComponents();
 	return { Components.GetData(), Components.Num() };
 }
 
-TArrayView<FValuePtr> FVectorValue::GetMutableComponents()
+TArrayView<FValuePtr> FDimensional::GetMutableComponents()
 {
 	FArithmeticTypePtr ArithmeticType = Type->ToArithmetic();
 	check(ArithmeticType);
 
-	FValuePtr* Ptr = (FValuePtr*)static_cast<TVectorValue<1>*>(this)->Components;
+	FValuePtr* Ptr = (FValuePtr*)static_cast<TDimensional<1>*>(this)->Components;
 	return { Ptr, ArithmeticType->NumRows };
 }
 
-} // namespace MaterialIR
+} // namespace UE::MIR
 
 #endif // #if WITH_EDITOR
diff --git a/Engine/Source/Runtime/Engine/Private/Materials/MaterialIRBuilder.cpp b/Engine/Source/Runtime/Engine/Private/Materials/MaterialIRBuilder.cpp
deleted file mode 100644
diff --git a/Engine/Source/Runtime/Engine/Private/Materials/MaterialIREmitter.cpp b/Engine/Source/Runtime/Engine/Private/Materials/MaterialIREmitter.cpp
new file mode 100644
--- /dev/null
+++ b/Engine/Source/Runtime/Engine/Private/Materials/MaterialIREmitter.cpp
@@ -0,0 +1,419 @@
+// Copyright Epic Games, Inc. All Rights Reserved.
+
+#include "Materials/MaterialIREmitter.h"
+#include "Materials/MaterialIRModule.h"
+#include "Materials/MaterialIRModuleBuilder.h"
+#include "Shader/ShaderTypes.h"
+#include "MaterialShared.h"
+#include "MaterialExpressionIO.h"
+
+#if WITH_EDITOR
+
+namespace UE::MIR
+{
+
+template <typename T>
+static T InitValue(FTypePtr InType)
+{
+	T Value{};
+	Value.Kind = T::TypeKind;
+	Value.Type = InType;
+	return Value;
+}
+
+struct FEmitter::FPrivate : FEmitter {
+	
+	// Creates a new `FDimensional` value of specified `Type` and returns it.
+	static FDimensional* NewDimensionalValue(FEmitter* Emitter, FArithmeticTypePtr Type)
+	{
+		check(!Type->IsScalar());
+		int Dimensions = Type->NumRows * Type->NumColumns;
+		FDimensional* Value = (FDimensional*)FMemory::Malloc(sizeof(FDimensional) + sizeof(FValuePtr) * Dimensions);
+		Value->Kind = VK_Dimensional;
+		Value->Type = Type;
+		return Value;
+	}
+
+	// Searches for an existing value in module that matches specified `Prototype`.
+	// If none found, it creates a new value as a copy of the prototype, adds it to
+	// the module then returns it.
+	template <typename TValueType>
+	static FValuePtr EmitPrototype(FEmitter* Emitter, const TValueType& Prototype)
+	{
+		if (FValuePtr Existing = FindValue(Emitter, &Prototype))
+		{
+			return Existing;
+		}
+
+		TValueType* Value = new TValueType{ Prototype };
+		Emitter->Module->Values.Add(Value);
+		return Value;
+	}
+
+	// Emits specified newly created `Value`. If the exact value already exists,
+	// specified one is *destroyed* and existing one is returned instead.
+	static FValuePtr EmitNew(FEmitter* Emitter, FValuePtr Value)
+	{
+		if (FValuePtr Existing = FindValue(Emitter, Value))
+		{
+			delete Value;
+			return Existing;
+		}
+
+		Emitter->Module->Values.Add(Value);
+		return Value;
+	}
+
+	// Looks for an existing value in the module that matches `Prototype` and returns it if found.
+	static FValuePtr FindValue(const FEmitter* Emitter, FValuePtr Prototype)
+	{
+		/* todo: improve this with hashmap */
+		for (FValuePtr CurrValue : Emitter->Module->Values)
+		{
+			if (CurrValue->Equals(Prototype))
+			{
+				return CurrValue;
+			}
+		}
+
+		return nullptr;
+	}
+};
+
+FEmitter::FEmitter(FMaterialIRModuleBuilder* InBuilder, UMaterial* InMaterial, FMaterialIRModule* InModule)
+{
+	Builder = InBuilder;
+	Material = InMaterial;
+	Module = InModule;
+
+	ZeroIntValue = EmitConstantScalarZero(SK_Int);
+	ZeroFloatValue = EmitConstantScalarZero(SK_Float);
+}
+
+FValuePtr FEmitter::GetAndCheckInputTypeKindIs(const FExpressionInput* Input, ETypeKind Kind)
+{
+	FValuePtr InputValue = Get(Input);
+	CheckInputTypeIs(Input, InputValue, Kind);
+	return InputValue;
+}
+
+bool FEmitter::CheckInputTypeIs(const FExpressionInput* Input, FValuePtr InputValue, ETypeKind Kind)
+{
+	if (InputValue->Type->Kind != Kind)
+	{
+		Errorf(TEXT("Unexpected input '%s' value type (%s). Expected %s type."), *Input->InputName.ToString(), TEXT(""), TEXT(""));
+		return false;
+	}
+	return true;
+}
+
+FValuePtr FEmitter::EmitConstantFromShaderValue(const UE::Shader::FValue& InValue)
+{
+	using namespace UE::Shader;
+
+	switch (InValue.Type.ValueType)
+	{
+		case UE::Shader::EValueType::Float1: return EmitConstantFloat1(InValue.AsFloatScalar());
+		case UE::Shader::EValueType::Float2: return EmitConstantFloat2(FVector2f{ InValue.Component[0].Float, InValue.Component[1].Float });
+		case UE::Shader::EValueType::Float3: return EmitConstantFloat3(FVector3f{ InValue.Component[0].Float, InValue.Component[1].Float, InValue.Component[2].Float });
+		case UE::Shader::EValueType::Float4: return EmitConstantFloat4(FVector4f{ InValue.Component[0].Float, InValue.Component[1].Float, InValue.Component[2].Float, InValue.Component[3].Float });
+
+		case UE::Shader::EValueType::Int1: return EmitConstantInt1(InValue.AsFloatScalar());
+		case UE::Shader::EValueType::Int2: return EmitConstantInt2(FIntVector2{ InValue.Component[0].Int, InValue.Component[1].Int });
+		case UE::Shader::EValueType::Int3: return EmitConstantInt3(FIntVector3{ InValue.Component[0].Int, InValue.Component[1].Int, InValue.Component[2].Int });
+		case UE::Shader::EValueType::Int4: return EmitConstantInt4(FIntVector4{ InValue.Component[0].Int, InValue.Component[1].Int, InValue.Component[2].Int, InValue.Component[3].Int });
+	}
+
+	UE_MIR_UNREACHABLE();
+}
+
+FValuePtr FEmitter::EmitConstantScalarZero(EScalarKind Kind)
+{
+	switch (Kind)
+	{
+		case SK_Bool: UE_MIR_TODO();
+		case SK_Int: return ZeroIntValue;
+		case SK_Float: return ZeroFloatValue;
+		default: UE_MIR_UNREACHABLE();
+	}
+}
+
+FValuePtr FEmitter::EmitConstantFloat1(float InX)
+{
+	FScalarConstant Scalar = InitValue<FScalarConstant>(FArithmeticType::GetScalar(SK_Float));
+	Scalar.Float = InX;
+
+	return FPrivate::EmitPrototype(this, Scalar);
+}
+
+FValuePtr FEmitter::EmitConstantFloat2(const FVector2f& InValue)
+{
+	FValuePtr X = EmitConstantFloat1(InValue.X);
+	FValuePtr Y = EmitConstantFloat1(InValue.Y);
+	return EmitVector2(X, Y);
+}
+
+FValuePtr FEmitter::EmitConstantFloat3(const FVector3f& InValue)
+{
+	FValuePtr X = EmitConstantFloat1(InValue.X);
+	FValuePtr Y = EmitConstantFloat1(InValue.Y);
+	FValuePtr Z = EmitConstantFloat1(InValue.Z);
+	return EmitVector3(X, Y, Z);
+}
+
+FValuePtr FEmitter::EmitConstantFloat4(const FVector4f& InValue)
+{
+	FValuePtr X = EmitConstantFloat1(InValue.X);
+	FValuePtr Y = EmitConstantFloat1(InValue.Y);
+	FValuePtr Z = EmitConstantFloat1(InValue.Z);
+	FValuePtr W = EmitConstantFloat1(InValue.W);
+	return EmitVector4(X, Y, Z, W);
+}
+
+FValuePtr FEmitter::EmitConstantInt1(int InX)
+{
+	FScalarConstant Scalar = InitValue<FScalarConstant>(FArithmeticType::GetScalar(SK_Int));
+	Scalar.Integer = InX;
+
+	return FPrivate::EmitPrototype(this, Scalar);
+}
+
+FValuePtr FEmitter::EmitConstantInt2(const FIntVector2& InValue)
+{
+	FValuePtr X = EmitConstantInt1(InValue.X);
+	FValuePtr Y = EmitConstantInt1(InValue.Y);
+	return EmitVector2(X, Y);
+}
+
+FValuePtr FEmitter::EmitConstantInt3(const FIntVector3& InValue)
+{
+	FValuePtr X = EmitConstantInt1(InValue.X);
+	FValuePtr Y = EmitConstantInt1(InValue.Y);
+	FValuePtr Z = EmitConstantInt1(InValue.Z);
+	return EmitVector3(X, Y, Z);
+}
+
+FValuePtr FEmitter::EmitConstantInt4(const FIntVector4& InValue)
+{
+	FValuePtr X = EmitConstantInt1(InValue.X);
+	FValuePtr Y = EmitConstantInt1(InValue.Y);
+	FValuePtr Z = EmitConstantInt1(InValue.Z);
+	FValuePtr W = EmitConstantInt1(InValue.W);
+	return EmitVector4(X, Y, Z, W);
+}
+
+FValuePtr FEmitter::EmitVector2(FValuePtr InX, FValuePtr InY)
+{
+	check(InX->Type->ToScalar());
+	check(InX->Type == InY->Type);
+
+	TDimensional<2> Vector = InitValue<TDimensional<2>>(FArithmeticType::GetVector(InX->Type->ToArithmetic()->ScalarKind, 2));
+	TArrayView<FValuePtr> Components = Vector.GetMutableComponents();
+	Components[0] = InX;
+	Components[1] = InY;
+
+	return FPrivate::EmitPrototype(this, Vector);
+}
+
+FValuePtr FEmitter::EmitVector3(FValuePtr InX, FValuePtr InY, FValuePtr InZ)
+{
+	check(InX->Type->ToScalar());
+	check(InX->Type == InY->Type);
+	check(InY->Type == InZ->Type);
+
+	TDimensional<3> Vector = InitValue<TDimensional<3>>(FArithmeticType::GetVector(InX->Type->ToArithmetic()->ScalarKind, 3));
+	TArrayView<FValuePtr> Components = Vector.GetMutableComponents();
+	Components[0] = InX;
+	Components[1] = InY;
+	Components[2] = InZ;
+
+	return FPrivate::EmitPrototype(this, Vector);
+}
+
+FValuePtr FEmitter::EmitVector4(FValuePtr InX, FValuePtr InY, FValuePtr InZ, FValuePtr InW)
+{
+	check(InX->Type->ToScalar());
+	check(InX->Type == InY->Type);
+	check(InY->Type == InZ->Type);
+	check(InZ->Type == InW->Type);
+
+	TDimensional<4> Vector = InitValue<TDimensional<4>>(FArithmeticType::GetVector(InX->Type->ToArithmetic()->ScalarKind, 4));
+	TArrayView<FValuePtr> Components = Vector.GetMutableComponents();
+	Components[0] = InX;
+	Components[1] = InY;
+	Components[2] = InZ;
+	Components[3] = InW;
+
+	return FPrivate::EmitPrototype(this, Vector);
+}
+
+FValuePtr FEmitter::EmitArithmetic(FArithmeticTypePtr Type, FValuePtr Scalar)
+{
+	check(Scalar->Type->ToScalar());
+	if (Type->IsScalar())
+	{
+		check(Type == Scalar->Type);
+		return Scalar;
+	}
+
+	FDimensional* Value = FPrivate::NewDimensionalValue(this, Type);
+	for (FValuePtr& Component : Value->GetMutableComponents())
+	{
+		Component = Scalar;
+	}
+
+	return FPrivate::EmitNew(this, Value);
+}
+
+FSetMaterialOutput* FEmitter::EmitSetMaterialOutput(EMaterialProperty InProperty, FValue* InArgValue)
+{
+	FSetMaterialOutput* Instr = new FSetMaterialOutput;
+	Instr->Kind 	= VK_SetMaterialOutput;
+	Instr->Property = InProperty;
+	Instr->ArgValue = InArgValue;
+
+	Module->Values.Add(Instr);
+	Module->Outputs.Add(Instr);
+	return Instr;
+}
+
+static FTypePtr GetBinaryOperatorResultType(EBinaryOperator Operator, FTypePtr LhsType, FTypePtr RhsType)
+{
+	switch (Operator)
+	{
+		case BO_Add:
+		case BO_Subtract:
+		case BO_Multiply:
+		case BO_Divide:
+			check(LhsType == RhsType);
+			return LhsType;
+		
+		default:
+			UE_MIR_UNREACHABLE();
+	}
+}
+
+FValuePtr FEmitter::EmitBinaryOperator(EBinaryOperator Operator, FValuePtr Lhs, FValuePtr Rhs)
+{
+	check(Lhs->Type == Rhs->Type);
+
+	/* todo: operation folding */
+
+	FTypePtr ResultType = GetBinaryOperatorResultType(Operator, Lhs->Type, Rhs->Type);
+	FBinaryOperator Prototype = InitValue<FBinaryOperator>(ResultType);
+	Prototype.Operator = Operator;
+	Prototype.Lhs= Lhs;
+	Prototype.Rhs = Rhs;
+
+	return FPrivate::EmitPrototype(this, Prototype);
+}
+
+static FValuePtr ConvertScalarValue(FEmitter* Emitter, FValuePtr ScalarValue, EScalarKind ScalarKind, EScalarKind TargetKind)
+{
+	if (ScalarKind == TargetKind)
+	{
+		return ScalarValue;
+	}
+
+	auto ScalarConstant = ScalarValue->As<FScalarConstant>();
+	if (!ScalarConstant)
+	{
+		UE_MIR_UNREACHABLE(); // todo
+	}
+
+	switch (ScalarKind)
+	{
+		case SK_Bool:
+		case SK_Int:
+		{
+			switch (TargetKind)
+			{
+				case SK_Bool: UE_MIR_TODO();
+				case SK_Int: return Emitter->EmitConstantInt1(ScalarConstant->Integer);
+				case SK_Float: return Emitter->EmitConstantFloat1((float)ScalarConstant->Integer);
+				default: UE_MIR_UNREACHABLE();
+			}
+		}
+
+		case SK_Float:
+		{
+			switch (TargetKind)
+			{
+				case SK_Bool: UE_MIR_TODO();
+				case SK_Int: return Emitter->EmitConstantInt1((int)ScalarConstant->Float);
+				default: UE_MIR_UNREACHABLE();
+			}
+		}
+
+		default: break;
+	}
+
+	UE_MIR_UNREACHABLE();
+}
+
+static FValuePtr ConvertArithmeticValue(FEmitter* Emitter, FValuePtr Value, FArithmeticTypePtr TargetArithmeticType)
+{
+	FArithmeticTypePtr ValueArithmeticType = static_cast<FArithmeticTypePtr>(Value->Type);
+	if (ValueArithmeticType->IsScalar())
+	{
+		FValuePtr ConvertedScalar = ConvertScalarValue(Emitter, Value, ValueArithmeticType->ScalarKind, TargetArithmeticType->ScalarKind);
+		return Emitter->EmitArithmetic(TargetArithmeticType, ConvertedScalar);
+	}
+	else
+	{
+		UE_MIR_TODO();
+	}
+}
+
+FValuePtr FEmitter::TryEmitConvert(FValuePtr Value, FTypePtr TargetType)
+{
+	// If types already match, nothing else to do, simply return the same value.
+	FTypePtr ValueType = Value->Type;
+	if (ValueType == TargetType)
+	{
+		return Value;
+	}
+
+	FValuePtr ResultValue{};
+
+	// Converting from an arithmetic value to another arithmetic type.
+	FArithmeticTypePtr ValueArithmeticType = Value->Type->ToArithmetic();
+	FArithmeticTypePtr TargetArithmeticType = TargetType->ToArithmetic();
+	if (ValueArithmeticType && TargetArithmeticType)
+	{
+		ResultValue = ConvertArithmeticValue(this, Value, TargetArithmeticType);
+	}
+
+	/* ... */
+
+	// No other legal conversions applicable. Report error if we haven't converted the value.
+	if (!ResultValue)
+	{
+		Errorf(TEXT("Cannot convert value of type '%s' to '%s'"), TEXT(""), TEXT(""));
+	}
+
+	return ResultValue;
+}
+
+FValuePtr FEmitter::Get(const FExpressionInput* Input)
+{
+	check(Input);
+	FValuePtr* Value = Builder->InputValues.Find(Input);
+	return *Value;
+}
+	
+void FEmitter::Put(const FExpressionOutput* Output, FValuePtr Value)
+{
+	check(Output);
+	Builder->OutputValues.Add(Output, Value);
+}
+
+void FEmitter::Error(FString Message)
+{
+	Module->Errors.Emplace(Expression, MoveTemp(Message));
+	bHasExprBuildError = true;
+}
+
+} // namespace UE::MIR
+
+#endif // #if WITH_EDITOR
diff --git a/Engine/Source/Runtime/Engine/Private/Materials/MaterialIRModule.cpp b/Engine/Source/Runtime/Engine/Private/Materials/MaterialIRModule.cpp
--- a/Engine/Source/Runtime/Engine/Private/Materials/MaterialIRModule.cpp
+++ b/Engine/Source/Runtime/Engine/Private/Materials/MaterialIRModule.cpp
@@ -5,7 +5,7 @@
 
 #if WITH_EDITOR
 
-namespace MIR = MaterialIR;
+namespace IR = UE::MIR;
 
 FMaterialIRModule::~FMaterialIRModule()
 {
@@ -14,7 +14,7 @@ FMaterialIRModule::~FMaterialIRModule()
 
 void FMaterialIRModule::Empty()
 {
-	for (MIR::FValue* Value : Values)
+	for (IR::FValuePtr Value : Values)
 	{
 		delete Value;
 	}
@@ -23,12 +23,4 @@ void FMaterialIRModule::Empty()
 	Outputs.Empty();
 }
 
-void FMaterialIRModule::Reset(EShaderPlatform InShaderPlatform, ERHIFeatureLevel::Type InFeatureLevel, const ITargetPlatform* InTargetPlatform)
-{
-	Empty();
-
-	ShaderPlatform = InShaderPlatform;
-	TargetPlatform = InTargetPlatform;
-}
-
-#endif
+#endif // #if WITH_EDITOR
diff --git a/Engine/Source/Runtime/Engine/Private/Materials/MaterialIRModuleBuilder.cpp b/Engine/Source/Runtime/Engine/Private/Materials/MaterialIRModuleBuilder.cpp
--- a/Engine/Source/Runtime/Engine/Private/Materials/MaterialIRModuleBuilder.cpp
+++ b/Engine/Source/Runtime/Engine/Private/Materials/MaterialIRModuleBuilder.cpp
@@ -7,76 +7,134 @@
 #include "Materials/MaterialIRModule.h"
 #include "Materials/MaterialIR.h"
 #include "Materials/MaterialIRTypes.h"
-#include "Materials/MaterialIRBuilder.h"
+#include "Materials/MaterialIREmitter.h"
 #include "MaterialIRUtility.h"
 
 #include "Materials/MaterialAttributeDefinitionMap.h"
 #include "Materials/Material.h"
 #include "MaterialExpressionIO.h"
 #include "MaterialShared.h"
+#include "Materials/MaterialExpression.h"
 
-namespace IR = MaterialIR;
+namespace IR = UE::MIR;
 
-struct FMaterialIRModuleBuilder::FHelper
+struct FMaterialIRModuleBuilder::FPrivate
 {
-};
-
-void FMaterialIRModuleBuilder::SetSource(FMaterial* InMaterial, const FStaticParameterSet* InStaticParameters)
-{
-	Material = InMaterial;
-	StaticParameters = InStaticParameters;
-}
-
-void FMaterialIRModuleBuilder::SetPlatform(EShaderPlatform InShaderPlatform, ERHIFeatureLevel::Type InFeatureLevel, const ITargetPlatform* InTargetPlatform)
-{
-	ShaderPlatform	= InShaderPlatform;
-	FeatureLevel		= InFeatureLevel;
-	TargetPlatform	= InTargetPlatform;
-}
-
-void FMaterialIRModuleBuilder::SetTarget(FMaterialIRModule* InModule)
-{
-	Module = InModule;
-}
-
-bool FMaterialIRModuleBuilder::Build()
-{
-	Module->Reset(ShaderPlatform, FeatureLevel, TargetPlatform);
-
-	IR::FBuilder Builder{ Module };
-
-	UMaterial* BaseMaterial = Material->GetMaterialInterface()->GetMaterial();
-
-	for (int32 PropertyIndex = 0; PropertyIndex < MP_MAX; ++PropertyIndex)
+	static void Build_GenerateOutputInstructions(FMaterialIRModuleBuilder& Builder, UE::MIR::FEmitter& Emitter)
 	{
-		EMaterialProperty Property = (EMaterialProperty)PropertyIndex;
+		// Prepare the array of FSetMaterialOutputInstr outputs from the material attributes inputs.
 		FMaterialInputDescription Input;
-
-		if (!Utility::IsMaterialPropertyShared(Property) 
-			|| Property == MP_SubsurfaceColor
-			|| Property == MP_FrontMaterial
-			|| !BaseMaterial->GetExpressionInputDescription(Property, Input))
+		for (int32 Index = 0; UE::Utility::NextMaterialAttributeInput(Builder.BaseMaterial, Index, Input); ++Index)
 		{
-			continue;
+			EMaterialProperty Property = (EMaterialProperty)Index;
+
+			IR::FSetMaterialOutput* Output = Emitter.EmitSetMaterialOutput(Property, nullptr);
+
+			if (Input.bUseConstant)
+			{
+				Output->ArgValue = Emitter.EmitConstantFromShaderValue(Input.ConstantValue);
+			}
+			else if (!Input.Input->IsConnected())
+			{
+				Output->ArgValue = UE::Utility::CreateMaterialAttributeDefaultValue(Emitter, Builder.BaseMaterial, Property);
+			}
+			else
+			{
+				Builder.ExpressionAnalysisStack.Add(Input.Input->Expression);
+			}
 		}
+	}
 
-		IR::FSetMaterialOutputInstr* Output = Builder.EmitSetMaterialOutput(Property, nullptr);
-	
-		if (Input.bUseConstant)
-		{
-			Output->ArgValue = Builder.NewConstantFromShaderValue(Input.ConstantValue);
-		}
-		else if (!Input.Input->IsConnected())
+	static void Build_AnalyzeExpressionGraph(FMaterialIRModuleBuilder& Builder, UE::MIR::FEmitter & Emitter)
+	{
+		TSet<UMaterialExpression*> BuiltExpressions;
+
+		while (!Builder.ExpressionAnalysisStack.IsEmpty())
 		{
-			Output->ArgValue = Utility::CreateMaterialAttributeDefaultValue(Builder, Material, Property);
+			Emitter.Expression = Builder.ExpressionAnalysisStack.Last();
+
+			// If expression is clean, nothing to be done.
+			if (BuiltExpressions.Contains(Emitter.Expression))
+			{
+				Builder.ExpressionAnalysisStack.Pop(EAllowShrinking::No);
+				continue;
+			}
+
+			// Push to the expression stack all dependencies that still need to be analyzed.
+			for (FExpressionInputIterator It{ Emitter.Expression}; It; ++It)
+			{
+				// Ignore disconnected inputs and connected expressions already built.
+				if (!It->IsConnected() || BuiltExpressions.Contains(It->Expression))
+				{
+					continue;
+				}
+
+				Builder.ExpressionAnalysisStack.Push(It->Expression);
+			}
+
+			// If on top of the stack there's a different expression, we have a dependency to analyze first.
+			if (Builder.ExpressionAnalysisStack.Last() != Emitter.Expression) {
+				continue;
+			}
+
+			// 
+			Builder.ExpressionAnalysisStack.Pop();
+			BuiltExpressions.Add(Emitter.Expression);
+
+			//
+			for (FExpressionInputIterator It{ Emitter.Expression}; It; ++It)
+			{
+				if (FExpressionOutput* ConnectedOutput = It->GetConnectedOutput())
+				{
+					// Fetch the value flowing through connected output.
+					IR::FValuePtr* ValuePtr = Builder.OutputValues.Find(ConnectedOutput);
+					check(ValuePtr && TEXT("Output value not found."));
+
+					// Set the value flowing into this input.
+					Builder.InputValues.Add(It.Input, *ValuePtr);
+				}
+			}
+
+			// And clear the expression errors.
+			Emitter.bHasExprBuildError = false;
+
+			// Invoke the expression build function. This will perform semantic analysis, error reporting and
+			// emit IR values for its outputs (which will flow into connected expressions inputs).
+			Emitter.Expression->Build(Emitter);
 		}
-		else if (Input.Input->IsConnected())
+	}
+
+	static void Build_LinkMaterialOutputsToIncomingValues(FMaterialIRModuleBuilder& Builder, UE::MIR::FEmitter& Emitter)
+	{
+		FMaterialInputDescription Input;
+		for (IR::FSetMaterialOutput* Output : Builder.Module->Outputs)
 		{
-			UE_MIR_UNREACHABLE();
+			if (Output->ArgValue || !ensure(Builder.BaseMaterial->GetExpressionInputDescription(Output->Property, Input))) {
+				continue;
+			}
+
+			IR::FValuePtr* ValuePtr = Builder.OutputValues.Find(Input.Input->Expression->GetOutput(0));
+			check(ValuePtr);
+
+			Output->ArgValue = *ValuePtr;
 		}
 	}
+};
+
+bool FMaterialIRModuleBuilder::Build(UMaterial* InMaterial, FMaterialIRModule* TargetModule)
+{
+	BaseMaterial = InMaterial;
+
+	Module = TargetModule;
+	Module->Empty();
+
+	IR::FEmitter Emitter{ this, BaseMaterial, Module };
+
+	FPrivate::Build_GenerateOutputInstructions(*this, Emitter);
+	FPrivate::Build_AnalyzeExpressionGraph(*this, Emitter);
+	FPrivate::Build_LinkMaterialOutputsToIncomingValues(*this, Emitter);
 
 	return true;
 }
 
-#endif
+#endif // #if WITH_EDITOR
diff --git a/Engine/Source/Runtime/Engine/Private/Materials/MaterialIRToHLSLTranslator.cpp b/Engine/Source/Runtime/Engine/Private/Materials/MaterialIRToHLSLTranslator.cpp
--- a/Engine/Source/Runtime/Engine/Private/Materials/MaterialIRToHLSLTranslator.cpp
+++ b/Engine/Source/Runtime/Engine/Private/Materials/MaterialIRToHLSLTranslator.cpp
@@ -17,7 +17,7 @@
 
 #if WITH_EDITOR
 
-namespace IR = MaterialIR;
+namespace IR = UE::MIR;
 
 struct FHLSLPrinter
 {
@@ -120,269 +120,303 @@ static const TCHAR* GetShadingModelParameterName(EMaterialShadingModel InModel)
 	}
 }
 
-UE_MIR_BEGIN_PRIVATE(FMaterialIRToHLSLTranslator);
-
-void GenerateHLSL(const FMaterial& InMaterial, const FMaterialIRModule& InModule, FParametersMap& InParams)
+struct FMaterialIRToHLSLTranslator::FPrivate
 {
-	FString PixelAttributes;
-	FString EvaluateOtherMaterialAttributesHLSL;
-	FHLSLPrinter Printer{ EvaluateOtherMaterialAttributesHLSL };
+	FMaterialIRToHLSLTranslator* Translator;
 
- 	for (const IR::FSetMaterialOutputInstr* Output : InModule.GetOutputs())
+	void GenerateHLSL(const FMaterial& InMaterial,
+					  const FMaterialIRModule& InModule,
+					  FParametersMap& InParams)
 	{
-		LowerValue(Printer, Output);
-	}
-
-	EvaluateOtherMaterialAttributesHLSL.Append(TEXT("\tPixelMaterialInputs.FrontMaterial = GetInitialisedSubstrateData();\n"));
-	EvaluateOtherMaterialAttributesHLSL.Append(TEXT("\tPixelMaterialInputs.Subsurface = 0;\n"));
+		FString PixelAttributes;
+		FString EvaluateOtherMaterialAttributesHLSL;
+		FHLSLPrinter Printer{ EvaluateOtherMaterialAttributesHLSL };
 
-	for (int32 PropertyIndex = 0; PropertyIndex < MP_MAX; ++PropertyIndex)
-	{
-		EMaterialProperty Property = (EMaterialProperty)PropertyIndex;
-		if (!Utility::IsMaterialPropertyShared(Property))
+		for (const IR::FSetMaterialOutput* Output : InModule.GetOutputs())
 		{
-			continue;
+			LowerValue(Printer, Output);
 		}
+
+		EvaluateOtherMaterialAttributesHLSL.Append(TEXT("\tPixelMaterialInputs.FrontMaterial = GetInitialisedSubstrateData();\n"));
+		EvaluateOtherMaterialAttributesHLSL.Append(TEXT("\tPixelMaterialInputs.Subsurface = 0;\n"));
+
+		for (int32 PropertyIndex = 0; PropertyIndex < MP_MAX; ++PropertyIndex)
+		{
+			EMaterialProperty Property = (EMaterialProperty)PropertyIndex;
+			if (!UE::Utility::IsMaterialPropertyShared(Property))
+			{
+				continue;
+			}
 		
-		check(FMaterialAttributeDefinitionMap::GetShaderFrequency(Property) == SF_Pixel);
+			check(FMaterialAttributeDefinitionMap::GetShaderFrequency(Property) == SF_Pixel);
 			
-		// Special case MP_SubsurfaceColor as the actual property is a combination of the color and the profile but we don't want to expose the profile
-		FString PropertyName = (Property == MP_SubsurfaceColor) ? "Subsurface" : FMaterialAttributeDefinitionMap::GetAttributeName(Property);
-		EMaterialValueType Type = (Property == MP_SubsurfaceColor) ? MCT_Float4 : FMaterialAttributeDefinitionMap::GetValueType(Property);
-		check(PropertyName.Len() > 0);
+			// Special case MP_SubsurfaceColor as the actual property is a combination of the color and the profile but we don't want to expose the profile
+			FString PropertyName = (Property == MP_SubsurfaceColor) ? "Subsurface" : FMaterialAttributeDefinitionMap::GetAttributeName(Property);
+			EMaterialValueType Type = (Property == MP_SubsurfaceColor) ? MCT_Float4 : FMaterialAttributeDefinitionMap::GetValueType(Property);
+			check(PropertyName.Len() > 0);
 
-		PixelAttributes.Appendf(TEXT("\t%s %s;\n"), GetHLSLTypeString(Type), *PropertyName);
-	}
+			PixelAttributes.Appendf(TEXT("\t%s %s;\n"), GetHLSLTypeString(Type), *PropertyName);
+		}
 	
-	InParams.Add(TEXT("pixel_material_inputs"), MoveTemp(PixelAttributes));
+		InParams.Add(TEXT("pixel_material_inputs"), MoveTemp(PixelAttributes));
 
-	InParams.Add(TEXT("calc_pixel_material_inputs_normal"), TEXT("\tPixelMaterialInputs.Normal = MaterialFloat3(0.00000000, 0.00000000, 1.00000000);"));
-	InParams.Add(TEXT("calc_pixel_material_inputs_other_inputs"), EvaluateOtherMaterialAttributesHLSL);
+		InParams.Add(TEXT("calc_pixel_material_inputs_normal"), TEXT("\tPixelMaterialInputs.Normal = MaterialFloat3(0.00000000, 0.00000000, 1.00000000);"));
+		InParams.Add(TEXT("calc_pixel_material_inputs_other_inputs"), EvaluateOtherMaterialAttributesHLSL);
 
-	InParams.Add(TEXT("calc_pixel_material_inputs_analytic_derivatives_normal"), TEXT("\tPixelMaterialInputs.Normal = MaterialFloat3(0.00000000, 0.00000000, 1.00000000);"));
-	InParams.Add(TEXT("calc_pixel_material_inputs_analytic_derivatives_other_inputs"), MoveTemp(EvaluateOtherMaterialAttributesHLSL));
+		InParams.Add(TEXT("calc_pixel_material_inputs_analytic_derivatives_normal"), TEXT("\tPixelMaterialInputs.Normal = MaterialFloat3(0.00000000, 0.00000000, 1.00000000);"));
+		InParams.Add(TEXT("calc_pixel_material_inputs_analytic_derivatives_other_inputs"), MoveTemp(EvaluateOtherMaterialAttributesHLSL));
 	
-	InParams.Add(TEXT("material_declarations"), TEXT("struct FMaterialAttributes {};"));
+		InParams.Add(TEXT("material_declarations"), TEXT("struct FMaterialAttributes {};"));
 	
-	auto SetParamInt = [&] (const TCHAR* InParamName, int InValue)
-	{
-		InParams.Add(InParamName, FString::Printf(TEXT("%d"), InValue));
-	};
+		auto SetParamInt = [&] (const TCHAR* InParamName, int InValue)
+		{
+			InParams.Add(InParamName, FString::Printf(TEXT("%d"), InValue));
+		};
 
-	SetParamInt(TEXT("num_material_texcoords_vertex"), 0);
-	SetParamInt(TEXT("num_material_texcoords"), 0);
-	SetParamInt(TEXT("num_custom_vertex_interpolators"), 0);
-	SetParamInt(TEXT("num_tex_coord_interpolators"), 0);
-}
+		SetParamInt(TEXT("num_material_texcoords_vertex"), 0);
+		SetParamInt(TEXT("num_material_texcoords"), 0);
+		SetParamInt(TEXT("num_custom_vertex_interpolators"), 0);
+		SetParamInt(TEXT("num_tex_coord_interpolators"), 0);
+	}
 
-void SetMaterialParameters(const FMaterial& InMaterial, FParametersMap& InParams)
-{
-	auto SetParamReturnFloat = [&] (const TCHAR* InParamName, float InValue)
+	void SetMaterialParameters(const FMaterial& InMaterial, FParametersMap& InParams)
 	{
-		InParams.Add(InParamName, FString::Printf(TEXT("\treturn %.5f"), InValue));
-	};
-
-	SetParamReturnFloat(TEXT("get_material_emissive_for_cs"), 0.f);
-	SetParamReturnFloat(TEXT("get_material_translucency_directional_lighting_intensity"), InMaterial.GetTranslucencyDirectionalLightingIntensity());
-	SetParamReturnFloat(TEXT("get_material_translucent_shadow_density_scale"), InMaterial.GetTranslucentShadowDensityScale());
-	SetParamReturnFloat(TEXT("get_material_translucent_self_shadow_density_scale"), InMaterial.GetTranslucentSelfShadowDensityScale());
-	SetParamReturnFloat(TEXT("get_material_translucent_self_shadow_second_density_scale"), InMaterial.GetTranslucentSelfShadowSecondDensityScale());
-	SetParamReturnFloat(TEXT("get_material_translucent_self_shadow_second_opacity"), InMaterial.GetTranslucentSelfShadowSecondOpacity());
-	SetParamReturnFloat(TEXT("get_material_translucent_backscattering_exponent"), InMaterial.GetTranslucentBackscatteringExponent());
-
-	FLinearColor Extinction = InMaterial.GetTranslucentMultipleScatteringExtinction();
-	InParams.Add(TEXT("get_material_translucent_multiple_scattering_extinction"), FString::Printf(TEXT("\treturn MaterialFloat3(%.5f, %.5f, %.5f)"), Extinction.R, Extinction.G, Extinction.B));
-
-	SetParamReturnFloat(TEXT("get_material_opacity_mask_clip_value"), InMaterial.GetOpacityMaskClipValue());
-	InParams.Add(TEXT("get_material_world_position_offset_raw"), TEXT("\treturn 0; // todo"));
-	InParams.Add(TEXT("get_material_previous_world_position_offset_raw"), TEXT("\treturn 0; // todo"));
+		auto SetParamReturnFloat = [&] (const TCHAR* InParamName, float InValue)
+		{
+			InParams.Add(InParamName, FString::Printf(TEXT("\treturn %.5f"), InValue));
+		};
+
+		SetParamReturnFloat(TEXT("get_material_emissive_for_cs"), 0.f);
+		SetParamReturnFloat(TEXT("get_material_translucency_directional_lighting_intensity"), InMaterial.GetTranslucencyDirectionalLightingIntensity());
+		SetParamReturnFloat(TEXT("get_material_translucent_shadow_density_scale"), InMaterial.GetTranslucentShadowDensityScale());
+		SetParamReturnFloat(TEXT("get_material_translucent_self_shadow_density_scale"), InMaterial.GetTranslucentSelfShadowDensityScale());
+		SetParamReturnFloat(TEXT("get_material_translucent_self_shadow_second_density_scale"), InMaterial.GetTranslucentSelfShadowSecondDensityScale());
+		SetParamReturnFloat(TEXT("get_material_translucent_self_shadow_second_opacity"), InMaterial.GetTranslucentSelfShadowSecondOpacity());
+		SetParamReturnFloat(TEXT("get_material_translucent_backscattering_exponent"), InMaterial.GetTranslucentBackscatteringExponent());
+
+		FLinearColor Extinction = InMaterial.GetTranslucentMultipleScatteringExtinction();
+		InParams.Add(TEXT("get_material_translucent_multiple_scattering_extinction"), FString::Printf(TEXT("\treturn MaterialFloat3(%.5f, %.5f, %.5f)"), Extinction.R, Extinction.G, Extinction.B));
+
+		SetParamReturnFloat(TEXT("get_material_opacity_mask_clip_value"), InMaterial.GetOpacityMaskClipValue());
+		InParams.Add(TEXT("get_material_world_position_offset_raw"), TEXT("\treturn 0; // todo"));
+		InParams.Add(TEXT("get_material_previous_world_position_offset_raw"), TEXT("\treturn 0; // todo"));
 	
-	// CustomData0/1 are named ClearCoat/ClearCoatRoughness
-	InParams.Add(TEXT("get_material_custom_data0"), TEXT("\treturn 1.0f; // todo"));
-	InParams.Add(TEXT("get_material_custom_data1"), TEXT("\treturn 0.1f; // todo"));
-
-	FString EvaluateMaterialDeclaration;
-	EvaluateMaterialDeclaration.Append(TEXT("void EvaluateVertexMaterialAttributes(in out FMaterialVertexParameters Parameters)\n{\n"));
-	EvaluateMaterialDeclaration.Append(TEXT("\n}\n"));
-	InParams.Add(TEXT("evaluate_material_attributes"), EvaluateMaterialDeclaration);
-}
+		// CustomData0/1 are named ClearCoat/ClearCoatRoughness
+		InParams.Add(TEXT("get_material_custom_data0"), TEXT("\treturn 1.0f; // todo"));
+		InParams.Add(TEXT("get_material_custom_data1"), TEXT("\treturn 0.1f; // todo"));
+
+		FString EvaluateMaterialDeclaration;
+		EvaluateMaterialDeclaration.Append(TEXT("void EvaluateVertexMaterialAttributes(in out FMaterialVertexParameters Parameters)\n{\n"));
+		EvaluateMaterialDeclaration.Append(TEXT("\n}\n"));
+		InParams.Add(TEXT("evaluate_material_attributes"), EvaluateMaterialDeclaration);
+	}
 
-void GetShaderCompilerEnvironment(const FMaterial& InMaterial, const FMaterialIRModule& InModule, FShaderCompilerEnvironment& OutEnvironment)
-{
-	EShaderPlatform ShaderPlatform = InModule.GetShaderPlatform();
-	const FMaterialCompilationOutput& CompilationOutput = InModule.GetCompilationOutput();
-
-	OutEnvironment.TargetPlatform = InModule.GetTargetPlatform();
-	OutEnvironment.SetDefine(TEXT("ENABLE_NEW_HLSL_GENERATOR"), 1);
-	OutEnvironment.SetDefine(TEXT("MATERIAL_ATMOSPHERIC_FOG"), false);
-	OutEnvironment.SetDefine(TEXT("MATERIAL_SKY_ATMOSPHERE"), false);
-	OutEnvironment.SetDefine(TEXT("INTERPOLATE_VERTEX_COLOR"), false);
-	OutEnvironment.SetDefine(TEXT("NEEDS_PARTICLE_COLOR"), false);
-	OutEnvironment.SetDefine(TEXT("NEEDS_PARTICLE_LOCAL_TO_WORLD"), false);
-	OutEnvironment.SetDefine(TEXT("NEEDS_PARTICLE_WORLD_TO_LOCAL"), false);
-	OutEnvironment.SetDefine(TEXT("NEEDS_PER_INSTANCE_RANDOM_PS"), false);
-	OutEnvironment.SetDefine(TEXT("USES_TRANSFORM_VECTOR"), false);
-	OutEnvironment.SetDefine(TEXT("WANT_PIXEL_DEPTH_OFFSET"), CompilationOutput.bUsesPixelDepthOffset);
-	OutEnvironment.SetDefineAndCompileArgument(TEXT("USES_WORLD_POSITION_OFFSET"), (bool)CompilationOutput.bUsesWorldPositionOffset);
-	OutEnvironment.SetDefineAndCompileArgument(TEXT("USES_DISPLACEMENT"), false);
-	OutEnvironment.SetDefine(TEXT("USES_EMISSIVE_COLOR"), false);
-	OutEnvironment.SetDefine(TEXT("USES_DISTORTION"), InMaterial.IsDistorted());
-	OutEnvironment.SetDefine(TEXT("MATERIAL_ENABLE_TRANSLUCENCY_FOGGING"), InMaterial.ShouldApplyFogging());
-	OutEnvironment.SetDefine(TEXT("MATERIAL_ENABLE_TRANSLUCENCY_CLOUD_FOGGING"), InMaterial.ShouldApplyCloudFogging());
-	OutEnvironment.SetDefine(TEXT("MATERIAL_IS_SKY"), InMaterial.IsSky());
-	OutEnvironment.SetDefine(TEXT("MATERIAL_COMPUTE_FOG_PER_PIXEL"), InMaterial.ComputeFogPerPixel());
-	OutEnvironment.SetDefine(TEXT("MATERIAL_FULLY_ROUGH"), false);
-	OutEnvironment.SetDefine(TEXT("MATERIAL_USES_ANISOTROPY"), false);
-	OutEnvironment.SetDefine(TEXT("MATERIAL_NEURAL_POST_PROCESS"), (CompilationOutput.bUsedWithNeuralNetworks || InMaterial.IsUsedWithNeuralNetworks()) && InMaterial.IsPostProcessMaterial());
-	OutEnvironment.SetDefine(TEXT("NUM_VIRTUALTEXTURE_SAMPLES"), 0);
-	OutEnvironment.SetDefine(TEXT("MATERIAL_VIRTUALTEXTURE_FEEDBACK"), false);
-	OutEnvironment.SetDefine(TEXT("IS_MATERIAL_SHADER"), true);
-
-	FMaterialShadingModelField ShadingModels = InMaterial.GetShadingModels();
-	ensure(ShadingModels.IsValid());
-
-	int32 NumActiveShadingModels = 0;
-	if (ShadingModels.IsLit())
+	void GetShaderCompilerEnvironment(const FMaterial& InMaterial, const FMaterialIRModule& InModule, FShaderCompilerEnvironment& OutEnvironment)
 	{
-		// This is to have platforms use the simple single layer water shading similar to mobile: no dynamic lights, only sun and sky, no distortion, no colored transmittance on background, no custom depth read.
-		const bool bSingleLayerWaterUsesSimpleShading = FDataDrivenShaderPlatformInfo::GetWaterUsesSimpleForwardShading(ShaderPlatform) && IsForwardShadingEnabled(ShaderPlatform);
-
-		for (int i = 0; i < MSM_NUM; ++i)
+		const FMaterialCompilationOutput& CompilationOutput = InModule.GetCompilationOutput();
+
+		OutEnvironment.TargetPlatform = Translator->TargetPlatform;
+		OutEnvironment.SetDefine(TEXT("ENABLE_NEW_HLSL_GENERATOR"), 1);
+		OutEnvironment.SetDefine(TEXT("MATERIAL_ATMOSPHERIC_FOG"), false);
+		OutEnvironment.SetDefine(TEXT("MATERIAL_SKY_ATMOSPHERE"), false);
+		OutEnvironment.SetDefine(TEXT("INTERPOLATE_VERTEX_COLOR"), false);
+		OutEnvironment.SetDefine(TEXT("NEEDS_PARTICLE_COLOR"), false);
+		OutEnvironment.SetDefine(TEXT("NEEDS_PARTICLE_LOCAL_TO_WORLD"), false);
+		OutEnvironment.SetDefine(TEXT("NEEDS_PARTICLE_WORLD_TO_LOCAL"), false);
+		OutEnvironment.SetDefine(TEXT("NEEDS_PER_INSTANCE_RANDOM_PS"), false);
+		OutEnvironment.SetDefine(TEXT("USES_TRANSFORM_VECTOR"), false);
+		OutEnvironment.SetDefine(TEXT("WANT_PIXEL_DEPTH_OFFSET"), CompilationOutput.bUsesPixelDepthOffset);
+		OutEnvironment.SetDefineAndCompileArgument(TEXT("USES_WORLD_POSITION_OFFSET"), (bool)CompilationOutput.bUsesWorldPositionOffset);
+		OutEnvironment.SetDefineAndCompileArgument(TEXT("USES_DISPLACEMENT"), false);
+		OutEnvironment.SetDefine(TEXT("USES_EMISSIVE_COLOR"), false);
+		OutEnvironment.SetDefine(TEXT("USES_DISTORTION"), InMaterial.IsDistorted());
+		OutEnvironment.SetDefine(TEXT("MATERIAL_ENABLE_TRANSLUCENCY_FOGGING"), InMaterial.ShouldApplyFogging());
+		OutEnvironment.SetDefine(TEXT("MATERIAL_ENABLE_TRANSLUCENCY_CLOUD_FOGGING"), InMaterial.ShouldApplyCloudFogging());
+		OutEnvironment.SetDefine(TEXT("MATERIAL_IS_SKY"), InMaterial.IsSky());
+		OutEnvironment.SetDefine(TEXT("MATERIAL_COMPUTE_FOG_PER_PIXEL"), InMaterial.ComputeFogPerPixel());
+		OutEnvironment.SetDefine(TEXT("MATERIAL_FULLY_ROUGH"), false);
+		OutEnvironment.SetDefine(TEXT("MATERIAL_USES_ANISOTROPY"), false);
+		OutEnvironment.SetDefine(TEXT("MATERIAL_NEURAL_POST_PROCESS"), (CompilationOutput.bUsedWithNeuralNetworks || InMaterial.IsUsedWithNeuralNetworks()) && InMaterial.IsPostProcessMaterial());
+		OutEnvironment.SetDefine(TEXT("NUM_VIRTUALTEXTURE_SAMPLES"), 0);
+		OutEnvironment.SetDefine(TEXT("MATERIAL_VIRTUALTEXTURE_FEEDBACK"), false);
+		OutEnvironment.SetDefine(TEXT("IS_MATERIAL_SHADER"), true);
+
+		FMaterialShadingModelField ShadingModels = InMaterial.GetShadingModels();
+		ensure(ShadingModels.IsValid());
+
+		int32 NumActiveShadingModels = 0;
+		if (ShadingModels.IsLit())
 		{
-			EMaterialShadingModel Model = (EMaterialShadingModel)i;
-			if (Model == MSM_Strata || !ShadingModels.HasShadingModel(Model))
-			{
-				continue;
-			}
+			// This is to have platforms use the simple single layer water shading similar to mobile: no dynamic lights, only sun and sky, no distortion, no colored transmittance on background, no custom depth read.
+			const bool bSingleLayerWaterUsesSimpleShading = FDataDrivenShaderPlatformInfo::GetWaterUsesSimpleForwardShading(Translator->ShaderPlatform) && IsForwardShadingEnabled(Translator->ShaderPlatform);
 
-			if (Model == MSM_SingleLayerWater && !FDataDrivenShaderPlatformInfo::GetRequiresDisableForwardLocalLights(ShaderPlatform))
+			for (int i = 0; i < MSM_NUM; ++i)
 			{
-				continue;
+				EMaterialShadingModel Model = (EMaterialShadingModel)i;
+				if (Model == MSM_Strata || !ShadingModels.HasShadingModel(Model))
+				{
+					continue;
+				}
+
+				if (Model == MSM_SingleLayerWater && !FDataDrivenShaderPlatformInfo::GetRequiresDisableForwardLocalLights(Translator->ShaderPlatform))
+				{
+					continue;
+				}
+
+				if (Model == MSM_SingleLayerWater && bSingleLayerWaterUsesSimpleShading)
+				{
+					// Value must match SINGLE_LAYER_WATER_SHADING_QUALITY_MOBILE_WITH_DEPTH_TEXTURE in SingleLayerWaterCommon.ush!
+					OutEnvironment.SetDefine(TEXT("SINGLE_LAYER_WATER_SHADING_QUALITY"), true);
+				}
+
+				OutEnvironment.SetDefine(GetShadingModelParameterName(Model), true);
+				NumActiveShadingModels += 1;
 			}
+		}
+		else
+		{
+			// Unlit shading model can only exist by itself
+			OutEnvironment.SetDefine(TEXT("MATERIAL_SINGLE_SHADINGMODEL"), true);
+			OutEnvironment.SetDefine(GetShadingModelParameterName(MSM_Unlit), true);
+		}
+
+		if (NumActiveShadingModels == 1)
+		{
+			OutEnvironment.SetDefine(TEXT("MATERIAL_SINGLE_SHADINGMODEL"), true);
+		}
+		else if (!ensure(NumActiveShadingModels > 0))
+		{
+			UE_LOG(LogMaterial, Warning, TEXT("Unknown material shading model(s). Setting to MSM_DefaultLit"));
+			OutEnvironment.SetDefine(GetShadingModelParameterName(MSM_DefaultLit), true);
+		}
 
-			if (Model == MSM_SingleLayerWater && bSingleLayerWaterUsesSimpleShading)
+		static IConsoleVariable* CVarLWCIsEnabled = IConsoleManager::Get().FindConsoleVariable(TEXT("r.MaterialEditor.LWCEnabled"));
+		OutEnvironment.SetDefine(TEXT("MATERIAL_LWC_ENABLED"), CVarLWCIsEnabled->GetInt());
+		OutEnvironment.SetDefine(TEXT("WSVECTOR_IS_TILEOFFSET"), true);
+		OutEnvironment.SetDefine(TEXT("WSVECTOR_IS_DOUBLEFLOAT"), false);
+
+		if (InMaterial.GetMaterialDomain() == MD_Volume)
+		{
+			TArray<const UMaterialExpressionVolumetricAdvancedMaterialOutput*> VolumetricAdvancedExpressions;
+			InMaterial.GetMaterialInterface()->GetMaterial()->GetAllExpressionsOfType(VolumetricAdvancedExpressions);
+			if (VolumetricAdvancedExpressions.Num() > 0)
 			{
-				// Value must match SINGLE_LAYER_WATER_SHADING_QUALITY_MOBILE_WITH_DEPTH_TEXTURE in SingleLayerWaterCommon.ush!
-				OutEnvironment.SetDefine(TEXT("SINGLE_LAYER_WATER_SHADING_QUALITY"), true);
+				if (VolumetricAdvancedExpressions.Num() > 1)
+				{
+					UE_LOG(LogMaterial, Fatal, TEXT("Only a single UMaterialExpressionVolumetricAdvancedMaterialOutput node is supported."));
+				}
+
+				const UMaterialExpressionVolumetricAdvancedMaterialOutput* VolumetricAdvancedNode = VolumetricAdvancedExpressions[0];
+				const TCHAR* Param = VolumetricAdvancedNode->GetEvaluatePhaseOncePerSample() ? TEXT("MATERIAL_VOLUMETRIC_ADVANCED_PHASE_PERSAMPLE") : TEXT("MATERIAL_VOLUMETRIC_ADVANCED_PHASE_PERPIXEL");
+				OutEnvironment.SetDefine(Param, true);
+
+				OutEnvironment.SetDefine(TEXT("MATERIAL_VOLUMETRIC_ADVANCED"), true);
+				OutEnvironment.SetDefine(TEXT("MATERIAL_VOLUMETRIC_ADVANCED_GRAYSCALE_MATERIAL"), VolumetricAdvancedNode->bGrayScaleMaterial);
+				OutEnvironment.SetDefine(TEXT("MATERIAL_VOLUMETRIC_ADVANCED_RAYMARCH_VOLUME_SHADOW"), VolumetricAdvancedNode->bRayMarchVolumeShadow);
+				OutEnvironment.SetDefine(TEXT("MATERIAL_VOLUMETRIC_ADVANCED_CLAMP_MULTISCATTERING_CONTRIBUTION"), VolumetricAdvancedNode->bClampMultiScatteringContribution);
+				OutEnvironment.SetDefine(TEXT("MATERIAL_VOLUMETRIC_ADVANCED_MULTISCATTERING_OCTAVE_COUNT"), VolumetricAdvancedNode->GetMultiScatteringApproximationOctaveCount());
+				OutEnvironment.SetDefine(TEXT("MATERIAL_VOLUMETRIC_ADVANCED_CONSERVATIVE_DENSITY"), VolumetricAdvancedNode->ConservativeDensity.IsConnected());
+				OutEnvironment.SetDefine(TEXT("MATERIAL_VOLUMETRIC_ADVANCED_OVERRIDE_AMBIENT_OCCLUSION"), InMaterial.HasAmbientOcclusionConnected());
+				OutEnvironment.SetDefine(TEXT("MATERIAL_VOLUMETRIC_ADVANCED_GROUND_CONTRIBUTION"), VolumetricAdvancedNode->bGroundContribution);
 			}
-
-			OutEnvironment.SetDefine(GetShadingModelParameterName(Model), true);
-			NumActiveShadingModels += 1;
 		}
-	}
-	else
-	{
-		// Unlit shading model can only exist by itself
-		OutEnvironment.SetDefine(TEXT("MATERIAL_SINGLE_SHADINGMODEL"), true);
-		OutEnvironment.SetDefine(GetShadingModelParameterName(MSM_Unlit), true);
-	}
 
-	if (NumActiveShadingModels == 1)
-	{
-		OutEnvironment.SetDefine(TEXT("MATERIAL_SINGLE_SHADINGMODEL"), true);
+		OutEnvironment.SetDefine(TEXT("MATERIAL_IS_SUBSTRATE"), false);
+		OutEnvironment.SetDefine(TEXT("DUAL_SOURCE_COLOR_BLENDING_ENABLED"), false);
+		OutEnvironment.SetDefine(TEXT("TEXTURE_SAMPLE_DEBUG"), false);
 	}
-	else if (!ensure(NumActiveShadingModels > 0))
-	{
-		UE_LOG(LogMaterial, Warning, TEXT("Unknown material shading model(s). Setting to MSM_DefaultLit"));
-		OutEnvironment.SetDefine(GetShadingModelParameterName(MSM_DefaultLit), true);
-	}
-
-	static IConsoleVariable* CVarLWCIsEnabled = IConsoleManager::Get().FindConsoleVariable(TEXT("r.MaterialEditor.LWCEnabled"));
-	OutEnvironment.SetDefine(TEXT("MATERIAL_LWC_ENABLED"), CVarLWCIsEnabled->GetInt());
-	OutEnvironment.SetDefine(TEXT("WSVECTOR_IS_TILEOFFSET"), true);
-	OutEnvironment.SetDefine(TEXT("WSVECTOR_IS_DOUBLEFLOAT"), false);
 
-	if (InMaterial.GetMaterialDomain() == MD_Volume)
+	void LowerValue( FHLSLPrinter& Printer, const IR::FValue* InValue)
 	{
-		TArray<const UMaterialExpressionVolumetricAdvancedMaterialOutput*> VolumetricAdvancedExpressions;
-		InMaterial.GetMaterialInterface()->GetMaterial()->GetAllExpressionsOfType(VolumetricAdvancedExpressions);
-		if (VolumetricAdvancedExpressions.Num() > 0)
+		if (const IR::FScalarConstant* Scalar = InValue->As<IR::FScalarConstant>())
 		{
-			if (VolumetricAdvancedExpressions.Num() > 1)
+			IR::FArithmeticTypePtr ArithType = Scalar->Type->ToArithmetic();
+			check(ArithType && ArithType->IsScalar());
+
+			switch (ArithType->ScalarKind)
 			{
-				UE_LOG(LogMaterial, Fatal, TEXT("Only a single UMaterialExpressionVolumetricAdvancedMaterialOutput node is supported."));
+				case IR::SK_Bool: Printer << Scalar->Boolean; break;
+				case IR::SK_Int: Printer << Scalar->Integer; break;
+				case IR::SK_Float: Printer << Scalar->Float; break;
 			}
-
-			const UMaterialExpressionVolumetricAdvancedMaterialOutput* VolumetricAdvancedNode = VolumetricAdvancedExpressions[0];
-			const TCHAR* Param = VolumetricAdvancedNode->GetEvaluatePhaseOncePerSample() ? TEXT("MATERIAL_VOLUMETRIC_ADVANCED_PHASE_PERSAMPLE") : TEXT("MATERIAL_VOLUMETRIC_ADVANCED_PHASE_PERPIXEL");
-			OutEnvironment.SetDefine(Param, true);
-
-			OutEnvironment.SetDefine(TEXT("MATERIAL_VOLUMETRIC_ADVANCED"), true);
-			OutEnvironment.SetDefine(TEXT("MATERIAL_VOLUMETRIC_ADVANCED_GRAYSCALE_MATERIAL"), VolumetricAdvancedNode->bGrayScaleMaterial);
-			OutEnvironment.SetDefine(TEXT("MATERIAL_VOLUMETRIC_ADVANCED_RAYMARCH_VOLUME_SHADOW"), VolumetricAdvancedNode->bRayMarchVolumeShadow);
-			OutEnvironment.SetDefine(TEXT("MATERIAL_VOLUMETRIC_ADVANCED_CLAMP_MULTISCATTERING_CONTRIBUTION"), VolumetricAdvancedNode->bClampMultiScatteringContribution);
-			OutEnvironment.SetDefine(TEXT("MATERIAL_VOLUMETRIC_ADVANCED_MULTISCATTERING_OCTAVE_COUNT"), VolumetricAdvancedNode->GetMultiScatteringApproximationOctaveCount());
-			OutEnvironment.SetDefine(TEXT("MATERIAL_VOLUMETRIC_ADVANCED_CONSERVATIVE_DENSITY"), VolumetricAdvancedNode->ConservativeDensity.IsConnected());
-			OutEnvironment.SetDefine(TEXT("MATERIAL_VOLUMETRIC_ADVANCED_OVERRIDE_AMBIENT_OCCLUSION"), InMaterial.HasAmbientOcclusionConnected());
-			OutEnvironment.SetDefine(TEXT("MATERIAL_VOLUMETRIC_ADVANCED_GROUND_CONTRIBUTION"), VolumetricAdvancedNode->bGroundContribution);
 		}
-	}
+		else if (const IR::FDimensional* Vector = InValue->As<IR::FDimensional>())
+		{
+			IR::FArithmeticTypePtr ArithType = Vector->Type->ToArithmetic();
+			check(ArithType && ArithType->IsVector());
 
-	OutEnvironment.SetDefine(TEXT("MATERIAL_IS_SUBSTRATE"), false);
-	OutEnvironment.SetDefine(TEXT("DUAL_SOURCE_COLOR_BLENDING_ENABLED"), false);
-	OutEnvironment.SetDefine(TEXT("TEXTURE_SAMPLE_DEBUG"), false);
-}
+			Printer << ScalarKindToString(ArithType->ScalarKind) << ArithType->NumRows << TEXT("(");
 
-void LowerValue(FHLSLPrinter& Printer, const IR::FValue* InValue)
-{
-	if (const IR::FScalarValue* Scalar = InValue->Cast<IR::FScalarValue>())
-	{
-		IR::FArithmeticTypePtr ArithType = Scalar->Type->ToArithmetic();
-		check(ArithType && ArithType->IsScalar());
+			Printer.BeginList();
+			for (IR::FValuePtr Component : Vector->GetComponents())
+			{
+				Printer.PrintListSeparator();
+				LowerValue(Printer, Component);
+			}
 
-		switch (ArithType->ScalarKind)
+			Printer << TEXT(")");
+		}
+		else if (const IR::FSetMaterialOutput* Output = InValue->As<IR::FSetMaterialOutput>())
 		{
-			case IR::SK_Bool: Printer << Scalar->Boolean; break;
-			case IR::SK_Int: Printer << Scalar->Integer; break;
-			case IR::SK_Float: Printer << Scalar->Float; break;
+			// Special case MP_SubsurfaceColor as the actual property is a combination of the color and the profile but we don't want to expose the profile
+			const FString& PropertyName = (Output->Property == MP_SubsurfaceColor) ? "Subsurface" : FMaterialAttributeDefinitionMap::GetAttributeName(Output->Property);
+
+			Printer.Printf(TEXT("\tPixelMaterialInputs.%s = "), *PropertyName);
+		
+			LowerValue(Printer, Output->ArgValue);
+
+			Printer << TEXT(";\n");
 		}
-	}
-	else if (const IR::FVectorValue* Vector = InValue->Cast<IR::FVectorValue>())
-	{
-		IR::FArithmeticTypePtr ArithType = Vector->Type->ToArithmetic();
-		check(ArithType && ArithType->IsVector());
+		else if (const IR::FBinaryOperator* BinOp = InValue->As<IR::FBinaryOperator>())
+		{
+			LowerValue(Printer, BinOp->Lhs);
 
-		Printer << ScalarKindToString(ArithType->ScalarKind) << ArithType->NumRows << TEXT("(");
+			const TCHAR* OpString;
+			switch (BinOp->Operator)
+			{
+				case IR::BO_Add: OpString = TEXT(" + "); break;
+				case IR::BO_Subtract: OpString = TEXT(" - "); break;
+				case IR::BO_Multiply: OpString = TEXT(" * "); break;
+				case IR::BO_Divide: OpString = TEXT(" / "); break;
+				default: UE_MIR_UNREACHABLE();
+			}
+			Printer << OpString;
 
-		Printer.BeginList();
-		for (IR::FValuePtr Component : Vector->GetComponents())
+			LowerValue(Printer, BinOp->Rhs);
+		}
+		else
 		{
-			Printer.PrintListSeparator();
-			LowerValue(Printer, Component);
+			UE_MIR_UNREACHABLE();
 		}
-
-		Printer << TEXT(")");
 	}
-	else if (const IR::FSetMaterialOutputInstr* Output = InValue->Cast<IR::FSetMaterialOutputInstr>())
-	{
-		// Special case MP_SubsurfaceColor as the actual property is a combination of the color and the profile but we don't want to expose the profile
-		const FString& PropertyName = (Output->Property == MP_SubsurfaceColor) ? "Subsurface" : FMaterialAttributeDefinitionMap::GetAttributeName(Output->Property);
-
-		Printer.Printf(TEXT("\tPixelMaterialInputs.%s = "), *PropertyName);
-		
-		LowerValue(Printer, Output->ArgValue);
+};
 
-		Printer << TEXT(";\n");
-	}
-	else
-	{
-		UE_MIR_UNREACHABLE();
-	}
+FMaterialIRToHLSLTranslator::FMaterialIRToHLSLTranslator()
+{
 }
 
-UE_MIR_END_PRIVATE(); //FMaterialIRToHLSLTranslator
-
-FMaterialIRToHLSLTranslator::FMaterialIRToHLSLTranslator()
+void FMaterialIRToHLSLTranslator::SetTarget(EShaderPlatform InShaderPlatform, const ITargetPlatform* InTargetPlatform, ERHIFeatureLevel::Type InFeatureLevel)
 {
+	ShaderPlatform = InShaderPlatform;
+	TargetPlatform = InTargetPlatform;
+	FeatureLevel = InFeatureLevel;
 }
 
-void FMaterialIRToHLSLTranslator::Translate(const FMaterial& InMaterial, const FMaterialIRModule& InModule, FParametersMap& OutParametersMap, FShaderCompilerEnvironment& OutEnvironment)
+bool FMaterialIRToHLSLTranslator::Translate(
+	const FMaterial& InMaterial,
+	const FStaticParameterSet& StaticParameters,
+	const FMaterialIRModule& InModule,
+	FParametersMap& OutParametersMap,
+	FShaderCompilerEnvironment& OutEnvironment)
 {
 	OutParametersMap.Empty();
 
-	AsPrivate()->GenerateHLSL(InMaterial, InModule, OutParametersMap);
-	AsPrivate()->SetMaterialParameters(InMaterial, OutParametersMap);
-	AsPrivate()->GetShaderCompilerEnvironment(InMaterial, InModule, OutEnvironment);
+	FPrivate Private{ this };
+	Private.GenerateHLSL(InMaterial, InModule, OutParametersMap);
+	Private.SetMaterialParameters(InMaterial, OutParametersMap);
+	Private.GetShaderCompilerEnvironment(InMaterial, InModule, OutEnvironment);
+
+	return true;
 }
 
 #endif // #if WITH_EDITOR
diff --git a/Engine/Source/Runtime/Engine/Private/Materials/MaterialIRTypes.cpp b/Engine/Source/Runtime/Engine/Private/Materials/MaterialIRTypes.cpp
--- a/Engine/Source/Runtime/Engine/Private/Materials/MaterialIRTypes.cpp
+++ b/Engine/Source/Runtime/Engine/Private/Materials/MaterialIRTypes.cpp
@@ -4,7 +4,7 @@
 
 #if WITH_EDITOR
 
-namespace MaterialIR
+namespace UE::MIR
 {
 
 FTypePtr FType::FromShaderType(const UE::Shader::FType& InShaderType)
@@ -46,18 +46,27 @@ FTypePtr FType::GetVoid()
 	return &Type;
 }
 
-const FArithmeticType* FType::ToArithmetic() const
+FArithmeticTypePtr FType::ToArithmetic() const
 {
-	return Kind == TK_Arithmetic ? static_cast<const FArithmeticType*>(this) : nullptr; 
+	return Kind == TK_Arithmetic ? static_cast<FArithmeticTypePtr>(this) : nullptr; 
 }
 
-bool FType::IsScalar() const
+FArithmeticTypePtr FType::ToScalar() const
 {
-	if (FArithmeticTypePtr Arith = ToArithmetic())
-	{
-		return Arith->IsScalar();
-	}
-	return false;
+	FArithmeticTypePtr Type = ToArithmetic();
+	return Type->IsScalar() ? Type : nullptr;
+}
+
+FArithmeticTypePtr FType::ToVector() const
+{
+	FArithmeticTypePtr Type = ToArithmetic();
+	return Type->IsVector() ? Type : nullptr;
+}
+
+FArithmeticTypePtr FType::ToMatrix() const
+{
+	FArithmeticTypePtr Type = ToArithmetic();
+	return Type->IsMatrix() ? Type : nullptr;
 }
 
 const TCHAR* ScalarKindToString(EScalarKind Kind)
@@ -149,6 +158,6 @@ const FArithmeticType* FArithmeticType::GetMatrix(EScalarKind InScalarKind, int
 	return GetNumericalType(InScalarKind, NumRows, NumColumns);
 }
 
-} // namespace MaterialIR
+} // namespace UE::MIR
 
 #endif // #if WITH_EDITOR
diff --git a/Engine/Source/Runtime/Engine/Private/Materials/MaterialIRUtility.cpp b/Engine/Source/Runtime/Engine/Private/Materials/MaterialIRUtility.cpp
--- a/Engine/Source/Runtime/Engine/Private/Materials/MaterialIRUtility.cpp
+++ b/Engine/Source/Runtime/Engine/Private/Materials/MaterialIRUtility.cpp
@@ -2,14 +2,15 @@
 
 #include "Materials/MaterialIRUtility.h"
 #include "Materials/MaterialAttributeDefinitionMap.h"
-#include "Materials/MaterialIRBuilder.h"
+#include "Materials/MaterialIREmitter.h"
+#include "Materials/Material.h"
 #include "MaterialShared.h"
 
 #if WITH_EDITOR
 
-namespace Utility {
+namespace IR = UE::MIR;
 
-namespace MIR = MaterialIR;
+namespace UE::Utility {
 
 bool IsMaterialPropertyShared(EMaterialProperty InProperty)
 {
@@ -39,24 +40,41 @@ bool IsMaterialPropertyShared(EMaterialProperty InProperty)
 	}
 }
 
-MIR::FValuePtr CreateMaterialAttributeDefaultValue(MIR::FBuilder& Builder, const FMaterial* Material, EMaterialProperty Property)
+bool NextMaterialAttributeInput(UMaterial* BaseMaterial, int32& PropertyIndex, FMaterialInputDescription& Input)
+{
+	for (; PropertyIndex < MP_MAX; ++PropertyIndex)
+	{
+		EMaterialProperty Property = (EMaterialProperty)PropertyIndex;
+		if (Utility::IsMaterialPropertyShared(Property)
+			&& Property != MP_SubsurfaceColor
+			&& Property != MP_FrontMaterial
+			&& BaseMaterial->GetExpressionInputDescription(Property, Input))
+		{
+			return true;
+		}
+	}
+
+	return false;
+}
+
+IR::FValuePtr CreateMaterialAttributeDefaultValue(IR::FEmitter& Emitter, const UMaterial* Material, EMaterialProperty Property)
 {
 	EMaterialValueType Type = FMaterialAttributeDefinitionMap::GetValueType(Property);
 	FVector4f DefaultValue = FMaterialAttributeDefinitionMap::GetDefaultValue(Property);
 
 	switch (Type)
 	{
-		case MCT_ShadingModel: return Builder.NewConstantInt1(Material->GetShadingModels().GetFirstShadingModel());
+		case MCT_ShadingModel: return Emitter.EmitConstantInt1(Material->GetShadingModels().GetFirstShadingModel());
 
-		case MCT_Float1: return Builder.NewConstantFloat1(  DefaultValue.X );
-		case MCT_Float2: return Builder.NewConstantFloat2({ DefaultValue.X, DefaultValue.Y });
-		case MCT_Float3: return Builder.NewConstantFloat3({ DefaultValue.X, DefaultValue.Y, DefaultValue.Z });
-		case MCT_Float4: return Builder.NewConstantFloat4(DefaultValue);
+		case MCT_Float1: return Emitter.EmitConstantFloat1(  DefaultValue.X );
+		case MCT_Float2: return Emitter.EmitConstantFloat2({ DefaultValue.X, DefaultValue.Y });
+		case MCT_Float3: return Emitter.EmitConstantFloat3({ DefaultValue.X, DefaultValue.Y, DefaultValue.Z });
+		case MCT_Float4: return Emitter.EmitConstantFloat4(DefaultValue);
 
-		case MCT_UInt1: return Builder.NewConstantInt1(  (int32)DefaultValue.X );
-		case MCT_UInt2: return Builder.NewConstantInt2({ (int32)DefaultValue.X, (int32)DefaultValue.Y });
-		case MCT_UInt3: return Builder.NewConstantInt3({ (int32)DefaultValue.X, (int32)DefaultValue.Y, (int32)DefaultValue.Z });
-		case MCT_UInt4: return Builder.NewConstantInt4({ (int32)DefaultValue.X, (int32)DefaultValue.Y, (int32)DefaultValue.Z, (int32)DefaultValue.W });
+		case MCT_UInt1: return Emitter.EmitConstantInt1(  (int32)DefaultValue.X );
+		case MCT_UInt2: return Emitter.EmitConstantInt2({ (int32)DefaultValue.X, (int32)DefaultValue.Y });
+		case MCT_UInt3: return Emitter.EmitConstantInt3({ (int32)DefaultValue.X, (int32)DefaultValue.Y, (int32)DefaultValue.Z });
+		case MCT_UInt4: return Emitter.EmitConstantInt4({ (int32)DefaultValue.X, (int32)DefaultValue.Y, (int32)DefaultValue.Z, (int32)DefaultValue.W });
 
 		default: UE_MIR_UNREACHABLE();
 	}
diff --git a/Engine/Source/Runtime/Engine/Private/Materials/MaterialIRUtility.h b/Engine/Source/Runtime/Engine/Private/Materials/MaterialIRUtility.h
--- a/Engine/Source/Runtime/Engine/Private/Materials/MaterialIRUtility.h
+++ b/Engine/Source/Runtime/Engine/Private/Materials/MaterialIRUtility.h
@@ -5,12 +5,16 @@
 
 #if WITH_EDITOR
 
-namespace Utility {
+namespace UE::Utility {
 
+//
 bool IsMaterialPropertyShared(EMaterialProperty InProperty);
 
 //
-MaterialIR::FValuePtr CreateMaterialAttributeDefaultValue(MaterialIR::FBuilder& Builder, const FMaterial* Material, EMaterialProperty Property);
+bool NextMaterialAttributeInput(UMaterial* BaseMaterial, int32& PropertyIndex, FMaterialInputDescription& Input);
+
+//
+UE::MIR::FValuePtr CreateMaterialAttributeDefaultValue(UE::MIR::FEmitter& Emitter, const UMaterial* Material, EMaterialProperty Property);
 
 } // namespace Utility
 
diff --git a/Engine/Source/Runtime/Engine/Private/Materials/MaterialInterface.cpp b/Engine/Source/Runtime/Engine/Private/Materials/MaterialInterface.cpp
--- a/Engine/Source/Runtime/Engine/Private/Materials/MaterialInterface.cpp
+++ b/Engine/Source/Runtime/Engine/Private/Materials/MaterialInterface.cpp
@@ -307,6 +307,11 @@ bool UMaterialInterface::IsUsingNewHLSLGenerator() const
 	return BaseMaterial ? BaseMaterial->bEnableNewHLSLGenerator : false;
 }
 
+bool UMaterialInterface::IsUsingNewTranslatorPrototype() const
+{
+	return IsUsingNewHLSLGenerator() && bIsUsingNewMaterialTranslatorPrototype;
+}
+
 const FSubstrateCompilationConfig& UMaterialInterface::GetSubstrateCompilationConfig() const
 {
 	const UMaterial* BaseMaterial = GetMaterial_Concurrent();
diff --git a/Engine/Source/Runtime/Engine/Private/Materials/MaterialShared.cpp b/Engine/Source/Runtime/Engine/Private/Materials/MaterialShared.cpp
--- a/Engine/Source/Runtime/Engine/Private/Materials/MaterialShared.cpp
+++ b/Engine/Source/Runtime/Engine/Private/Materials/MaterialShared.cpp
@@ -95,9 +95,10 @@ static TAutoConsoleVariable<bool> CVarMaterialEdPreshaderDumpToHLSL(
 	TEXT("Controls whether to append preshader expressions and parameter reference counts to the HLSL source window (as comments at the end of the code)."),
 	ECVF_RenderThreadSafe);
 
-static TAutoConsoleVariable<bool> CVarMaterialEdUseNewTranslatorPrototype(
+bool bIsUsingNewMaterialTranslatorPrototype;
+static FAutoConsoleVariableRef CVarUsingUseNewMaterialTranslatorPrototype(
 	TEXT("r.MaterialEditor.UseNewTranslatorPrototype"),
-	false,
+	bIsUsingNewMaterialTranslatorPrototype,
 	TEXT("Controls whether to enable the new material translator prototype (WIP) ."),
 	ECVF_RenderThreadSafe);
 
@@ -461,6 +462,11 @@ FExpressionInput FExpressionInput::GetTracedInput() const
 	return *this;
 }
 
+FExpressionOutput* FExpressionInput::GetConnectedOutput()
+{
+	return IsConnected() ? &Expression->GetOutputs()[OutputIndex] : nullptr;
+}
+
 int32 FExpressionExecOutput::Compile(class FMaterialCompiler* Compiler) const
 {
 	if (Expression)
@@ -3368,37 +3374,40 @@ bool FMaterial::Translate_New(const FMaterialShaderMapId& InShaderMapId,
 	FMaterialCompilationOutput& OutCompilationOutput,
 	TRefCountPtr<FSharedShaderCompilerEnvironment>& OutMaterialEnvironment)
 {
-	const bool bUseNewTranslatorPrototype = CVarMaterialEdUseNewTranslatorPrototype->GetBool();
-	if (bUseNewTranslatorPrototype)
+	if (bIsUsingNewMaterialTranslatorPrototype)
 	{
-		FMaterialIRModule Module;
+		const FMaterialIRModule& Module = this->GetMaterialInterface()->GetIRModule();
 
-		/* Setup the builder and bild the material */
-		FMaterialIRModuleBuilder Builder;
-		Builder.SetSource(this, &InStaticParameters);
-		Builder.SetPlatform(InPlatform, InShaderMapId.FeatureLevel, InTargetPlatform);
-		Builder.SetTarget(&Module);
-		if (!Builder.Build())
+		if (!Module.GetErrors().IsEmpty())
 		{
+			// TODO: This could be moved elsewhere, when the Module is built.
+			for (const FMaterialIRModule::FError& Error : Module.GetErrors())
+			{
+				ErrorExpressions.Push(Error.Expression);
+				CompileErrors.Push(Error.Message);
+			}
 			return false;
 		}
-
-		/* Copy over the compilation output */
+	
+		// Copy over the compilation output
 		OutCompilationOutput = Module.GetCompilationOutput();
 		OutMaterialEnvironment = new FSharedShaderCompilerEnvironment();
 
-		/* Translate the material IR module to HLSL template string parameters and material environment */
+		// Translate the material IR module to HLSL template string parameters and material environment
 		FMaterialIRToHLSLTranslator::FParametersMap ShaderStringParameters;
 		FMaterialIRToHLSLTranslator Translator;
-		Translator.Translate(*this, Module, ShaderStringParameters, *OutMaterialEnvironment);
+		if (!Translator.Translate(*this, InStaticParameters, Module, ShaderStringParameters, *OutMaterialEnvironment))
+		{
+			return false;
+		}
 
-		/* Interpolate HLSL parameters with the material shader template to produce the final shader source */
+		// Interpolate HLSL parameters with the material shader template to produce the final shader source
 		int32 LineNumber;
 		FStringTemplateResolver Resolver = FMaterialSourceTemplate::Get().BeginResolve(InPlatform, &LineNumber);
 		ShaderStringParameters.Add({TEXT("line_number"), FString::Printf(TEXT("%u"), LineNumber)});
 		Resolver.SetParameterMap(&ShaderStringParameters);
 
-		// Interpolate the string
+		// Interpolate the final material shader source string
 		FString MaterialShaderCode = Resolver.Finalize();
 		OutMaterialEnvironment->IncludeVirtualPathToContentsMap.Add(TEXT("/Engine/Generated/Material.ush"), MoveTemp(MaterialShaderCode));
 
diff --git a/Engine/Source/Runtime/Engine/Public/MaterialExpressionIO.h b/Engine/Source/Runtime/Engine/Public/MaterialExpressionIO.h
--- a/Engine/Source/Runtime/Engine/Public/MaterialExpressionIO.h
+++ b/Engine/Source/Runtime/Engine/Public/MaterialExpressionIO.h
@@ -7,6 +7,7 @@
 #include "SceneTypes.h"
 
 class FMaterialHLSLGenerator;
+struct FExpressionOutput;
 
 namespace UE::HLSLTree
 {
@@ -110,6 +111,10 @@ struct FExpressionInput
 		MaskB = B;
 		MaskA = A;
 	}
+
+	/** Returns the connected expression output if this input is connected to one,	nullptr otherwise. */
+	FExpressionOutput* GetConnectedOutput();
+
 #endif // WITH_EDITOR
 };
 
diff --git a/Engine/Source/Runtime/Engine/Public/MaterialShared.h b/Engine/Source/Runtime/Engine/Public/MaterialShared.h
--- a/Engine/Source/Runtime/Engine/Public/MaterialShared.h
+++ b/Engine/Source/Runtime/Engine/Public/MaterialShared.h
@@ -100,6 +100,11 @@ class FExpressionStaticTerrainLayerWeight;
 }
 }
 
+namespace UE::MIR
+{
+	class FEmitter;
+}
+
 template <class ElementType> class TLinkedList;
 
 #define ME_CAPTION_HEIGHT		18
@@ -191,6 +196,8 @@ ENGINE_API bool IsAlphaCompositeBlendMode(const FMaterial& In);
 ENGINE_API bool IsAlphaCompositeBlendMode(const UMaterialInterface& In);
 ENGINE_API bool IsAlphaCompositeBlendMode(const FMaterialShaderParameters& In);
 
+extern bool bIsUsingNewMaterialTranslatorPrototype;
+
 #if WITH_EDITOR
 
 void GetCmdLineFilterShaderFormats(TArray<FName>& InOutShderFormats);
@@ -2771,6 +2778,7 @@ private:
 	friend class FShaderCompilingManager;
 	friend class FHLSLMaterialTranslator;
 	friend class FMaterialHLSLErrorHandler;
+	friend class UE::MIR::FEmitter;
 };
 
 
diff --git a/Engine/Source/Runtime/Engine/Public/Materials/MaterialExpression.h b/Engine/Source/Runtime/Engine/Public/Materials/MaterialExpression.h
--- a/Engine/Source/Runtime/Engine/Public/Materials/MaterialExpression.h
+++ b/Engine/Source/Runtime/Engine/Public/Materials/MaterialExpression.h
@@ -33,12 +33,19 @@ enum class EMaterialParameterType : uint8;
 
 namespace UE
 {
+
 namespace HLSLTree
 {
 class FScope;
 class FStatement;
 class FExpression;
 }
+
+namespace MIR
+{
+	class FEmitter;
+}
+
 }
 
 //@warning: FExpressionInput is mirrored in MaterialShared.h and manually "subclassed" in Material.h (FMaterialInput)
@@ -321,6 +328,14 @@ class UMaterialExpression : public UObject
 
 	ENGINE_API bool IsUsingNewHLSLGenerator() const;
 
+	/**
+	 * Builds this expression's material IR.
+	 * 
+	 * This function is responsible for doing the semantic analysis on this error, error checking and reporting and
+	 * implementing this expression's logic by emitting appropriate material IR values.
+	 */
+	ENGINE_API virtual void Build(UE::MIR::FEmitter& Emitter);
+
 #endif // WITH_EDITOR
 
 	/**
@@ -348,6 +363,10 @@ class UMaterialExpression : public UObject
 
 	/** Counts the number of inputs this expression has. Default implementation has O(n) complexity. */
 	ENGINE_API virtual int32 CountInputs() const;
+	/**
+	 * Returns the output at specified index is valid
+	 */
+	ENGINE_API FExpressionOutput* GetOutput(int32 OutputIndex) { return &Outputs[OutputIndex]; }
 
 	/**
 	 * Returns the input at index InputIndex if valid, nullptr otherwise.
diff --git a/Engine/Source/Runtime/Engine/Public/Materials/MaterialExpressionAdd.h b/Engine/Source/Runtime/Engine/Public/Materials/MaterialExpressionAdd.h
--- a/Engine/Source/Runtime/Engine/Public/Materials/MaterialExpressionAdd.h
+++ b/Engine/Source/Runtime/Engine/Public/Materials/MaterialExpressionAdd.h
@@ -31,6 +31,7 @@ class UMaterialExpressionAdd : public UMaterialExpression
 
 	//~ Begin UMaterialExpression Interface
 #if WITH_EDITOR
+	virtual void Build(UE::MIR::FEmitter& Emitter) override;
 	virtual int32 Compile(class FMaterialCompiler* Compiler, int32 OutputIndex) override;
 	virtual void GetCaption(TArray<FString>& OutCaptions) const override;
 	virtual FText GetKeywords() const override {return FText::FromString(TEXT("+"));}
diff --git a/Engine/Source/Runtime/Engine/Public/Materials/MaterialExpressionConstant.h b/Engine/Source/Runtime/Engine/Public/Materials/MaterialExpressionConstant.h
--- a/Engine/Source/Runtime/Engine/Public/Materials/MaterialExpressionConstant.h
+++ b/Engine/Source/Runtime/Engine/Public/Materials/MaterialExpressionConstant.h
@@ -19,6 +19,7 @@ class UMaterialExpressionConstant : public UMaterialExpression
 
 	//~ Begin UMaterialExpression Interface
 #if WITH_EDITOR
+	virtual void Build(UE::MIR::FEmitter& Emitter) override;
 	virtual int32 Compile(class FMaterialCompiler* Compiler, int32 OutputIndex) override;
 	virtual void GetCaption(TArray<FString>& OutCaptions) const override;
 	virtual FString GetDescription() const override;
diff --git a/Engine/Source/Runtime/Engine/Public/Materials/MaterialExpressionConstant2Vector.h b/Engine/Source/Runtime/Engine/Public/Materials/MaterialExpressionConstant2Vector.h
--- a/Engine/Source/Runtime/Engine/Public/Materials/MaterialExpressionConstant2Vector.h
+++ b/Engine/Source/Runtime/Engine/Public/Materials/MaterialExpressionConstant2Vector.h
@@ -23,6 +23,7 @@ class UMaterialExpressionConstant2Vector : public UMaterialExpression
 
 	//~ Begin UMaterialExpression Interface
 #if WITH_EDITOR
+	virtual void Build(UE::MIR::FEmitter& Emitter) override;
 	virtual int32 Compile(class FMaterialCompiler* Compiler, int32 OutputIndex) override;
 	virtual void GetCaption(TArray<FString>& OutCaptions) const override;
 	virtual FString GetDescription() const override;
diff --git a/Engine/Source/Runtime/Engine/Public/Materials/MaterialExpressionConstant3Vector.h b/Engine/Source/Runtime/Engine/Public/Materials/MaterialExpressionConstant3Vector.h
--- a/Engine/Source/Runtime/Engine/Public/Materials/MaterialExpressionConstant3Vector.h
+++ b/Engine/Source/Runtime/Engine/Public/Materials/MaterialExpressionConstant3Vector.h
@@ -21,9 +21,9 @@ public:
 
 	//~ UMaterialExpression interface
 #if WITH_EDITOR
+	virtual void Build(UE::MIR::FEmitter& Emitter) override;
 	virtual int32 Compile(class FMaterialCompiler* Compiler, int32 OutputIndex) override;
 	virtual void GetCaption(TArray<FString>& OutCaptions) const override;
-
 	virtual FString GetDescription() const override;
 	virtual uint32 GetOutputType(int32 OutputIndex) override {return MCT_Float3;}
 	virtual bool GenerateHLSLExpression(FMaterialHLSLGenerator& Generator, UE::HLSLTree::FScope& Scope, int32 OutputIndex, UE::HLSLTree::FExpression const*& OutExpression) const override;
diff --git a/Engine/Source/Runtime/Engine/Public/Materials/MaterialExpressionConstant4Vector.h b/Engine/Source/Runtime/Engine/Public/Materials/MaterialExpressionConstant4Vector.h
--- a/Engine/Source/Runtime/Engine/Public/Materials/MaterialExpressionConstant4Vector.h
+++ b/Engine/Source/Runtime/Engine/Public/Materials/MaterialExpressionConstant4Vector.h
@@ -19,6 +19,7 @@ class UMaterialExpressionConstant4Vector : public UMaterialExpression
 
 	//~ Begin UMaterialExpression Interface
 #if WITH_EDITOR
+	virtual void Build(UE::MIR::FEmitter& Emitter) override;
 	virtual int32 Compile(class FMaterialCompiler* Compiler, int32 OutputIndex) override;
 	virtual void GetCaption(TArray<FString>& OutCaptions) const override;
 	virtual FString GetDescription() const override;
diff --git a/Engine/Source/Runtime/Engine/Public/Materials/MaterialExpressionDivide.h b/Engine/Source/Runtime/Engine/Public/Materials/MaterialExpressionDivide.h
--- a/Engine/Source/Runtime/Engine/Public/Materials/MaterialExpressionDivide.h
+++ b/Engine/Source/Runtime/Engine/Public/Materials/MaterialExpressionDivide.h
@@ -31,6 +31,7 @@ class UMaterialExpressionDivide : public UMaterialExpression
 
 	//~ Begin UMaterialExpression Interface
 #if WITH_EDITOR
+	virtual void Build(UE::MIR::FEmitter& Emitter) override;
 	virtual int32 Compile(class FMaterialCompiler* Compiler, int32 OutputIndex) override;
 	virtual bool GenerateHLSLExpression(FMaterialHLSLGenerator& Generator, UE::HLSLTree::FScope& Scope, int32 OutputIndex, UE::HLSLTree::FExpression const*& OutExpression) const override;
 	virtual void GetCaption(TArray<FString>& OutCaptions) const override;
diff --git a/Engine/Source/Runtime/Engine/Public/Materials/MaterialExpressionMultiply.h b/Engine/Source/Runtime/Engine/Public/Materials/MaterialExpressionMultiply.h
--- a/Engine/Source/Runtime/Engine/Public/Materials/MaterialExpressionMultiply.h
+++ b/Engine/Source/Runtime/Engine/Public/Materials/MaterialExpressionMultiply.h
@@ -31,6 +31,7 @@ class UMaterialExpressionMultiply : public UMaterialExpression
 
 	//~ Begin UMaterialExpression Interface
 #if WITH_EDITOR
+	virtual void Build(UE::MIR::FEmitter& Emitter) override;
 	virtual int32 Compile(class FMaterialCompiler* Compiler, int32 OutputIndex) override;
 	virtual void GetCaption(TArray<FString>& OutCaptions) const override;
 	virtual FText GetKeywords() const override {return FText::FromString(TEXT("*"));}
diff --git a/Engine/Source/Runtime/Engine/Public/Materials/MaterialExpressionSubtract.h b/Engine/Source/Runtime/Engine/Public/Materials/MaterialExpressionSubtract.h
--- a/Engine/Source/Runtime/Engine/Public/Materials/MaterialExpressionSubtract.h
+++ b/Engine/Source/Runtime/Engine/Public/Materials/MaterialExpressionSubtract.h
@@ -31,6 +31,7 @@ class UMaterialExpressionSubtract : public UMaterialExpression
 
 	//~ Begin UMaterialExpression Interface
 #if WITH_EDITOR
+	virtual void Build(UE::MIR::FEmitter& Emitter) override;
 	virtual bool GenerateHLSLExpression(FMaterialHLSLGenerator& Generator, UE::HLSLTree::FScope& Scope, int32 OutputIndex, UE::HLSLTree::FExpression const*& OutExpression) const override;
 	virtual int32 Compile(class FMaterialCompiler* Compiler, int32 OutputIndex) override;
 	virtual void GetCaption(TArray<FString>& OutCaptions) const override;
diff --git a/Engine/Source/Runtime/Engine/Public/Materials/MaterialIR.h b/Engine/Source/Runtime/Engine/Public/Materials/MaterialIR.h
--- a/Engine/Source/Runtime/Engine/Public/Materials/MaterialIR.h
+++ b/Engine/Source/Runtime/Engine/Public/Materials/MaterialIR.h
@@ -7,19 +7,21 @@ enum EMaterialProperty : int;
 
 #if WITH_EDITOR
 
-namespace MaterialIR
-{
+namespace UE::MIR {
 
 enum EValueKind
 {
 	/* Values */
-	VK_Value = 1 << 0,
-	VK_Scalar = VK_Value | 1 << 2,
-	VK_Vector = VK_Value | 2 << 2,
-	
+	VK_ScalarConstant,
+	VK_Dimensional,
+
 	/* Instructions */
-	VK_Instruction = 1 << 1,
- 	VK_SetMaterialOutputInstr = VK_Instruction | 1 << 2,
+	VK_InstructionBegin,
+
+	VK_SetMaterialOutput = VK_InstructionBegin,
+	VK_BinaryOperator,
+
+	VK_InstructionEnd,
 };
 
 /* Values */
@@ -29,27 +31,23 @@ struct FValue
 	EValueKind Kind{};
 	FTypePtr  Type{};
 
-	bool IsA(EValueKind InKind) const
-	{
-		return (Kind & InKind) == InKind;
-	}
+	bool IsA(EValueKind InKind) const { return Kind == InKind; }
+	FInstructionPtr AsInstruction() const;
+	bool Equals(FValuePtr Other) const;
+	uint32 GetSizeInBytes() const;
 
 	template <typename T>
-	const T* Cast() const
-	{
-		return this && IsA(T::TypeKind) ? static_cast<const T*>(this) : nullptr;
-	}
+	const T* As() const { return this && IsA(T::TypeKind) ? static_cast<const T*>(this) : nullptr; }
 
-	void Destroy();
 };
 
-	template <EValueKind TTypeKind>
-	struct TValue : FValue
-	{
-		static constexpr EValueKind TypeKind = TTypeKind;
-	};
+template <EValueKind TTypeKind>
+struct TValue : FValue
+{
+	static constexpr EValueKind TypeKind = TTypeKind;
+};
 
-struct FScalarValue : TValue<VK_Scalar>
+struct FScalarConstant : TValue<VK_ScalarConstant>
 {
 	union
 	{
@@ -59,37 +57,53 @@ struct FScalarValue : TValue<VK_Scalar>
 	};
 };
 
-struct FVectorValue : TValue<VK_Vector>
+struct FDimensional : TValue<VK_Dimensional>
 {
 	TArrayView<const FValuePtr> GetComponents() const;
 	TArrayView<FValuePtr> GetMutableComponents();
+	uint32 GetSizeInBytes() const;
 };
 
 template <int TDimension>
-struct TVectorValue : FVectorValue
+struct TDimensional : FDimensional
 {
 	FValuePtr Components[TDimension];
 };
 
 /* Instructions */
 
-struct FInstruction : TValue<VK_Instruction>
+struct FInstruction : FValue
 {
 	FInstruction* Next{};
 };
 
-	template <EValueKind TTypeKind>
-	struct TInstruction : FInstruction
-	{
-		static constexpr EValueKind TypeKind = TTypeKind;
-	};
+template <EValueKind TTypeKind>
+struct TInstruction : FInstruction
+{
+	static constexpr EValueKind TypeKind = TTypeKind;
+};
 
-struct FSetMaterialOutputInstr : TInstruction<VK_SetMaterialOutputInstr>
+struct FSetMaterialOutput : TInstruction<VK_SetMaterialOutput>
 {
 	EMaterialProperty Property;
 	FValuePtr ArgValue;
 };
 
-} // namespace MaterialIR
+enum EBinaryOperator
+{
+	BO_Invalid,
+	BO_Add,
+	BO_Subtract,
+	BO_Multiply,
+	BO_Divide,
+};
+
+struct FBinaryOperator : TInstruction<VK_BinaryOperator>
+{
+	EBinaryOperator Operator = BO_Invalid;
+	FValuePtr Lhs{};
+	FValuePtr Rhs{};
+};
 
+} // namespace UE::MIR
 #endif
diff --git a/Engine/Source/Runtime/Engine/Public/Materials/MaterialIRBuilder.h b/Engine/Source/Runtime/Engine/Public/Materials/MaterialIRBuilder.h
deleted file mode 100644
diff --git a/Engine/Source/Runtime/Engine/Public/Materials/MaterialIRCommon.h b/Engine/Source/Runtime/Engine/Public/Materials/MaterialIRCommon.h
--- a/Engine/Source/Runtime/Engine/Public/Materials/MaterialIRCommon.h
+++ b/Engine/Source/Runtime/Engine/Public/Materials/MaterialIRCommon.h
@@ -17,16 +17,20 @@ namespace UE::Shader
 
 namespace ERHIFeatureLevel { enum Type : int; }
 
-class ITargetPlatform;
-class FMaterialIRModule;
 class FMaterial;
-struct FStaticParameterSet;
-struct FShaderCompilerEnvironment;
+class FMaterialIRModule;
 class FMaterialIRModule;
 class FMaterialIRModuleBuilder;
+class ITargetPlatform;
+class UMaterial;
+class UMaterialExpression;
 struct FExpressionInput;
+struct FExpressionOutput;
+struct FShaderCompilerEnvironment;
+struct FStaticParameterSet;
+struct FMaterialInputDescription;
 
-namespace MaterialIR
+namespace UE::MIR
 {
 
 /* Types*/
@@ -38,28 +42,17 @@ using FArithmeticTypePtr = const FArithmeticType*;
 
 /* IR */
 struct FValue;
+struct FInstruction;
+struct FSetMaterialOutput;
 using FValuePtr = const FValue*;
-struct FScalarValue;
-struct FVector;
-struct FSetMaterialOutputInstr;
+using FInstructionPtr = const FInstruction*;
 
 /* Others */
-class FBuilder;
+class FEmitter;
 
 }
 
 #define UE_MIR_UNREACHABLE() { check(!"Unreachable"); UE_ASSUME(false); }
-
-#define UE_MIR_PRIVATE() \
-		struct FPrivate; \
-		friend FPrivate; \
-		FPrivate* AsPrivate() { return reinterpret_cast<FPrivate*>(this); }\
-		const FPrivate* AsPrivate() const { return reinterpret_cast<const FPrivate*>(this); }
-
-#define UE_MIR_BEGIN_PRIVATE(TypeName)\
-		struct TypeName::FPrivate : TypeName {
-
-#define UE_MIR_END_PRIVATE()\
-		}
+#define UE_MIR_TODO() UE_MIR_UNREACHABLE()
 
 #endif // #if WITH_EDITOR
diff --git a/Engine/Source/Runtime/Engine/Public/Materials/MaterialIREmitter.h b/Engine/Source/Runtime/Engine/Public/Materials/MaterialIREmitter.h
new file mode 100644
--- /dev/null
+++ b/Engine/Source/Runtime/Engine/Public/Materials/MaterialIREmitter.h
@@ -0,0 +1,79 @@
+// Copyright Epic Games, Inc. All Rights Reserved.
+#pragma once
+
+#include "Materials/MaterialIRCommon.h"
+#include "Materials/MaterialIR.h"
+#include "Materials/MaterialIRTypes.h"
+
+#if WITH_EDITOR
+
+namespace UE::MIR {
+
+//
+class FEmitter
+{
+public:
+	FEmitter(FMaterialIRModuleBuilder* InBuilder, UMaterial* InMaterial, FMaterialIRModule* InModule);
+	
+	/* Analysis */
+	FValuePtr GetAndCheckInputTypeKindIs(const FExpressionInput* Input, ETypeKind Kind);
+	bool CheckInputTypeIs(const FExpressionInput* Input, FValuePtr InputValue, ETypeKind Kind);
+
+	/* Constants */
+
+	FValuePtr EmitConstantFromShaderValue(const UE::Shader::FValue& InValue);
+	FValuePtr EmitConstantScalarZero(EScalarKind Kind);
+	FValuePtr EmitConstantFloat1(float InX);
+	FValuePtr EmitConstantFloat2(const FVector2f& InValue);
+	FValuePtr EmitConstantFloat3(const FVector3f& InValue);
+	FValuePtr EmitConstantFloat4(const FVector4f& InValue);
+	FValuePtr EmitConstantInt1(int InX);
+	FValuePtr EmitConstantInt2(const FIntVector2& InValue);
+	FValuePtr EmitConstantInt3(const FIntVector3& InValue);
+	FValuePtr EmitConstantInt4(const FIntVector4& InValue);
+	FValuePtr EmitVector2(FValuePtr InX, FValuePtr InY);
+	FValuePtr EmitVector3(FValuePtr InX, FValuePtr InY, FValuePtr InZ);
+	FValuePtr EmitVector4(FValuePtr InX, FValuePtr InY, FValuePtr InZ, FValuePtr InW);
+
+	/* Other Values */
+
+	FValuePtr EmitArithmetic(FArithmeticTypePtr Type, FValuePtr Scalar);
+
+	/* Instructions */
+
+	FSetMaterialOutput* EmitSetMaterialOutput(EMaterialProperty InProperty, FValue* InArgValue);
+	FValuePtr EmitBinaryOperator(EBinaryOperator Operator, FValuePtr Lhs, FValuePtr Rhs);
+	FValuePtr TryEmitConvert(FValuePtr Value, FTypePtr TargetType);
+
+	/* IO */
+
+	FValuePtr Get(const FExpressionInput* Input);
+	void Put(const FExpressionOutput* Output, FValuePtr Value);
+
+	/* Error reporting */
+	bool IsInvalid() const { return bHasExprBuildError; }
+
+	template <int TFormatLength, typename... TArgs>
+	void Errorf(const TCHAR (&Format)[TFormatLength], TArgs&&... Args)
+	{
+		Error(FString::Printf(Format, Forward<TArgs>(Args)...));
+	}
+
+	void Error(FString Message);
+
+private:
+	UMaterial* Material{};
+	FMaterialIRModule* Module{};
+	UMaterialExpression* Expression{};
+	FMaterialIRModuleBuilder* Builder{};
+	FValuePtr ZeroIntValue{};
+	FValuePtr ZeroFloatValue{};
+	bool bHasExprBuildError = true;
+
+	struct FPrivate;
+	friend FMaterialIRModuleBuilder;
+};
+
+} // namespace UE::MIR
+
+#endif // #if WITH_EDITOR
diff --git a/Engine/Source/Runtime/Engine/Public/Materials/MaterialIRModule.h b/Engine/Source/Runtime/Engine/Public/Materials/MaterialIRModule.h
--- a/Engine/Source/Runtime/Engine/Public/Materials/MaterialIRModule.h
+++ b/Engine/Source/Runtime/Engine/Public/Materials/MaterialIRModule.h
@@ -8,23 +8,29 @@
 
 class FMaterialIRModule
 {
+public:
+	struct FError
+	{
+		UMaterialExpression* Expression;
+		FString Message;
+	};
+
 public:
 	~FMaterialIRModule();
 	void Empty();
-	void Reset(EShaderPlatform InShaderPlatform, ERHIFeatureLevel::Type InFeatureLevel, const ITargetPlatform* InTargetPlatform);
-	EShaderPlatform GetShaderPlatform() const { return ShaderPlatform; }
-	const ITargetPlatform* GetTargetPlatform() const { return TargetPlatform; }
 	const FMaterialCompilationOutput& GetCompilationOutput() const { return CompilationOutput; }
-	TArrayView<const MaterialIR::FSetMaterialOutputInstr* const> GetOutputs() const { return Outputs; }
+	TArrayView<const UE::MIR::FSetMaterialOutput* const> GetOutputs() const { return Outputs; }
+	TArrayView<const FError> GetErrors() const { return Errors; }
 
 private:
-	EShaderPlatform ShaderPlatform;
-	const ITargetPlatform* TargetPlatform;
 	FMaterialCompilationOutput CompilationOutput;
-	TArray<MaterialIR::FValue*> Values;
-	TArray<MaterialIR::FSetMaterialOutputInstr*> Outputs;
+	TArray<UE::MIR::FValuePtr> Values;
+	TArray<UE::MIR::FSetMaterialOutput*> Outputs;
+	TArray<FError> Errors;
+
 
-	friend MaterialIR::FBuilder;
+	friend UE::MIR::FEmitter;
+	friend FMaterialIRModuleBuilder;
 };
 
-#endif
+#endif // #if WITH_EDITOR
diff --git a/Engine/Source/Runtime/Engine/Public/Materials/MaterialIRModuleBuilder.h b/Engine/Source/Runtime/Engine/Public/Materials/MaterialIRModuleBuilder.h
--- a/Engine/Source/Runtime/Engine/Public/Materials/MaterialIRModuleBuilder.h
+++ b/Engine/Source/Runtime/Engine/Public/Materials/MaterialIRModuleBuilder.h
@@ -9,30 +9,17 @@
 class FMaterialIRModuleBuilder
 {
 public:
-	//
-	void SetSource(FMaterial* InMaterial, const FStaticParameterSet* InStaticParameters);
-
-	//
-	void SetPlatform(EShaderPlatform InShaderPlatform, ERHIFeatureLevel::Type InFeatureLevel, const ITargetPlatform* InTargetPlatform);
-
-	//
-	void SetTarget(FMaterialIRModule* InModule);
-
-	//
-	bool Build();
+	bool Build(UMaterial* InMaterial, FMaterialIRModule* TargetModule);
 
 private:
-	FMaterial* Material;
-	const FStaticParameterSet* StaticParameters;
-	EShaderPlatform ShaderPlatform{};
-	ERHIFeatureLevel::Type FeatureLevel{};
-	const ITargetPlatform* TargetPlatform{};
+	UMaterial* BaseMaterial;
 	FMaterialIRModule* Module;
+	TArray<UMaterialExpression*> ExpressionAnalysisStack;
+	TMap<const FExpressionInput*, UE::MIR::FValuePtr> InputValues;
+	TMap<const FExpressionOutput*, UE::MIR::FValuePtr> OutputValues;
 
-	TArray<FExpressionInput*> ExpressionAnalysisStack;
-
-	struct FHelper;
-	friend FHelper;
+	friend class UE::MIR::FEmitter;
+	struct FPrivate;
 };
 
 #endif
diff --git a/Engine/Source/Runtime/Engine/Public/Materials/MaterialIRToHLSLTranslator.h b/Engine/Source/Runtime/Engine/Public/Materials/MaterialIRToHLSLTranslator.h
--- a/Engine/Source/Runtime/Engine/Public/Materials/MaterialIRToHLSLTranslator.h
+++ b/Engine/Source/Runtime/Engine/Public/Materials/MaterialIRToHLSLTranslator.h
@@ -13,11 +13,15 @@ public:
 
 public:
 	FMaterialIRToHLSLTranslator();
-	void Translate(const FMaterial& InMaterial, const FMaterialIRModule& InModule, FParametersMap& OutParameters, FShaderCompilerEnvironment& OutEnvironment);
+	void SetTarget(EShaderPlatform InShaderPlatform, const ITargetPlatform* InTargetPlatform, ERHIFeatureLevel::Type InFeatureLevel);
+	bool Translate(const FMaterial& InMaterial, const FStaticParameterSet& StaticParameters, const FMaterialIRModule& InModule, FParametersMap& OutParameters, FShaderCompilerEnvironment& OutEnvironment);
 
 private:
- 	
-	UE_MIR_PRIVATE();
+	EShaderPlatform ShaderPlatform{};
+	const ITargetPlatform* TargetPlatform{};
+	ERHIFeatureLevel::Type FeatureLevel{};
+
+	struct FPrivate;
 };
 
 #endif // #if WITH_EDITOR
diff --git a/Engine/Source/Runtime/Engine/Public/Materials/MaterialIRTypes.h b/Engine/Source/Runtime/Engine/Public/Materials/MaterialIRTypes.h
--- a/Engine/Source/Runtime/Engine/Public/Materials/MaterialIRTypes.h
+++ b/Engine/Source/Runtime/Engine/Public/Materials/MaterialIRTypes.h
@@ -6,7 +6,7 @@
 
 #if WITH_EDITOR
 
-namespace MaterialIR
+namespace UE::MIR
 {
 
 enum ETypeKind
@@ -22,9 +22,10 @@ struct FType
 	static FTypePtr FromShaderType(const UE::Shader::FType& InShaderType);
 	static FTypePtr GetVoid();
 
-	const FArithmeticType* ToArithmetic() const;
-	bool IsScalar() const;
-
+	FArithmeticTypePtr ToArithmetic() const;
+	FArithmeticTypePtr ToScalar() const;
+	FArithmeticTypePtr ToVector() const;
+	FArithmeticTypePtr ToMatrix() const;
 };
 
 enum EScalarKind
@@ -49,6 +50,6 @@ struct FArithmeticType : FType
 	bool IsMatrix() const { return NumRows > 1 && NumColumns > 1; }
 };
 
-} // namespace MaterialIR
+} // namespace UE::MIR
 
-#endif
+#endif // #if WITH_EDITOR
diff --git a/Engine/Source/Runtime/Engine/Public/Materials/MaterialInterface.h b/Engine/Source/Runtime/Engine/Public/Materials/MaterialInterface.h
--- a/Engine/Source/Runtime/Engine/Public/Materials/MaterialInterface.h
+++ b/Engine/Source/Runtime/Engine/Public/Materials/MaterialInterface.h
@@ -20,6 +20,7 @@
 #include "SceneTypes.h"
 #include "Engine/BlendableInterface.h"
 #include "Materials/MaterialLayersFunctions.h"
+#include "Materials/MaterialIRModule.h"
 #include "Interfaces/Interface_AssetUserData.h"
 #include "MaterialSceneTextureId.h"
 #include "Materials/MaterialRelevance.h"
@@ -433,9 +434,11 @@ public:
 
 #if WITH_EDITOR
 	ENGINE_API virtual const FMaterialCachedHLSLTree& GetCachedHLSLTree(TMicRecursionGuard RecursionGuard = TMicRecursionGuard()) const;
+	ENGINE_API const FMaterialIRModule& GetIRModule() const { return IRModule; }
 #endif
 
 	ENGINE_API bool IsUsingNewHLSLGenerator() const;
+	ENGINE_API bool IsUsingNewTranslatorPrototype() const;
 	ENGINE_API bool IsUsingControlFlow() const;
 
 	ENGINE_API const FSubstrateCompilationConfig& GetSubstrateCompilationConfig() const;
@@ -1192,6 +1195,7 @@ protected:
 
 #if WITH_EDITOR
 	TUniquePtr<FMaterialCachedHLSLTree> CachedHLSLTree;
+	FMaterialIRModule IRModule;
 #endif // WITH_EDITOR
 private:
 	/**
