[Material Editor] Build fix for 34087979.

diff --git a/Engine/Source/Runtime/Engine/Private/Materials/MaterialIREmitter.cpp b/Engine/Source/Runtime/Engine/Private/Materials/MaterialIREmitter.cpp
--- a/Engine/Source/Runtime/Engine/Private/Materials/MaterialIREmitter.cpp
+++ b/Engine/Source/Runtime/Engine/Private/Materials/MaterialIREmitter.cpp
@@ -410,7 +410,10 @@ void FEmitter::Put(const FExpressionOutput* Output, FValuePtr Value)
 
 void FEmitter::Error(FString Message)
 {
-	Module->Errors.Emplace(Expression, MoveTemp(Message));
+	FMaterialIRModule::FError Error;
+	Error.Expression = Expression;
+	Error.Message = MoveTemp(Message);
+	Module->Errors.Push(Error);
 	bHasExprBuildError = true;
 }
 
