[Backout] - CL34087979
Original CL Desc
-----------------------------------------------------------------
[Material Editor] New Material Translator advances.

* Added constants compilation
* Added binary operators
* Added conversion instructions (casting)
* General improvements and refactorings

diff --git a/Engine/Source/Editor/UnrealEd/Private/MaterialGraph.cpp b/Engine/Source/Editor/UnrealEd/Private/MaterialGraph.cpp
--- a/Engine/Source/Editor/UnrealEd/Private/MaterialGraph.cpp
+++ b/Engine/Source/Editor/UnrealEd/Private/MaterialGraph.cpp
@@ -153,7 +153,7 @@ void UMaterialGraph::UpdatePinTypes()
 {
 	using namespace UE::HLSLTree;
 
-	if (!Material->IsUsingNewHLSLGenerator() || Material->IsUsingNewTranslatorPrototype())
+	if (!Material->IsUsingNewHLSLGenerator())
 	{
 		return;
 	}
diff --git a/Engine/Source/Runtime/Engine/Private/Materials/Material.cpp b/Engine/Source/Runtime/Engine/Private/Materials/Material.cpp
--- a/Engine/Source/Runtime/Engine/Private/Materials/Material.cpp
+++ b/Engine/Source/Runtime/Engine/Private/Materials/Material.cpp
@@ -84,7 +84,6 @@
 #include "Materials/MaterialExpressionConstant.h"
 #include "Materials/MaterialExpressionConstant3Vector.h"
 #include "Materials/MaterialExpressionBreakMaterialAttributes.h"
-#include "Materials/MaterialIRModuleBuilder.h"
 #include "MaterialCachedData.h"
 #include "Misc/OutputDeviceArchiveWrapper.h"
 #include "HAL/FileManager.h"
@@ -2085,19 +2084,9 @@ void UMaterial::UpdateCachedExpressionData()
 		// Relinks function call inputs. Otherwise, we can get invalid inputs and they will cause errors when generating the syntax tree
 		UpdateTransientExpressionData();
 
-		if (bIsUsingNewMaterialTranslatorPrototype)
-		{
-			FMaterialIRModuleBuilder Builder;
-			Builder.Build(this, &IRModule);
-
-			LocalCachedExpressionData->AnalyzeMaterial(*this); // Temporary
-		}
-		else
-		{
-			LocalCachedTree = new FMaterialCachedHLSLTree();
-			LocalCachedTree->GenerateTree(this, nullptr, nullptr);
-			LocalCachedExpressionData->UpdateForCachedHLSLTree(*LocalCachedTree, nullptr, this);
-		}
+		LocalCachedTree = new FMaterialCachedHLSLTree();
+		LocalCachedTree->GenerateTree(this, nullptr, nullptr);
+		LocalCachedExpressionData->UpdateForCachedHLSLTree(*LocalCachedTree, nullptr, this);
 	}
 	else
 	{
diff --git a/Engine/Source/Runtime/Engine/Private/Materials/MaterialExpressionsIR.cpp b/Engine/Source/Runtime/Engine/Private/Materials/MaterialExpressionsIR.cpp
deleted file mode 100644
diff --git a/Engine/Source/Runtime/Engine/Private/Materials/MaterialIR.cpp b/Engine/Source/Runtime/Engine/Private/Materials/MaterialIR.cpp
--- a/Engine/Source/Runtime/Engine/Private/Materials/MaterialIR.cpp
+++ b/Engine/Source/Runtime/Engine/Private/Materials/MaterialIR.cpp
@@ -5,57 +5,27 @@
 
 #if WITH_EDITOR
 
-namespace UE::MIR
+namespace MaterialIR
 {
 
-FInstructionPtr FValue::AsInstruction() const
-{
-	return (Kind >= VK_InstructionBegin && Kind < VK_InstructionEnd) ? static_cast<FInstructionPtr>(this) : nullptr;
-}
-
-bool FValue::Equals(FValuePtr Other) const
-{
-	// If kinds are different the two values are surely different.
-	if (Kind != Other->Kind)
-	{
-		return false;
-	}
-
-	// Get the size of this value in bytes. It should match that of Other, since the value kinds are the same.
-	uint32 SizeInBytes = GetSizeInBytes();
-	check(SizeInBytes == Other->GetSizeInBytes());
-
-	// Values are PODs by design, therefore simply comparing bytes is sufficient.
-	return FMemory::Memcmp(this, Other, SizeInBytes) == 0;
-}
-
-uint32 FValue::GetSizeInBytes() const
-{
-	switch (Kind)
-	{
-		case VK_ScalarConstant: return sizeof(FScalarConstant);
-		case VK_SetMaterialOutput: return sizeof(FScalarConstant);
-		case VK_BinaryOperator: return sizeof(FBinaryOperator);
-		case VK_Dimensional: return sizeof(FDimensional) + sizeof(FValuePtr) * static_cast<const FDimensional*>(this)->GetComponents().Num();
-		default: UE_MIR_UNREACHABLE();
-	}
-}
+struct FValue;
+using FValuePtr = const FValue*;
 
-TArrayView<const FValuePtr> FDimensional::GetComponents() const
+TArrayView<const FValuePtr> FVectorValue::GetComponents() const
 {
-	TArrayView<FValuePtr> Components = const_cast<FDimensional*>(this)->GetMutableComponents();
+	TArrayView<FValuePtr> Components = const_cast<FVectorValue*>(this)->GetMutableComponents();
 	return { Components.GetData(), Components.Num() };
 }
 
-TArrayView<FValuePtr> FDimensional::GetMutableComponents()
+TArrayView<FValuePtr> FVectorValue::GetMutableComponents()
 {
 	FArithmeticTypePtr ArithmeticType = Type->ToArithmetic();
 	check(ArithmeticType);
 
-	FValuePtr* Ptr = (FValuePtr*)static_cast<TDimensional<1>*>(this)->Components;
+	FValuePtr* Ptr = (FValuePtr*)static_cast<TVectorValue<1>*>(this)->Components;
 	return { Ptr, ArithmeticType->NumRows };
 }
 
-} // namespace UE::MIR
+} // namespace MaterialIR
 
 #endif // #if WITH_EDITOR
diff --git a/Engine/Source/Runtime/Engine/Private/Materials/MaterialIRBuilder.cpp b/Engine/Source/Runtime/Engine/Private/Materials/MaterialIRBuilder.cpp
new file mode 100644
--- /dev/null
+++ b/Engine/Source/Runtime/Engine/Private/Materials/MaterialIRBuilder.cpp
@@ -0,0 +1,179 @@
+// Copyright Epic Games, Inc. All Rights Reserved.
+
+#include "Materials/MaterialIRBuilder.h"
+#include "Materials/MaterialIRModule.h"
+#include "Materials/MaterialIR.h"
+#include "Materials/MaterialIRTypes.h"
+#include "Shader/ShaderTypes.h"
+
+#if WITH_EDITOR
+
+namespace MaterialIR
+{
+
+struct FBuilder::FPrivate : FBuilder
+{
+	template <typename T>
+	T* New(FTypePtr InType)
+			{
+		T* Val = new T;
+		Val->Kind = T::TypeKind;
+		Val->Type = InType;
+		Module->Values.Add(Val);
+		return Val;
+	}
+};
+
+FBuilder::FBuilder()
+{
+}
+
+FBuilder::FBuilder(FMaterialIRModule* InModule)
+{
+	SetTargetModule(InModule);
+}
+
+void FBuilder::SetTargetModule(FMaterialIRModule* InModule)
+	{
+	Module = InModule;
+}
+
+FValuePtr FBuilder::NewConstantFromShaderValue(const UE::Shader::FValue& InValue)
+{
+	using namespace UE::Shader;
+
+	switch (InValue.Type.ValueType)
+	{
+	case UE::Shader::EValueType::Float1: return NewConstantFloat1(InValue.AsFloatScalar());
+	case UE::Shader::EValueType::Float2: return NewConstantFloat2(FVector2f{ InValue.Component[0].Float, InValue.Component[1].Float });
+	case UE::Shader::EValueType::Float3: return NewConstantFloat3(FVector3f{ InValue.Component[0].Float, InValue.Component[1].Float, InValue.Component[2].Float });
+	case UE::Shader::EValueType::Float4: return NewConstantFloat4(FVector4f{ InValue.Component[0].Float, InValue.Component[1].Float, InValue.Component[2].Float, InValue.Component[3].Float });
+
+	case UE::Shader::EValueType::Int1: return NewConstantInt1(InValue.AsFloatScalar());
+	case UE::Shader::EValueType::Int2: return NewConstantInt2(FIntVector2{ InValue.Component[0].Int, InValue.Component[1].Int });
+	case UE::Shader::EValueType::Int3: return NewConstantInt3(FIntVector3{ InValue.Component[0].Int, InValue.Component[1].Int, InValue.Component[2].Int });
+	case UE::Shader::EValueType::Int4: return NewConstantInt4(FIntVector4{ InValue.Component[0].Int, InValue.Component[1].Int, InValue.Component[2].Int, InValue.Component[3].Int });
+	}
+
+	UE_MIR_UNREACHABLE();
+}
+
+FValuePtr FBuilder::NewConstantFloat1(float InX)
+	{
+	FScalarValue* S = AsPrivate()->New<FScalarValue>(FArithmeticType::GetScalar(SK_Float));
+	S->Float = InX;
+	return S;
+}
+
+FValuePtr FBuilder::NewConstantFloat2(const FVector2f& InValue)
+{
+	FValuePtr X = NewConstantFloat1(InValue.X);
+	FValuePtr Y = NewConstantFloat1(InValue.Y);
+	return NewVector2(X, Y);
+}
+
+FValuePtr FBuilder::NewConstantFloat3(const FVector3f& InValue)
+{
+	FValuePtr X = NewConstantFloat1(InValue.X);
+	FValuePtr Y = NewConstantFloat1(InValue.Y);
+	FValuePtr Z = NewConstantFloat1(InValue.Z);
+	return NewVector3(X, Y, Z);
+}
+
+FValuePtr FBuilder::NewConstantFloat4(const FVector4f& InValue)
+{
+	FValuePtr X = NewConstantFloat1(InValue.X);
+	FValuePtr Y = NewConstantFloat1(InValue.Y);
+	FValuePtr Z = NewConstantFloat1(InValue.Z);
+	FValuePtr W = NewConstantFloat1(InValue.W);
+	return NewVector4(X, Y, Z, W);
+}
+
+FValuePtr FBuilder::NewConstantInt1(int InX)
+{
+	FScalarValue* S = AsPrivate()->New<FScalarValue>(FArithmeticType::GetScalar(SK_Int));
+	S->Integer = InX;
+	return S;
+}
+
+FValuePtr FBuilder::NewConstantInt2(const FIntVector2& InValue)
+{
+	FValuePtr X = NewConstantInt1(InValue.X);
+	FValuePtr Y = NewConstantInt1(InValue.Y);
+	return NewVector2(X, Y);
+}
+
+FValuePtr FBuilder::NewConstantInt3(const FIntVector3& InValue)
+{
+	FValuePtr X = NewConstantInt1(InValue.X);
+	FValuePtr Y = NewConstantInt1(InValue.Y);
+	FValuePtr Z = NewConstantInt1(InValue.Z);
+	return NewVector3(X, Y, Z);
+}
+
+FValuePtr FBuilder::NewConstantInt4(const FIntVector4& InValue)
+{
+	FValuePtr X = NewConstantInt1(InValue.X);
+	FValuePtr Y = NewConstantInt1(InValue.Y);
+	FValuePtr Z = NewConstantInt1(InValue.Z);
+	FValuePtr W = NewConstantInt1(InValue.W);
+	return NewVector4(X, Y, Z, W);
+}
+
+FValuePtr FBuilder::NewVector2(FValuePtr InX, FValuePtr InY)
+{
+	check(InX->Type->IsScalar());
+	check(InX->Type == InY->Type);
+
+	FVectorValue* V = AsPrivate()->New<TVectorValue<2>>(FArithmeticType::GetVector(InX->Type->ToArithmetic()->ScalarKind, 2));
+	TArrayView<FValuePtr> Components = V->GetMutableComponents();
+	Components[0] = InX;
+	Components[1] = InY;
+
+	return V;
+}
+
+FValuePtr FBuilder::NewVector3(FValuePtr InX, FValuePtr InY, FValuePtr InZ)
+{
+	check(InX->Type->IsScalar());
+	check(InX->Type == InY->Type);
+	check(InY->Type == InZ->Type);
+
+	FVectorValue* V = AsPrivate()->New<TVectorValue<3>>(FArithmeticType::GetVector(InX->Type->ToArithmetic()->ScalarKind, 3));
+	TArrayView<FValuePtr> Components = V->GetMutableComponents();
+	Components[0] = InX;
+	Components[1] = InY;
+	Components[2] = InZ;
+
+	return V;
+}
+
+FValuePtr FBuilder::NewVector4(FValuePtr InX, FValuePtr InY, FValuePtr InZ, FValuePtr InW)
+{
+	check(InX->Type->IsScalar());
+	check(InX->Type == InY->Type);
+	check(InY->Type == InZ->Type);
+	check(InZ->Type == InW->Type);
+
+	FVectorValue* V = AsPrivate()->New<TVectorValue<4>>(FArithmeticType::GetVector(InX->Type->ToArithmetic()->ScalarKind, 4));
+	TArrayView<FValuePtr> Components = V->GetMutableComponents();
+	Components[0] = InX;
+	Components[1] = InY;
+	Components[2] = InZ;
+	Components[3] = InW;
+
+	return V;
+}
+
+FSetMaterialOutputInstr* FBuilder::EmitSetMaterialOutput(EMaterialProperty InProperty, FValue* InArgValue)
+{
+	FSetMaterialOutputInstr* Instr = AsPrivate()->New<FSetMaterialOutputInstr>(nullptr);
+	Instr->Property = InProperty;
+	Instr->ArgValue = InArgValue;
+	Module->Outputs.Push(Instr);
+	return Instr;
+}
+
+} // namespace MaterialIR
+
+#endif // #if WITH_EDITOR
diff --git a/Engine/Source/Runtime/Engine/Private/Materials/MaterialIREmitter.cpp b/Engine/Source/Runtime/Engine/Private/Materials/MaterialIREmitter.cpp
deleted file mode 100644
diff --git a/Engine/Source/Runtime/Engine/Private/Materials/MaterialIRModule.cpp b/Engine/Source/Runtime/Engine/Private/Materials/MaterialIRModule.cpp
--- a/Engine/Source/Runtime/Engine/Private/Materials/MaterialIRModule.cpp
+++ b/Engine/Source/Runtime/Engine/Private/Materials/MaterialIRModule.cpp
@@ -5,7 +5,7 @@
 
 #if WITH_EDITOR
 
-namespace IR = UE::MIR;
+namespace MIR = MaterialIR;
 
 FMaterialIRModule::~FMaterialIRModule()
 {
@@ -14,7 +14,7 @@ FMaterialIRModule::~FMaterialIRModule()
 
 void FMaterialIRModule::Empty()
 {
-	for (IR::FValuePtr Value : Values)
+	for (MIR::FValue* Value : Values)
 	{
 		delete Value;
 	}
@@ -23,4 +23,12 @@ void FMaterialIRModule::Empty()
 	Outputs.Empty();
 }
 
-#endif // #if WITH_EDITOR
+void FMaterialIRModule::Reset(EShaderPlatform InShaderPlatform, ERHIFeatureLevel::Type InFeatureLevel, const ITargetPlatform* InTargetPlatform)
+{
+	Empty();
+
+	ShaderPlatform = InShaderPlatform;
+	TargetPlatform = InTargetPlatform;
+}
+
+#endif
diff --git a/Engine/Source/Runtime/Engine/Private/Materials/MaterialIRModuleBuilder.cpp b/Engine/Source/Runtime/Engine/Private/Materials/MaterialIRModuleBuilder.cpp
--- a/Engine/Source/Runtime/Engine/Private/Materials/MaterialIRModuleBuilder.cpp
+++ b/Engine/Source/Runtime/Engine/Private/Materials/MaterialIRModuleBuilder.cpp
@@ -7,134 +7,76 @@
 #include "Materials/MaterialIRModule.h"
 #include "Materials/MaterialIR.h"
 #include "Materials/MaterialIRTypes.h"
-#include "Materials/MaterialIREmitter.h"
+#include "Materials/MaterialIRBuilder.h"
 #include "MaterialIRUtility.h"
 
 #include "Materials/MaterialAttributeDefinitionMap.h"
 #include "Materials/Material.h"
 #include "MaterialExpressionIO.h"
 #include "MaterialShared.h"
-#include "Materials/MaterialExpression.h"
 
-namespace IR = UE::MIR;
+namespace IR = MaterialIR;
 
-struct FMaterialIRModuleBuilder::FPrivate
+struct FMaterialIRModuleBuilder::FHelper
 {
-	static void Build_GenerateOutputInstructions(FMaterialIRModuleBuilder& Builder, UE::MIR::FEmitter& Emitter)
+};
+
+void FMaterialIRModuleBuilder::SetSource(FMaterial* InMaterial, const FStaticParameterSet* InStaticParameters)
+{
+	Material = InMaterial;
+	StaticParameters = InStaticParameters;
+}
+
+void FMaterialIRModuleBuilder::SetPlatform(EShaderPlatform InShaderPlatform, ERHIFeatureLevel::Type InFeatureLevel, const ITargetPlatform* InTargetPlatform)
+{
+	ShaderPlatform	= InShaderPlatform;
+	FeatureLevel		= InFeatureLevel;
+	TargetPlatform	= InTargetPlatform;
+}
+
+void FMaterialIRModuleBuilder::SetTarget(FMaterialIRModule* InModule)
+{
+	Module = InModule;
+}
+
+bool FMaterialIRModuleBuilder::Build()
+{
+	Module->Reset(ShaderPlatform, FeatureLevel, TargetPlatform);
+
+	IR::FBuilder Builder{ Module };
+
+	UMaterial* BaseMaterial = Material->GetMaterialInterface()->GetMaterial();
+
+	for (int32 PropertyIndex = 0; PropertyIndex < MP_MAX; ++PropertyIndex)
 	{
-		// Prepare the array of FSetMaterialOutputInstr outputs from the material attributes inputs.
+		EMaterialProperty Property = (EMaterialProperty)PropertyIndex;
 		FMaterialInputDescription Input;
-		for (int32 Index = 0; UE::Utility::NextMaterialAttributeInput(Builder.BaseMaterial, Index, Input); ++Index)
+
+		if (!Utility::IsMaterialPropertyShared(Property) 
+			|| Property == MP_SubsurfaceColor
+			|| Property == MP_FrontMaterial
+			|| !BaseMaterial->GetExpressionInputDescription(Property, Input))
 		{
-			EMaterialProperty Property = (EMaterialProperty)Index;
-
-			IR::FSetMaterialOutput* Output = Emitter.EmitSetMaterialOutput(Property, nullptr);
-
-			if (Input.bUseConstant)
-			{
-				Output->ArgValue = Emitter.EmitConstantFromShaderValue(Input.ConstantValue);
-			}
-			else if (!Input.Input->IsConnected())
-			{
-				Output->ArgValue = UE::Utility::CreateMaterialAttributeDefaultValue(Emitter, Builder.BaseMaterial, Property);
-			}
-			else
-			{
-				Builder.ExpressionAnalysisStack.Add(Input.Input->Expression);
-			}
+			continue;
 		}
-	}
 
-	static void Build_AnalyzeExpressionGraph(FMaterialIRModuleBuilder& Builder, UE::MIR::FEmitter & Emitter)
-	{
-		TSet<UMaterialExpression*> BuiltExpressions;
-
-		while (!Builder.ExpressionAnalysisStack.IsEmpty())
+		IR::FSetMaterialOutputInstr* Output = Builder.EmitSetMaterialOutput(Property, nullptr);
+	
+		if (Input.bUseConstant)
 		{
-			Emitter.Expression = Builder.ExpressionAnalysisStack.Last();
-
-			// If expression is clean, nothing to be done.
-			if (BuiltExpressions.Contains(Emitter.Expression))
-			{
-				Builder.ExpressionAnalysisStack.Pop(EAllowShrinking::No);
-				continue;
-			}
-
-			// Push to the expression stack all dependencies that still need to be analyzed.
-			for (FExpressionInputIterator It{ Emitter.Expression}; It; ++It)
-			{
-				// Ignore disconnected inputs and connected expressions already built.
-				if (!It->IsConnected() || BuiltExpressions.Contains(It->Expression))
-				{
-					continue;
-				}
-
-				Builder.ExpressionAnalysisStack.Push(It->Expression);
-			}
-
-			// If on top of the stack there's a different expression, we have a dependency to analyze first.
-			if (Builder.ExpressionAnalysisStack.Last() != Emitter.Expression) {
-				continue;
-			}
-
-			// 
-			Builder.ExpressionAnalysisStack.Pop();
-			BuiltExpressions.Add(Emitter.Expression);
-
-			//
-			for (FExpressionInputIterator It{ Emitter.Expression}; It; ++It)
-			{
-				if (FExpressionOutput* ConnectedOutput = It->GetConnectedOutput())
-				{
-					// Fetch the value flowing through connected output.
-					IR::FValuePtr* ValuePtr = Builder.OutputValues.Find(ConnectedOutput);
-					check(ValuePtr && TEXT("Output value not found."));
-
-					// Set the value flowing into this input.
-					Builder.InputValues.Add(It.Input, *ValuePtr);
-				}
-			}
-
-			// And clear the expression errors.
-			Emitter.bHasExprBuildError = false;
-
-			// Invoke the expression build function. This will perform semantic analysis, error reporting and
-			// emit IR values for its outputs (which will flow into connected expressions inputs).
-			Emitter.Expression->Build(Emitter);
+			Output->ArgValue = Builder.NewConstantFromShaderValue(Input.ConstantValue);
 		}
-	}
-
-	static void Build_LinkMaterialOutputsToIncomingValues(FMaterialIRModuleBuilder& Builder, UE::MIR::FEmitter& Emitter)
-	{
-		FMaterialInputDescription Input;
-		for (IR::FSetMaterialOutput* Output : Builder.Module->Outputs)
+		else if (!Input.Input->IsConnected())
 		{
-			if (Output->ArgValue || !ensure(Builder.BaseMaterial->GetExpressionInputDescription(Output->Property, Input))) {
-				continue;
-			}
-
-			IR::FValuePtr* ValuePtr = Builder.OutputValues.Find(Input.Input->Expression->GetOutput(0));
-			check(ValuePtr);
-
-			Output->ArgValue = *ValuePtr;
+			Output->ArgValue = Utility::CreateMaterialAttributeDefaultValue(Builder, Material, Property);
+		}
+		else if (Input.Input->IsConnected())
+		{
+			UE_MIR_UNREACHABLE();
 		}
 	}
-};
-
-bool FMaterialIRModuleBuilder::Build(UMaterial* InMaterial, FMaterialIRModule* TargetModule)
-{
-	BaseMaterial = InMaterial;
-
-	Module = TargetModule;
-	Module->Empty();
-
-	IR::FEmitter Emitter{ this, BaseMaterial, Module };
-
-	FPrivate::Build_GenerateOutputInstructions(*this, Emitter);
-	FPrivate::Build_AnalyzeExpressionGraph(*this, Emitter);
-	FPrivate::Build_LinkMaterialOutputsToIncomingValues(*this, Emitter);
 
 	return true;
 }
 
-#endif // #if WITH_EDITOR
+#endif
diff --git a/Engine/Source/Runtime/Engine/Private/Materials/MaterialIRToHLSLTranslator.cpp b/Engine/Source/Runtime/Engine/Private/Materials/MaterialIRToHLSLTranslator.cpp
--- a/Engine/Source/Runtime/Engine/Private/Materials/MaterialIRToHLSLTranslator.cpp
+++ b/Engine/Source/Runtime/Engine/Private/Materials/MaterialIRToHLSLTranslator.cpp
@@ -17,7 +17,7 @@
 
 #if WITH_EDITOR
 
-namespace IR = UE::MIR;
+namespace IR = MaterialIR;
 
 struct FHLSLPrinter
 {
@@ -120,303 +120,269 @@ static const TCHAR* GetShadingModelParameterName(EMaterialShadingModel InModel)
 	}
 }
 
-struct FMaterialIRToHLSLTranslator::FPrivate
+UE_MIR_BEGIN_PRIVATE(FMaterialIRToHLSLTranslator);
+
+void GenerateHLSL(const FMaterial& InMaterial, const FMaterialIRModule& InModule, FParametersMap& InParams)
 {
-	FMaterialIRToHLSLTranslator* Translator;
+	FString PixelAttributes;
+	FString EvaluateOtherMaterialAttributesHLSL;
+	FHLSLPrinter Printer{ EvaluateOtherMaterialAttributesHLSL };
 
-	void GenerateHLSL(const FMaterial& InMaterial,
-					  const FMaterialIRModule& InModule,
-					  FParametersMap& InParams)
+ 	for (const IR::FSetMaterialOutputInstr* Output : InModule.GetOutputs())
 	{
-		FString PixelAttributes;
-		FString EvaluateOtherMaterialAttributesHLSL;
-		FHLSLPrinter Printer{ EvaluateOtherMaterialAttributesHLSL };
-
-		for (const IR::FSetMaterialOutput* Output : InModule.GetOutputs())
-		{
-			LowerValue(Printer, Output);
-		}
+		LowerValue(Printer, Output);
+	}
 
-		EvaluateOtherMaterialAttributesHLSL.Append(TEXT("\tPixelMaterialInputs.FrontMaterial = GetInitialisedSubstrateData();\n"));
-		EvaluateOtherMaterialAttributesHLSL.Append(TEXT("\tPixelMaterialInputs.Subsurface = 0;\n"));
+	EvaluateOtherMaterialAttributesHLSL.Append(TEXT("\tPixelMaterialInputs.FrontMaterial = GetInitialisedSubstrateData();\n"));
+	EvaluateOtherMaterialAttributesHLSL.Append(TEXT("\tPixelMaterialInputs.Subsurface = 0;\n"));
 
-		for (int32 PropertyIndex = 0; PropertyIndex < MP_MAX; ++PropertyIndex)
+	for (int32 PropertyIndex = 0; PropertyIndex < MP_MAX; ++PropertyIndex)
+	{
+		EMaterialProperty Property = (EMaterialProperty)PropertyIndex;
+		if (!Utility::IsMaterialPropertyShared(Property))
 		{
-			EMaterialProperty Property = (EMaterialProperty)PropertyIndex;
-			if (!UE::Utility::IsMaterialPropertyShared(Property))
-			{
-				continue;
-			}
+			continue;
+		}
 		
-			check(FMaterialAttributeDefinitionMap::GetShaderFrequency(Property) == SF_Pixel);
+		check(FMaterialAttributeDefinitionMap::GetShaderFrequency(Property) == SF_Pixel);
 			
-			// Special case MP_SubsurfaceColor as the actual property is a combination of the color and the profile but we don't want to expose the profile
-			FString PropertyName = (Property == MP_SubsurfaceColor) ? "Subsurface" : FMaterialAttributeDefinitionMap::GetAttributeName(Property);
-			EMaterialValueType Type = (Property == MP_SubsurfaceColor) ? MCT_Float4 : FMaterialAttributeDefinitionMap::GetValueType(Property);
-			check(PropertyName.Len() > 0);
+		// Special case MP_SubsurfaceColor as the actual property is a combination of the color and the profile but we don't want to expose the profile
+		FString PropertyName = (Property == MP_SubsurfaceColor) ? "Subsurface" : FMaterialAttributeDefinitionMap::GetAttributeName(Property);
+		EMaterialValueType Type = (Property == MP_SubsurfaceColor) ? MCT_Float4 : FMaterialAttributeDefinitionMap::GetValueType(Property);
+		check(PropertyName.Len() > 0);
 
-			PixelAttributes.Appendf(TEXT("\t%s %s;\n"), GetHLSLTypeString(Type), *PropertyName);
-		}
+		PixelAttributes.Appendf(TEXT("\t%s %s;\n"), GetHLSLTypeString(Type), *PropertyName);
+	}
 	
-		InParams.Add(TEXT("pixel_material_inputs"), MoveTemp(PixelAttributes));
+	InParams.Add(TEXT("pixel_material_inputs"), MoveTemp(PixelAttributes));
 
-		InParams.Add(TEXT("calc_pixel_material_inputs_normal"), TEXT("\tPixelMaterialInputs.Normal = MaterialFloat3(0.00000000, 0.00000000, 1.00000000);"));
-		InParams.Add(TEXT("calc_pixel_material_inputs_other_inputs"), EvaluateOtherMaterialAttributesHLSL);
+	InParams.Add(TEXT("calc_pixel_material_inputs_normal"), TEXT("\tPixelMaterialInputs.Normal = MaterialFloat3(0.00000000, 0.00000000, 1.00000000);"));
+	InParams.Add(TEXT("calc_pixel_material_inputs_other_inputs"), EvaluateOtherMaterialAttributesHLSL);
 
-		InParams.Add(TEXT("calc_pixel_material_inputs_analytic_derivatives_normal"), TEXT("\tPixelMaterialInputs.Normal = MaterialFloat3(0.00000000, 0.00000000, 1.00000000);"));
-		InParams.Add(TEXT("calc_pixel_material_inputs_analytic_derivatives_other_inputs"), MoveTemp(EvaluateOtherMaterialAttributesHLSL));
+	InParams.Add(TEXT("calc_pixel_material_inputs_analytic_derivatives_normal"), TEXT("\tPixelMaterialInputs.Normal = MaterialFloat3(0.00000000, 0.00000000, 1.00000000);"));
+	InParams.Add(TEXT("calc_pixel_material_inputs_analytic_derivatives_other_inputs"), MoveTemp(EvaluateOtherMaterialAttributesHLSL));
 	
-		InParams.Add(TEXT("material_declarations"), TEXT("struct FMaterialAttributes {};"));
+	InParams.Add(TEXT("material_declarations"), TEXT("struct FMaterialAttributes {};"));
 	
-		auto SetParamInt = [&] (const TCHAR* InParamName, int InValue)
-		{
-			InParams.Add(InParamName, FString::Printf(TEXT("%d"), InValue));
-		};
+	auto SetParamInt = [&] (const TCHAR* InParamName, int InValue)
+	{
+		InParams.Add(InParamName, FString::Printf(TEXT("%d"), InValue));
+	};
 
-		SetParamInt(TEXT("num_material_texcoords_vertex"), 0);
-		SetParamInt(TEXT("num_material_texcoords"), 0);
-		SetParamInt(TEXT("num_custom_vertex_interpolators"), 0);
-		SetParamInt(TEXT("num_tex_coord_interpolators"), 0);
-	}
+	SetParamInt(TEXT("num_material_texcoords_vertex"), 0);
+	SetParamInt(TEXT("num_material_texcoords"), 0);
+	SetParamInt(TEXT("num_custom_vertex_interpolators"), 0);
+	SetParamInt(TEXT("num_tex_coord_interpolators"), 0);
+}
 
-	void SetMaterialParameters(const FMaterial& InMaterial, FParametersMap& InParams)
+void SetMaterialParameters(const FMaterial& InMaterial, FParametersMap& InParams)
+{
+	auto SetParamReturnFloat = [&] (const TCHAR* InParamName, float InValue)
 	{
-		auto SetParamReturnFloat = [&] (const TCHAR* InParamName, float InValue)
-		{
-			InParams.Add(InParamName, FString::Printf(TEXT("\treturn %.5f"), InValue));
-		};
-
-		SetParamReturnFloat(TEXT("get_material_emissive_for_cs"), 0.f);
-		SetParamReturnFloat(TEXT("get_material_translucency_directional_lighting_intensity"), InMaterial.GetTranslucencyDirectionalLightingIntensity());
-		SetParamReturnFloat(TEXT("get_material_translucent_shadow_density_scale"), InMaterial.GetTranslucentShadowDensityScale());
-		SetParamReturnFloat(TEXT("get_material_translucent_self_shadow_density_scale"), InMaterial.GetTranslucentSelfShadowDensityScale());
-		SetParamReturnFloat(TEXT("get_material_translucent_self_shadow_second_density_scale"), InMaterial.GetTranslucentSelfShadowSecondDensityScale());
-		SetParamReturnFloat(TEXT("get_material_translucent_self_shadow_second_opacity"), InMaterial.GetTranslucentSelfShadowSecondOpacity());
-		SetParamReturnFloat(TEXT("get_material_translucent_backscattering_exponent"), InMaterial.GetTranslucentBackscatteringExponent());
-
-		FLinearColor Extinction = InMaterial.GetTranslucentMultipleScatteringExtinction();
-		InParams.Add(TEXT("get_material_translucent_multiple_scattering_extinction"), FString::Printf(TEXT("\treturn MaterialFloat3(%.5f, %.5f, %.5f)"), Extinction.R, Extinction.G, Extinction.B));
-
-		SetParamReturnFloat(TEXT("get_material_opacity_mask_clip_value"), InMaterial.GetOpacityMaskClipValue());
-		InParams.Add(TEXT("get_material_world_position_offset_raw"), TEXT("\treturn 0; // todo"));
-		InParams.Add(TEXT("get_material_previous_world_position_offset_raw"), TEXT("\treturn 0; // todo"));
+		InParams.Add(InParamName, FString::Printf(TEXT("\treturn %.5f"), InValue));
+	};
+
+	SetParamReturnFloat(TEXT("get_material_emissive_for_cs"), 0.f);
+	SetParamReturnFloat(TEXT("get_material_translucency_directional_lighting_intensity"), InMaterial.GetTranslucencyDirectionalLightingIntensity());
+	SetParamReturnFloat(TEXT("get_material_translucent_shadow_density_scale"), InMaterial.GetTranslucentShadowDensityScale());
+	SetParamReturnFloat(TEXT("get_material_translucent_self_shadow_density_scale"), InMaterial.GetTranslucentSelfShadowDensityScale());
+	SetParamReturnFloat(TEXT("get_material_translucent_self_shadow_second_density_scale"), InMaterial.GetTranslucentSelfShadowSecondDensityScale());
+	SetParamReturnFloat(TEXT("get_material_translucent_self_shadow_second_opacity"), InMaterial.GetTranslucentSelfShadowSecondOpacity());
+	SetParamReturnFloat(TEXT("get_material_translucent_backscattering_exponent"), InMaterial.GetTranslucentBackscatteringExponent());
+
+	FLinearColor Extinction = InMaterial.GetTranslucentMultipleScatteringExtinction();
+	InParams.Add(TEXT("get_material_translucent_multiple_scattering_extinction"), FString::Printf(TEXT("\treturn MaterialFloat3(%.5f, %.5f, %.5f)"), Extinction.R, Extinction.G, Extinction.B));
+
+	SetParamReturnFloat(TEXT("get_material_opacity_mask_clip_value"), InMaterial.GetOpacityMaskClipValue());
+	InParams.Add(TEXT("get_material_world_position_offset_raw"), TEXT("\treturn 0; // todo"));
+	InParams.Add(TEXT("get_material_previous_world_position_offset_raw"), TEXT("\treturn 0; // todo"));
 	
-		// CustomData0/1 are named ClearCoat/ClearCoatRoughness
-		InParams.Add(TEXT("get_material_custom_data0"), TEXT("\treturn 1.0f; // todo"));
-		InParams.Add(TEXT("get_material_custom_data1"), TEXT("\treturn 0.1f; // todo"));
-
-		FString EvaluateMaterialDeclaration;
-		EvaluateMaterialDeclaration.Append(TEXT("void EvaluateVertexMaterialAttributes(in out FMaterialVertexParameters Parameters)\n{\n"));
-		EvaluateMaterialDeclaration.Append(TEXT("\n}\n"));
-		InParams.Add(TEXT("evaluate_material_attributes"), EvaluateMaterialDeclaration);
-	}
+	// CustomData0/1 are named ClearCoat/ClearCoatRoughness
+	InParams.Add(TEXT("get_material_custom_data0"), TEXT("\treturn 1.0f; // todo"));
+	InParams.Add(TEXT("get_material_custom_data1"), TEXT("\treturn 0.1f; // todo"));
+
+	FString EvaluateMaterialDeclaration;
+	EvaluateMaterialDeclaration.Append(TEXT("void EvaluateVertexMaterialAttributes(in out FMaterialVertexParameters Parameters)\n{\n"));
+	EvaluateMaterialDeclaration.Append(TEXT("\n}\n"));
+	InParams.Add(TEXT("evaluate_material_attributes"), EvaluateMaterialDeclaration);
+}
 
-	void GetShaderCompilerEnvironment(const FMaterial& InMaterial, const FMaterialIRModule& InModule, FShaderCompilerEnvironment& OutEnvironment)
+void GetShaderCompilerEnvironment(const FMaterial& InMaterial, const FMaterialIRModule& InModule, FShaderCompilerEnvironment& OutEnvironment)
+{
+	EShaderPlatform ShaderPlatform = InModule.GetShaderPlatform();
+	const FMaterialCompilationOutput& CompilationOutput = InModule.GetCompilationOutput();
+
+	OutEnvironment.TargetPlatform = InModule.GetTargetPlatform();
+	OutEnvironment.SetDefine(TEXT("ENABLE_NEW_HLSL_GENERATOR"), 1);
+	OutEnvironment.SetDefine(TEXT("MATERIAL_ATMOSPHERIC_FOG"), false);
+	OutEnvironment.SetDefine(TEXT("MATERIAL_SKY_ATMOSPHERE"), false);
+	OutEnvironment.SetDefine(TEXT("INTERPOLATE_VERTEX_COLOR"), false);
+	OutEnvironment.SetDefine(TEXT("NEEDS_PARTICLE_COLOR"), false);
+	OutEnvironment.SetDefine(TEXT("NEEDS_PARTICLE_LOCAL_TO_WORLD"), false);
+	OutEnvironment.SetDefine(TEXT("NEEDS_PARTICLE_WORLD_TO_LOCAL"), false);
+	OutEnvironment.SetDefine(TEXT("NEEDS_PER_INSTANCE_RANDOM_PS"), false);
+	OutEnvironment.SetDefine(TEXT("USES_TRANSFORM_VECTOR"), false);
+	OutEnvironment.SetDefine(TEXT("WANT_PIXEL_DEPTH_OFFSET"), CompilationOutput.bUsesPixelDepthOffset);
+	OutEnvironment.SetDefineAndCompileArgument(TEXT("USES_WORLD_POSITION_OFFSET"), (bool)CompilationOutput.bUsesWorldPositionOffset);
+	OutEnvironment.SetDefineAndCompileArgument(TEXT("USES_DISPLACEMENT"), false);
+	OutEnvironment.SetDefine(TEXT("USES_EMISSIVE_COLOR"), false);
+	OutEnvironment.SetDefine(TEXT("USES_DISTORTION"), InMaterial.IsDistorted());
+	OutEnvironment.SetDefine(TEXT("MATERIAL_ENABLE_TRANSLUCENCY_FOGGING"), InMaterial.ShouldApplyFogging());
+	OutEnvironment.SetDefine(TEXT("MATERIAL_ENABLE_TRANSLUCENCY_CLOUD_FOGGING"), InMaterial.ShouldApplyCloudFogging());
+	OutEnvironment.SetDefine(TEXT("MATERIAL_IS_SKY"), InMaterial.IsSky());
+	OutEnvironment.SetDefine(TEXT("MATERIAL_COMPUTE_FOG_PER_PIXEL"), InMaterial.ComputeFogPerPixel());
+	OutEnvironment.SetDefine(TEXT("MATERIAL_FULLY_ROUGH"), false);
+	OutEnvironment.SetDefine(TEXT("MATERIAL_USES_ANISOTROPY"), false);
+	OutEnvironment.SetDefine(TEXT("MATERIAL_NEURAL_POST_PROCESS"), (CompilationOutput.bUsedWithNeuralNetworks || InMaterial.IsUsedWithNeuralNetworks()) && InMaterial.IsPostProcessMaterial());
+	OutEnvironment.SetDefine(TEXT("NUM_VIRTUALTEXTURE_SAMPLES"), 0);
+	OutEnvironment.SetDefine(TEXT("MATERIAL_VIRTUALTEXTURE_FEEDBACK"), false);
+	OutEnvironment.SetDefine(TEXT("IS_MATERIAL_SHADER"), true);
+
+	FMaterialShadingModelField ShadingModels = InMaterial.GetShadingModels();
+	ensure(ShadingModels.IsValid());
+
+	int32 NumActiveShadingModels = 0;
+	if (ShadingModels.IsLit())
 	{
-		const FMaterialCompilationOutput& CompilationOutput = InModule.GetCompilationOutput();
-
-		OutEnvironment.TargetPlatform = Translator->TargetPlatform;
-		OutEnvironment.SetDefine(TEXT("ENABLE_NEW_HLSL_GENERATOR"), 1);
-		OutEnvironment.SetDefine(TEXT("MATERIAL_ATMOSPHERIC_FOG"), false);
-		OutEnvironment.SetDefine(TEXT("MATERIAL_SKY_ATMOSPHERE"), false);
-		OutEnvironment.SetDefine(TEXT("INTERPOLATE_VERTEX_COLOR"), false);
-		OutEnvironment.SetDefine(TEXT("NEEDS_PARTICLE_COLOR"), false);
-		OutEnvironment.SetDefine(TEXT("NEEDS_PARTICLE_LOCAL_TO_WORLD"), false);
-		OutEnvironment.SetDefine(TEXT("NEEDS_PARTICLE_WORLD_TO_LOCAL"), false);
-		OutEnvironment.SetDefine(TEXT("NEEDS_PER_INSTANCE_RANDOM_PS"), false);
-		OutEnvironment.SetDefine(TEXT("USES_TRANSFORM_VECTOR"), false);
-		OutEnvironment.SetDefine(TEXT("WANT_PIXEL_DEPTH_OFFSET"), CompilationOutput.bUsesPixelDepthOffset);
-		OutEnvironment.SetDefineAndCompileArgument(TEXT("USES_WORLD_POSITION_OFFSET"), (bool)CompilationOutput.bUsesWorldPositionOffset);
-		OutEnvironment.SetDefineAndCompileArgument(TEXT("USES_DISPLACEMENT"), false);
-		OutEnvironment.SetDefine(TEXT("USES_EMISSIVE_COLOR"), false);
-		OutEnvironment.SetDefine(TEXT("USES_DISTORTION"), InMaterial.IsDistorted());
-		OutEnvironment.SetDefine(TEXT("MATERIAL_ENABLE_TRANSLUCENCY_FOGGING"), InMaterial.ShouldApplyFogging());
-		OutEnvironment.SetDefine(TEXT("MATERIAL_ENABLE_TRANSLUCENCY_CLOUD_FOGGING"), InMaterial.ShouldApplyCloudFogging());
-		OutEnvironment.SetDefine(TEXT("MATERIAL_IS_SKY"), InMaterial.IsSky());
-		OutEnvironment.SetDefine(TEXT("MATERIAL_COMPUTE_FOG_PER_PIXEL"), InMaterial.ComputeFogPerPixel());
-		OutEnvironment.SetDefine(TEXT("MATERIAL_FULLY_ROUGH"), false);
-		OutEnvironment.SetDefine(TEXT("MATERIAL_USES_ANISOTROPY"), false);
-		OutEnvironment.SetDefine(TEXT("MATERIAL_NEURAL_POST_PROCESS"), (CompilationOutput.bUsedWithNeuralNetworks || InMaterial.IsUsedWithNeuralNetworks()) && InMaterial.IsPostProcessMaterial());
-		OutEnvironment.SetDefine(TEXT("NUM_VIRTUALTEXTURE_SAMPLES"), 0);
-		OutEnvironment.SetDefine(TEXT("MATERIAL_VIRTUALTEXTURE_FEEDBACK"), false);
-		OutEnvironment.SetDefine(TEXT("IS_MATERIAL_SHADER"), true);
-
-		FMaterialShadingModelField ShadingModels = InMaterial.GetShadingModels();
-		ensure(ShadingModels.IsValid());
-
-		int32 NumActiveShadingModels = 0;
-		if (ShadingModels.IsLit())
-		{
-			// This is to have platforms use the simple single layer water shading similar to mobile: no dynamic lights, only sun and sky, no distortion, no colored transmittance on background, no custom depth read.
-			const bool bSingleLayerWaterUsesSimpleShading = FDataDrivenShaderPlatformInfo::GetWaterUsesSimpleForwardShading(Translator->ShaderPlatform) && IsForwardShadingEnabled(Translator->ShaderPlatform);
+		// This is to have platforms use the simple single layer water shading similar to mobile: no dynamic lights, only sun and sky, no distortion, no colored transmittance on background, no custom depth read.
+		const bool bSingleLayerWaterUsesSimpleShading = FDataDrivenShaderPlatformInfo::GetWaterUsesSimpleForwardShading(ShaderPlatform) && IsForwardShadingEnabled(ShaderPlatform);
 
-			for (int i = 0; i < MSM_NUM; ++i)
+		for (int i = 0; i < MSM_NUM; ++i)
+		{
+			EMaterialShadingModel Model = (EMaterialShadingModel)i;
+			if (Model == MSM_Strata || !ShadingModels.HasShadingModel(Model))
 			{
-				EMaterialShadingModel Model = (EMaterialShadingModel)i;
-				if (Model == MSM_Strata || !ShadingModels.HasShadingModel(Model))
-				{
-					continue;
-				}
-
-				if (Model == MSM_SingleLayerWater && !FDataDrivenShaderPlatformInfo::GetRequiresDisableForwardLocalLights(Translator->ShaderPlatform))
-				{
-					continue;
-				}
-
-				if (Model == MSM_SingleLayerWater && bSingleLayerWaterUsesSimpleShading)
-				{
-					// Value must match SINGLE_LAYER_WATER_SHADING_QUALITY_MOBILE_WITH_DEPTH_TEXTURE in SingleLayerWaterCommon.ush!
-					OutEnvironment.SetDefine(TEXT("SINGLE_LAYER_WATER_SHADING_QUALITY"), true);
-				}
-
-				OutEnvironment.SetDefine(GetShadingModelParameterName(Model), true);
-				NumActiveShadingModels += 1;
+				continue;
 			}
-		}
-		else
-		{
-			// Unlit shading model can only exist by itself
-			OutEnvironment.SetDefine(TEXT("MATERIAL_SINGLE_SHADINGMODEL"), true);
-			OutEnvironment.SetDefine(GetShadingModelParameterName(MSM_Unlit), true);
-		}
-
-		if (NumActiveShadingModels == 1)
-		{
-			OutEnvironment.SetDefine(TEXT("MATERIAL_SINGLE_SHADINGMODEL"), true);
-		}
-		else if (!ensure(NumActiveShadingModels > 0))
-		{
-			UE_LOG(LogMaterial, Warning, TEXT("Unknown material shading model(s). Setting to MSM_DefaultLit"));
-			OutEnvironment.SetDefine(GetShadingModelParameterName(MSM_DefaultLit), true);
-		}
 
-		static IConsoleVariable* CVarLWCIsEnabled = IConsoleManager::Get().FindConsoleVariable(TEXT("r.MaterialEditor.LWCEnabled"));
-		OutEnvironment.SetDefine(TEXT("MATERIAL_LWC_ENABLED"), CVarLWCIsEnabled->GetInt());
-		OutEnvironment.SetDefine(TEXT("WSVECTOR_IS_TILEOFFSET"), true);
-		OutEnvironment.SetDefine(TEXT("WSVECTOR_IS_DOUBLEFLOAT"), false);
+			if (Model == MSM_SingleLayerWater && !FDataDrivenShaderPlatformInfo::GetRequiresDisableForwardLocalLights(ShaderPlatform))
+			{
+				continue;
+			}
 
-		if (InMaterial.GetMaterialDomain() == MD_Volume)
-		{
-			TArray<const UMaterialExpressionVolumetricAdvancedMaterialOutput*> VolumetricAdvancedExpressions;
-			InMaterial.GetMaterialInterface()->GetMaterial()->GetAllExpressionsOfType(VolumetricAdvancedExpressions);
-			if (VolumetricAdvancedExpressions.Num() > 0)
+			if (Model == MSM_SingleLayerWater && bSingleLayerWaterUsesSimpleShading)
 			{
-				if (VolumetricAdvancedExpressions.Num() > 1)
-				{
-					UE_LOG(LogMaterial, Fatal, TEXT("Only a single UMaterialExpressionVolumetricAdvancedMaterialOutput node is supported."));
-				}
-
-				const UMaterialExpressionVolumetricAdvancedMaterialOutput* VolumetricAdvancedNode = VolumetricAdvancedExpressions[0];
-				const TCHAR* Param = VolumetricAdvancedNode->GetEvaluatePhaseOncePerSample() ? TEXT("MATERIAL_VOLUMETRIC_ADVANCED_PHASE_PERSAMPLE") : TEXT("MATERIAL_VOLUMETRIC_ADVANCED_PHASE_PERPIXEL");
-				OutEnvironment.SetDefine(Param, true);
-
-				OutEnvironment.SetDefine(TEXT("MATERIAL_VOLUMETRIC_ADVANCED"), true);
-				OutEnvironment.SetDefine(TEXT("MATERIAL_VOLUMETRIC_ADVANCED_GRAYSCALE_MATERIAL"), VolumetricAdvancedNode->bGrayScaleMaterial);
-				OutEnvironment.SetDefine(TEXT("MATERIAL_VOLUMETRIC_ADVANCED_RAYMARCH_VOLUME_SHADOW"), VolumetricAdvancedNode->bRayMarchVolumeShadow);
-				OutEnvironment.SetDefine(TEXT("MATERIAL_VOLUMETRIC_ADVANCED_CLAMP_MULTISCATTERING_CONTRIBUTION"), VolumetricAdvancedNode->bClampMultiScatteringContribution);
-				OutEnvironment.SetDefine(TEXT("MATERIAL_VOLUMETRIC_ADVANCED_MULTISCATTERING_OCTAVE_COUNT"), VolumetricAdvancedNode->GetMultiScatteringApproximationOctaveCount());
-				OutEnvironment.SetDefine(TEXT("MATERIAL_VOLUMETRIC_ADVANCED_CONSERVATIVE_DENSITY"), VolumetricAdvancedNode->ConservativeDensity.IsConnected());
-				OutEnvironment.SetDefine(TEXT("MATERIAL_VOLUMETRIC_ADVANCED_OVERRIDE_AMBIENT_OCCLUSION"), InMaterial.HasAmbientOcclusionConnected());
-				OutEnvironment.SetDefine(TEXT("MATERIAL_VOLUMETRIC_ADVANCED_GROUND_CONTRIBUTION"), VolumetricAdvancedNode->bGroundContribution);
+				// Value must match SINGLE_LAYER_WATER_SHADING_QUALITY_MOBILE_WITH_DEPTH_TEXTURE in SingleLayerWaterCommon.ush!
+				OutEnvironment.SetDefine(TEXT("SINGLE_LAYER_WATER_SHADING_QUALITY"), true);
 			}
+
+			OutEnvironment.SetDefine(GetShadingModelParameterName(Model), true);
+			NumActiveShadingModels += 1;
 		}
+	}
+	else
+	{
+		// Unlit shading model can only exist by itself
+		OutEnvironment.SetDefine(TEXT("MATERIAL_SINGLE_SHADINGMODEL"), true);
+		OutEnvironment.SetDefine(GetShadingModelParameterName(MSM_Unlit), true);
+	}
 
-		OutEnvironment.SetDefine(TEXT("MATERIAL_IS_SUBSTRATE"), false);
-		OutEnvironment.SetDefine(TEXT("DUAL_SOURCE_COLOR_BLENDING_ENABLED"), false);
-		OutEnvironment.SetDefine(TEXT("TEXTURE_SAMPLE_DEBUG"), false);
+	if (NumActiveShadingModels == 1)
+	{
+		OutEnvironment.SetDefine(TEXT("MATERIAL_SINGLE_SHADINGMODEL"), true);
 	}
+	else if (!ensure(NumActiveShadingModels > 0))
+	{
+		UE_LOG(LogMaterial, Warning, TEXT("Unknown material shading model(s). Setting to MSM_DefaultLit"));
+		OutEnvironment.SetDefine(GetShadingModelParameterName(MSM_DefaultLit), true);
+	}
+
+	static IConsoleVariable* CVarLWCIsEnabled = IConsoleManager::Get().FindConsoleVariable(TEXT("r.MaterialEditor.LWCEnabled"));
+	OutEnvironment.SetDefine(TEXT("MATERIAL_LWC_ENABLED"), CVarLWCIsEnabled->GetInt());
+	OutEnvironment.SetDefine(TEXT("WSVECTOR_IS_TILEOFFSET"), true);
+	OutEnvironment.SetDefine(TEXT("WSVECTOR_IS_DOUBLEFLOAT"), false);
 
-	void LowerValue( FHLSLPrinter& Printer, const IR::FValue* InValue)
+	if (InMaterial.GetMaterialDomain() == MD_Volume)
 	{
-		if (const IR::FScalarConstant* Scalar = InValue->As<IR::FScalarConstant>())
+		TArray<const UMaterialExpressionVolumetricAdvancedMaterialOutput*> VolumetricAdvancedExpressions;
+		InMaterial.GetMaterialInterface()->GetMaterial()->GetAllExpressionsOfType(VolumetricAdvancedExpressions);
+		if (VolumetricAdvancedExpressions.Num() > 0)
 		{
-			IR::FArithmeticTypePtr ArithType = Scalar->Type->ToArithmetic();
-			check(ArithType && ArithType->IsScalar());
-
-			switch (ArithType->ScalarKind)
+			if (VolumetricAdvancedExpressions.Num() > 1)
 			{
-				case IR::SK_Bool: Printer << Scalar->Boolean; break;
-				case IR::SK_Int: Printer << Scalar->Integer; break;
-				case IR::SK_Float: Printer << Scalar->Float; break;
+				UE_LOG(LogMaterial, Fatal, TEXT("Only a single UMaterialExpressionVolumetricAdvancedMaterialOutput node is supported."));
 			}
+
+			const UMaterialExpressionVolumetricAdvancedMaterialOutput* VolumetricAdvancedNode = VolumetricAdvancedExpressions[0];
+			const TCHAR* Param = VolumetricAdvancedNode->GetEvaluatePhaseOncePerSample() ? TEXT("MATERIAL_VOLUMETRIC_ADVANCED_PHASE_PERSAMPLE") : TEXT("MATERIAL_VOLUMETRIC_ADVANCED_PHASE_PERPIXEL");
+			OutEnvironment.SetDefine(Param, true);
+
+			OutEnvironment.SetDefine(TEXT("MATERIAL_VOLUMETRIC_ADVANCED"), true);
+			OutEnvironment.SetDefine(TEXT("MATERIAL_VOLUMETRIC_ADVANCED_GRAYSCALE_MATERIAL"), VolumetricAdvancedNode->bGrayScaleMaterial);
+			OutEnvironment.SetDefine(TEXT("MATERIAL_VOLUMETRIC_ADVANCED_RAYMARCH_VOLUME_SHADOW"), VolumetricAdvancedNode->bRayMarchVolumeShadow);
+			OutEnvironment.SetDefine(TEXT("MATERIAL_VOLUMETRIC_ADVANCED_CLAMP_MULTISCATTERING_CONTRIBUTION"), VolumetricAdvancedNode->bClampMultiScatteringContribution);
+			OutEnvironment.SetDefine(TEXT("MATERIAL_VOLUMETRIC_ADVANCED_MULTISCATTERING_OCTAVE_COUNT"), VolumetricAdvancedNode->GetMultiScatteringApproximationOctaveCount());
+			OutEnvironment.SetDefine(TEXT("MATERIAL_VOLUMETRIC_ADVANCED_CONSERVATIVE_DENSITY"), VolumetricAdvancedNode->ConservativeDensity.IsConnected());
+			OutEnvironment.SetDefine(TEXT("MATERIAL_VOLUMETRIC_ADVANCED_OVERRIDE_AMBIENT_OCCLUSION"), InMaterial.HasAmbientOcclusionConnected());
+			OutEnvironment.SetDefine(TEXT("MATERIAL_VOLUMETRIC_ADVANCED_GROUND_CONTRIBUTION"), VolumetricAdvancedNode->bGroundContribution);
 		}
-		else if (const IR::FDimensional* Vector = InValue->As<IR::FDimensional>())
-		{
-			IR::FArithmeticTypePtr ArithType = Vector->Type->ToArithmetic();
-			check(ArithType && ArithType->IsVector());
+	}
 
-			Printer << ScalarKindToString(ArithType->ScalarKind) << ArithType->NumRows << TEXT("(");
+	OutEnvironment.SetDefine(TEXT("MATERIAL_IS_SUBSTRATE"), false);
+	OutEnvironment.SetDefine(TEXT("DUAL_SOURCE_COLOR_BLENDING_ENABLED"), false);
+	OutEnvironment.SetDefine(TEXT("TEXTURE_SAMPLE_DEBUG"), false);
+}
 
-			Printer.BeginList();
-			for (IR::FValuePtr Component : Vector->GetComponents())
-			{
-				Printer.PrintListSeparator();
-				LowerValue(Printer, Component);
-			}
+void LowerValue(FHLSLPrinter& Printer, const IR::FValue* InValue)
+{
+	if (const IR::FScalarValue* Scalar = InValue->Cast<IR::FScalarValue>())
+	{
+		IR::FArithmeticTypePtr ArithType = Scalar->Type->ToArithmetic();
+		check(ArithType && ArithType->IsScalar());
 
-			Printer << TEXT(")");
-		}
-		else if (const IR::FSetMaterialOutput* Output = InValue->As<IR::FSetMaterialOutput>())
+		switch (ArithType->ScalarKind)
 		{
-			// Special case MP_SubsurfaceColor as the actual property is a combination of the color and the profile but we don't want to expose the profile
-			const FString& PropertyName = (Output->Property == MP_SubsurfaceColor) ? "Subsurface" : FMaterialAttributeDefinitionMap::GetAttributeName(Output->Property);
-
-			Printer.Printf(TEXT("\tPixelMaterialInputs.%s = "), *PropertyName);
-		
-			LowerValue(Printer, Output->ArgValue);
-
-			Printer << TEXT(";\n");
+			case IR::SK_Bool: Printer << Scalar->Boolean; break;
+			case IR::SK_Int: Printer << Scalar->Integer; break;
+			case IR::SK_Float: Printer << Scalar->Float; break;
 		}
-		else if (const IR::FBinaryOperator* BinOp = InValue->As<IR::FBinaryOperator>())
-		{
-			LowerValue(Printer, BinOp->Lhs);
+	}
+	else if (const IR::FVectorValue* Vector = InValue->Cast<IR::FVectorValue>())
+	{
+		IR::FArithmeticTypePtr ArithType = Vector->Type->ToArithmetic();
+		check(ArithType && ArithType->IsVector());
 
-			const TCHAR* OpString;
-			switch (BinOp->Operator)
-			{
-				case IR::BO_Add: OpString = TEXT(" + "); break;
-				case IR::BO_Subtract: OpString = TEXT(" - "); break;
-				case IR::BO_Multiply: OpString = TEXT(" * "); break;
-				case IR::BO_Divide: OpString = TEXT(" / "); break;
-				default: UE_MIR_UNREACHABLE();
-			}
-			Printer << OpString;
+		Printer << ScalarKindToString(ArithType->ScalarKind) << ArithType->NumRows << TEXT("(");
 
-			LowerValue(Printer, BinOp->Rhs);
-		}
-		else
+		Printer.BeginList();
+		for (IR::FValuePtr Component : Vector->GetComponents())
 		{
-			UE_MIR_UNREACHABLE();
+			Printer.PrintListSeparator();
+			LowerValue(Printer, Component);
 		}
+
+		Printer << TEXT(")");
 	}
-};
+	else if (const IR::FSetMaterialOutputInstr* Output = InValue->Cast<IR::FSetMaterialOutputInstr>())
+	{
+		// Special case MP_SubsurfaceColor as the actual property is a combination of the color and the profile but we don't want to expose the profile
+		const FString& PropertyName = (Output->Property == MP_SubsurfaceColor) ? "Subsurface" : FMaterialAttributeDefinitionMap::GetAttributeName(Output->Property);
 
-FMaterialIRToHLSLTranslator::FMaterialIRToHLSLTranslator()
-{
+		Printer.Printf(TEXT("\tPixelMaterialInputs.%s = "), *PropertyName);
+		
+		LowerValue(Printer, Output->ArgValue);
+
+		Printer << TEXT(";\n");
+	}
+	else
+	{
+		UE_MIR_UNREACHABLE();
+	}
 }
 
-void FMaterialIRToHLSLTranslator::SetTarget(EShaderPlatform InShaderPlatform, const ITargetPlatform* InTargetPlatform, ERHIFeatureLevel::Type InFeatureLevel)
+UE_MIR_END_PRIVATE(); //FMaterialIRToHLSLTranslator
+
+FMaterialIRToHLSLTranslator::FMaterialIRToHLSLTranslator()
 {
-	ShaderPlatform = InShaderPlatform;
-	TargetPlatform = InTargetPlatform;
-	FeatureLevel = InFeatureLevel;
 }
 
-bool FMaterialIRToHLSLTranslator::Translate(
-	const FMaterial& InMaterial,
-	const FStaticParameterSet& StaticParameters,
-	const FMaterialIRModule& InModule,
-	FParametersMap& OutParametersMap,
-	FShaderCompilerEnvironment& OutEnvironment)
+void FMaterialIRToHLSLTranslator::Translate(const FMaterial& InMaterial, const FMaterialIRModule& InModule, FParametersMap& OutParametersMap, FShaderCompilerEnvironment& OutEnvironment)
 {
 	OutParametersMap.Empty();
 
-	FPrivate Private{ this };
-	Private.GenerateHLSL(InMaterial, InModule, OutParametersMap);
-	Private.SetMaterialParameters(InMaterial, OutParametersMap);
-	Private.GetShaderCompilerEnvironment(InMaterial, InModule, OutEnvironment);
-
-	return true;
+	AsPrivate()->GenerateHLSL(InMaterial, InModule, OutParametersMap);
+	AsPrivate()->SetMaterialParameters(InMaterial, OutParametersMap);
+	AsPrivate()->GetShaderCompilerEnvironment(InMaterial, InModule, OutEnvironment);
 }
 
 #endif // #if WITH_EDITOR
diff --git a/Engine/Source/Runtime/Engine/Private/Materials/MaterialIRTypes.cpp b/Engine/Source/Runtime/Engine/Private/Materials/MaterialIRTypes.cpp
--- a/Engine/Source/Runtime/Engine/Private/Materials/MaterialIRTypes.cpp
+++ b/Engine/Source/Runtime/Engine/Private/Materials/MaterialIRTypes.cpp
@@ -4,7 +4,7 @@
 
 #if WITH_EDITOR
 
-namespace UE::MIR
+namespace MaterialIR
 {
 
 FTypePtr FType::FromShaderType(const UE::Shader::FType& InShaderType)
@@ -46,27 +46,18 @@ FTypePtr FType::GetVoid()
 	return &Type;
 }
 
-FArithmeticTypePtr FType::ToArithmetic() const
+const FArithmeticType* FType::ToArithmetic() const
 {
-	return Kind == TK_Arithmetic ? static_cast<FArithmeticTypePtr>(this) : nullptr; 
+	return Kind == TK_Arithmetic ? static_cast<const FArithmeticType*>(this) : nullptr; 
 }
 
-FArithmeticTypePtr FType::ToScalar() const
+bool FType::IsScalar() const
 {
-	FArithmeticTypePtr Type = ToArithmetic();
-	return Type->IsScalar() ? Type : nullptr;
-}
-
-FArithmeticTypePtr FType::ToVector() const
-{
-	FArithmeticTypePtr Type = ToArithmetic();
-	return Type->IsVector() ? Type : nullptr;
-}
-
-FArithmeticTypePtr FType::ToMatrix() const
-{
-	FArithmeticTypePtr Type = ToArithmetic();
-	return Type->IsMatrix() ? Type : nullptr;
+	if (FArithmeticTypePtr Arith = ToArithmetic())
+	{
+		return Arith->IsScalar();
+	}
+	return false;
 }
 
 const TCHAR* ScalarKindToString(EScalarKind Kind)
@@ -158,6 +149,6 @@ const FArithmeticType* FArithmeticType::GetMatrix(EScalarKind InScalarKind, int
 	return GetNumericalType(InScalarKind, NumRows, NumColumns);
 }
 
-} // namespace UE::MIR
+} // namespace MaterialIR
 
 #endif // #if WITH_EDITOR
diff --git a/Engine/Source/Runtime/Engine/Private/Materials/MaterialIRUtility.cpp b/Engine/Source/Runtime/Engine/Private/Materials/MaterialIRUtility.cpp
--- a/Engine/Source/Runtime/Engine/Private/Materials/MaterialIRUtility.cpp
+++ b/Engine/Source/Runtime/Engine/Private/Materials/MaterialIRUtility.cpp
@@ -2,15 +2,14 @@
 
 #include "Materials/MaterialIRUtility.h"
 #include "Materials/MaterialAttributeDefinitionMap.h"
-#include "Materials/MaterialIREmitter.h"
-#include "Materials/Material.h"
+#include "Materials/MaterialIRBuilder.h"
 #include "MaterialShared.h"
 
 #if WITH_EDITOR
 
-namespace IR = UE::MIR;
+namespace Utility {
 
-namespace UE::Utility {
+namespace MIR = MaterialIR;
 
 bool IsMaterialPropertyShared(EMaterialProperty InProperty)
 {
@@ -40,41 +39,24 @@ bool IsMaterialPropertyShared(EMaterialProperty InProperty)
 	}
 }
 
-bool NextMaterialAttributeInput(UMaterial* BaseMaterial, int32& PropertyIndex, FMaterialInputDescription& Input)
-{
-	for (; PropertyIndex < MP_MAX; ++PropertyIndex)
-	{
-		EMaterialProperty Property = (EMaterialProperty)PropertyIndex;
-		if (Utility::IsMaterialPropertyShared(Property)
-			&& Property != MP_SubsurfaceColor
-			&& Property != MP_FrontMaterial
-			&& BaseMaterial->GetExpressionInputDescription(Property, Input))
-		{
-			return true;
-		}
-	}
-
-	return false;
-}
-
-IR::FValuePtr CreateMaterialAttributeDefaultValue(IR::FEmitter& Emitter, const UMaterial* Material, EMaterialProperty Property)
+MIR::FValuePtr CreateMaterialAttributeDefaultValue(MIR::FBuilder& Builder, const FMaterial* Material, EMaterialProperty Property)
 {
 	EMaterialValueType Type = FMaterialAttributeDefinitionMap::GetValueType(Property);
 	FVector4f DefaultValue = FMaterialAttributeDefinitionMap::GetDefaultValue(Property);
 
 	switch (Type)
 	{
-		case MCT_ShadingModel: return Emitter.EmitConstantInt1(Material->GetShadingModels().GetFirstShadingModel());
+		case MCT_ShadingModel: return Builder.NewConstantInt1(Material->GetShadingModels().GetFirstShadingModel());
 
-		case MCT_Float1: return Emitter.EmitConstantFloat1(  DefaultValue.X );
-		case MCT_Float2: return Emitter.EmitConstantFloat2({ DefaultValue.X, DefaultValue.Y });
-		case MCT_Float3: return Emitter.EmitConstantFloat3({ DefaultValue.X, DefaultValue.Y, DefaultValue.Z });
-		case MCT_Float4: return Emitter.EmitConstantFloat4(DefaultValue);
+		case MCT_Float1: return Builder.NewConstantFloat1(  DefaultValue.X );
+		case MCT_Float2: return Builder.NewConstantFloat2({ DefaultValue.X, DefaultValue.Y });
+		case MCT_Float3: return Builder.NewConstantFloat3({ DefaultValue.X, DefaultValue.Y, DefaultValue.Z });
+		case MCT_Float4: return Builder.NewConstantFloat4(DefaultValue);
 
-		case MCT_UInt1: return Emitter.EmitConstantInt1(  (int32)DefaultValue.X );
-		case MCT_UInt2: return Emitter.EmitConstantInt2({ (int32)DefaultValue.X, (int32)DefaultValue.Y });
-		case MCT_UInt3: return Emitter.EmitConstantInt3({ (int32)DefaultValue.X, (int32)DefaultValue.Y, (int32)DefaultValue.Z });
-		case MCT_UInt4: return Emitter.EmitConstantInt4({ (int32)DefaultValue.X, (int32)DefaultValue.Y, (int32)DefaultValue.Z, (int32)DefaultValue.W });
+		case MCT_UInt1: return Builder.NewConstantInt1(  (int32)DefaultValue.X );
+		case MCT_UInt2: return Builder.NewConstantInt2({ (int32)DefaultValue.X, (int32)DefaultValue.Y });
+		case MCT_UInt3: return Builder.NewConstantInt3({ (int32)DefaultValue.X, (int32)DefaultValue.Y, (int32)DefaultValue.Z });
+		case MCT_UInt4: return Builder.NewConstantInt4({ (int32)DefaultValue.X, (int32)DefaultValue.Y, (int32)DefaultValue.Z, (int32)DefaultValue.W });
 
 		default: UE_MIR_UNREACHABLE();
 	}
diff --git a/Engine/Source/Runtime/Engine/Private/Materials/MaterialIRUtility.h b/Engine/Source/Runtime/Engine/Private/Materials/MaterialIRUtility.h
--- a/Engine/Source/Runtime/Engine/Private/Materials/MaterialIRUtility.h
+++ b/Engine/Source/Runtime/Engine/Private/Materials/MaterialIRUtility.h
@@ -5,16 +5,12 @@
 
 #if WITH_EDITOR
 
-namespace UE::Utility {
+namespace Utility {
 
-//
 bool IsMaterialPropertyShared(EMaterialProperty InProperty);
 
 //
-bool NextMaterialAttributeInput(UMaterial* BaseMaterial, int32& PropertyIndex, FMaterialInputDescription& Input);
-
-//
-UE::MIR::FValuePtr CreateMaterialAttributeDefaultValue(UE::MIR::FEmitter& Emitter, const UMaterial* Material, EMaterialProperty Property);
+MaterialIR::FValuePtr CreateMaterialAttributeDefaultValue(MaterialIR::FBuilder& Builder, const FMaterial* Material, EMaterialProperty Property);
 
 } // namespace Utility
 
diff --git a/Engine/Source/Runtime/Engine/Private/Materials/MaterialInterface.cpp b/Engine/Source/Runtime/Engine/Private/Materials/MaterialInterface.cpp
--- a/Engine/Source/Runtime/Engine/Private/Materials/MaterialInterface.cpp
+++ b/Engine/Source/Runtime/Engine/Private/Materials/MaterialInterface.cpp
@@ -307,11 +307,6 @@ bool UMaterialInterface::IsUsingNewHLSLGenerator() const
 	return BaseMaterial ? BaseMaterial->bEnableNewHLSLGenerator : false;
 }
 
-bool UMaterialInterface::IsUsingNewTranslatorPrototype() const
-{
-	return IsUsingNewHLSLGenerator() && bIsUsingNewMaterialTranslatorPrototype;
-}
-
 const FSubstrateCompilationConfig& UMaterialInterface::GetSubstrateCompilationConfig() const
 {
 	const UMaterial* BaseMaterial = GetMaterial_Concurrent();
diff --git a/Engine/Source/Runtime/Engine/Private/Materials/MaterialShared.cpp b/Engine/Source/Runtime/Engine/Private/Materials/MaterialShared.cpp
--- a/Engine/Source/Runtime/Engine/Private/Materials/MaterialShared.cpp
+++ b/Engine/Source/Runtime/Engine/Private/Materials/MaterialShared.cpp
@@ -95,10 +95,9 @@ static TAutoConsoleVariable<bool> CVarMaterialEdPreshaderDumpToHLSL(
 	TEXT("Controls whether to append preshader expressions and parameter reference counts to the HLSL source window (as comments at the end of the code)."),
 	ECVF_RenderThreadSafe);
 
-bool bIsUsingNewMaterialTranslatorPrototype;
-static FAutoConsoleVariableRef CVarUsingUseNewMaterialTranslatorPrototype(
+static TAutoConsoleVariable<bool> CVarMaterialEdUseNewTranslatorPrototype(
 	TEXT("r.MaterialEditor.UseNewTranslatorPrototype"),
-	bIsUsingNewMaterialTranslatorPrototype,
+	false,
 	TEXT("Controls whether to enable the new material translator prototype (WIP) ."),
 	ECVF_RenderThreadSafe);
 
@@ -462,11 +461,6 @@ FExpressionInput FExpressionInput::GetTracedInput() const
 	return *this;
 }
 
-FExpressionOutput* FExpressionInput::GetConnectedOutput()
-{
-	return IsConnected() ? &Expression->GetOutputs()[OutputIndex] : nullptr;
-}
-
 int32 FExpressionExecOutput::Compile(class FMaterialCompiler* Compiler) const
 {
 	if (Expression)
@@ -3374,40 +3368,37 @@ bool FMaterial::Translate_New(const FMaterialShaderMapId& InShaderMapId,
 	FMaterialCompilationOutput& OutCompilationOutput,
 	TRefCountPtr<FSharedShaderCompilerEnvironment>& OutMaterialEnvironment)
 {
-	if (bIsUsingNewMaterialTranslatorPrototype)
+	const bool bUseNewTranslatorPrototype = CVarMaterialEdUseNewTranslatorPrototype->GetBool();
+	if (bUseNewTranslatorPrototype)
 	{
-		const FMaterialIRModule& Module = this->GetMaterialInterface()->GetIRModule();
+		FMaterialIRModule Module;
 
-		if (!Module.GetErrors().IsEmpty())
+		/* Setup the builder and bild the material */
+		FMaterialIRModuleBuilder Builder;
+		Builder.SetSource(this, &InStaticParameters);
+		Builder.SetPlatform(InPlatform, InShaderMapId.FeatureLevel, InTargetPlatform);
+		Builder.SetTarget(&Module);
+		if (!Builder.Build())
 		{
-			// TODO: This could be moved elsewhere, when the Module is built.
-			for (const FMaterialIRModule::FError& Error : Module.GetErrors())
-			{
-				ErrorExpressions.Push(Error.Expression);
-				CompileErrors.Push(Error.Message);
-			}
 			return false;
 		}
-	
-		// Copy over the compilation output
+
+		/* Copy over the compilation output */
 		OutCompilationOutput = Module.GetCompilationOutput();
 		OutMaterialEnvironment = new FSharedShaderCompilerEnvironment();
 
-		// Translate the material IR module to HLSL template string parameters and material environment
+		/* Translate the material IR module to HLSL template string parameters and material environment */
 		FMaterialIRToHLSLTranslator::FParametersMap ShaderStringParameters;
 		FMaterialIRToHLSLTranslator Translator;
-		if (!Translator.Translate(*this, InStaticParameters, Module, ShaderStringParameters, *OutMaterialEnvironment))
-		{
-			return false;
-		}
+		Translator.Translate(*this, Module, ShaderStringParameters, *OutMaterialEnvironment);
 
-		// Interpolate HLSL parameters with the material shader template to produce the final shader source
+		/* Interpolate HLSL parameters with the material shader template to produce the final shader source */
 		int32 LineNumber;
 		FStringTemplateResolver Resolver = FMaterialSourceTemplate::Get().BeginResolve(InPlatform, &LineNumber);
 		ShaderStringParameters.Add({TEXT("line_number"), FString::Printf(TEXT("%u"), LineNumber)});
 		Resolver.SetParameterMap(&ShaderStringParameters);
 
-		// Interpolate the final material shader source string
+		// Interpolate the string
 		FString MaterialShaderCode = Resolver.Finalize();
 		OutMaterialEnvironment->IncludeVirtualPathToContentsMap.Add(TEXT("/Engine/Generated/Material.ush"), MoveTemp(MaterialShaderCode));
 
diff --git a/Engine/Source/Runtime/Engine/Public/MaterialExpressionIO.h b/Engine/Source/Runtime/Engine/Public/MaterialExpressionIO.h
--- a/Engine/Source/Runtime/Engine/Public/MaterialExpressionIO.h
+++ b/Engine/Source/Runtime/Engine/Public/MaterialExpressionIO.h
@@ -7,7 +7,6 @@
 #include "SceneTypes.h"
 
 class FMaterialHLSLGenerator;
-struct FExpressionOutput;
 
 namespace UE::HLSLTree
 {
@@ -111,10 +110,6 @@ struct FExpressionInput
 		MaskB = B;
 		MaskA = A;
 	}
-
-	/** Returns the connected expression output if this input is connected to one,	nullptr otherwise. */
-	FExpressionOutput* GetConnectedOutput();
-
 #endif // WITH_EDITOR
 };
 
diff --git a/Engine/Source/Runtime/Engine/Public/MaterialShared.h b/Engine/Source/Runtime/Engine/Public/MaterialShared.h
--- a/Engine/Source/Runtime/Engine/Public/MaterialShared.h
+++ b/Engine/Source/Runtime/Engine/Public/MaterialShared.h
@@ -100,11 +100,6 @@ class FExpressionStaticTerrainLayerWeight;
 }
 }
 
-namespace UE::MIR
-{
-	class FEmitter;
-}
-
 template <class ElementType> class TLinkedList;
 
 #define ME_CAPTION_HEIGHT		18
@@ -196,8 +191,6 @@ ENGINE_API bool IsAlphaCompositeBlendMode(const FMaterial& In);
 ENGINE_API bool IsAlphaCompositeBlendMode(const UMaterialInterface& In);
 ENGINE_API bool IsAlphaCompositeBlendMode(const FMaterialShaderParameters& In);
 
-extern bool bIsUsingNewMaterialTranslatorPrototype;
-
 #if WITH_EDITOR
 
 void GetCmdLineFilterShaderFormats(TArray<FName>& InOutShderFormats);
@@ -2778,7 +2771,6 @@ private:
 	friend class FShaderCompilingManager;
 	friend class FHLSLMaterialTranslator;
 	friend class FMaterialHLSLErrorHandler;
-	friend class UE::MIR::FEmitter;
 };
 
 
diff --git a/Engine/Source/Runtime/Engine/Public/Materials/MaterialExpression.h b/Engine/Source/Runtime/Engine/Public/Materials/MaterialExpression.h
--- a/Engine/Source/Runtime/Engine/Public/Materials/MaterialExpression.h
+++ b/Engine/Source/Runtime/Engine/Public/Materials/MaterialExpression.h
@@ -33,19 +33,12 @@ enum class EMaterialParameterType : uint8;
 
 namespace UE
 {
-
 namespace HLSLTree
 {
 class FScope;
 class FStatement;
 class FExpression;
 }
-
-namespace MIR
-{
-	class FEmitter;
-}
-
 }
 
 //@warning: FExpressionInput is mirrored in MaterialShared.h and manually "subclassed" in Material.h (FMaterialInput)
@@ -328,14 +321,6 @@ class UMaterialExpression : public UObject
 
 	ENGINE_API bool IsUsingNewHLSLGenerator() const;
 
-	/**
-	 * Builds this expression's material IR.
-	 * 
-	 * This function is responsible for doing the semantic analysis on this error, error checking and reporting and
-	 * implementing this expression's logic by emitting appropriate material IR values.
-	 */
-	ENGINE_API virtual void Build(UE::MIR::FEmitter& Emitter);
-
 #endif // WITH_EDITOR
 
 	/**
@@ -363,10 +348,6 @@ class UMaterialExpression : public UObject
 
 	/** Counts the number of inputs this expression has. Default implementation has O(n) complexity. */
 	ENGINE_API virtual int32 CountInputs() const;
-	/**
-	 * Returns the output at specified index is valid
-	 */
-	ENGINE_API FExpressionOutput* GetOutput(int32 OutputIndex) { return &Outputs[OutputIndex]; }
 
 	/**
 	 * Returns the input at index InputIndex if valid, nullptr otherwise.
diff --git a/Engine/Source/Runtime/Engine/Public/Materials/MaterialExpressionAdd.h b/Engine/Source/Runtime/Engine/Public/Materials/MaterialExpressionAdd.h
--- a/Engine/Source/Runtime/Engine/Public/Materials/MaterialExpressionAdd.h
+++ b/Engine/Source/Runtime/Engine/Public/Materials/MaterialExpressionAdd.h
@@ -31,7 +31,6 @@ class UMaterialExpressionAdd : public UMaterialExpression
 
 	//~ Begin UMaterialExpression Interface
 #if WITH_EDITOR
-	virtual void Build(UE::MIR::FEmitter& Emitter) override;
 	virtual int32 Compile(class FMaterialCompiler* Compiler, int32 OutputIndex) override;
 	virtual void GetCaption(TArray<FString>& OutCaptions) const override;
 	virtual FText GetKeywords() const override {return FText::FromString(TEXT("+"));}
diff --git a/Engine/Source/Runtime/Engine/Public/Materials/MaterialExpressionConstant.h b/Engine/Source/Runtime/Engine/Public/Materials/MaterialExpressionConstant.h
--- a/Engine/Source/Runtime/Engine/Public/Materials/MaterialExpressionConstant.h
+++ b/Engine/Source/Runtime/Engine/Public/Materials/MaterialExpressionConstant.h
@@ -19,7 +19,6 @@ class UMaterialExpressionConstant : public UMaterialExpression
 
 	//~ Begin UMaterialExpression Interface
 #if WITH_EDITOR
-	virtual void Build(UE::MIR::FEmitter& Emitter) override;
 	virtual int32 Compile(class FMaterialCompiler* Compiler, int32 OutputIndex) override;
 	virtual void GetCaption(TArray<FString>& OutCaptions) const override;
 	virtual FString GetDescription() const override;
diff --git a/Engine/Source/Runtime/Engine/Public/Materials/MaterialExpressionConstant2Vector.h b/Engine/Source/Runtime/Engine/Public/Materials/MaterialExpressionConstant2Vector.h
--- a/Engine/Source/Runtime/Engine/Public/Materials/MaterialExpressionConstant2Vector.h
+++ b/Engine/Source/Runtime/Engine/Public/Materials/MaterialExpressionConstant2Vector.h
@@ -23,7 +23,6 @@ class UMaterialExpressionConstant2Vector : public UMaterialExpression
 
 	//~ Begin UMaterialExpression Interface
 #if WITH_EDITOR
-	virtual void Build(UE::MIR::FEmitter& Emitter) override;
 	virtual int32 Compile(class FMaterialCompiler* Compiler, int32 OutputIndex) override;
 	virtual void GetCaption(TArray<FString>& OutCaptions) const override;
 	virtual FString GetDescription() const override;
diff --git a/Engine/Source/Runtime/Engine/Public/Materials/MaterialExpressionConstant3Vector.h b/Engine/Source/Runtime/Engine/Public/Materials/MaterialExpressionConstant3Vector.h
--- a/Engine/Source/Runtime/Engine/Public/Materials/MaterialExpressionConstant3Vector.h
+++ b/Engine/Source/Runtime/Engine/Public/Materials/MaterialExpressionConstant3Vector.h
@@ -21,9 +21,9 @@ public:
 
 	//~ UMaterialExpression interface
 #if WITH_EDITOR
-	virtual void Build(UE::MIR::FEmitter& Emitter) override;
 	virtual int32 Compile(class FMaterialCompiler* Compiler, int32 OutputIndex) override;
 	virtual void GetCaption(TArray<FString>& OutCaptions) const override;
+
 	virtual FString GetDescription() const override;
 	virtual uint32 GetOutputType(int32 OutputIndex) override {return MCT_Float3;}
 	virtual bool GenerateHLSLExpression(FMaterialHLSLGenerator& Generator, UE::HLSLTree::FScope& Scope, int32 OutputIndex, UE::HLSLTree::FExpression const*& OutExpression) const override;
diff --git a/Engine/Source/Runtime/Engine/Public/Materials/MaterialExpressionConstant4Vector.h b/Engine/Source/Runtime/Engine/Public/Materials/MaterialExpressionConstant4Vector.h
--- a/Engine/Source/Runtime/Engine/Public/Materials/MaterialExpressionConstant4Vector.h
+++ b/Engine/Source/Runtime/Engine/Public/Materials/MaterialExpressionConstant4Vector.h
@@ -19,7 +19,6 @@ class UMaterialExpressionConstant4Vector : public UMaterialExpression
 
 	//~ Begin UMaterialExpression Interface
 #if WITH_EDITOR
-	virtual void Build(UE::MIR::FEmitter& Emitter) override;
 	virtual int32 Compile(class FMaterialCompiler* Compiler, int32 OutputIndex) override;
 	virtual void GetCaption(TArray<FString>& OutCaptions) const override;
 	virtual FString GetDescription() const override;
diff --git a/Engine/Source/Runtime/Engine/Public/Materials/MaterialExpressionDivide.h b/Engine/Source/Runtime/Engine/Public/Materials/MaterialExpressionDivide.h
--- a/Engine/Source/Runtime/Engine/Public/Materials/MaterialExpressionDivide.h
+++ b/Engine/Source/Runtime/Engine/Public/Materials/MaterialExpressionDivide.h
@@ -31,7 +31,6 @@ class UMaterialExpressionDivide : public UMaterialExpression
 
 	//~ Begin UMaterialExpression Interface
 #if WITH_EDITOR
-	virtual void Build(UE::MIR::FEmitter& Emitter) override;
 	virtual int32 Compile(class FMaterialCompiler* Compiler, int32 OutputIndex) override;
 	virtual bool GenerateHLSLExpression(FMaterialHLSLGenerator& Generator, UE::HLSLTree::FScope& Scope, int32 OutputIndex, UE::HLSLTree::FExpression const*& OutExpression) const override;
 	virtual void GetCaption(TArray<FString>& OutCaptions) const override;
diff --git a/Engine/Source/Runtime/Engine/Public/Materials/MaterialExpressionMultiply.h b/Engine/Source/Runtime/Engine/Public/Materials/MaterialExpressionMultiply.h
--- a/Engine/Source/Runtime/Engine/Public/Materials/MaterialExpressionMultiply.h
+++ b/Engine/Source/Runtime/Engine/Public/Materials/MaterialExpressionMultiply.h
@@ -31,7 +31,6 @@ class UMaterialExpressionMultiply : public UMaterialExpression
 
 	//~ Begin UMaterialExpression Interface
 #if WITH_EDITOR
-	virtual void Build(UE::MIR::FEmitter& Emitter) override;
 	virtual int32 Compile(class FMaterialCompiler* Compiler, int32 OutputIndex) override;
 	virtual void GetCaption(TArray<FString>& OutCaptions) const override;
 	virtual FText GetKeywords() const override {return FText::FromString(TEXT("*"));}
diff --git a/Engine/Source/Runtime/Engine/Public/Materials/MaterialExpressionSubtract.h b/Engine/Source/Runtime/Engine/Public/Materials/MaterialExpressionSubtract.h
--- a/Engine/Source/Runtime/Engine/Public/Materials/MaterialExpressionSubtract.h
+++ b/Engine/Source/Runtime/Engine/Public/Materials/MaterialExpressionSubtract.h
@@ -31,7 +31,6 @@ class UMaterialExpressionSubtract : public UMaterialExpression
 
 	//~ Begin UMaterialExpression Interface
 #if WITH_EDITOR
-	virtual void Build(UE::MIR::FEmitter& Emitter) override;
 	virtual bool GenerateHLSLExpression(FMaterialHLSLGenerator& Generator, UE::HLSLTree::FScope& Scope, int32 OutputIndex, UE::HLSLTree::FExpression const*& OutExpression) const override;
 	virtual int32 Compile(class FMaterialCompiler* Compiler, int32 OutputIndex) override;
 	virtual void GetCaption(TArray<FString>& OutCaptions) const override;
diff --git a/Engine/Source/Runtime/Engine/Public/Materials/MaterialIR.h b/Engine/Source/Runtime/Engine/Public/Materials/MaterialIR.h
--- a/Engine/Source/Runtime/Engine/Public/Materials/MaterialIR.h
+++ b/Engine/Source/Runtime/Engine/Public/Materials/MaterialIR.h
@@ -7,21 +7,19 @@ enum EMaterialProperty : int;
 
 #if WITH_EDITOR
 
-namespace UE::MIR {
+namespace MaterialIR
+{
 
 enum EValueKind
 {
 	/* Values */
-	VK_ScalarConstant,
-	VK_Dimensional,
-
+	VK_Value = 1 << 0,
+	VK_Scalar = VK_Value | 1 << 2,
+	VK_Vector = VK_Value | 2 << 2,
+	
 	/* Instructions */
-	VK_InstructionBegin,
-
-	VK_SetMaterialOutput = VK_InstructionBegin,
-	VK_BinaryOperator,
-
-	VK_InstructionEnd,
+	VK_Instruction = 1 << 1,
+ 	VK_SetMaterialOutputInstr = VK_Instruction | 1 << 2,
 };
 
 /* Values */
@@ -31,23 +29,27 @@ struct FValue
 	EValueKind Kind{};
 	FTypePtr  Type{};
 
-	bool IsA(EValueKind InKind) const { return Kind == InKind; }
-	FInstructionPtr AsInstruction() const;
-	bool Equals(FValuePtr Other) const;
-	uint32 GetSizeInBytes() const;
+	bool IsA(EValueKind InKind) const
+	{
+		return (Kind & InKind) == InKind;
+	}
 
 	template <typename T>
-	const T* As() const { return this && IsA(T::TypeKind) ? static_cast<const T*>(this) : nullptr; }
+	const T* Cast() const
+	{
+		return this && IsA(T::TypeKind) ? static_cast<const T*>(this) : nullptr;
+	}
 
+	void Destroy();
 };
 
-template <EValueKind TTypeKind>
-struct TValue : FValue
-{
-	static constexpr EValueKind TypeKind = TTypeKind;
-};
+	template <EValueKind TTypeKind>
+	struct TValue : FValue
+	{
+		static constexpr EValueKind TypeKind = TTypeKind;
+	};
 
-struct FScalarConstant : TValue<VK_ScalarConstant>
+struct FScalarValue : TValue<VK_Scalar>
 {
 	union
 	{
@@ -57,53 +59,37 @@ struct FScalarConstant : TValue<VK_ScalarConstant>
 	};
 };
 
-struct FDimensional : TValue<VK_Dimensional>
+struct FVectorValue : TValue<VK_Vector>
 {
 	TArrayView<const FValuePtr> GetComponents() const;
 	TArrayView<FValuePtr> GetMutableComponents();
-	uint32 GetSizeInBytes() const;
 };
 
 template <int TDimension>
-struct TDimensional : FDimensional
+struct TVectorValue : FVectorValue
 {
 	FValuePtr Components[TDimension];
 };
 
 /* Instructions */
 
-struct FInstruction : FValue
+struct FInstruction : TValue<VK_Instruction>
 {
 	FInstruction* Next{};
 };
 
-template <EValueKind TTypeKind>
-struct TInstruction : FInstruction
-{
-	static constexpr EValueKind TypeKind = TTypeKind;
-};
+	template <EValueKind TTypeKind>
+	struct TInstruction : FInstruction
+	{
+		static constexpr EValueKind TypeKind = TTypeKind;
+	};
 
-struct FSetMaterialOutput : TInstruction<VK_SetMaterialOutput>
+struct FSetMaterialOutputInstr : TInstruction<VK_SetMaterialOutputInstr>
 {
 	EMaterialProperty Property;
 	FValuePtr ArgValue;
 };
 
-enum EBinaryOperator
-{
-	BO_Invalid,
-	BO_Add,
-	BO_Subtract,
-	BO_Multiply,
-	BO_Divide,
-};
-
-struct FBinaryOperator : TInstruction<VK_BinaryOperator>
-{
-	EBinaryOperator Operator = BO_Invalid;
-	FValuePtr Lhs{};
-	FValuePtr Rhs{};
-};
+} // namespace MaterialIR
 
-} // namespace UE::MIR
 #endif
diff --git a/Engine/Source/Runtime/Engine/Public/Materials/MaterialIRBuilder.h b/Engine/Source/Runtime/Engine/Public/Materials/MaterialIRBuilder.h
new file mode 100644
--- /dev/null
+++ b/Engine/Source/Runtime/Engine/Public/Materials/MaterialIRBuilder.h
@@ -0,0 +1,48 @@
+// Copyright Epic Games, Inc. All Rights Reserved.
+#pragma once
+
+#include "Materials/MaterialIRCommon.h"
+
+#if WITH_EDITOR
+
+namespace MaterialIR
+{
+
+//
+class FBuilder
+{
+public:
+
+	FBuilder();
+	FBuilder(FMaterialIRModule* InModule);
+
+	void SetTargetModule(FMaterialIRModule* InModule);
+
+	/* Constants */
+
+	FValuePtr NewConstantFromShaderValue(const UE::Shader::FValue& InValue);
+	FValuePtr NewConstantFloat1(float InX);
+	FValuePtr NewConstantFloat2(const FVector2f& InValue);
+	FValuePtr NewConstantFloat3(const FVector3f& InValue);
+	FValuePtr NewConstantFloat4(const FVector4f& InValue);
+	FValuePtr NewConstantInt1(int InX);
+	FValuePtr NewConstantInt2(const FIntVector2& InValue);
+	FValuePtr NewConstantInt3(const FIntVector3& InValue);
+	FValuePtr NewConstantInt4(const FIntVector4& InValue);
+	FValuePtr NewVector2(FValuePtr InX, FValuePtr InY);
+	FValuePtr NewVector3(FValuePtr InX, FValuePtr InY, FValuePtr InZ);
+	FValuePtr NewVector4(FValuePtr InX, FValuePtr InY, FValuePtr InZ, FValuePtr InW);
+
+	/* Instructions */
+
+	FSetMaterialOutputInstr* EmitSetMaterialOutput(EMaterialProperty InProperty, FValue* InArgValue);
+
+private:
+	FMaterialIRModule* Module{};
+
+	UE_MIR_PRIVATE();
+};
+
+} // namespace MaterialIR
+
+#endif
diff --git a/Engine/Source/Runtime/Engine/Public/Materials/MaterialIRCommon.h b/Engine/Source/Runtime/Engine/Public/Materials/MaterialIRCommon.h
--- a/Engine/Source/Runtime/Engine/Public/Materials/MaterialIRCommon.h
+++ b/Engine/Source/Runtime/Engine/Public/Materials/MaterialIRCommon.h
@@ -17,20 +17,16 @@ namespace UE::Shader
 
 namespace ERHIFeatureLevel { enum Type : int; }
 
-class FMaterial;
+class ITargetPlatform;
 class FMaterialIRModule;
+class FMaterial;
+struct FStaticParameterSet;
+struct FShaderCompilerEnvironment;
 class FMaterialIRModule;
 class FMaterialIRModuleBuilder;
-class ITargetPlatform;
-class UMaterial;
-class UMaterialExpression;
 struct FExpressionInput;
-struct FExpressionOutput;
-struct FShaderCompilerEnvironment;
-struct FStaticParameterSet;
-struct FMaterialInputDescription;
 
-namespace UE::MIR
+namespace MaterialIR
 {
 
 /* Types*/
@@ -42,17 +38,28 @@ using FArithmeticTypePtr = const FArithmeticType*;
 
 /* IR */
 struct FValue;
-struct FInstruction;
-struct FSetMaterialOutput;
 using FValuePtr = const FValue*;
-using FInstructionPtr = const FInstruction*;
+struct FScalarValue;
+struct FVector;
+struct FSetMaterialOutputInstr;
 
 /* Others */
-class FEmitter;
+class FBuilder;
 
 }
 
 #define UE_MIR_UNREACHABLE() { check(!"Unreachable"); UE_ASSUME(false); }
-#define UE_MIR_TODO() UE_MIR_UNREACHABLE()
+
+#define UE_MIR_PRIVATE() \
+		struct FPrivate; \
+		friend FPrivate; \
+		FPrivate* AsPrivate() { return reinterpret_cast<FPrivate*>(this); }\
+		const FPrivate* AsPrivate() const { return reinterpret_cast<const FPrivate*>(this); }
+
+#define UE_MIR_BEGIN_PRIVATE(TypeName)\
+		struct TypeName::FPrivate : TypeName {
+
+#define UE_MIR_END_PRIVATE()\
+		}
 
 #endif // #if WITH_EDITOR
diff --git a/Engine/Source/Runtime/Engine/Public/Materials/MaterialIREmitter.h b/Engine/Source/Runtime/Engine/Public/Materials/MaterialIREmitter.h
deleted file mode 100644
diff --git a/Engine/Source/Runtime/Engine/Public/Materials/MaterialIRModule.h b/Engine/Source/Runtime/Engine/Public/Materials/MaterialIRModule.h
--- a/Engine/Source/Runtime/Engine/Public/Materials/MaterialIRModule.h
+++ b/Engine/Source/Runtime/Engine/Public/Materials/MaterialIRModule.h
@@ -8,29 +8,23 @@
 
 class FMaterialIRModule
 {
-public:
-	struct FError
-	{
-		UMaterialExpression* Expression;
-		FString Message;
-	};
-
 public:
 	~FMaterialIRModule();
 	void Empty();
+	void Reset(EShaderPlatform InShaderPlatform, ERHIFeatureLevel::Type InFeatureLevel, const ITargetPlatform* InTargetPlatform);
+	EShaderPlatform GetShaderPlatform() const { return ShaderPlatform; }
+	const ITargetPlatform* GetTargetPlatform() const { return TargetPlatform; }
 	const FMaterialCompilationOutput& GetCompilationOutput() const { return CompilationOutput; }
-	TArrayView<const UE::MIR::FSetMaterialOutput* const> GetOutputs() const { return Outputs; }
-	TArrayView<const FError> GetErrors() const { return Errors; }
+	TArrayView<const MaterialIR::FSetMaterialOutputInstr* const> GetOutputs() const { return Outputs; }
 
 private:
+	EShaderPlatform ShaderPlatform;
+	const ITargetPlatform* TargetPlatform;
 	FMaterialCompilationOutput CompilationOutput;
-	TArray<UE::MIR::FValuePtr> Values;
-	TArray<UE::MIR::FSetMaterialOutput*> Outputs;
-	TArray<FError> Errors;
-
+	TArray<MaterialIR::FValue*> Values;
+	TArray<MaterialIR::FSetMaterialOutputInstr*> Outputs;
 
-	friend UE::MIR::FEmitter;
-	friend FMaterialIRModuleBuilder;
+	friend MaterialIR::FBuilder;
 };
 
-#endif // #if WITH_EDITOR
+#endif
diff --git a/Engine/Source/Runtime/Engine/Public/Materials/MaterialIRModuleBuilder.h b/Engine/Source/Runtime/Engine/Public/Materials/MaterialIRModuleBuilder.h
--- a/Engine/Source/Runtime/Engine/Public/Materials/MaterialIRModuleBuilder.h
+++ b/Engine/Source/Runtime/Engine/Public/Materials/MaterialIRModuleBuilder.h
@@ -9,17 +9,30 @@
 class FMaterialIRModuleBuilder
 {
 public:
-	bool Build(UMaterial* InMaterial, FMaterialIRModule* TargetModule);
+	//
+	void SetSource(FMaterial* InMaterial, const FStaticParameterSet* InStaticParameters);
+
+	//
+	void SetPlatform(EShaderPlatform InShaderPlatform, ERHIFeatureLevel::Type InFeatureLevel, const ITargetPlatform* InTargetPlatform);
+
+	//
+	void SetTarget(FMaterialIRModule* InModule);
+
+	//
+	bool Build();
 
 private:
-	UMaterial* BaseMaterial;
+	FMaterial* Material;
+	const FStaticParameterSet* StaticParameters;
+	EShaderPlatform ShaderPlatform{};
+	ERHIFeatureLevel::Type FeatureLevel{};
+	const ITargetPlatform* TargetPlatform{};
 	FMaterialIRModule* Module;
-	TArray<UMaterialExpression*> ExpressionAnalysisStack;
-	TMap<const FExpressionInput*, UE::MIR::FValuePtr> InputValues;
-	TMap<const FExpressionOutput*, UE::MIR::FValuePtr> OutputValues;
 
-	friend class UE::MIR::FEmitter;
-	struct FPrivate;
+	TArray<FExpressionInput*> ExpressionAnalysisStack;
+
+	struct FHelper;
+	friend FHelper;
 };
 
 #endif
diff --git a/Engine/Source/Runtime/Engine/Public/Materials/MaterialIRToHLSLTranslator.h b/Engine/Source/Runtime/Engine/Public/Materials/MaterialIRToHLSLTranslator.h
--- a/Engine/Source/Runtime/Engine/Public/Materials/MaterialIRToHLSLTranslator.h
+++ b/Engine/Source/Runtime/Engine/Public/Materials/MaterialIRToHLSLTranslator.h
@@ -13,15 +13,11 @@ public:
 
 public:
 	FMaterialIRToHLSLTranslator();
-	void SetTarget(EShaderPlatform InShaderPlatform, const ITargetPlatform* InTargetPlatform, ERHIFeatureLevel::Type InFeatureLevel);
-	bool Translate(const FMaterial& InMaterial, const FStaticParameterSet& StaticParameters, const FMaterialIRModule& InModule, FParametersMap& OutParameters, FShaderCompilerEnvironment& OutEnvironment);
+	void Translate(const FMaterial& InMaterial, const FMaterialIRModule& InModule, FParametersMap& OutParameters, FShaderCompilerEnvironment& OutEnvironment);
 
 private:
-	EShaderPlatform ShaderPlatform{};
-	const ITargetPlatform* TargetPlatform{};
-	ERHIFeatureLevel::Type FeatureLevel{};
-
-	struct FPrivate;
+ 	
+	UE_MIR_PRIVATE();
 };
 
 #endif // #if WITH_EDITOR
diff --git a/Engine/Source/Runtime/Engine/Public/Materials/MaterialIRTypes.h b/Engine/Source/Runtime/Engine/Public/Materials/MaterialIRTypes.h
--- a/Engine/Source/Runtime/Engine/Public/Materials/MaterialIRTypes.h
+++ b/Engine/Source/Runtime/Engine/Public/Materials/MaterialIRTypes.h
@@ -6,7 +6,7 @@
 
 #if WITH_EDITOR
 
-namespace UE::MIR
+namespace MaterialIR
 {
 
 enum ETypeKind
@@ -22,10 +22,9 @@ struct FType
 	static FTypePtr FromShaderType(const UE::Shader::FType& InShaderType);
 	static FTypePtr GetVoid();
 
-	FArithmeticTypePtr ToArithmetic() const;
-	FArithmeticTypePtr ToScalar() const;
-	FArithmeticTypePtr ToVector() const;
-	FArithmeticTypePtr ToMatrix() const;
+	const FArithmeticType* ToArithmetic() const;
+	bool IsScalar() const;
+
 };
 
 enum EScalarKind
@@ -50,6 +49,6 @@ struct FArithmeticType : FType
 	bool IsMatrix() const { return NumRows > 1 && NumColumns > 1; }
 };
 
-} // namespace UE::MIR
+} // namespace MaterialIR
 
-#endif // #if WITH_EDITOR
+#endif
diff --git a/Engine/Source/Runtime/Engine/Public/Materials/MaterialInterface.h b/Engine/Source/Runtime/Engine/Public/Materials/MaterialInterface.h
--- a/Engine/Source/Runtime/Engine/Public/Materials/MaterialInterface.h
+++ b/Engine/Source/Runtime/Engine/Public/Materials/MaterialInterface.h
@@ -20,7 +20,6 @@
 #include "SceneTypes.h"
 #include "Engine/BlendableInterface.h"
 #include "Materials/MaterialLayersFunctions.h"
-#include "Materials/MaterialIRModule.h"
 #include "Interfaces/Interface_AssetUserData.h"
 #include "MaterialSceneTextureId.h"
 #include "Materials/MaterialRelevance.h"
@@ -434,11 +433,9 @@ public:
 
 #if WITH_EDITOR
 	ENGINE_API virtual const FMaterialCachedHLSLTree& GetCachedHLSLTree(TMicRecursionGuard RecursionGuard = TMicRecursionGuard()) const;
-	ENGINE_API const FMaterialIRModule& GetIRModule() const { return IRModule; }
 #endif
 
 	ENGINE_API bool IsUsingNewHLSLGenerator() const;
-	ENGINE_API bool IsUsingNewTranslatorPrototype() const;
 	ENGINE_API bool IsUsingControlFlow() const;
 
 	ENGINE_API const FSubstrateCompilationConfig& GetSubstrateCompilationConfig() const;
@@ -1195,7 +1192,6 @@ protected:
 
 #if WITH_EDITOR
 	TUniquePtr<FMaterialCachedHLSLTree> CachedHLSLTree;
-	FMaterialIRModule IRModule;
 #endif // WITH_EDITOR
 private:
 	/**
