reduce hitching on ODSC shader reload by replacing FMaterialUpdateContext::EOptions::RecreateRenderStates with 0.

Doing so will only re-cache the FMeshDrawCommand for the materials we provide (and their dependencies). This should be good enough since we're only interested in calling FMaterial::TryGetShaders again to fetch the right shader for the MDC

#rnx

diff --git a/Engine/Source/Runtime/Engine/Private/ShaderCompiler/ShaderCompiler.cpp b/Engine/Source/Runtime/Engine/Private/ShaderCompiler/ShaderCompiler.cpp
--- a/Engine/Source/Runtime/Engine/Private/ShaderCompiler/ShaderCompiler.cpp
+++ b/Engine/Source/Runtime/Engine/Private/ShaderCompiler/ShaderCompiler.cpp
@@ -8725,6 +8725,17 @@ namespace
 	}
 }
 
+static int32 GODSCMaterialUpdateFlags = 0;
+static FAutoConsoleVariableRef CVarODSCMaterialUpdateFlags(
+	TEXT("ODSC.MaterialUpdateFlags"),
+	GODSCMaterialUpdateFlags,
+	TEXT("Changes the material update flags when ODSC receives new shaders and needs to update the materials\n")
+	TEXT("0 (default): no additional work\n")
+	TEXT("1: Reregister all components while updating the material\n")
+	TEXT("2: Sync with the rendering thread after all the calls to RecacheUniformExpressions\n")
+	TEXT("4 (legacy): Recreates only the render state for *all* components, including the ones not changed by ODSC\n")
+);
+
 void ProcessCookOnTheFlyShaders(bool bReloadGlobalShaders, const TArray<uint8>& MeshMaterialMaps, const TArray<FString>& MaterialsToLoad, const TArray<uint8>& GlobalShaderMap)
 {
 	TRACE_CPUPROFILER_EVENT_SCOPE(ProcessCookOnTheFlyShaders);
@@ -8767,12 +8778,17 @@ void ProcessCookOnTheFlyShaders(bool bReloadGlobalShaders, const TArray<uint8>&
 		if (LoadedMaterials.Num())
 		{
 			// this will stop the rendering thread, and reattach components, in the destructor
-			FMaterialUpdateContext UpdateContext(FMaterialUpdateContext::EOptions::RecreateRenderStates);
+			FMaterialUpdateContext UpdateContext(GODSCMaterialUpdateFlags);
 
 			// gather the shader maps to reattach
 			for (UMaterialInterface* Material : LoadedMaterials)
 			{
-				Material->RecacheUniformExpressions(true);
+				// ~FMaterialUpdateContext takes care of calling RecacheUniformExpressions on all MaterialInstances, no need to call it twice
+				if (Cast<UMaterialInstance>(Material) == nullptr)
+				{
+					Material->RecacheUniformExpressions(true);
+				}
+
 				UpdateContext.AddMaterialInterface(Material);
 			}
 		}
