[NetPhysics] Remove deprecated logic from UE5.3

diff --git a/Engine/Source/Runtime/Engine/Classes/GameFramework/PlayerController.h b/Engine/Source/Runtime/Engine/Classes/GameFramework/PlayerController.h
--- a/Engine/Source/Runtime/Engine/Classes/GameFramework/PlayerController.h
+++ b/Engine/Source/Runtime/Engine/Classes/GameFramework/PlayerController.h
@@ -377,29 +377,9 @@ protected:
 #if UE_WITH_IRIS
 	ENGINE_API virtual void BeginReplication() override;
 #endif // UE_WITH_IRIS
-	/** The type of async physics data object to use*/
-	UPROPERTY(meta = (DeprecatedProperty, DeprecationMessage = "AsyncPhysicsDataClass is deprecated. See UInputSettings::bEnableLegacyInputScales to enable legacy behavior"))
-	TSubclassOf<UAsyncPhysicsData> AsyncPhysicsDataClass_DEPRECATED;
-
-	/** Get the async physics data to write to. This data will make its way to the async physics tick on client and server. Should not be used during async tick */
-	UE_DEPRECATED(5.3, "GetAsyncPhysicsDataToConsume is deprecated, please see the new C++ NetworkPhysicsComponent")
-	UFUNCTION(BlueprintPure, Category = PlayerController)
-	ENGINE_API UAsyncPhysicsData* GetAsyncPhysicsDataToWrite() const;
-
-	/** Get the async physics data to execute logic off of. This data should not be modified and will NOT make its way back. Must be used during async tick */
-	UE_DEPRECATED(5.3, "GetAsyncPhysicsDataToConsume is deprecated, please see the new C++ NetworkPhysicsComponent")
-	UFUNCTION(BlueprintPure, Category = PlayerController)
-	ENGINE_API const UAsyncPhysicsData* GetAsyncPhysicsDataToConsume() const;
 
 private:
 
-	UPROPERTY(ReplicatedUsing=OnRep_AsyncPhysicsDataComponent, meta = (DeprecatedProperty, DeprecationMessage = "AsyncPhysicsDataComponent is deprecated. please see the new C++ NetworkPhysicsComponent"))
-	TObjectPtr<UAsyncPhysicsInputComponent> AsyncPhysicsDataComponent_DEPRECARED;
-
-	UE_DEPRECATED(5.3, "OnRep_AsyncPhysicsDataComponent is deprecated, please see the new C++ NetworkPhysicsComponent")
-	UFUNCTION()
-	ENGINE_API void OnRep_AsyncPhysicsDataComponent();
-
 	struct FDynamicForceFeedbackAction
 	{
 		/** Time over which interpolation should happen */
@@ -2407,9 +2387,6 @@ private:
 	UFUNCTION(Client, Unreliable)
 	ENGINE_API void ClientCorrectionAsyncPhysicsTimestamp(FAsyncPhysicsTimestamp Timestamp);
 
-	/** Update the tick offset in between the local client and the server */
-	ENGINE_API virtual void AsyncPhysicsTickActor(float DeltaTime, float SimTime) override;
-
 public:
 
 	/** Generates a timestamp for the upcoming physics step (plus any pending time). Useful for synchronizing client and server events on a specific physics step */
diff --git a/Engine/Source/Runtime/Engine/Private/PlayerController.cpp b/Engine/Source/Runtime/Engine/Private/PlayerController.cpp
--- a/Engine/Source/Runtime/Engine/Private/PlayerController.cpp
+++ b/Engine/Source/Runtime/Engine/Private/PlayerController.cpp
@@ -65,8 +65,6 @@
 #include "Particles/EmitterCameraLensEffectBase.h"
 #include "LevelUtils.h"
 #include "WorldPartition/WorldPartitionSubsystem.h"
-#include "Physics/AsyncPhysicsInputComponent.h"
-#include "Physics/NetworkPhysicsComponent.h"
 #include "PBDRigidsSolver.h"
 #include "PhysicsEngine/PhysicsSettings.h"
 
@@ -239,10 +237,6 @@ APlayerController::APlayerController(const FObjectInitializer& ObjectInitializer
 		RootComponent->SetUsingAbsoluteRotation(true);
 	}
 
-	if (UPhysicsSettings::Get()->PhysicsPrediction.bEnablePhysicsPrediction)
-	{
-		bAsyncPhysicsTickEnabled = true;
-	}
 #if UE_ENABLE_DEBUG_DRAWING
 	CurrentInputModeDebugString = TEXT("Default");
 #endif	// UE_ENABLE_DEBUG_DRAWING
@@ -4868,8 +4862,6 @@ void APlayerController::GetLifetimeReplicatedProps(TArray< FLifetimeProperty > &
 {
 	Super::GetLifetimeReplicatedProps(OutLifetimeProps);
 
-	DISABLE_REPLICATED_PROPERTY(APlayerController, AsyncPhysicsDataComponent_DEPRECARED);
-
 	// These used to only replicate if PlayerCameraManager->GetViewTargetPawn() != GetPawn()
 	// But, since they also don't update unless that condition is true, these values won't change, thus won't send
 	// This is a little less efficient, but fits into the new condition system well, and shouldn't really add much overhead
@@ -4879,9 +4871,6 @@ void APlayerController::GetLifetimeReplicatedProps(TArray< FLifetimeProperty > &
 	DOREPLIFETIME_CONDITION(APlayerController, SpawnLocation, COND_OwnerOnly);
 }
 
-void APlayerController::OnRep_AsyncPhysicsDataComponent()
-{}
-
 void APlayerController::SetPlayer( UPlayer* InPlayer )
 {
 	FMoviePlayerProxyBlock MoviePlayerBlock;
@@ -6116,16 +6105,6 @@ void APlayerController::BeginReplication()
 }
 #endif // UE_WITH_IRIS
 
-UAsyncPhysicsData* APlayerController::GetAsyncPhysicsDataToWrite() const
-{
-	return nullptr;
-}
-
-const UAsyncPhysicsData* APlayerController::GetAsyncPhysicsDataToConsume() const
-{
-	return nullptr;  
-}
-
 void APlayerController::ExecuteAsyncPhysicsCommand(const FAsyncPhysicsTimestamp& AsyncPhysicsTimestamp, UObject* OwningObject, const TFunction<void()>& Command, const bool bEnableResim)
 {
 	if(UWorld* World = GetWorld())
@@ -6342,11 +6321,5 @@ void APlayerController::UpdateServerTimestampToCorrect()
 	PRAGMA_ENABLE_DEPRECATION_WARNINGS
 }
 
-void APlayerController::AsyncPhysicsTickActor(float DeltaTime, float SimTime)
-{
-	Super::AsyncPhysicsTickActor(DeltaTime, SimTime);
-}
-
-
 #undef LOCTEXT_NAMESPACE
 
