Renamed some visualization menu text.

diff --git a/Engine/Source/Runtime/Engine/Public/ShowFlagsValues.inl b/Engine/Source/Runtime/Engine/Public/ShowFlagsValues.inl
--- a/Engine/Source/Runtime/Engine/Public/ShowFlagsValues.inl
+++ b/Engine/Source/Runtime/Engine/Public/ShowFlagsValues.inl
@@ -37,7 +37,7 @@ SHOWFLAG_ALWAYS_ACCESSIBLE(EyeAdaptation, SFG_PostProcess, NSLOCTEXT("UnrealEd",
 /** Display a histogram of the scene HDR color */
 SHOWFLAG_FIXED_IN_SHIPPING(0, VisualizeHDR, SFG_Visualize, NSLOCTEXT("UnrealEd", "VisualizeHDRSF", "HDR (Eye Adaptation)"))
 /** Display the illuminance debug view for the skylight */
-SHOWFLAG_FIXED_IN_SHIPPING(0, VisualizeSkyLightIlluminance, SFG_Visualize, NSLOCTEXT("UnrealEd", "VisualizeSkyLightIlluminance", "Visualize SkyLight Illuminance"))
+SHOWFLAG_FIXED_IN_SHIPPING(0, VisualizeSkyLightIlluminance, SFG_Visualize, NSLOCTEXT("UnrealEd", "VisualizeSkyLightIlluminance", "SkyLight Illuminance"))
 /** Helper to tweak local expsoure settings */
 SHOWFLAG_FIXED_IN_SHIPPING(0, VisualizeLocalExposure, SFG_Visualize, NSLOCTEXT("UnrealEd", "VisualizeLocalExposureSF", "Local Exposure"))
 /** Image based lens flares (Simulate artifact of reflections within a camera system) */
@@ -334,7 +334,7 @@ SHOWFLAG_FIXED_IN_SHIPPING(0, PhysicsField, SFG_Visualize, NSLOCTEXT("UnrealEd",
 /** Global Distance field */
 SHOWFLAG_FIXED_IN_SHIPPING(0, VisualizeGlobalDistanceField, SFG_Visualize, NSLOCTEXT("UnrealEd", "GlobalDistanceFieldSF", "Global DistanceField"))
 /** Enable the debug visualization of diffuse/specular lighting (direct and indirect) using probes */
-SHOWFLAG_FIXED_IN_SHIPPING(0, VisualizeLightingOnProbes, SFG_Visualize, NSLOCTEXT("UnrealEd", "VisualizeLightingOnProbesSF", "Visualize Lighting on Probes"))
+SHOWFLAG_FIXED_IN_SHIPPING(0, VisualizeLightingOnProbes, SFG_Visualize, NSLOCTEXT("UnrealEd", "VisualizeLightingOnProbesSF", "Debug Lighting on Probes"))
 /** Screen space AO, for now SHOWFLAG_ALWAYS_ACCESSIBLE because r.GBuffer need that */
 SHOWFLAG_ALWAYS_ACCESSIBLE(ScreenSpaceAO, SFG_LightingFeatures, NSLOCTEXT("UnrealEd", "ScreenSpaceAOSF", "Screen Space Ambient Occlusion"))
 /** Distance field AO, for now SHOWFLAG_ALWAYS_ACCESSIBLE because it's exposed in SceneCapture */
@@ -394,7 +394,7 @@ SHOWFLAG_FIXED_IN_SHIPPING(0, RayTracingDebug, SFG_Hidden, NSLOCTEXT("UnrealEd",
 /** Enable the SkyAtmosphere visualization to be drawn on screen */
 SHOWFLAG_FIXED_IN_SHIPPING(0, VisualizeSkyAtmosphere, SFG_Visualize, NSLOCTEXT("UnrealEd", "VisualizeSkyAtmosphereSF", "Sky Atmosphere"))
 /** Enable the light function atlas debug visualization to be drawn on screen */
-SHOWFLAG_FIXED_IN_SHIPPING(0, VisualizeLightFunctionAtlas, SFG_Visualize, NSLOCTEXT("UnrealEd", "VisualizeLightFunctionAtlasSF", "Visualize Light Function Atlas"))
+SHOWFLAG_FIXED_IN_SHIPPING(0, VisualizeLightFunctionAtlas, SFG_Visualize, NSLOCTEXT("UnrealEd", "VisualizeLightFunctionAtlasSF", "Light Function Atlas"))
 /** Shows a full-screen calibration color post process material defined in renderer config  */
 SHOWFLAG_ALWAYS_ACCESSIBLE(VisualizeCalibrationColor, SFG_Developer, NSLOCTEXT("UnrealEd", "VisualizeCalibrationColorSF", "Visualize Calibration Color"))
 /** Shows a full-screen calibration grayscale post process material defined in renderer config  */
@@ -402,7 +402,7 @@ SHOWFLAG_ALWAYS_ACCESSIBLE(VisualizeCalibrationGrayscale, SFG_Developer, NSLOCTE
 /** Shows a full-screen calibration custom post process material defined in renderer config  */
 SHOWFLAG_ALWAYS_ACCESSIBLE(VisualizeCalibrationCustom, SFG_Developer, NSLOCTEXT("UnrealEd", "VisualizeCalibrationCustomSF", "Visualize Calibration Custom"))
 /** Visualise information about the post processes stacked on screen */
-SHOWFLAG_FIXED_IN_SHIPPING(0, VisualizePostProcessStack, SFG_Visualize, NSLOCTEXT("UnrealEd", "VisualizePostProcessStackSF", "Visualize Post Process Stack"))
+SHOWFLAG_FIXED_IN_SHIPPING(0, VisualizePostProcessStack, SFG_Visualize, NSLOCTEXT("UnrealEd", "VisualizePostProcessStackSF", "Post Process Stack"))
 
 /** Draw in the main pass the primitives that would normally only contribute to runtime virtual textures. */
 SHOWFLAG_FIXED_IN_SHIPPING(0, VirtualTexturePrimitives, SFG_Developer, NSLOCTEXT("UnrealEd", "VirtualTexturePrimitivesSF", "Virtual Texture Primitives"))
