[World Partition Stats] Made World Partition stats group public and added a new stat for UWorldPartitionRuntimeCellData::AppendStreamingSourceInfo.

diff --git a/Engine/Source/Runtime/Engine/Private/WorldPartition/WorldPartitionRuntimeCellData.cpp b/Engine/Source/Runtime/Engine/Private/WorldPartition/WorldPartitionRuntimeCellData.cpp
--- a/Engine/Source/Runtime/Engine/Private/WorldPartition/WorldPartitionRuntimeCellData.cpp
+++ b/Engine/Source/Runtime/Engine/Private/WorldPartition/WorldPartitionRuntimeCellData.cpp
@@ -1,6 +1,7 @@
 // Copyright Epic Games, Inc. All Rights Reserved.
 
 #include "WorldPartition/WorldPartitionRuntimeCellData.h"
+#include "WorldPartition/WorldPartitionLog.h"
 #include "Misc/HierarchicalLogArchive.h"
 
 int32 UWorldPartitionRuntimeCellData::StreamingSourceCacheEpoch = 0;
@@ -46,8 +47,11 @@ void UWorldPartitionRuntimeCellData::ResetStreamingSourceInfo() const
 	CachedSourceInfoEpoch = StreamingSourceCacheEpoch;	
 }
 
+DECLARE_CYCLE_STAT(TEXT("Append Streaming Source Info"), STAT_WorldPartitionAppendStreamingSourceInfo, STATGROUP_WorldPartition);
 void UWorldPartitionRuntimeCellData::AppendStreamingSourceInfo(const FWorldPartitionStreamingSource& Source, const FSphericalSector& SourceShape) const
 {
+	SCOPE_CYCLE_COUNTER(STAT_WorldPartitionAppendStreamingSourceInfo);
+
 	if (CachedSourceInfoEpoch != StreamingSourceCacheEpoch)
 	{
 		ResetStreamingSourceInfo();
diff --git a/Engine/Source/Runtime/Engine/Private/WorldPartition/WorldPartitionSubsystem.cpp b/Engine/Source/Runtime/Engine/Private/WorldPartition/WorldPartitionSubsystem.cpp
--- a/Engine/Source/Runtime/Engine/Private/WorldPartition/WorldPartitionSubsystem.cpp
+++ b/Engine/Source/Runtime/Engine/Private/WorldPartition/WorldPartitionSubsystem.cpp
@@ -33,7 +33,6 @@
 #include "Async/ParallelFor.h"
 #include "Algo/ForEach.h"
 #include "Misc/HashBuilder.h"
-#include "Stats/Stats.h"
 
 #if WITH_EDITOR
 #include "Editor.h"
@@ -50,7 +49,6 @@
 
 #include UE_INLINE_GENERATED_CPP_BY_NAME(WorldPartitionSubsystem)
 
-DECLARE_STATS_GROUP(TEXT("World Partition"), STATGROUP_WorldPartition, STATCAT_Advanced);
 DECLARE_CYCLE_STAT(TEXT("Update Streaming State"), STAT_WorldPartitionUpdateStreamingState, STATGROUP_WorldPartition);
 DECLARE_CYCLE_STAT(TEXT("Update Streaming State (External)"), STAT_WorldPartitionUpdateStreamingStateExternal, STATGROUP_WorldPartition);
 CSV_DEFINE_CATEGORY(WorldPartition, (!UE_BUILD_SHIPPING));
diff --git a/Engine/Source/Runtime/Engine/Public/WorldPartition/WorldPartitionLog.h b/Engine/Source/Runtime/Engine/Public/WorldPartition/WorldPartitionLog.h
--- a/Engine/Source/Runtime/Engine/Public/WorldPartition/WorldPartitionLog.h
+++ b/Engine/Source/Runtime/Engine/Public/WorldPartition/WorldPartitionLog.h
@@ -2,5 +2,7 @@
 #pragma once
 
 #include "CoreMinimal.h"
+#include "Stats/Stats.h"
 
-ENGINE_API DECLARE_LOG_CATEGORY_EXTERN(LogWorldPartition, Log, All);
\ No newline at end of file
+ENGINE_API DECLARE_LOG_CATEGORY_EXTERN(LogWorldPartition, Log, All);
+DECLARE_STATS_GROUP(TEXT("World Partition"), STATGROUP_WorldPartition, STATCAT_Advanced);
\ No newline at end of file
