load global shaders from ODSC at startup

diff --git a/Engine/Source/Runtime/Engine/Private/ODSC/ODSCManager.cpp b/Engine/Source/Runtime/Engine/Private/ODSC/ODSCManager.cpp
--- a/Engine/Source/Runtime/Engine/Private/ODSC/ODSCManager.cpp
+++ b/Engine/Source/Runtime/Engine/Private/ODSC/ODSCManager.cpp
@@ -10,6 +10,12 @@
 #include "Materials/MaterialInstance.h"
 #include "Materials/Material.h"
 
+// For GetMaxSupportedFeatureLevel
+#include "DataDrivenShaderPlatformInfo.h"
+
+// For GetCachedScalabilityCVars
+#include "UnrealEngine.h"
+
 DEFINE_LOG_CATEGORY(LogODSC);
 
 // FODSCManager
@@ -240,3 +246,22 @@ void FODSCManager::RegisterMaterialShaderMap(const FMaterialShaderMap& MaterialS
 		GODSCManager->Thread->RegisterMaterialShaderMap(MaterialShaderMap);
 	}
 }
+ 
+void FODSCManager::TryLoadGlobalShaders(EShaderPlatform ShaderPlatform)
+{
+	check(IsODSCActive());
+	ERHIFeatureLevel::Type TargetFeatureLevel = GetMaxSupportedFeatureLevel(ShaderPlatform);
+	const EMaterialQualityLevel::Type ActiveQualityLevel = GetCachedScalabilityCVars().MaterialQualityLevel;
+	Thread->AddRequest(TArray<FString>(), FString(), ShaderPlatform, TargetFeatureLevel, ActiveQualityLevel, ODSCRecompileCommand::Changed);
+	Thread->Wakeup();
+	Thread->WaitUntilAllRequestsDone();
+	TArray<FODSCMessageHandler*> CompletedThreadedRequests;
+	Thread->GetCompletedRequests(CompletedThreadedRequests);
+	// Finish and remove any completed requests
+	for (FODSCMessageHandler* CompletedRequest : CompletedThreadedRequests)
+	{
+		check(CompletedRequest);
+		ProcessCookOnTheFlyShaders(false, CompletedRequest->GetMeshMaterialMaps(), CompletedRequest->GetMaterialsToLoad(), CompletedRequest->GetGlobalShaderMap());
+		delete CompletedRequest;
+	}
+}
\ No newline at end of file
diff --git a/Engine/Source/Runtime/Engine/Private/ODSC/ODSCThread.cpp b/Engine/Source/Runtime/Engine/Private/ODSC/ODSCThread.cpp
--- a/Engine/Source/Runtime/Engine/Private/ODSC/ODSCThread.cpp
+++ b/Engine/Source/Runtime/Engine/Private/ODSC/ODSCThread.cpp
@@ -106,6 +106,7 @@ bool FODSCMessageHandler::ReloadGlobalShaders() const
 FODSCThread::FODSCThread(const FString& HostIP)
 	: Thread(nullptr)
 	, WakeupEvent(FPlatformProcess::GetSynchEventFromPool(true))
+	, AllRequestsDoneEvent(FPlatformProcess::GetSynchEventFromPool(true))
 	, ODSCHostIP(HostIP)
 {
 	UE_LOG(LogODSC, Log, TEXT("ODSC Thread active."));
@@ -121,6 +122,8 @@ FODSCThread::~FODSCThread()
 {
 	StopThread();
 
+	FPlatformProcess::ReturnSynchEventToPool(AllRequestsDoneEvent);
+	AllRequestsDoneEvent = nullptr;
 	FPlatformProcess::ReturnSynchEventToPool(WakeupEvent);
 	WakeupEvent = nullptr;
 }
@@ -268,9 +271,15 @@ void FODSCThread::GetCompletedRequests(TArray<FODSCMessageHandler*>& OutComplete
 
 void FODSCThread::Wakeup()
 {
+	AllRequestsDoneEvent->Reset();
 	WakeupEvent->Trigger();
 }
 
+void FODSCThread::WaitUntilAllRequestsDone()
+{
+	AllRequestsDoneEvent->Wait();
+}
+
 bool FODSCThread::Init()
 {
 	return true;
@@ -349,6 +358,7 @@ void FODSCThread::Process()
 	// SendMessageToServer is synchronous, so when we're here, we know we've processed all the requests
 	bHasPendingRequests = false;
 	WakeupEvent->Reset();
+	AllRequestsDoneEvent->Trigger();
 }
 
 void FODSCThread::SendMessageToServer(IPlatformFile::IFileServerMessageHandler* Handler)
diff --git a/Engine/Source/Runtime/Engine/Private/ODSC/ODSCThread.h b/Engine/Source/Runtime/Engine/Private/ODSC/ODSCThread.h
--- a/Engine/Source/Runtime/Engine/Private/ODSC/ODSCThread.h
+++ b/Engine/Source/Runtime/Engine/Private/ODSC/ODSCThread.h
@@ -149,6 +149,11 @@ public:
 	* Wakeup the thread to process requests.
 	*/
 	void Wakeup();
+	
+	/**
+	* Wait until all added requests are processed. Must be called after Wakeup.
+	*/
+	void WaitUntilAllRequestsDone();
 
 	bool HasPendingRequests() const { return bHasPendingRequests; }
 
@@ -206,6 +211,9 @@ private:
 
 	/** Holds an event signaling the thread to wake up. */
 	FEvent* WakeupEvent;
+	
+	/** Holds an event signaling when all the requests are processed*/
+	FEvent* AllRequestsDoneEvent;
 
 	void SendMessageToServer(IPlatformFile::IFileServerMessageHandler* Handler);
 
diff --git a/Engine/Source/Runtime/Engine/Private/ShaderCompiler/ShaderCompiler.cpp b/Engine/Source/Runtime/Engine/Private/ShaderCompiler/ShaderCompiler.cpp
--- a/Engine/Source/Runtime/Engine/Private/ShaderCompiler/ShaderCompiler.cpp
+++ b/Engine/Source/Runtime/Engine/Private/ShaderCompiler/ShaderCompiler.cpp
@@ -9782,6 +9782,16 @@ void CompileGlobalShaderMap(EShaderPlatform Platform, const ITargetPlatform* Tar
 		FlushShaderFileCache();
 	}
 
+#if WITH_ODSC
+	// First try to load the global shader map with ODSC if it's connected. TryLoadGlobalShaders will set GGlobalShaderMap[Platform]
+	if (!GGlobalShaderMap[Platform] && FODSCManager::IsODSCActive())
+	{
+		UE_LOG(LogShaders, Display, TEXT("Trying to load global shaders from ODSC ..."));
+		GODSCManager->TryLoadGlobalShaders(Platform);
+		UE_LOG(LogShaders, Display, TEXT("Global shaders from ODSC: %s"), (GGlobalShaderMap[Platform] != nullptr) ? TEXT("success") : TEXT("failed"));
+	}
+#endif
+
 	// If the global shader map hasn't been created yet, create it.
 	if (!GGlobalShaderMap[Platform])
 	{
@@ -10638,12 +10648,14 @@ void LoadGlobalShadersForRemoteRecompile(FArchive& Ar, EShaderPlatform ShaderPla
 
 			bool bIsNewGlobalShaderMapComplete = IsGlobalShaderMapComplete(nullptr, NewGlobalShaderMap, ShaderPlatform);
 
-			if (GGlobalShaderMap[ShaderPlatform] && bIsNewGlobalShaderMapComplete)
+			if (bIsNewGlobalShaderMapComplete)
 			{
-				GGlobalShaderMap[ShaderPlatform]->ReleaseAllSections();
-
-				delete GGlobalShaderMap[ShaderPlatform];
-				GGlobalShaderMap[ShaderPlatform] = nullptr;
+				if (GGlobalShaderMap[ShaderPlatform])
+				{
+					GGlobalShaderMap[ShaderPlatform]->ReleaseAllSections();
+					delete GGlobalShaderMap[ShaderPlatform];
+					GGlobalShaderMap[ShaderPlatform] = nullptr;
+				}
 				GGlobalShaderMap[ShaderPlatform] = NewGlobalShaderMap;
 
 				VerifyGlobalShaders(ShaderPlatform, nullptr, false);
@@ -10658,10 +10670,7 @@ void LoadGlobalShadersForRemoteRecompile(FArchive& Ar, EShaderPlatform ShaderPla
 			}
 			else
 			{
-				if (!bIsNewGlobalShaderMapComplete)
-				{
-					UE_LOG(LogShaderCompilers, Error, TEXT("New shader map is incomplete. Look at the ODSC server log to see shader errors"));
-				}
+				UE_LOG(LogShaderCompilers, Error, TEXT("New shader map is incomplete. Look at the ODSC server log to see shader errors"));
 				delete NewGlobalShaderMap;
 			}
 		}
diff --git a/Engine/Source/Runtime/Engine/Public/ODSC/ODSCManager.h b/Engine/Source/Runtime/Engine/Public/ODSC/ODSCManager.h
--- a/Engine/Source/Runtime/Engine/Public/ODSC/ODSCManager.h
+++ b/Engine/Source/Runtime/Engine/Public/ODSC/ODSCManager.h
@@ -105,6 +105,7 @@ public:
 
 	static void SuspendODSCForceRecompile();
 	static void ResumeODSCForceRecompile();
+	void TryLoadGlobalShaders(EShaderPlatform ShaderPlatform);
 
 private:
 
