Add MetalOptimizeBySize setting for MetalShaderFormat

diff --git a/Engine/Source/Developer/Apple/MetalShaderFormat/Private/MetalShaderCompiler.cpp b/Engine/Source/Developer/Apple/MetalShaderFormat/Private/MetalShaderCompiler.cpp
--- a/Engine/Source/Developer/Apple/MetalShaderFormat/Private/MetalShaderCompiler.cpp
+++ b/Engine/Source/Developer/Apple/MetalShaderFormat/Private/MetalShaderCompiler.cpp
@@ -752,6 +752,7 @@ void BuildMetalShaderOutput(
 			Job.SourceCRCLen = SourceCRCLen;
 			Job.SourceCRC = SourceCRC;
 			Job.bRetainObjectFile = ShaderInput.Environment.CompilerFlags.Contains(CFLAG_Archive);
+			Job.bOptimizeForSize = ShaderInput.Environment.GetCompileArgument(TEXT("METAL_OPTIMIZE_FOR_SIZE"), false);
 			Job.bCompileAsPCH = false;
 			Job.ReturnCode = 0;
 
diff --git a/Engine/Source/Developer/Apple/MetalShaderFormat/Private/MetalShaderCompiler.h b/Engine/Source/Developer/Apple/MetalShaderFormat/Private/MetalShaderCompiler.h
--- a/Engine/Source/Developer/Apple/MetalShaderFormat/Private/MetalShaderCompiler.h
+++ b/Engine/Source/Developer/Apple/MetalShaderFormat/Private/MetalShaderCompiler.h
@@ -58,6 +58,7 @@ struct FMetalShaderBytecodeJob
 	uint32 SourceCRC;
 	bool bRetainObjectFile;
 	bool bCompileAsPCH;
+	bool bOptimizeForSize;
 	
 	FString Message;
 	FString Results;
diff --git a/Engine/Source/Developer/Apple/MetalShaderFormat/Private/MetalShaderFormat.cpp b/Engine/Source/Developer/Apple/MetalShaderFormat/Private/MetalShaderFormat.cpp
--- a/Engine/Source/Developer/Apple/MetalShaderFormat/Private/MetalShaderFormat.cpp
+++ b/Engine/Source/Developer/Apple/MetalShaderFormat/Private/MetalShaderFormat.cpp
@@ -1106,6 +1106,11 @@ bool FMetalCompilerToolchain::CompileMetalShader(FMetalShaderBytecodeJob& Job, F
 	{
 		// Invoke the metal frontend.
 		FString MetalParams = FString::Printf(TEXT("%s %s %s %s -Wno-null-character -fbracket-depth=1024 %s %s %s %s %s -o %s"), *Job.MinOSVersion, *Job.PreserveInvariance, *Job.DebugInfo, *Job.MathMode, TEXT("-c"), *Job.Standard, *Job.Defines, *IncludeArgs, *LocalInputMetalFilePath, *LocalOutputMetalAIRFilePath);
+		if (Job.bOptimizeForSize)
+		{
+			MetalParams += TEXT(" -Os");
+		}
+
 		bool bSuccess = this->ExecMetalFrontend(SDK, *MetalParams, &Job.ReturnCode, &Job.Results, &Job.Errors);
 
 		if (!bSuccess || (Job.ReturnCode != 0))
diff --git a/Engine/Source/Runtime/Engine/Private/ShaderCompiler/ShaderCompiler.cpp b/Engine/Source/Runtime/Engine/Private/ShaderCompiler/ShaderCompiler.cpp
--- a/Engine/Source/Runtime/Engine/Private/ShaderCompiler/ShaderCompiler.cpp
+++ b/Engine/Source/Runtime/Engine/Private/ShaderCompiler/ShaderCompiler.cpp
@@ -8061,12 +8061,14 @@ void GlobalBeginCompileShader(
 			int32 IndirectArgumentTier = 0;
 			bool bEnableMathOptimisations = true;
             bool bSupportAppleA8 = false;
-            
+			bool bMetalOptimizeForSize = false;
+
 			if (IsPCPlatform(EShaderPlatform(Target.Platform)))
 			{
 				GConfig->GetBool(TEXT("/Script/MacTargetPlatform.MacTargetSettings"), TEXT("UseFastIntrinsics"), bAllowFastIntrinsics, GEngineIni);
 				GConfig->GetBool(TEXT("/Script/MacTargetPlatform.MacTargetSettings"), TEXT("EnableMathOptimisations"), bEnableMathOptimisations, GEngineIni);
 				GConfig->GetInt(TEXT("/Script/MacTargetPlatform.MacTargetSettings"), TEXT("IndirectArgumentTier"), IndirectArgumentTier, GEngineIni);
+				GConfig->GetBool(TEXT("/Script/MacTargetPlatform.MacTargetSettings"), TEXT("MetalOptimizeForSize"), bMetalOptimizeForSize, GEngineIni);
                 
                 // No half precision support on MacOS at the moment
                 bForceFloats = true;
@@ -8078,6 +8080,7 @@ void GlobalBeginCompileShader(
 				GConfig->GetBool(TEXT("/Script/IOSRuntimeSettings.IOSRuntimeSettings"), TEXT("ForceFloats"), bForceFloats, GEngineIni);
 				GConfig->GetInt(TEXT("/Script/IOSRuntimeSettings.IOSRuntimeSettings"), TEXT("IndirectArgumentTier"), IndirectArgumentTier, GEngineIni);
                 GConfig->GetBool(TEXT("/Script/IOSRuntimeSettings.IOSRuntimeSettings"), TEXT("bSupportAppleA8"), bSupportAppleA8, GEngineIni);
+				GConfig->GetBool(TEXT("/Script/IOSRuntimeSettings.IOSRuntimeSettings"), TEXT("MetalOptimizeForSize"), bMetalOptimizeForSize, GEngineIni);
                 
 				// Force no development shaders on iOS
 				bAllowDevelopmentShaderCompile = false;
@@ -8088,6 +8091,7 @@ void GlobalBeginCompileShader(
 			Input.Environment.SetCompileArgument(TEXT("METAL_USE_FAST_INTRINSICS"), bAllowFastIntrinsics);
 			Input.Environment.SetCompileArgument(TEXT("METAL_INDIRECT_ARGUMENT_BUFFERS"), IndirectArgumentTier);
             Input.Environment.SetCompileArgument(TEXT("SUPPORT_APPLE_A8"), bSupportAppleA8);
+            Input.Environment.SetCompileArgument(TEXT("METAL_OPTIMIZE_FOR_SIZE"), bMetalOptimizeForSize);
 			
 			// Same as console-variable above, but that's global and this is per-platform, per-project
 			if (!bEnableMathOptimisations)
diff --git a/Engine/Source/Runtime/RenderCore/Private/Shader.cpp b/Engine/Source/Runtime/RenderCore/Private/Shader.cpp
--- a/Engine/Source/Runtime/RenderCore/Private/Shader.cpp
+++ b/Engine/Source/Runtime/RenderCore/Private/Shader.cpp
@@ -2069,12 +2069,14 @@ void ShaderMapAppendKey(EShaderPlatform Platform, FShaderKeyGenerator& KeyGen)
 		bool bForceFloats = false;
         bool bSupportAppleA8 = false;
 		int32 IndirectArgumentTier = 0;
-        
+		bool bMetalOptimizeBySize = false;
+
 		if (IsPCPlatform(Platform))
 		{
 			GConfig->GetBool(TEXT("/Script/MacTargetPlatform.MacTargetSettings"), TEXT("UseFastIntrinsics"), bAllowFastIntrinsics, GEngineIni);
 			GConfig->GetBool(TEXT("/Script/MacTargetPlatform.MacTargetSettings"), TEXT("EnableMathOptimisations"), bEnableMathOptimisations, GEngineIni);
 			GConfig->GetInt(TEXT("/Script/MacTargetPlatform.MacTargetSettings"), TEXT("IndirectArgumentTier"), IndirectArgumentTier, GEngineIni);
+			GConfig->GetBool(TEXT("/Script/MacTargetPlatform.MacTargetSettings"), TEXT("MetalOptimizeBySize"), bMetalOptimizeBySize, GEngineIni);
 		}
 		else
 		{
@@ -2083,6 +2085,7 @@ void ShaderMapAppendKey(EShaderPlatform Platform, FShaderKeyGenerator& KeyGen)
 			GConfig->GetBool(TEXT("/Script/IOSRuntimeSettings.IOSRuntimeSettings"), TEXT("ForceFloats"), bForceFloats, GEngineIni);
 			GConfig->GetInt(TEXT("/Script/IOSRuntimeSettings.IOSRuntimeSettings"), TEXT("IndirectArgumentTier"), IndirectArgumentTier, GEngineIni);
             GConfig->GetBool(TEXT("/Script/IOSRuntimeSettings.IOSRuntimeSettings"), TEXT("bSupportAppleA8"), bSupportAppleA8, GEngineIni);
+			GConfig->GetBool(TEXT("/Script/IOSRuntimeSettings.IOSRuntimeSettings"), TEXT("MetalOptimizeBySize"), bMetalOptimizeBySize, GEngineIni);
 		}
 		
 		if (bAllowFastIntrinsics)
@@ -2109,6 +2112,12 @@ void ShaderMapAppendKey(EShaderPlatform Platform, FShaderKeyGenerator& KeyGen)
 			KeyGen.AppendSeparator();
 			KeyGen.Append(TEXT("A8GPU"));
         }
+		
+		if(bMetalOptimizeBySize)
+        {
+			KeyGen.AppendSeparator();
+			KeyGen.Append(TEXT("Os"));
+        }
         
 		KeyGen.AppendSeparator();
 		KeyGen.Append(TEXT("IAB"));
