fix duplicate points in error-based sampling of splines w/ 'constant' control points

diff --git a/Engine/Source/Runtime/Engine/Private/Components/SplineComponent.cpp b/Engine/Source/Runtime/Engine/Private/Components/SplineComponent.cpp
--- a/Engine/Source/Runtime/Engine/Private/Components/SplineComponent.cpp
+++ b/Engine/Source/Runtime/Engine/Private/Components/SplineComponent.cpp
@@ -1662,9 +1662,14 @@ bool USplineComponent::DivideSplineIntoPolylineRecursiveWithDistancesHelper(floa
 	{
 		// The middle point is close enough to the other 2 points, let's keep those and stop the recursion :
 		OutPoints.Add(Samples[0]);
-		OutPoints.Add(Samples[2]);
 		OutDistancesAlongSpline.Add(StartDistanceAlongSpline);
-		OutDistancesAlongSpline.Add(EndDistanceAlongSpline);
+		// For a constant spline, the end can be the exact same as the start; in this case, just add the point once
+		if (Samples[0] != Samples[2])
+		{
+			OutPoints.Add(Samples[2]);
+			OutDistancesAlongSpline.Add(EndDistanceAlongSpline);
+		}
+		
 	}
 
 	check(OutPoints.Num() == OutDistancesAlongSpline.Num())
