Change errors to warnings when falling back to default material due to MATUSAGE_SkeletalMesh or MATUSAGE_Clothing not being set on a material where required.
#rnx

diff --git a/Engine/Source/Runtime/Engine/Private/SkeletalMesh.cpp b/Engine/Source/Runtime/Engine/Private/SkeletalMesh.cpp
--- a/Engine/Source/Runtime/Engine/Private/SkeletalMesh.cpp
+++ b/Engine/Source/Runtime/Engine/Private/SkeletalMesh.cpp
@@ -6286,7 +6286,7 @@ FSkeletalMeshSceneProxy::FSkeletalMeshSceneProxy(const USkinnedMeshComponent* Co
 
 			if(!Material || !bValidUsage)
 			{
-				UE_CLOG(Material && !bValidUsage, LogSkeletalMesh, Error,
+				UE_CLOG(Material && !bValidUsage, LogSkeletalMesh, Warning,
 					TEXT("Material with missing usage flag was applied to skeletal mesh %s"),
 					*Component->GetSkinnedAsset()->GetPathName());
 
@@ -6318,7 +6318,7 @@ FSkeletalMeshSceneProxy::FSkeletalMeshSceneProxy(const USkinnedMeshComponent* Co
 		if (!OverlayMaterial->CheckMaterialUsage_Concurrent(MATUSAGE_SkeletalMesh))
 		{
 			OverlayMaterial = UMaterial::GetDefaultMaterial(MD_Surface);
-			UE_LOG(LogSkeletalMesh, Error, TEXT("Overlay material with missing usage flag was applied to skeletal mesh %s"),	*Component->GetSkinnedAsset()->GetPathName());
+			UE_LOG(LogSkeletalMesh, Warning, TEXT("Overlay material with missing usage flag was applied to skeletal mesh %s"),	*Component->GetSkinnedAsset()->GetPathName());
 		}
 	}
 
