More PostCompilation protections for texture compiler.
#rnx

diff --git a/Engine/Source/Runtime/Engine/Private/TextureCompiler.cpp b/Engine/Source/Runtime/Engine/Private/TextureCompiler.cpp
--- a/Engine/Source/Runtime/Engine/Private/TextureCompiler.cpp
+++ b/Engine/Source/Runtime/Engine/Private/TextureCompiler.cpp
@@ -768,6 +768,16 @@ void FTextureCompilingManager::ProcessAsyncTasks(bool bLimitExecutionTime)
 
 void FTextureCompilingManager::ProcessAsyncTasks(const AssetCompilation::FProcessAsyncTaskParams& Params)
 {
+	if (bIsRoutingPostCompilation)
+	{
+		// This potentially affects RegisteredTextureBuckets which can't be touched inside PostCompilation.
+		// This is likely because a worker task got scheduled during a wait inside
+		// PostCompilation and it's randomly running during the wait, causing crashes.
+		// Workers that need to interact with textures should do that work in response to
+		// a game tick via e.g. ExecuteOnGameThread
+		UE_LOG(LogTexture, Fatal, TEXT("Calling ProcessAsyncTasks is not allowed during PostCompilation."));
+	}
+
 	FObjectCacheContextScope ObjectCacheScope;
 	ProcessDeferredRequests();
 	FinishCompilationsForGame();
