[World Partition HashSet]
- Fixed streaming sources queries not taking into account 2D cells.
- Fixed crash when enabling 3D debug draw display now that partitions can be fully 2D.
- Fixed crash when changing LHGrid Is2D flag without having the show grid preview enabled (recent regression).

diff --git a/Engine/Source/Runtime/Engine/Private/WorldPartition/RuntimeHashSet/RuntimePartitionLHGrid.cpp b/Engine/Source/Runtime/Engine/Private/WorldPartition/RuntimeHashSet/RuntimePartitionLHGrid.cpp
--- a/Engine/Source/Runtime/Engine/Private/WorldPartition/RuntimeHashSet/RuntimePartitionLHGrid.cpp
+++ b/Engine/Source/Runtime/Engine/Private/WorldPartition/RuntimeHashSet/RuntimePartitionLHGrid.cpp
@@ -111,15 +111,12 @@ void URuntimePartitionLHGrid::PostEditChangeProperty(FPropertyChangedEvent& InPr
 	else if ((PropertyName == GET_MEMBER_NAME_CHECKED(URuntimePartitionLHGrid, bShowGridPreview)) || 
 			 (PropertyName == GET_MEMBER_NAME_CHECKED(URuntimePartitionLHGrid, bIs2D)))
 	{
+		WorldGridPreviewer.Reset();
+
 		if (bShowGridPreview)
 		{
 			WorldGridPreviewer = MakeUnique<FWorldGridPreviewer>(GetTypedOuter<UWorld>(), bIs2D);
 		}
-		else
-		{
-			check(WorldGridPreviewer);
-			WorldGridPreviewer.Reset();
-		}
 
 		if (!GPackageWasDirty)
 		{
diff --git a/Engine/Source/Runtime/Engine/Private/WorldPartition/RuntimeHashSet/WorldPartitionRuntimeHashSet.cpp b/Engine/Source/Runtime/Engine/Private/WorldPartition/RuntimeHashSet/WorldPartitionRuntimeHashSet.cpp
--- a/Engine/Source/Runtime/Engine/Private/WorldPartition/RuntimeHashSet/WorldPartitionRuntimeHashSet.cpp
+++ b/Engine/Source/Runtime/Engine/Private/WorldPartition/RuntimeHashSet/WorldPartitionRuntimeHashSet.cpp
@@ -471,7 +471,7 @@ void UWorldPartitionRuntimeHashSet::ForEachStreamingCellsQuery(const FWorldParti
 		return false;
 	};
 
-	auto ForEachSpatiallyLoadedCells = [&ShouldAddCell, QueryCache, &QuerySource, &Func](FStaticSpatialIndexType* InSpatialIndex, int32 InLoadingRange, FName InGridName)
+	auto ForEachSpatiallyLoadedCells = [&ShouldAddCell, QueryCache, &QuerySource, &Func]<typename SpatialIndexType>(SpatialIndexType* InSpatialIndex, int32 InLoadingRange, FName InGridName)
 	{
 		if (InSpatialIndex)
 		{
@@ -511,7 +511,9 @@ void UWorldPartitionRuntimeHashSet::ForEachStreamingCellsQuery(const FWorldParti
 
 	ForEachStreamingData([&QuerySource, &ForEachSpatiallyLoadedCells, &ForEachNonSpatiallyLoadedCells](const FRuntimePartitionStreamingData& StreamingData)
 	{
-		return ForEachSpatiallyLoadedCells(StreamingData.SpatialIndex.Get(), StreamingData.GetLoadingRange(), StreamingData.Name) && ForEachNonSpatiallyLoadedCells(StreamingData.NonSpatiallyLoadedCells);
+		return ForEachSpatiallyLoadedCells(StreamingData.SpatialIndex.Get(), StreamingData.GetLoadingRange(), StreamingData.Name) && 
+			   ForEachSpatiallyLoadedCells(StreamingData.SpatialIndex2D.Get(), StreamingData.GetLoadingRange(), StreamingData.Name) &&
+			   ForEachNonSpatiallyLoadedCells(StreamingData.NonSpatiallyLoadedCells);
 	});
 }
 
diff --git a/Engine/Source/Runtime/Engine/Private/WorldPartition/RuntimeHashSet/WorldPartitionRuntimeHashSetDebugDraw.cpp b/Engine/Source/Runtime/Engine/Private/WorldPartition/RuntimeHashSet/WorldPartitionRuntimeHashSetDebugDraw.cpp
--- a/Engine/Source/Runtime/Engine/Private/WorldPartition/RuntimeHashSet/WorldPartitionRuntimeHashSetDebugDraw.cpp
+++ b/Engine/Source/Runtime/Engine/Private/WorldPartition/RuntimeHashSet/WorldPartitionRuntimeHashSetDebugDraw.cpp
@@ -358,63 +358,66 @@ void UWorldPartitionRuntimeHashSet::Draw3D(const TArray<FWorldPartitionStreaming
 
 	ForEachStreamingData([this, &Sources, VisualizeMode, &DataLayerDebugColors, ContentBundleManager, OwningWorld, &WorldPartitionTransform](const FRuntimePartitionStreamingData& StreamingData)
 	{
-		for (const FWorldPartitionStreamingSource& Source : Sources)
+		if (StreamingData.SpatialIndex)
 		{
-			Source.ForEachShape(StreamingData.GetLoadingRange(), StreamingData.Name, false, [this, &StreamingData, VisualizeMode, &DataLayerDebugColors, ContentBundleManager, OwningWorld, &WorldPartitionTransform](const FSphericalSector& Shape)
+			for (const FWorldPartitionStreamingSource& Source : Sources)
 			{
-				const FSphere ShapeSphere(Shape.GetCenter(), Shape.GetRadius());
-
-				StreamingData.SpatialIndex.Get()->ForEachIntersectingElement(ShapeSphere, [this, VisualizeMode, &DataLayerDebugColors, ContentBundleManager, OwningWorld, &WorldPartitionTransform](UWorldPartitionRuntimeCell* Cell)
+				Source.ForEachShape(StreamingData.GetLoadingRange(), StreamingData.Name, false, [this, &StreamingData, VisualizeMode, &DataLayerDebugColors, ContentBundleManager, OwningWorld, &WorldPartitionTransform](const FSphericalSector& Shape)
 				{
-					const FVector2D CellStreamingBoundsSize = FVector2D(Cell->GetStreamingBounds().GetSize());
-					const FVector2D CellStreamingBoundsMin = FVector2D(Cell->GetStreamingBounds().Min);
-
-					float CellOpacity = 0.0f;
-					TArray<FLinearColor> CellColors;
+					const FSphere ShapeSphere(Shape.GetCenter(), Shape.GetRadius());
 
-					switch (GShowRuntimeHashSetDebugDisplayMode)
+					StreamingData.SpatialIndex.Get()->ForEachIntersectingElement(ShapeSphere, [this, VisualizeMode, &DataLayerDebugColors, ContentBundleManager, OwningWorld, &WorldPartitionTransform](UWorldPartitionRuntimeCell* Cell)
 					{
-					case 0:
-						CellColors.Add(Cell->GetDebugColor(VisualizeMode));
-						CellOpacity = 0.25f / FMath::Max<float>(GShowRuntimeHashSetDebugDisplayLevelCount, 1);
-						break;
-					case 1:
-						if (DataLayerDebugColors.Num() && Cell->GetDataLayers().Num())
+						const FVector2D CellStreamingBoundsSize = FVector2D(Cell->GetStreamingBounds().GetSize());
+						const FVector2D CellStreamingBoundsMin = FVector2D(Cell->GetStreamingBounds().Min);
+
+						float CellOpacity = 0.0f;
+						TArray<FLinearColor> CellColors;
+
+						switch (GShowRuntimeHashSetDebugDisplayMode)
 						{
-							for (const FName& DataLayer : Cell->GetDataLayers())
+						case 0:
+							CellColors.Add(Cell->GetDebugColor(VisualizeMode));
+							CellOpacity = 0.25f / FMath::Max<float>(GShowRuntimeHashSetDebugDisplayLevelCount, 1);
+							break;
+						case 1:
+							if (DataLayerDebugColors.Num() && Cell->GetDataLayers().Num())
 							{
-								CellColors.Add(DataLayerDebugColors[DataLayer]);
+								for (const FName& DataLayer : Cell->GetDataLayers())
+								{
+									CellColors.Add(DataLayerDebugColors[DataLayer]);
+								}
+								CellOpacity = 0.67f;
 							}
-							CellOpacity = 0.67f;
-						}
-						break;
-					case 2:
-						if (ContentBundleManager && Cell->GetContentBundleID().IsValid())
-						{
-							if (const FContentBundleBase* ContentBundle = ContentBundleManager->GetContentBundle(OwningWorld, Cell->GetContentBundleID()))
+							break;
+						case 2:
+							if (ContentBundleManager && Cell->GetContentBundleID().IsValid())
 							{
-								check(ContentBundle->GetDescriptor());
-								CellColors.Add(ContentBundle->GetDescriptor()->GetDebugColor());
-								CellOpacity = 0.67f;
+								if (const FContentBundleBase* ContentBundle = ContentBundleManager->GetContentBundle(OwningWorld, Cell->GetContentBundleID()))
+								{
+									check(ContentBundle->GetDescriptor());
+									CellColors.Add(ContentBundle->GetDescriptor()->GetDebugColor());
+									CellOpacity = 0.67f;
+								}
 							}
+							break;
 						}
-						break;
-					}
 
-					if (CellColors.IsEmpty())
-					{
-						CellColors.Add(FLinearColor::White);
-						CellOpacity = 0.1f;
-					}
+						if (CellColors.IsEmpty())
+						{
+							CellColors.Add(FLinearColor::White);
+							CellOpacity = 0.1f;
+						}
 
-					// Draw Cell using its debug color
-					const FBox Box(Cell->GetStreamingBounds());
-					const FVector BoxCenter(Box.GetCenter());
-					const FColor BoxColor(CellColors[0].CopyWithNewOpacity(CellOpacity).ToFColor(true));
-					const FVector CellPos = WorldPartitionTransform.TransformPosition(BoxCenter);
-					DrawDebugBox(OwningWorld, CellPos, Box.GetExtent(), WorldPartitionTransform.GetRotation(), BoxColor.WithAlpha(255), false, -1.f, 255, 20.f);
+						// Draw Cell using its debug color
+						const FBox Box(Cell->GetStreamingBounds());
+						const FVector BoxCenter(Box.GetCenter());
+						const FColor BoxColor(CellColors[0].CopyWithNewOpacity(CellOpacity).ToFColor(true));
+						const FVector CellPos = WorldPartitionTransform.TransformPosition(BoxCenter);
+						DrawDebugBox(OwningWorld, CellPos, Box.GetExtent(), WorldPartitionTransform.GetRotation(), BoxColor.WithAlpha(255), false, -1.f, 255, 20.f);
+					});
 				});
-			});
+			}
 		}
 		return true;
 	});
