Enable COMPILE_DYNAMIC_FRAME_TIME on all platforms.
- disable r.DynamicRes.DynamicFrameTime by default and only enable on Windows to maintain same behaviour.

diff --git a/Engine/Config/Windows/WindowsEngine.ini b/Engine/Config/Windows/WindowsEngine.ini
--- a/Engine/Config/Windows/WindowsEngine.ini
+++ b/Engine/Config/Windows/WindowsEngine.ini
@@ -24,6 +24,7 @@ r.Nanite.Streaming.ReservedResources=1
 D3D12.Bindless.ResourceDescriptorHeapSize=100000
 D3D12.Bindless.SamplerDescriptorHeapSize=2048
 r.PSOPrecache.GlobalComputeShaders=1
+r.DynamicRes.DynamicFrameTime=1
 ; Allow ray tracing to be toggled using r.RayTracing.Enable=1/0 without restarting
 r.RayTracing.EnableOnDemand=True
 
diff --git a/Engine/Source/Runtime/Engine/Private/DynamicResolution.cpp b/Engine/Source/Runtime/Engine/Private/DynamicResolution.cpp
--- a/Engine/Source/Runtime/Engine/Private/DynamicResolution.cpp
+++ b/Engine/Source/Runtime/Engine/Private/DynamicResolution.cpp
@@ -10,12 +10,6 @@
 #include "SceneView.h"
 #include "Stats/StatsTrace.h"
 
-#if PLATFORM_DESKTOP
-	#define COMPILE_DYNAMIC_FRAME_TIME 1
-#else
-	#define COMPILE_DYNAMIC_FRAME_TIME 0
-#endif
-
 
 static TAutoConsoleVariable<float> CVarDynamicResMinSP(
 	TEXT("r.DynamicRes.MinScreenPercentage"),
@@ -49,10 +43,8 @@ static TAutoConsoleVariable<int32> CVarUseGameThreadCriticalPath(
 	TEXT("Whether to use game thread critical path time when determining whether game is CPU bound."),
 	ECVF_RenderThreadSafe | ECVF_Default);
 
-#if COMPILE_DYNAMIC_FRAME_TIME
-
 static TAutoConsoleVariable<int32> CVarDynamicFrameTimeEnable(
-	TEXT("r.DynamicRes.DynamicFrameTime"), 1,
+	TEXT("r.DynamicRes.DynamicFrameTime"), 0,
 	TEXT("Whether the r.DynamicRes.FrameTimeBudget should automatically increases when frame rate is bound by CPU."),
 	ECVF_RenderThreadSafe | ECVF_Default);
 
@@ -73,8 +65,6 @@ static TAutoConsoleVariable<float> CVarDynamicFrameTimeRoundUpToVsync(
 	TEXT("Error to use to round up the dynamic frame time to vsync boundaries (default=10%)."),
 	ECVF_RenderThreadSafe | ECVF_Default);
 
-#endif
-
 
 static TAutoConsoleVariable<float> CVarTargetedGPUHeadRoomPercentage(
 	TEXT("r.DynamicRes.TargetedGPUHeadRoomPercentage"),
@@ -284,7 +274,6 @@ void FDynamicResolutionHeuristicProxy::RefreshCurrentFrameResolutionFraction_Ren
 
 	const bool bCanChangeResolution = NumberOfFramesSinceScreenPercentageChange >= CVarFrameChangePeriod.GetValueOnRenderThread();
 
-#if COMPILE_DYNAMIC_FRAME_TIME
 	float MinGlobalFrameTime = 0.0f;
 	if (CVarDynamicFrameTimeEnable.GetValueOnRenderThread())
 	{
@@ -373,7 +362,6 @@ void FDynamicResolutionHeuristicProxy::RefreshCurrentFrameResolutionFraction_Ren
 			MinGlobalFrameTime = DynamicFrameTimeBudgetMs;
 		}
 	}
-#endif // COMPILE_DYNAMIC_FRAME_TIME
 
 	// New ResolutionFraction to use for this frame.
 	DynamicRenderScaling::TMap<float> NewFrameResolutionFractions = CurrentFrameResolutionFractions;
@@ -393,9 +381,7 @@ void FDynamicResolutionHeuristicProxy::RefreshCurrentFrameResolutionFraction_Ren
 		float BudgetBudgetMs = BudgetSettings.BudgetMs;
 		if (Budget == GDynamicPrimaryResolutionFraction)
 		{
-			#if COMPILE_DYNAMIC_FRAME_TIME
-				BudgetBudgetMs = FMath::Max(BudgetBudgetMs, MinGlobalFrameTime);
-			#endif
+			BudgetBudgetMs = FMath::Max(BudgetBudgetMs, MinGlobalFrameTime);
 		}
 		else if (!BudgetSettings.IsEnabled())
 		{
