Fix async post load RemoveLegacyClothingSections should not do a PreEditChange since it is always execute in a worker thread.

#rnx

diff --git a/Engine/Source/Runtime/Engine/Private/SkeletalMesh.cpp b/Engine/Source/Runtime/Engine/Private/SkeletalMesh.cpp
--- a/Engine/Source/Runtime/Engine/Private/SkeletalMesh.cpp
+++ b/Engine/Source/Runtime/Engine/Private/SkeletalMesh.cpp
@@ -2280,11 +2280,6 @@ void USkeletalMesh::RemoveLegacyClothingSections()
 						FSkelMeshSection& DuplicatedSection = LodModel.Sections[Section.CorrespondClothSectionIndex_DEPRECATED];
 
 						// Cache the base index for the first clothing section (will be in correct order)
-						if(ClothingSectionCount == 0)
-						{
-							PreEditChange(nullptr);
-						}
-						
 						BaseVertex = FMath::Min(DuplicatedSection.BaseVertexIndex, BaseVertex);
 						BaseIndex = FMath::Min(DuplicatedSection.BaseIndex, BaseIndex);
 
