[MemReport] Expose memreport metadata generation to the UEngine override interface, also expose game specific BugIt log for CheatManager
#tests local build

diff --git a/Engine/Source/Runtime/Engine/Classes/Engine/Engine.h b/Engine/Source/Runtime/Engine/Classes/Engine/Engine.h
--- a/Engine/Source/Runtime/Engine/Classes/Engine/Engine.h
+++ b/Engine/Source/Runtime/Engine/Classes/Engine/Engine.h
@@ -2306,6 +2306,7 @@ public:
 	ENGINE_API bool HandleListParticleSystemsCommand( const TCHAR* Cmd, FOutputDevice& Ar );
 	ENGINE_API bool HandleListSpawnedActorsCommand( const TCHAR* Cmd, FOutputDevice& Ar, UWorld* InWorld );
 	ENGINE_API bool HandleLogoutStatLevelsCommand( const TCHAR* Cmd, FOutputDevice& Ar, UWorld* InWorld );
+	ENGINE_API virtual void WriteMemReportMetadata( FOutputDevice& Ar, UWorld* InWorld );
 	ENGINE_API bool HandleMemReportCommand( const TCHAR* Cmd, FOutputDevice& Ar, UWorld* InWorld );
 	ENGINE_API bool HandleMemReportDeferredCommand( const TCHAR* Cmd, FOutputDevice& Ar, UWorld* InWorld );
 	ENGINE_API bool HandleSkeletalMeshReportCommand(const TCHAR* Cmd, FOutputDevice& Ar, UWorld* InWorld);
diff --git a/Engine/Source/Runtime/Engine/Private/UnrealEngine.cpp b/Engine/Source/Runtime/Engine/Private/UnrealEngine.cpp
--- a/Engine/Source/Runtime/Engine/Private/UnrealEngine.cpp
+++ b/Engine/Source/Runtime/Engine/Private/UnrealEngine.cpp
@@ -7465,6 +7465,47 @@ bool UEngine::HandleMemReportCommand( const TCHAR* Cmd, FOutputDevice& Ar, UWorl
 	return true;
 }
 
+void UEngine::WriteMemReportMetadata( FOutputDevice& Ar, UWorld* InWorld )
+{
+	FString ConfigString = LexToString(FApp::GetBuildConfiguration());
+	uint32 ChangelistNumber = FEngineVersion::Current().GetChangelist();
+
+	// Set the device name to the platform name as a fallback
+	FString DeviceName = FPlatformProperties::PlatformName();
+
+	// Attempt to get the specific device name from the runtime device profile selector
+	FString DeviceProfile = TEXT("None");
+	FString DeviceProfileSelectionModule;
+	if (GConfig->GetString(TEXT("DeviceProfileManager"), TEXT("DeviceProfileSelectionModule"), DeviceProfileSelectionModule, GEngineIni))
+	{
+		if (IDeviceProfileSelectorModule* DPSelectorModule = FModuleManager::LoadModulePtr<IDeviceProfileSelectorModule>(*DeviceProfileSelectionModule))
+		{
+			DeviceName = DPSelectorModule->GetRuntimeDeviceProfileName();
+		}
+	}
+
+	Ar.Logf(TEXT("Changelist: %d"), ChangelistNumber);
+	Ar.Logf(TEXT("Config: %s"), *ConfigString);
+	Ar.Logf(TEXT("Device Name: %s"), *DeviceName);
+	Ar.Logf(TEXT("Device Profile: %s"), *DeviceProfile);
+	Ar.Logf(TEXT("CommandLine Options: %s"), FCommandLine::Get());
+	Ar.Logf(TEXT("Time Since Boot: %.02f Seconds"), FPlatformTime::Seconds() - GStartTime);
+
+	// List Name, Location, and Rotation of each local PlayerController (useful context for client mem)
+	TArray<APlayerController*> LocalPlayerControllers;
+	GetAllLocalPlayerControllers(LocalPlayerControllers);
+	for (APlayerController* LocalPlayerController : LocalPlayerControllers)
+	{
+		FVector OutLocation;
+		FRotator OutRotation;
+		LocalPlayerController->GetPlayerViewPoint(OutLocation, OutRotation);
+		Ar.Logf(TEXT("Local PlayerController: %s View Location: %s View Rotation: %s"),
+			*LocalPlayerController->GetName(),
+			*OutLocation.ToString(),
+			*OutRotation.ToString());
+	}
+}
+
 bool UEngine::HandleMemReportDeferredCommand( const TCHAR* Cmd, FOutputDevice& Ar, UWorld* InWorld )
 {
 #if ALLOW_DEBUG_FILES
@@ -7522,44 +7563,7 @@ bool UEngine::HandleMemReportDeferredCommand( const TCHAR* Cmd, FOutputDevice& A
 
 	// log out some useful information in the header
 	{
-		FString ConfigString = LexToString(FApp::GetBuildConfiguration());
-		uint32 ChangelistNumber = FEngineVersion::Current().GetChangelist();
-
-		// Set the device name to the platform name as a fallback
-		FString DeviceName = FPlatformProperties::PlatformName();
-
-		// Attempt to get the specific device name from the runtime device profile selector
-		FString DeviceProfile = TEXT("None");
-		FString DeviceProfileSelectionModule;
-		if (GConfig->GetString(TEXT("DeviceProfileManager"), TEXT("DeviceProfileSelectionModule"), DeviceProfileSelectionModule, GEngineIni))
-		{
-			if (IDeviceProfileSelectorModule* DPSelectorModule = FModuleManager::LoadModulePtr<IDeviceProfileSelectorModule>(*DeviceProfileSelectionModule))
-			{
-				DeviceName = DPSelectorModule->GetRuntimeDeviceProfileName();
-			}
-		}
-
-		ReportAr->Logf(TEXT("Changelist: %d"), ChangelistNumber);
-		ReportAr->Logf(TEXT("Config: %s"), *ConfigString);
-		ReportAr->Logf(TEXT("Device Name: %s"), *DeviceName);
-		ReportAr->Logf(TEXT("Device Profile: %s"), *DeviceProfile);
-		ReportAr->Logf(TEXT("CommandLine Options: %s"), FCommandLine::Get());
-		ReportAr->Logf(TEXT("Time Since Boot: %.02f Seconds"), FPlatformTime::Seconds() - GStartTime);
-
-		// List Name, Location, and Rotation of each local PlayerController (useful context for client mem)
-		TArray<APlayerController*> LocalPlayerControllers;
-		GetAllLocalPlayerControllers(LocalPlayerControllers);
-		for (APlayerController* LocalPlayerController : LocalPlayerControllers)
-		{
-			FVector OutLocation;
-			FRotator OutRotation;
-			LocalPlayerController->GetPlayerViewPoint(OutLocation, OutRotation);
-			ReportAr->Logf(TEXT("Local PlayerController: %s View Location: %s View Rotation: %s"),
-				*LocalPlayerController->GetName(),
-				*OutLocation.ToString(),
-				*OutRotation.ToString());
-		}
-
+		WriteMemReportMetadata(*ReportAr, InWorld);
 		ReportAr->Logf(TEXT(""));
 	}
 
