[AssetManager] Improve performance of inserting elements into the preload list if the uniqueness is guaranteed by the caller

#tests local build

diff --git a/Engine/Source/Runtime/Engine/Classes/Engine/AssetManager.h b/Engine/Source/Runtime/Engine/Classes/Engine/AssetManager.h
--- a/Engine/Source/Runtime/Engine/Classes/Engine/AssetManager.h
+++ b/Engine/Source/Runtime/Engine/Classes/Engine/AssetManager.h
@@ -325,9 +325,10 @@ public:
 	 * @param PrimaryAssetId	Asset that would be loaded
 	 * @param LoadBundles		List of bundles to load for those assets
 	 * @param bLoadRecursive	If true, this will call RecursivelyExpandBundleData and recurse into sub bundles of other primary assets loaded by a bundle reference
+	 * @param bEnsureUniqueness Only unique items will be inserted in OutAssetLoadList
 	 * @return					True if primary asset id was found
 	 */
-	ENGINE_API bool GetPrimaryAssetLoadList(TArray<FSoftObjectPath>& OutAssetLoadList, const FPrimaryAssetId& PrimaryAssetId, const TArray<FName>& LoadBundles, bool bLoadRecursive) const;
+	ENGINE_API bool GetPrimaryAssetLoadList(TArray<FSoftObjectPath>& OutAssetLoadList, const FPrimaryAssetId& PrimaryAssetId, const TArray<FName>& LoadBundles, bool bLoadRecursive, bool bEnsureUniqueness = true) const;
 
 	/**
 	 * Preloads data for a set of assets in a specific bundle state, and returns a handle you must keep active.
diff --git a/Engine/Source/Runtime/Engine/Private/AssetManager.cpp b/Engine/Source/Runtime/Engine/Private/AssetManager.cpp
--- a/Engine/Source/Runtime/Engine/Private/AssetManager.cpp
+++ b/Engine/Source/Runtime/Engine/Private/AssetManager.cpp
@@ -1908,12 +1908,12 @@ TSharedPtr<FStreamableHandle> UAssetManager::ChangeBundleStateForMatchingPrimary
 bool UAssetManager::GetPrimaryAssetLoadSet(TSet<FSoftObjectPath>& OutAssetLoadSet, const FPrimaryAssetId& PrimaryAssetId, const TArray<FName>& LoadBundles, bool bLoadRecursive) const
 {
 	TArray<FSoftObjectPath> Array;
-	bool bReturnValue = GetPrimaryAssetLoadList(Array, PrimaryAssetId, LoadBundles, bLoadRecursive);
+	bool bReturnValue = GetPrimaryAssetLoadList(Array, PrimaryAssetId, LoadBundles, bLoadRecursive, false);
 	OutAssetLoadSet.Append(Array);
 	return bReturnValue;
 }
 
-bool UAssetManager::GetPrimaryAssetLoadList(TArray<FSoftObjectPath>& OutAssetLoadList, const FPrimaryAssetId& PrimaryAssetId, const TArray<FName>& LoadBundles, bool bLoadRecursive) const
+bool UAssetManager::GetPrimaryAssetLoadList(TArray<FSoftObjectPath>& OutAssetLoadList, const FPrimaryAssetId& PrimaryAssetId, const TArray<FName>& LoadBundles, bool bLoadRecursive, bool bEnsureUniqueness) const
 {
 	const FPrimaryAssetData* NameData = GetNameData(PrimaryAssetId);
 	if (NameData)
@@ -1923,7 +1923,14 @@ bool UAssetManager::GetPrimaryAssetLoadList(TArray<FSoftObjectPath>& OutAssetLoa
 		if (!AssetPath.IsNull())
 		{
 			// Dynamic types can have no base asset path
-			OutAssetLoadList.AddUnique(AssetPath);
+			if (bEnsureUniqueness)
+			{
+				OutAssetLoadList.AddUnique(AssetPath);
+			}
+			else
+			{
+				OutAssetLoadList.Add(AssetPath);
+			}
 		}
 
 		// Construct a temporary bundle data with the bundles specified
@@ -1947,7 +1954,14 @@ bool UAssetManager::GetPrimaryAssetLoadList(TArray<FSoftObjectPath>& OutAssetLoa
 		{
 			for (const FTopLevelAssetPath& Path : Entry.AssetPaths)
 			{
-				OutAssetLoadList.AddUnique(FSoftObjectPath(Path));
+				if (bEnsureUniqueness)
+				{
+					OutAssetLoadList.AddUnique(FSoftObjectPath(Path));
+				}
+				else
+				{
+					OutAssetLoadList.Add(FSoftObjectPath(Path));
+				}
 			}
 		}
 	}
@@ -1970,7 +1984,7 @@ TSharedPtr<FStreamableHandle> UAssetManager::PreloadPrimaryAssets(const TArray<F
 
 	for (const FPrimaryAssetId& PrimaryAssetId : AssetsToLoad)
 	{
-		if (GetPrimaryAssetLoadList(PathsToLoad, PrimaryAssetId, LoadBundles, bLoadRecursive))
+		if (GetPrimaryAssetLoadList(PathsToLoad, PrimaryAssetId, LoadBundles, bLoadRecursive, false))
 		{
 			if (DebugValid.Len() < MaxDebugWarningLen)
 			{
