Add option to ignore adding missing columns when importing CSV into a data table

diff --git a/Engine/Source/Editor/UnrealEd/Classes/Factories/CSVImportFactory.h b/Engine/Source/Editor/UnrealEd/Classes/Factories/CSVImportFactory.h
--- a/Engine/Source/Editor/UnrealEd/Classes/Factories/CSVImportFactory.h
+++ b/Engine/Source/Editor/UnrealEd/Classes/Factories/CSVImportFactory.h
@@ -94,6 +94,9 @@ private:
 
 	bool bImportAll = false;
 
+	// Function to Copy DataTable row to avoid dangling pointers
+	TUniquePtr<uint8[]> CopyRow(const uint8* RowData, const UScriptStruct* RowStruct) const;
+
 public:
 	UPROPERTY(BlueprintReadWrite, Category="Automation")
 	FCSVImportSettings AutomatedImportSettings;
diff --git a/Engine/Source/Editor/UnrealEd/Private/Factories/CSVImportFactory.cpp b/Engine/Source/Editor/UnrealEd/Private/Factories/CSVImportFactory.cpp
--- a/Engine/Source/Editor/UnrealEd/Private/Factories/CSVImportFactory.cpp
+++ b/Engine/Source/Editor/UnrealEd/Private/Factories/CSVImportFactory.cpp
@@ -277,6 +277,8 @@ UObject* UCSVImportFactory::FactoryCreateText(UClass* InClass, UObject* InParent
 			{
 				TempImportDataTable->CopyImportOptions(DataTableImportOptions);
 			}
+			
+			TMap<FName, TUniquePtr<uint8[]>> TempRowMap;
 
 			// If there is an existing table, need to call this to free data memory before recreating object
 			UDataTable::FOnDataTableChanged OldOnDataTableChanged;
@@ -287,6 +289,19 @@ UObject* UCSVImportFactory::FactoryCreateText(UClass* InClass, UObject* InParent
 				OldOnDataTableChanged = MoveTemp(ExistingTable->OnDataTableChanged());
 				ExistingTable->OnDataTableChanged().Clear();
 				DataTableClass = ExistingTable->GetClass();
+				
+				TMap<FName, uint8*> ExistingTableRowMap;
+				ExistingTableRowMap = ExistingTable->GetRowMap();
+
+				if (ExistingTable->bPreserveExistingValues)
+				{
+					// Copy each existing row to the temporary map
+					for (const TPair<FName, uint8*>& TableRow : ExistingTableRowMap)
+					{
+						TempRowMap.Add(TableRow.Key, CopyRow(TableRow.Value, ExistingTable->RowStruct));
+					}
+				}
+				
 				ExistingTable->EmptyTable();
 			}
 
@@ -298,7 +313,19 @@ UObject* UCSVImportFactory::FactoryCreateText(UClass* InClass, UObject* InParent
 			if (!CurrentFilename.IsEmpty())
 			{
 				NewTable->AssetImportData->Update(CurrentFilename);
+			}
 
+			if (NewTable->bPreserveExistingValues && ExistingTable)
+			{
+				ensure(NewTable->RowStruct == ExistingTable->RowStruct);
+				
+				for (TPair<FName, TUniquePtr<uint8[]>>& Row : TempRowMap)
+				{
+					if (const FTableRowBase* RowBase = reinterpret_cast<FTableRowBase*>(Row.Value.Get()))
+					{
+						NewTable->AddRow(Row.Key, *RowBase);
+					}
+				}
 			}
 
 			// Go ahead and create table from string
@@ -429,6 +456,15 @@ EReimportResult::Type UCSVImportFactory::Reimport(UObject* Obj, const FString& P
 	return EReimportResult::Failed;
 }
 
+TUniquePtr<uint8[]> UCSVImportFactory::CopyRow(const uint8* RowData, const UScriptStruct* RowStruct) const
+{
+	const int32 RowSize = RowStruct->GetStructureSize();
+	TUniquePtr<uint8[]> NewRowData = MakeUnique<uint8[]>(RowSize);
+	RowStruct->InitializeStruct(NewRowData.Get());
+	RowStruct->CopyScriptStruct(NewRowData.Get(), RowData);
+	return NewRowData;
+}
+
 TArray<FString> UCSVImportFactory::DoImportDataTable(const FCSVImportSettings& InImportSettings, UDataTable* TargetDataTable)
 {
 	if (InImportSettings.bDataIsJson)
diff --git a/Engine/Source/Runtime/Engine/Classes/Engine/DataTable.h b/Engine/Source/Runtime/Engine/Classes/Engine/DataTable.h
--- a/Engine/Source/Runtime/Engine/Classes/Engine/DataTable.h
+++ b/Engine/Source/Runtime/Engine/Classes/Engine/DataTable.h
@@ -118,6 +118,10 @@ public:
 	UPROPERTY(EditAnywhere, Category = ImportOptions)
 	uint8 bIgnoreMissingFields : 1;
 
+	/** Set to true to preserve existing values for any fields that are expected but missing in the CSV file. If false, missing fields will be populated with default values. */
+	UPROPERTY(EditAnywhere, Category = ImportOptions)
+	uint8 bPreserveExistingValues : 1;
+
 	/** Explicit field in import data to use as key. If this is empty it uses Name for JSON and the first field found for CSV */
 	UPROPERTY(EditAnywhere, Category=ImportOptions)
 	FString ImportKeyField;
diff --git a/Engine/Source/Runtime/Engine/Private/DataTable.cpp b/Engine/Source/Runtime/Engine/Private/DataTable.cpp
--- a/Engine/Source/Runtime/Engine/Private/DataTable.cpp
+++ b/Engine/Source/Runtime/Engine/Private/DataTable.cpp
@@ -73,6 +73,7 @@ UDataTable::UDataTable(const FObjectInitializer& ObjectInitializer)
 	bIgnoreExtraFields = false;
 	bIgnoreMissingFields = false;
 	bStripFromClientBuilds = false;
+	bPreserveExistingValues = false;
 
 #if WITH_EDITORONLY_DATA
 	{ static const FAutoRegisterLocalizationDataGatheringCallback AutomaticRegistrationOfLocalizationGatherer(UDataTable::StaticClass(), &GatherDataTableForLocalization); }
@@ -656,6 +657,7 @@ bool UDataTable::CopyImportOptions(UDataTable* SourceTable)
 	bStripFromClientBuilds = SourceTable->bStripFromClientBuilds;
 	bIgnoreExtraFields = SourceTable->bIgnoreExtraFields;
 	bIgnoreMissingFields = SourceTable->bIgnoreMissingFields;
+	bPreserveExistingValues = SourceTable->bPreserveExistingValues;
 	ImportKeyField = SourceTable->ImportKeyField;
 	RowStruct = SourceTable->RowStruct;
 
diff --git a/Engine/Source/Runtime/Engine/Private/DataTableCSV.cpp b/Engine/Source/Runtime/Engine/Private/DataTableCSV.cpp
--- a/Engine/Source/Runtime/Engine/Private/DataTableCSV.cpp
+++ b/Engine/Source/Runtime/Engine/Private/DataTableCSV.cpp
@@ -165,9 +165,14 @@ bool FDataTableImporterCSV::ReadTable()
 	
 	TArray<FProperty*> ColumnProps = DataTable->GetTablePropertyArray(Rows[0], DataTable->RowStruct, ImportProblems, KeyColumn);
 
-	// Empty existing data
-	DataTable->EmptyTable();
+	// Empty existing data if don't want to keep existing values
+	if (!DataTable->bPreserveExistingValues)
+	{
+		DataTable->EmptyTable();
+	}
 
+	TSet<FName> NewlyAddedRows;
+	
 	// Iterate over rows
 	for(int32 RowIdx=1; RowIdx<Rows.Num(); RowIdx++)
 	{
@@ -199,38 +204,53 @@ bool FDataTableImporterCSV::ReadTable()
 			}
 			else
 			{
-			ImportProblems.Add(FString::Printf(TEXT("Row '%d' missing a name."), RowIdx));
+				ImportProblems.Add(FString::Printf(TEXT("Row '%d' missing a name."), RowIdx));
 			}
 
 			continue;
 		}
-
-		// Check its not a duplicate
-		if(!DataTable->AllowDuplicateRowsOnImport() && DataTable->GetRowMap().Find(RowName) != nullptr)
+		
+		uint8* ExistingRowData = DataTable->bPreserveExistingValues ? DataTable->FindRowUnchecked(RowName) : nullptr;
+		
+		if (!DataTable->AllowDuplicateRowsOnImport())
 		{
-			ImportProblems.Add(FString::Printf(TEXT("Duplicate row name '%s'."), *RowName.ToString()));
-			continue;
+			bool bRowNameAlreadyAdded;
+			NewlyAddedRows.Add(RowName, &bRowNameAlreadyAdded);
+			if (bRowNameAlreadyAdded)
+			{
+				ImportProblems.Add(FString::Printf(TEXT("Duplicate row name '%s'."), *RowName.ToString()));
+				continue;
+			}
 		}
-
+		
 		// Allocate data to store information, using UScriptStruct to know its size
-		uint8* RowData = (uint8*)FMemory::Malloc(DataTable->RowStruct->GetStructureSize());
-		DataTable->RowStruct->InitializeStruct(RowData);
-		// And be sure to call DestroyScriptStruct later
-
-		// Add to row map
-		DataTable->AddRowInternal(RowName, RowData);
-
+		uint8* RowData = ExistingRowData ? ExistingRowData : (uint8*)FMemory::Malloc(DataTable->RowStruct->GetStructureSize());
+		
+		// Add to row map if we don't already have existing row because we are preserve existing values
+		if (!ExistingRowData)
+		{
+			DataTable->RowStruct->InitializeStruct(RowData);
+			DataTable->AddRowInternal(RowName, RowData);
+		}
+		
 		// Now iterate over cells (skipping first cell unless we had an explicit name)
-		for(int32 CellIdx = 0; CellIdx < Cells.Num(); CellIdx++)
+		for (int32 CellIdx = 0; CellIdx < Cells.Num(); CellIdx++)
 		{
 			if (CellIdx == KeyColumn)
-		{
+			{
 				continue;
 			}
 
 			// Try and assign string to data using the column property
 			FProperty* ColumnProp = ColumnProps[CellIdx];
 			const FString CellValue = Cells[CellIdx];
+
+			// If we have no value in the csv for this specific entry then do nothing
+			if (DataTable->bPreserveExistingValues && CellValue.IsEmpty())
+			{
+				continue;
+			}
+			
 			FString Error = DataTableUtils::AssignStringToProperty(CellValue, ColumnProp, RowData);
 
 			// If we failed, output a problem string
