Fixes cases where DelayUntilNextTick and SetTimerForNextTick wouldn't always wait until the next world tick before triggering. In particular, within Gameplay Abilities that chain delays off of abiliity task node out-exec pins.  This is gated by cvars and defaulted to off.

diff --git a/Engine/Source/Runtime/Engine/Classes/Engine/LatentActionManager.h b/Engine/Source/Runtime/Engine/Classes/Engine/LatentActionManager.h
--- a/Engine/Source/Runtime/Engine/Classes/Engine/LatentActionManager.h
+++ b/Engine/Source/Runtime/Engine/Classes/Engine/LatentActionManager.h
@@ -8,6 +8,11 @@
 
 class FPendingLatentAction;
 
+namespace LatentActionCVars
+{
+	extern int32 GuaranteeEngineTickDelay;	// 0: next-tick delays may run at end of current frame (default behavior pre-5.5),  1: next-tick delays will always wait until the engine frame has advanced
+}
+
 // Latent action info
 USTRUCT(BlueprintInternalUseOnly)
 struct FLatentActionInfo
diff --git a/Engine/Source/Runtime/Engine/Private/LatentActionManager.cpp b/Engine/Source/Runtime/Engine/Private/LatentActionManager.cpp
--- a/Engine/Source/Runtime/Engine/Private/LatentActionManager.cpp
+++ b/Engine/Source/Runtime/Engine/Private/LatentActionManager.cpp
@@ -3,6 +3,7 @@
 #include "Engine/LatentActionManager.h"
 #include "LatentActions.h"
 #include "Stats/Stats.h"
+#include "HAL/ConsoleManager.h"
 
 #include UE_INLINE_GENERATED_CPP_BY_NAME(LatentActionManager)
 
@@ -10,13 +11,19 @@
 #define LATENT_ACTION_PROFILING_ENABLED 0
 #endif
 
-#if LATENT_ACTION_PROFILING_ENABLED
-#include "HAL/IConsoleManager.h"
-#endif
 
 FOnLatentActionsChanged FLatentActionManager::LatentActionsChangedDelegate;
 
 
+namespace LatentActionCVars
+{
+	int32 GuaranteeEngineTickDelay = 0;
+	static FAutoConsoleVariableRef CVarLatentActionGuaranteeEngineTickDelay(
+		TEXT("LatentActions.GuaranteeNextTickDelay"),
+		GuaranteeEngineTickDelay,
+		TEXT("If true, latent actions delayed until next tick will guarantee the engine frame has advanced. If false, these would always run at the end of the same engine tick (default behavior prior to 5.5)."));
+}
+
 /////////////////////////////////////////////////////
 // FPendingLatentAction
 
diff --git a/Engine/Source/Runtime/Engine/Private/TimerManager.cpp b/Engine/Source/Runtime/Engine/Private/TimerManager.cpp
--- a/Engine/Source/Runtime/Engine/Private/TimerManager.cpp
+++ b/Engine/Source/Runtime/Engine/Private/TimerManager.cpp
@@ -49,6 +49,12 @@ static FAutoConsoleVariableRef CVarMaxExpiredTimersToLog(
 	MaxExpiredTimersToLog,
 	TEXT("Maximum number of TimerData exceeding the threshold to log in a single frame."));
 
+static int32 GuaranteeEngineTickDelay = 0;
+static FAutoConsoleVariableRef CVarTimerManagerGuaranteeEngineTickDelay(
+	TEXT("TimerManager.GuaranteeEngineTickDelay"),
+	GuaranteeEngineTickDelay,
+	TEXT("If true, timers delayed until next tick will guarantee the engine tick has advanced. If false, these could run during the same engine tick (default behavior prior to 5.5)."));
+
 #ifndef UE_ENABLE_DUMPALLTIMERLOGSTHRESHOLD
 #define UE_ENABLE_DUMPALLTIMERLOGSTHRESHOLD !UE_BUILD_SHIPPING
 #endif
@@ -538,7 +544,6 @@ FTimerHandle FTimerManager::InternalSetTimerForNextTick(FTimerUnifiedDelegate&&
 	NewTimerData.bRequiresDelegate = true;
 	NewTimerData.TimerDelegate = MoveTemp(InDelegate);
 	NewTimerData.ExpireTime = InternalTime;
-	NewTimerData.Status = ETimerStatus::Active;
 
 	// Set level collection
 	const UWorld* const OwningWorld = OwningGameInstance ? OwningGameInstance->GetWorld() : nullptr;
@@ -547,8 +552,23 @@ FTimerHandle FTimerManager::InternalSetTimerForNextTick(FTimerUnifiedDelegate&&
 		NewTimerData.LevelCollection = OwningWorld->GetActiveLevelCollection()->GetType();
 	}
 
-	FTimerHandle NewTimerHandle = AddTimer(MoveTemp(NewTimerData));
-	ActiveTimerHeap.HeapPush(NewTimerHandle, FTimerHeapOrder(Timers));
+	// Add to pending timer heap to guarantee it happens next frame if we want to delay
+	const bool bQueueForCurrentFrame = (GuaranteeEngineTickDelay == 0) || HasBeenTickedThisFrame();
+	
+	FTimerHandle NewTimerHandle;
+	if (bQueueForCurrentFrame)
+	{
+		NewTimerData.Status = ETimerStatus::Active;
+		NewTimerHandle = AddTimer(MoveTemp(NewTimerData));
+		ActiveTimerHeap.HeapPush(NewTimerHandle, FTimerHeapOrder(Timers));
+	}
+	else
+	{
+		NewTimerData.Status = ETimerStatus::Pending;
+		NewTimerData.ExpireTime = 0.0;	// Pending timers' ExpireTime means time remaining. It will be converted upon activation.
+		NewTimerHandle = AddTimer(MoveTemp(NewTimerData));
+		PendingTimerSet.Add(NewTimerHandle);
+	}
 
 	return NewTimerHandle;
 }
diff --git a/Engine/Source/Runtime/Engine/Public/DelayAction.h b/Engine/Source/Runtime/Engine/Public/DelayAction.h
--- a/Engine/Source/Runtime/Engine/Public/DelayAction.h
+++ b/Engine/Source/Runtime/Engine/Public/DelayAction.h
@@ -44,16 +44,18 @@ public:
 
 
 // FDelayUntilNextTickAction
-// A simple delay action; triggers on the next update.
+// A simple delay action; triggers on the next engine tick. See also @CVarLatentActionGuaranteeEngineTickDelay
 class FDelayUntilNextTickAction : public FPendingLatentAction
 {
 public:
+	uint8 InitialFrameParity; // engine frame parity when this action was created (0=even, 1=odd)
 	FName ExecutionFunction;
 	int32 OutputLink;
 	FWeakObjectPtr CallbackTarget;
 
 	FDelayUntilNextTickAction(const FLatentActionInfo& LatentInfo)
-		: ExecutionFunction(LatentInfo.ExecutionFunction)
+		: InitialFrameParity((uint8)GFrameCounter & 1)
+		, ExecutionFunction(LatentInfo.ExecutionFunction)
 		, OutputLink(LatentInfo.Linkage)
 		, CallbackTarget(LatentInfo.CallbackTarget)
 	{
@@ -61,7 +63,8 @@ public:
 
 	virtual void UpdateOperation(FLatentResponse& Response) override
 	{
-		Response.FinishAndTriggerIf(true, ExecutionFunction, OutputLink, CallbackTarget);
+		const bool bShouldFinish = !LatentActionCVars::GuaranteeEngineTickDelay || (((uint8)GFrameCounter & 1) != InitialFrameParity);
+		Response.FinishAndTriggerIf(bShouldFinish, ExecutionFunction, OutputLink, CallbackTarget);
 	}
 
 #if WITH_EDITOR
