Tidy up the spot light fallback test in a rect light component

diff --git a/Engine/Source/Runtime/Engine/Private/Components/RectLightComponent.cpp b/Engine/Source/Runtime/Engine/Private/Components/RectLightComponent.cpp
--- a/Engine/Source/Runtime/Engine/Private/Components/RectLightComponent.cpp
+++ b/Engine/Source/Runtime/Engine/Private/Components/RectLightComponent.cpp
@@ -201,6 +201,16 @@ void URectLightComponent::PostEditChangeProperty(FPropertyChangedEvent& Property
 }
 #endif // WITH_EDITOR
 
+static bool RenderRectLightsAsSpotLights(FStaticFeatureLevel FeatureLevel)
+{
+	if (FeatureLevel == ERHIFeatureLevel::ES3_1)
+	{
+		static const auto CVarRenderRectLightAsSpotLight = IConsoleManager::Get().FindTConsoleVariableDataInt(TEXT("r.Mobile.Forward.RenderRectLightsAsSpotLights"));
+		return (CVarRenderRectLightAsSpotLight && CVarRenderRectLightAsSpotLight->GetValueOnAnyThread() != 0);
+	}
+	return false;
+}
+
 FRectLightSceneProxy::FRectLightSceneProxy(const URectLightComponent* Component)
 	: FLocalLightSceneProxy(Component)
 	, SourceWidth(Component->SourceWidth)
@@ -211,6 +221,7 @@ FRectLightSceneProxy::FRectLightSceneProxy(const URectLightComponent* Component)
 	, SourceTexture(Component->SourceTexture)
 {
 	RectAtlasId = ~0u;
+	bRenderAsSpotLight = RenderRectLightsAsSpotLights(SceneInterface->GetFeatureLevel());
 }
 
 FRectLightSceneProxy::~FRectLightSceneProxy() {}
@@ -263,10 +274,7 @@ void FRectLightSceneProxy::GetLightShaderParameters(FLightRenderParameters& Ligh
 	}
 	
 	// Render RectLight approximately as SpotLight if the requester does not support rect light (e.g., translucent light grid or mobile)
-	static const auto CVarRenderRectLightAsSpotLight = IConsoleManager::Get().FindTConsoleVariableDataInt(TEXT("r.Mobile.Forward.RenderRectLightsAsSpotLights"));
-	const bool bRenderAsSpotLightCVar = (CVarRenderRectLightAsSpotLight && CVarRenderRectLightAsSpotLight->GetValueOnAnyThread() != 0);
-	const bool bRenderAsSpotLight = !!(Flags & ELightShaderParameterFlags::RectAsSpotLight) || (SceneInterface && IsMobilePlatform(SceneInterface->GetShaderPlatform()) && bRenderAsSpotLightCVar);
-	if (bRenderAsSpotLight)
+	if (!!(Flags & ELightShaderParameterFlags::RectAsSpotLight) || bRenderAsSpotLight)
 	{
 		float ClampedOuterConeAngle = FMath::DegreesToRadians(89.001f);
 		float ClampedInnerConeAngle = FMath::DegreesToRadians(70.0f);
diff --git a/Engine/Source/Runtime/Engine/Public/RectLightSceneProxy.h b/Engine/Source/Runtime/Engine/Public/RectLightSceneProxy.h
--- a/Engine/Source/Runtime/Engine/Public/RectLightSceneProxy.h
+++ b/Engine/Source/Runtime/Engine/Public/RectLightSceneProxy.h
@@ -31,6 +31,7 @@ public:
 	FRectLightRayTracingData* RayTracingData; // Render thread data only
 	UTexture*	SourceTexture;
 	uint32		RectAtlasId;
+	bool		bRenderAsSpotLight;
 
 	FRectLightSceneProxy(const URectLightComponent* Component);
 	virtual ~FRectLightSceneProxy();
