[Backout] - CL34089817, 34114150, 34114248, and 34115050
Original CL Desc
-----------------------------------------------------------------
Add rect light support to mobile forward renderer

diff --git a/Engine/Shaders/Private/DeferredLightingCommon.ush b/Engine/Shaders/Private/DeferredLightingCommon.ush
--- a/Engine/Shaders/Private/DeferredLightingCommon.ush
+++ b/Engine/Shaders/Private/DeferredLightingCommon.ush
@@ -260,10 +260,7 @@ float GetLocalLightAttenuation(
 		// This extra attenuation has been added for supporting existing 'legacy' shading model on mobile. 
 		// This is not needed for Substrate which unifies all lighting paths
 		#if SHADING_PATH_MOBILE && !SUBSTRATE_ENABLED
-		if (!LightData.bRectLight)
-		{ 
-			LightMask *= 1.0f / (DistanceSqr + 1.0f);
-		}
+		LightMask *= 1.0f / (DistanceSqr + 1.0f);
 		#endif
 	}
 	else
@@ -367,22 +364,7 @@ FLightAccumulator AccumulateDynamicLighting(
 	#if TRANSLUCENCY_NON_DIRECTIONAL
 		NoL = 1.0f;
 	#endif
-		BRANCH
-		if (LightData.bRectLight)
-		{
-			FRect Rect = GetRect( ToLight, LightData );
-			const FRectTexture SourceTexture = ConvertToRectTexture(LightData);
-
-		#if REFERENCE_QUALITY
-			Lighting = IntegrateBxDF( GBuffer, N, V, Rect, Shadow, SourceTexture, SVPos );
-		#else
-			Lighting = IntegrateBxDF( GBuffer, N, V, Rect, Shadow, SourceTexture);
-		#endif		
-		}
-		else
-		{
-			Lighting = EvaluateBxDF(GBuffer, N, V, L, NoL, Shadow);
-		}
+		Lighting = EvaluateBxDF(GBuffer, N, V, L, NoL, Shadow);
 
 		Lighting.Specular *= LightData.SpecularScale;
 				
diff --git a/Engine/Shaders/Private/LightData.ush b/Engine/Shaders/Private/LightData.ush
--- a/Engine/Shaders/Private/LightData.ush
+++ b/Engine/Shaders/Private/LightData.ush
@@ -377,23 +377,16 @@ FDeferredLightData ConvertToDeferredLight_Mobile(const FLocalLightData In)
 	half4 OutPreviewShadowMapChannelMask = 0;
 	uint OutLightingChannelMask = LIGHTING_CHANNEL_MASK;
 	FDeferredLightData Out = ConvertToDeferredLight(In, 1.0f, OutPreviewShadowMapChannelMask, OutLightingChannelMask);	
-
+	
 	// Override value for performance on mobile
+	Out.SourceRadius	= 0;
+	Out.SourceLength	= 0;
+	Out.Tangent			= 0;
 	Out.SoftSourceRadius= 0;
 	Out.SpecularScale	= 1.0f;
-#if RECT_LIGHT_AS_SPOTLIGHT
-	// Render RectLight approximately as SpotLight
-	Out.bSpotLight	= Out.bSpotLight || Out.bRectLight;
-	Out.bRectLight	= false;
-#endif
-
-	BRANCH
-	if (!Out.bRectLight)
-	{
-		Out.SourceRadius = 0;
-		Out.SourceLength = 0;
-		Out.Tangent = 0;
-	}
+	//Render RectLight approximately as SpotLight
+	Out.bSpotLight		= Out.bSpotLight || Out.bRectLight;
+	Out.bRectLight		= false;
 
 	return Out;
 }
diff --git a/Engine/Source/Runtime/Core/Private/HAL/ConsoleManager.cpp b/Engine/Source/Runtime/Core/Private/HAL/ConsoleManager.cpp
--- a/Engine/Source/Runtime/Core/Private/HAL/ConsoleManager.cpp
+++ b/Engine/Source/Runtime/Core/Private/HAL/ConsoleManager.cpp
@@ -3772,13 +3772,6 @@ static TAutoConsoleVariable<int32> CVarMobileForwardEnableLocalLights(
 		"2: Local Lights Buffer Enabled\n"),
 	ECVF_ReadOnly | ECVF_RenderThreadSafe);
 
-static TAutoConsoleVariable<int32> CVarMobileRenderRectLightsAsSpotLights(
-	TEXT("r.Mobile.Forward.RenderRectLightsAsSpotLights"),
-	0,
-	TEXT("0: Rect Lights are rendered as area lights (default)\n"
-		"1: Rect Lights are rendered as Spot Lights\n"),
-	ECVF_ReadOnly | ECVF_MobileShaderChange);
-
 static TAutoConsoleVariable<int32> CVarMobileForwardEnableParticleLights(
 	TEXT("r.Mobile.Forward.EnableParticleLights"),
 	0,
diff --git a/Engine/Source/Runtime/Engine/Private/Components/RectLightComponent.cpp b/Engine/Source/Runtime/Engine/Private/Components/RectLightComponent.cpp
--- a/Engine/Source/Runtime/Engine/Private/Components/RectLightComponent.cpp
+++ b/Engine/Source/Runtime/Engine/Private/Components/RectLightComponent.cpp
@@ -201,16 +201,6 @@ void URectLightComponent::PostEditChangeProperty(FPropertyChangedEvent& Property
 }
 #endif // WITH_EDITOR
 
-static bool RenderRectLightsAsSpotLights(FStaticFeatureLevel FeatureLevel)
-{
-	if (FeatureLevel == ERHIFeatureLevel::ES3_1)
-	{
-		static const auto CVarRenderRectLightAsSpotLight = IConsoleManager::Get().FindTConsoleVariableDataInt(TEXT("r.Mobile.Forward.RenderRectLightsAsSpotLights"));
-		return (CVarRenderRectLightAsSpotLight && CVarRenderRectLightAsSpotLight->GetValueOnAnyThread() != 0);
-	}
-	return false;
-}
-
 FRectLightSceneProxy::FRectLightSceneProxy(const URectLightComponent* Component)
 	: FLocalLightSceneProxy(Component)
 	, SourceWidth(Component->SourceWidth)
@@ -221,7 +211,6 @@ FRectLightSceneProxy::FRectLightSceneProxy(const URectLightComponent* Component)
 	, SourceTexture(Component->SourceTexture)
 {
 	RectAtlasId = ~0u;
-	bRenderAsSpotLight = RenderRectLightsAsSpotLights(SceneInterface->GetFeatureLevel());
 }
 
 FRectLightSceneProxy::~FRectLightSceneProxy() {}
@@ -274,7 +263,8 @@ void FRectLightSceneProxy::GetLightShaderParameters(FLightRenderParameters& Ligh
 	}
 	
 	// Render RectLight approximately as SpotLight if the requester does not support rect light (e.g., translucent light grid or mobile)
-	if (!!(Flags & ELightShaderParameterFlags::RectAsSpotLight) || bRenderAsSpotLight)
+	const bool bRenderAsSpotLight = !!(Flags & ELightShaderParameterFlags::RectAsSpotLight) || (SceneInterface && IsMobilePlatform(SceneInterface->GetShaderPlatform()));
+	if (bRenderAsSpotLight)
 	{
 		float ClampedOuterConeAngle = FMath::DegreesToRadians(89.001f);
 		float ClampedInnerConeAngle = FMath::DegreesToRadians(70.0f);
diff --git a/Engine/Source/Runtime/Engine/Private/ShaderCompiler/ShaderCompiler.cpp b/Engine/Source/Runtime/Engine/Private/ShaderCompiler/ShaderCompiler.cpp
--- a/Engine/Source/Runtime/Engine/Private/ShaderCompiler/ShaderCompiler.cpp
+++ b/Engine/Source/Runtime/Engine/Private/ShaderCompiler/ShaderCompiler.cpp
@@ -7806,12 +7806,6 @@ void GlobalBeginCompileShader(
 			SET_SHADER_DEFINE(Input.Environment, USE_GLES_FBF_DEFERRED, bGLESDeferredShading ? 1 : 0);
 			SET_SHADER_DEFINE(Input.Environment, MOBILE_EXTENDED_GBUFFER, MobileUsesExtenedGBuffer((EShaderPlatform)Target.Platform) ? 1 : 0);
 		}
-		else
-		{
-			static const auto CVarRenderRectLightAsSpotLight = IConsoleManager::Get().FindTConsoleVariableDataInt(TEXT("r.Mobile.Forward.RenderRectLightsAsSpotLights"));
-			const bool bRenderAsSpotLight = (CVarRenderRectLightAsSpotLight && CVarRenderRectLightAsSpotLight->GetValueOnAnyThread() != 0);
-			SET_SHADER_DEFINE(Input.Environment, RECT_LIGHT_AS_SPOTLIGHT, bRenderAsSpotLight ? 1 : 0);
-		}
 
 		SET_SHADER_DEFINE(Input.Environment, USE_SCENE_DEPTH_AUX, MobileRequiresSceneDepthAux(ShaderPlatform) ? 1 : 0);
 
diff --git a/Engine/Source/Runtime/Engine/Public/RectLightSceneProxy.h b/Engine/Source/Runtime/Engine/Public/RectLightSceneProxy.h
--- a/Engine/Source/Runtime/Engine/Public/RectLightSceneProxy.h
+++ b/Engine/Source/Runtime/Engine/Public/RectLightSceneProxy.h
@@ -31,7 +31,6 @@ public:
 	FRectLightRayTracingData* RayTracingData; // Render thread data only
 	UTexture*	SourceTexture;
 	uint32		RectAtlasId;
-	bool		bRenderAsSpotLight;
 
 	FRectLightSceneProxy(const URectLightComponent* Component);
 	virtual ~FRectLightSceneProxy();
