Audio Game Thread Queue
When the audio thread needs to synchronize data with the game thread, it fires off a new task, which is then executed on the game thread. These tasks have short execution times. Multiple tasks are triggered in each frame. Because the task graph comes with a certain amount of overhead, I thought it may be quicker to batch these calls together in a queue and then execute them on the game thread in a single task.

I have created a new queue, which contains the function calls (commands) that need to be executed on the game thread. Then the tick function runs through each of the function calls back to back, in a single task.

A replay run shows that this saves around
* 40us on the game thread
* 50us on the game thread critical path
* 80us on audio tick workers
* 10us on game thread audio - this demonstrates that the majority of the overhead comes from the task graph

#tests Replay Run & PIE

diff --git a/Engine/Source/Runtime/Engine/Private/AudioThread.cpp b/Engine/Source/Runtime/Engine/Private/AudioThread.cpp
--- a/Engine/Source/Runtime/Engine/Private/AudioThread.cpp
+++ b/Engine/Source/Runtime/Engine/Private/AudioThread.cpp
@@ -7,6 +7,9 @@
 #include "AudioThread.h"
 #include "Audio.h"
 #include "Async/Async.h"
+#include "Containers/SpscQueue.h"
+#include "Engine/World.h"
+#include "Kismet/GameplayStatics.h"
 #include "ProfilingDebugging/CsvProfiler.h"
 #include "Tasks/Pipe.h"
 
@@ -49,6 +52,132 @@ FAutoConsoleVariableRef  CVarAudioCommandFenceWaitTimeMs(
 	TEXT("Sets number of ms for fence wait"), 
 	ECVF_Default);
 
+bool bBatchGameThreadAudioCommands = true;
+FAutoConsoleVariableRef CVarBatchGameThreadAudioCommands(
+	TEXT("AudioThread.BatchCommands"),
+	bBatchGameThreadAudioCommands,
+	TEXT("Batch audio commands that are created from the audio thread and executed on the game thread so that they are executed in a single task."),
+	ECVF_Default
+	);
+
+namespace
+{
+	/*
+	 * An audio command to be executed on the game thread.
+	 */
+	struct FAudioCommand
+	{
+		FAudioCommand() = default;
+
+		FAudioCommand(TUniqueFunction<void()> InFunction, const TStatId InStatId)
+			: Function(MoveTemp(InFunction))
+			, StatId(InStatId)
+		{
+		}
+
+		FAudioCommand(FAudioCommand&& Other) = default;
+		FAudioCommand& operator=(FAudioCommand&& Other) = default;
+
+		static void Execute(const TUniqueFunction<void()>& InFunction, TStatId InStatId)
+		{
+			check(IsInGameThread());
+
+			CSV_SCOPED_TIMING_STAT_EXCLUSIVE(Audio);
+			QUICK_SCOPE_CYCLE_COUNTER(STAT_AudioCommandQueue_RunCommandOnGameThread);
+			FScopeCycleCounter ScopeCycleCounter(InStatId);
+
+			InFunction();
+		}
+
+		TUniqueFunction<void()> Function;
+		TStatId StatId;
+	};
+
+	const FLazyName DiagnosticName("FGameThreadAudioCommandQueue");
+
+	/*
+	 * Used to run commands from the audio thread, on the game thread.
+	 */
+	class FGameThreadAudioCommandQueue : public FTickFunction
+	{
+	public:
+		FGameThreadAudioCommandQueue()
+			: FTickFunction()
+		{
+			// Setup tick properties 
+			bCanEverTick = true;
+			bTickEvenWhenPaused = false;
+			bStartWithTickEnabled = true;
+			TickGroup = TG_PrePhysics;
+			bAllowTickOnDedicatedServer = false;
+			bRunOnAnyThread = false;
+
+			FWorldDelegates::OnWorldInitializedActors.AddRaw(this, &FGameThreadAudioCommandQueue::OnWorldActorsInitialized);
+		}
+
+		virtual ~FGameThreadAudioCommandQueue() override
+		{
+			FWorldDelegates::OnWorldInitializedActors.RemoveAll(this);
+		}
+
+		// Begin FTickFunction overrides
+		virtual void ExecuteTick(float InDeltaTime, ELevelTick InTickType, ENamedThreads::Type InCurrentThread, const FGraphEventRef& InMyCompletionGraphEvent) override
+		{
+			QUICK_SCOPE_CYCLE_COUNTER(STAT_AudioCommandQueue_ExecuteTick);
+
+			TOptional<FAudioCommand> AudioCommand = Commands.Dequeue();
+			while (AudioCommand.IsSet())
+			{
+				FAudioCommand::Execute(AudioCommand.GetValue().Function, AudioCommand.GetValue().StatId);
+				AudioCommand = Commands.Dequeue();
+			}
+		}
+
+		virtual FString DiagnosticMessage() override
+		{
+			return DiagnosticName.ToString();
+		}
+
+		virtual FName DiagnosticContext(bool bDetailed) override
+		{
+			return DiagnosticName;
+		}
+
+		void RunCommandOnGameThread(TUniqueFunction<void()> InFunction, const TStatId InStatId)
+		{
+			if (IsAudioThreadRunning())
+			{
+				check(IsInAudioThread());
+				Commands.Enqueue(FAudioCommand(MoveTemp(InFunction), InStatId));
+			}
+			else
+			{
+				FAudioCommand::Execute(InFunction, InStatId);
+			}
+		}
+
+	private:
+		void OnWorldActorsInitialized(const UWorld::FActorsInitializedParams& InParams)
+		{
+			// When a new world is up and running unregister the tick function with the old world
+			// and register it with the new world
+			UnRegisterTickFunction();
+
+			const UGameInstance* const GameInstance = UGameplayStatics::GetGameInstance(InParams.World);
+			check(GameInstance);
+
+			if (InParams.World && (InParams.World->GetGameInstance() == GameInstance))
+			{
+				RegisterTickFunction(InParams.World->PersistentLevel);
+			}
+		}
+
+		TSpscQueue<FAudioCommand> Commands;
+	};
+	
+	FGameThreadAudioCommandQueue CommandQueue;
+}
+
 struct FAudioThreadInteractor
 {
 	static void UseAudioThreadCVarSinkFunction()
@@ -366,28 +495,35 @@ void FAudioThread::ProcessAllCommands()
 
 void FAudioThread::RunCommandOnGameThread(TUniqueFunction<void()> InFunction, const TStatId InStatId)
 {
-	if (IsAudioThreadRunning())
+	if (bBatchGameThreadAudioCommands)
 	{
-		check(IsInAudioThread());
-		FFunctionGraphTask::CreateAndDispatchWhenReady(
-			[Function = MoveTemp(InFunction), InStatId]()
-			{
-				CSV_SCOPED_TIMING_STAT_EXCLUSIVE(Audio);
-				QUICK_SCOPE_CYCLE_COUNTER(STAT_AudioThread_RunCommandOnGameThread);
-				FScopeCycleCounter ScopeCycleCounter(InStatId);
-				Function();
-			},
-			TStatId(),
-			nullptr,
-			ENamedThreads::GameThread);
+		CommandQueue.RunCommandOnGameThread(MoveTemp(InFunction), InStatId);
 	}
 	else
 	{
-		check(IsInGameThread());
-		CSV_SCOPED_TIMING_STAT_EXCLUSIVE(Audio);
-		QUICK_SCOPE_CYCLE_COUNTER(STAT_AudioThread_RunCommandOnGameThread);
-		FScopeCycleCounter ScopeCycleCounter(InStatId);
-		InFunction();
+		if (IsAudioThreadRunning())
+		{
+			check(IsInAudioThread());
+			FFunctionGraphTask::CreateAndDispatchWhenReady(
+				[Function = MoveTemp(InFunction), InStatId]()
+				{
+					CSV_SCOPED_TIMING_STAT_EXCLUSIVE(Audio);
+					QUICK_SCOPE_CYCLE_COUNTER(STAT_AudioThread_RunCommandOnGameThread);
+					FScopeCycleCounter ScopeCycleCounter(InStatId);
+					Function();
+				},
+				TStatId(),
+				nullptr,
+				ENamedThreads::GameThread);
+		}
+		else
+		{
+			check(IsInGameThread());
+			CSV_SCOPED_TIMING_STAT_EXCLUSIVE(Audio);
+			QUICK_SCOPE_CYCLE_COUNTER(STAT_AudioThread_RunCommandOnGameThread);
+			FScopeCycleCounter ScopeCycleCounter(InStatId);
+			InFunction();
+		}
 	}
 }
 
@@ -486,4 +622,4 @@ void FAudioCommandFence::Wait(bool bProcessGameThreadTasks) const
 	}
 
 	Fence = UE::Tasks::FTask{}; // release the task as it can hold some references
-}
+}
\ No newline at end of file
