[Material Editor] Add check for null ParameterCollections references in translation data fetched from DDC and updated DDC key generation.

- When null ParameterCollections are found in material translation data fetched from the DDC, translator will run again locally to re-translate the material and push new results to the DDC.
- The material translation DDC key now includes project name.

diff --git a/Engine/Source/Runtime/Engine/Private/Materials/HLSLMaterialTranslator.cpp b/Engine/Source/Runtime/Engine/Private/Materials/HLSLMaterialTranslator.cpp
--- a/Engine/Source/Runtime/Engine/Private/Materials/HLSLMaterialTranslator.cpp
+++ b/Engine/Source/Runtime/Engine/Private/Materials/HLSLMaterialTranslator.cpp
@@ -639,6 +639,7 @@ void FHLSLMaterialTranslator::AppendVersion(FShaderKeyGenerator& KeyGen, EShader
 {
 	static const FGuid DDCVersion = FDevSystemGuids::GetSystemGuid(FDevSystemGuids::Get().MaterialTranslationDDCVersion);
 	KeyGen.AppendDebugText(TEXT("_MatTransl_"));
+	KeyGen.AppendDebugText(FApp::GetProjectName());
 	KeyGen.Append(FMaterialSourceTemplate::Get().GetTemplateHashString(Platform));
 	KeyGen.AppendSeparator();
 	KeyGen.Append(DDCVersion);
@@ -669,7 +670,7 @@ FHLSLMaterialTranslator::FHLSLMaterialTranslator(FMaterial* InMaterial,
 ,	CurrentCustomVertexInterpolatorOffset(0)
 ,	CompileErrorsSink(nullptr)
 ,	CompileErrorExpressionsSink(nullptr)
-,	bSuccess(false)
+,	TranslationResult(EHLSLMaterialTranslatorResult::Success)
 ,	bCompileForComputeShader(false)
 ,	bUsesSceneDepth(false)
 ,	bNeedsParticlePosition(false)
@@ -1169,7 +1170,7 @@ UE_TRACE_EVENT_BEGIN(Cpu, FHLSLMaterialTranslatorTranslate, NoSync)
 UE_TRACE_EVENT_END()
  
 
-bool FHLSLMaterialTranslator::Translate()
+EHLSLMaterialTranslatorResult FHLSLMaterialTranslator::Translate(bool bForceDisableDDCQuery)
 {
 	STAT(FDateTime TranslationDateTime = FDateTime::Now());
 	STAT(double TotalTime = FPlatformTime::Seconds());
@@ -1221,7 +1222,7 @@ bool FHLSLMaterialTranslator::Translate()
 	FSharedBuffer EnvironmentDefinesBuffer;
 
 	// Asynchronously query the DDC for results
-	if (!bDisableTranslationDDC)
+	if (!bForceDisableDDCQuery && !bDisableTranslationDDC)
 	{
 		AsyncQueryDDC(DDCRequestOwner, EnvironmentDefinesBuffer);
 	}
@@ -1229,7 +1230,7 @@ bool FHLSLMaterialTranslator::Translate()
 	// Synchronously begin translating the material
 	{
 		SCOPE_SECONDS_COUNTER(TranslationOnlyTime);
-		DoTranslate();
+		TranslateMaterial();
 	}
 
 	// One of the two has terminated. This will be a NOOP if DDC query task has completed.
@@ -1251,13 +1252,25 @@ bool FHLSLMaterialTranslator::Translate()
 		EnvironmentDefinesBufferReaderProxy.bResolveRedirectors = true;
 		EnvironmentDefines->Serialize(EnvironmentDefinesBufferReaderProxy);
 
-		MaterialCompilationOutput = DDCMaterialCompilationOutput;
+		// If a parameter collection failed to load, the DDC entry is invalid. Caller needs
+		// to invoke translate again forcing DDC querying off in order to re-translate the
+		// material from scratch. We can't do this ourselves because this class was designed
+		// to be single-use. The caller needs to create a new instance and invoke Translate
+		// again.
+		if (EnvironmentDefines->ParameterCollections.Contains(nullptr))
+		{
+			TranslationResult = EHLSLMaterialTranslatorResult::RetryWithoutDDC;
+		}
+		else
+		{
+			MaterialCompilationOutput = DDCMaterialCompilationOutput;
 		
-		STAT(DDCRequestSerializeTime += SerializeTime - FPlatformTime::Seconds());
-		STAT(GShaderCompilerStats->IncrementMaterialCacheHit());
-		bSuccess = true;
+			STAT(DDCRequestSerializeTime += SerializeTime - FPlatformTime::Seconds());
+			STAT(GShaderCompilerStats->IncrementMaterialCacheHit());
+			TranslationResult = EHLSLMaterialTranslatorResult::Success;
+		}
 	}
-	else if (bSuccess)
+	else if (TranslationResult == EHLSLMaterialTranslatorResult::Success)
 	{
 		if (bVerbose)
 		{
@@ -1294,12 +1307,12 @@ bool FHLSLMaterialTranslator::Translate()
 
 #endif // STATS
 
-	return bSuccess;
+	return TranslationResult;
 }
 
-void FHLSLMaterialTranslator::DoTranslate()
+void FHLSLMaterialTranslator::TranslateMaterial()
 {
-	bSuccess = true;
+	TranslationResult = EHLSLMaterialTranslatorResult::Success;
 	
 	// No cache hit, continue translating the material
 	check(ScopeStack.Num() == 0);
@@ -1316,8 +1329,8 @@ void FHLSLMaterialTranslator::DoTranslate()
 	// If pedantic error checking is enabled, log out the pre-compilation errors
 	if (GPedanticErrorChecksEnabled)
 	{
-		bSuccess = Material->CheckInValidStateForCompilation(this);
-		if (!bSuccess)
+		TranslationResult = Material->CheckInValidStateForCompilation(this) ? EHLSLMaterialTranslatorResult::Success : EHLSLMaterialTranslatorResult::Failure;
+		if (TranslationResult != EHLSLMaterialTranslatorResult::Success)
 		{
 			return;
 		}
@@ -1327,7 +1340,7 @@ void FHLSLMaterialTranslator::DoTranslate()
 	UMaterialInterface* Interface = Material->GetMaterialInterface();
 	if (!Interface)
 	{
-		bSuccess = false;
+		TranslationResult = EHLSLMaterialTranslatorResult::Failure;
 		return;
 	}
 
@@ -1340,7 +1353,7 @@ void FHLSLMaterialTranslator::DoTranslate()
 			if (Expression->ContainsInputLoop(VisitedExpressions))
 			{
 				AppendExpressionError(Expression, TEXT("Expression is part of a cycle. Please make sure the material graph is acyclic."));
-				bSuccess = false;
+				TranslationResult = EHLSLMaterialTranslatorResult::Failure;
 				return;
 			}
 		}
@@ -4394,7 +4407,7 @@ int32 FHLSLMaterialTranslator::Error(const TCHAR* Text)
 	if (!bUsingErrorProxy)
 	{
 		// Standard error handling, immediately append one-off errors and signal failure
-		bSuccess = false;
+		TranslationResult = EHLSLMaterialTranslatorResult::Failure;
 		if (ExpressionToError)
 		{
 			ExpressionToError->LastErrorText = Text;
diff --git a/Engine/Source/Runtime/Engine/Private/Materials/HLSLMaterialTranslator.h b/Engine/Source/Runtime/Engine/Private/Materials/HLSLMaterialTranslator.h
--- a/Engine/Source/Runtime/Engine/Private/Materials/HLSLMaterialTranslator.h
+++ b/Engine/Source/Runtime/Engine/Private/Materials/HLSLMaterialTranslator.h
@@ -242,6 +242,20 @@ enum ESubstrateCompilationContext : uint8
 	SCC_MAX = 2u
 };
 
+/** 
+ * Enumerates Translate() results.
+ * If RetryWithoutDDC is returned, the caller should call Translate() again
+ * on a new instance of FHLSLMaterialTranslator forcing DDC query off (i.e. call
+ * Translate(false)). This is because an instance of FHLSLMaterialTranslator
+ * is single use.
+ */
+enum class EHLSLMaterialTranslatorResult
+{
+	Success,
+	Failure,
+	RetryWithoutDDC,
+};
+
 class FHLSLMaterialTranslator : public FMaterialCompiler
 {
 	friend class FMaterialDerivativeAutogen;
@@ -338,8 +352,8 @@ protected:
 	/** Keeps track of which variations of analytic derivative functions are used, and generates the code during translation. **/
 	FMaterialDerivativeAutogen DerivativeAutogen;
 
-	/** Whether the translation succeeded. */
-	uint32 bSuccess : 1;
+	/** The translation result. */
+	EHLSLMaterialTranslatorResult TranslationResult;
 	/** Whether the compute shader material inputs were compiled. */
 	uint32 bCompileForComputeShader : 1;
 	/** Whether the compiled material uses scene depth. */
@@ -621,7 +635,8 @@ public:
 
 	void ValidateVtPropertyLimits();
 	void ValidateShadingModelsForFeatureLevel(const FMaterialShadingModelField& ShadingModels);
-	bool Translate();
+
+	EHLSLMaterialTranslatorResult Translate(bool bForceDisableDDCQuery);
 
 	void GetMaterialEnvironment(EShaderPlatform InPlatform, FShaderCompilerEnvironment& OutEnvironment);
 	
@@ -1400,11 +1415,11 @@ protected:
 	virtual bool IsCurrentlyCompilingForPreviousFrame() const;
 
 	virtual bool IsDevelopmentFeatureEnabled(const FName& FeatureName) const override;
-	
+
 	/**
-	 * EFfectively performs the translation without querying the DDC first.
+	 * Effectively performs the translation without querying the DDC first.
 	 */
-	void DoTranslate();
+	void TranslateMaterial();
 
 	/**
 	 * Queries the DDC cache for a cached translation.
diff --git a/Engine/Source/Runtime/Engine/Private/Materials/MaterialShared.cpp b/Engine/Source/Runtime/Engine/Private/Materials/MaterialShared.cpp
--- a/Engine/Source/Runtime/Engine/Private/Materials/MaterialShared.cpp
+++ b/Engine/Source/Runtime/Engine/Private/Materials/MaterialShared.cpp
@@ -3356,19 +3356,35 @@ bool FMaterial::Translate_Legacy(const FMaterialShaderMapId& ShaderMapId,
 	TRefCountPtr<FSharedShaderCompilerEnvironment>& OutMaterialEnvironment)
 {
 	FString MaterialTranslationDDCKeyString = GetMaterialShaderMapKeyString(ShaderMapId, FMaterialShaderParameters(this), InPlatform, false);
+	
 	FHLSLMaterialTranslator MaterialTranslator(this, OutCompilationOutput, InStaticParameters, InPlatform, GetQualityLevel(), ShaderMapId.FeatureLevel, InTargetPlatform, &ShaderMapId.SubstrateCompilationConfig, MoveTemp(MaterialTranslationDDCKeyString));
-	const bool bSuccess = MaterialTranslator.Translate();
-	if (bSuccess)
+	EHLSLMaterialTranslatorResult Result = MaterialTranslator.Translate(false);
+
+	// If the DDC result was invalid we need to invoke translation again turning the DDC off.
+	if (Result == EHLSLMaterialTranslatorResult::RetryWithoutDDC)
 	{
-		// Create a shader compiler environment for the material that will be shared by all jobs from this material
-		OutMaterialEnvironment = new FSharedShaderCompilerEnvironment();
-		OutMaterialEnvironment->TargetPlatform = InTargetPlatform;
-		MaterialTranslator.GetMaterialEnvironment(InPlatform, *OutMaterialEnvironment);
-		const FString MaterialShaderCode = MaterialTranslator.GetMaterialShaderCode();
+		// FHLSLMaterialTranslator is designed as single use. After a call to Translate() no other calls are allowed.
+		// Destruct the current instance and create a new one before translating the material again, forcing the translator
+		// to translate the material instead of accessing the DDC cache.
+		MaterialTranslator.~FHLSLMaterialTranslator();
+		new (&MaterialTranslator) FHLSLMaterialTranslator(this, OutCompilationOutput, InStaticParameters, InPlatform, GetQualityLevel(), ShaderMapId.FeatureLevel, InTargetPlatform, &ShaderMapId.SubstrateCompilationConfig, MoveTemp(MaterialTranslationDDCKeyString));
+		Result = MaterialTranslator.Translate(true);
+	}
 
-		OutMaterialEnvironment->IncludeVirtualPathToContentsMap.Add(TEXT("/Engine/Generated/Material.ush"), MaterialShaderCode);
+	if (Result != EHLSLMaterialTranslatorResult::Success)
+	{
+		return false;
 	}
-	return bSuccess;
+
+	// Create a shader compiler environment for the material that will be shared by all jobs from this material
+	OutMaterialEnvironment = new FSharedShaderCompilerEnvironment();
+	OutMaterialEnvironment->TargetPlatform = InTargetPlatform;
+	MaterialTranslator.GetMaterialEnvironment(InPlatform, *OutMaterialEnvironment);
+	const FString MaterialShaderCode = MaterialTranslator.GetMaterialShaderCode();
+
+	OutMaterialEnvironment->IncludeVirtualPathToContentsMap.Add(TEXT("/Engine/Generated/Material.ush"), MaterialShaderCode);
+	
+	return true;
 }
 
 bool FMaterial::Translate_New(const FMaterialShaderMapId& InShaderMapId,
