PR #11957: Update TryFreeCDOName() to check the renamed default object name and not the current default object name.

The PR fixes TryFreeCDOName, but this change includes some additional cosmetic changes to RenameGeneratedClasses (inconsistent whitespace, missing consts, etc).

diff --git a/Engine/Source/Runtime/Engine/Private/Blueprint.cpp b/Engine/Source/Runtime/Engine/Private/Blueprint.cpp
--- a/Engine/Source/Runtime/Engine/Private/Blueprint.cpp
+++ b/Engine/Source/Runtime/Engine/Private/Blueprint.cpp
@@ -501,49 +501,51 @@ void UBlueprint::Serialize(FArchive& Ar)
 
 #if WITH_EDITOR
 
-bool UBlueprint::RenameGeneratedClasses( const TCHAR* InName, UObject* NewOuter, ERenameFlags Flags )
+bool UBlueprint::RenameGeneratedClasses(const TCHAR* NewName, UObject* NewOuter, ERenameFlags Flags)
 {
-	const bool bRenameGeneratedClasses = !(Flags & REN_SkipGeneratedClasses );
+	const bool bRenameGeneratedClasses = !(Flags & REN_SkipGeneratedClasses);
 
-	if(bRenameGeneratedClasses)
+	if (bRenameGeneratedClasses)
 	{
-		const auto TryFreeCDOName = [](UClass* ForClass, UObject* ToOuter, ERenameFlags InFlags)
-		{
-			if(ForClass->ClassDefaultObject)
+		const auto TryFreeCDOName = [](FName NewClassName, UObject* NewOuter, ERenameFlags InFlags)
 			{
-				FName CDOName = ForClass->GetDefaultObjectName();
-				
-				if(UObject* Obj = StaticFindObjectFast(UObject::StaticClass(), ToOuter, CDOName))
+				const FString& NewCDOName = FString(DEFAULT_OBJECT_PREFIX) + NewClassName.ToString();
+
+				if (UObject* ExistingCDO = StaticFindObjectFast(UObject::StaticClass(), NewOuter, *NewCDOName))
 				{
-					FName NewName = MakeUniqueObjectName(ToOuter, Obj->GetClass(), CDOName);
-					Obj->Rename(*(NewName.ToString()), ToOuter, InFlags|REN_ForceNoResetLoaders|REN_DontCreateRedirectors);
+					const FName NewName = MakeUniqueObjectName(NewOuter, ExistingCDO->GetClass(), *NewCDOName);
+					ExistingCDO->Rename(*(NewName.ToString()), NewOuter, InFlags | REN_ForceNoResetLoaders | REN_DontCreateRedirectors);
 				}
-			}
-		};
+			};
 
-		const auto CheckRedirectors = [](FName ClassName, UClass* ForClass, UObject* NewOuter)
-		{
-			if (UObjectRedirector* Redirector = FindObjectFast<UObjectRedirector>(NewOuter, ClassName))
+		const auto CheckRedirectors = [](FName NewClassName, UClass* OldClass, UObject* NewOuter)
 			{
-				// If we found a redirector, check that the object it points to is of the same class.
-				if (Redirector->DestinationObject
-					&& Redirector->DestinationObject->GetClass() == ForClass->GetClass())
+				if (UObjectRedirector* Redirector = FindObjectFast<UObjectRedirector>(NewOuter, NewClassName))
 				{
-					Redirector->Rename(nullptr, GetTransientPackage(), REN_ForceNoResetLoaders | REN_DontCreateRedirectors);
+					// If we found a redirector, check that the object it points to is of the same class.
+					const bool bFoundRedirector =
+						Redirector->DestinationObject &&
+						(Redirector->DestinationObject->GetClass() == OldClass->GetClass())
+					;
+
+					if (bFoundRedirector)
+					{
+						Redirector->Rename(nullptr, GetTransientPackage(), REN_ForceNoResetLoaders | REN_DontCreateRedirectors);
+					}
 				}
-			}
-		};
+			};
 
-		FName SkelClassName, GenClassName;
-		GetBlueprintClassNames(GenClassName, SkelClassName, FName(InName));
+		FName NewSkelClassName;
+		FName NewGenClassName;
+		GetBlueprintClassNames(NewGenClassName, NewSkelClassName, FName(NewName));
 
-		UPackage* NewTopLevelObjectOuter = NewOuter ? NewOuter->GetOutermost() : NULL;
-		if (GeneratedClass != NULL)
+		UPackage* NewTopLevelObjectOuter = NewOuter ? NewOuter->GetOutermost() : nullptr;
+		if (GeneratedClass != nullptr)
 		{
 			// check for collision of CDO name, move aside if necessary:
-			TryFreeCDOName(GeneratedClass, NewTopLevelObjectOuter, Flags);
-			CheckRedirectors(GenClassName, GeneratedClass, NewTopLevelObjectOuter);
-			bool bMovedOK = GeneratedClass->Rename(*GenClassName.ToString(), NewTopLevelObjectOuter, Flags);
+			TryFreeCDOName(NewGenClassName, NewTopLevelObjectOuter, Flags);
+			CheckRedirectors(NewGenClassName, GeneratedClass, NewTopLevelObjectOuter);
+			const bool bMovedOK = GeneratedClass->Rename(*NewGenClassName.ToString(), NewTopLevelObjectOuter, Flags);
 			if (!bMovedOK)
 			{
 				return false;
@@ -551,11 +553,16 @@ bool UBlueprint::RenameGeneratedClasses( const TCHAR* InName, UObject* NewOuter,
 		}
 
 		// Also move skeleton class, if different from generated class, to new package (again, to create redirector)
-		if (SkeletonGeneratedClass != NULL && SkeletonGeneratedClass != GeneratedClass)
+		const bool bShouldMoveSkeleton =
+			(SkeletonGeneratedClass != nullptr) && 
+			(SkeletonGeneratedClass != GeneratedClass)
+		;
+
+		if (bShouldMoveSkeleton)
 		{
-			TryFreeCDOName(SkeletonGeneratedClass, NewTopLevelObjectOuter, Flags);
-			CheckRedirectors(SkelClassName, SkeletonGeneratedClass, NewTopLevelObjectOuter);
-			bool bMovedOK = SkeletonGeneratedClass->Rename(*SkelClassName.ToString(), NewTopLevelObjectOuter, Flags);
+			TryFreeCDOName(NewSkelClassName, NewTopLevelObjectOuter, Flags);
+			CheckRedirectors(NewSkelClassName, SkeletonGeneratedClass, NewTopLevelObjectOuter);
+			const bool bMovedOK = SkeletonGeneratedClass->Rename(*NewSkelClassName.ToString(), NewTopLevelObjectOuter, Flags);
 			if (!bMovedOK)
 			{
 				return false;
