[GC] Expose GC parameters inside UEngine. Allow override GC frame budget
#tests local build

diff --git a/Engine/Source/Runtime/Engine/Classes/Engine/Engine.h b/Engine/Source/Runtime/Engine/Classes/Engine/Engine.h
--- a/Engine/Source/Runtime/Engine/Classes/Engine/Engine.h
+++ b/Engine/Source/Runtime/Engine/Classes/Engine/Engine.h
@@ -2721,13 +2721,17 @@ protected:
 	  */
 	 static ENGINE_API void SendWorldEndOfFrameUpdates();
 
-	 /**
-	  * Allows derived classes to force garbage collection based on various factors (low on available UObject slots / other resources)
-	  */
-	 virtual EGarbageCollectionType ShouldForceGarbageCollection()
-	 {
-		 return EGarbageCollectionType::None;
-	 }
+	/**
+	 * Allows derived classes to force garbage collection based on various factors (low on available UObject slots / other resources)
+	 */
+	ENGINE_API virtual EGarbageCollectionType ShouldForceGarbageCollection();
+
+	/**
+	 * Allows derived classes to set per-frame GC budget depending on various factors.
+	 * Default implementation uses gc.IncrementalGCTimePerFrame and gc.LowMemory.IncrementalGCTimePerFrame
+	 * depending on current memory usage and value of gc.LowMemory.MemoryThresholdMB
+	 */
+	ENGINE_API virtual float GetIncrementalGCTimePerFrame();
 
 public:
 	/** @return the GIsEditor flag setting */
diff --git a/Engine/Source/Runtime/Engine/Private/UnrealEngine.cpp b/Engine/Source/Runtime/Engine/Private/UnrealEngine.cpp
--- a/Engine/Source/Runtime/Engine/Private/UnrealEngine.cpp
+++ b/Engine/Source/Runtime/Engine/Private/UnrealEngine.cpp
@@ -1795,6 +1795,28 @@ void UEngine::SetTimeUntilNextGarbageCollection(const float MinTimeUntilNextPass
 	TimeSinceLastPendingKillPurge = TimeBetweenPurgingPendingKillObjects - MinTimeUntilNextPass;
 }
 
+EGarbageCollectionType UEngine::ShouldForceGarbageCollection()
+{
+	return EGarbageCollectionType::None;
+}
+
+float UEngine::GetIncrementalGCTimePerFrame()
+{
+	float IncGCTime = GIncrementalGCTimePerFrame;
+	if (GLowMemoryMemoryThresholdMB > 0.0)
+	{
+		float MBFree = float(PlatformMemoryHelpers::GetFrameMemoryStats().AvailablePhysical / 1024 / 1024);
+#if !UE_BUILD_SHIPPING
+		MBFree -= float(FPlatformMemory::GetExtraDevelopmentMemorySize() / 1024 / 1024);
+#endif
+		if (MBFree <= GLowMemoryMemoryThresholdMB && GLowMemoryIncrementalGCTimePerFrame > GIncrementalGCTimePerFrame)
+		{
+			IncGCTime = GLowMemoryIncrementalGCTimePerFrame;
+		}
+	}
+	return IncGCTime;
+}
+
 void UEngine::ConditionalCollectGarbage()
 {
 	if (GFrameCounter != LastGCFrame)
@@ -1894,18 +1916,7 @@ void UEngine::ConditionalCollectGarbage()
 					else
 					{
 						SCOPE_CYCLE_COUNTER(STAT_GCSweepTime);
-						float IncGCTime = GIncrementalGCTimePerFrame;
-						if (GLowMemoryMemoryThresholdMB > 0.0)
-						{
-							float MBFree = float(PlatformMemoryHelpers::GetFrameMemoryStats().AvailablePhysical / 1024 / 1024);
-#if !UE_BUILD_SHIPPING
-							MBFree -= float(FPlatformMemory::GetExtraDevelopmentMemorySize() / 1024 / 1024);
-#endif
-							if (MBFree <= GLowMemoryMemoryThresholdMB && GLowMemoryIncrementalGCTimePerFrame > GIncrementalGCTimePerFrame)
-							{
-								IncGCTime = GLowMemoryIncrementalGCTimePerFrame;
-							}
-						}
+						float IncGCTime = GetIncrementalGCTimePerFrame();
 						IncrementalPurgeGarbage(true, IncGCTime);
 					}
 				}
