Skeletal Mesh: Ensure there's enough space to store recovered LODModels on load

#rnx

diff --git a/Engine/Source/Runtime/Engine/Private/SkeletalMesh.cpp b/Engine/Source/Runtime/Engine/Private/SkeletalMesh.cpp
--- a/Engine/Source/Runtime/Engine/Private/SkeletalMesh.cpp
+++ b/Engine/Source/Runtime/Engine/Private/SkeletalMesh.cpp
@@ -3112,6 +3112,12 @@ void USkeletalMesh::PostLoadVerifyAndFixBadTangent()
 
 void USkeletalMesh::PostLoadRecoverConvertLODModelsToMeshDescription()
 {
+	// Make sure we have enough space in the FMeshDescription storage for all the recovered LOD models.
+	if (GetNumSourceModels() < GetImportedModel()->LODModels.Num())
+	{
+		SetNumSourceModels(GetImportedModel()->LODModels.Num());
+	}
+
 	// If we didn't get any meshes from the bulk data, then try to recover them from the LODModel listings.
 	for (int32 LODIndex = 0; LODIndex < GetImportedModel()->LODModels.Num(); ++LODIndex)
 	{
