Skeletal Mesh: Don't error out if the vertex attribute data count does not match the point count

#rnx

diff --git a/Engine/Source/Runtime/Engine/Private/Rendering/SkeletalMeshLODImporterData.cpp b/Engine/Source/Runtime/Engine/Private/Rendering/SkeletalMeshLODImporterData.cpp
--- a/Engine/Source/Runtime/Engine/Private/Rendering/SkeletalMeshLODImporterData.cpp
+++ b/Engine/Source/Runtime/Engine/Private/Rendering/SkeletalMeshLODImporterData.cpp
@@ -2081,8 +2081,10 @@ bool FSkeletalMeshImportData::GetMeshDescription(const USkeletalMesh* InSkeletal
 	{
 		const FString& VertexAttributeName = VertexAttributeNames[AttributeIndex];
 		const SkeletalMeshImportData::FVertexAttribute& VertexAttribute = VertexAttributes[AttributeIndex];
-		if (!ensure(VertexAttribute.AttributeValues.Num() == (Points.Num() * VertexAttribute.ComponentCount)))
+		if (VertexAttribute.AttributeValues.Num() != (Points.Num() * VertexAttribute.ComponentCount))
 		{
+			UE_ASSET_LOG(LogSkeletalMeshLODImporterData, Warning, InSkeletalMesh, TEXT("Vertex attribute '%s' value count (%d) does not match the mesh's point count (%d)."),
+				*VertexAttributeName, VertexAttribute.AttributeValues.Num() / VertexAttribute.ComponentCount, Points.Num());
 			continue;
 		}
 
