Fix pose watches using multicolor bones

diff --git a/Engine/Source/Runtime/Engine/Private/SkeletalDebugRendering.cpp b/Engine/Source/Runtime/Engine/Private/SkeletalDebugRendering.cpp
--- a/Engine/Source/Runtime/Engine/Private/SkeletalDebugRendering.cpp
+++ b/Engine/Source/Runtime/Engine/Private/SkeletalDebugRendering.cpp
@@ -244,6 +244,7 @@ void DrawBonesFromPoseWatch(
 	DrawConfig.AffectedBoneColor = BoneColor;
 	DrawConfig.SelectedBoneColor = BoneColor;
 	DrawConfig.ParentOfSelectedBoneColor = BoneColor;
+	DrawConfig.bUseMultiColorAsDefaultColor = false;
 
 	SkeletalDebugRendering::DrawBonesInternal(
 		PDI,
