Editor : Fixed transient values being reset at SkeletalMeshComponent for Blueprint created components (data cleared at actor component PostEditChange). Added a skeletal mesh component instance data to store the two values that are currently flagged as problematic.

diff --git a/Engine/Source/Runtime/Engine/Classes/Components/SkeletalMeshComponent.h b/Engine/Source/Runtime/Engine/Classes/Components/SkeletalMeshComponent.h
--- a/Engine/Source/Runtime/Engine/Classes/Components/SkeletalMeshComponent.h
+++ b/Engine/Source/Runtime/Engine/Classes/Components/SkeletalMeshComponent.h
@@ -1706,6 +1706,10 @@ public:
 	ENGINE_API virtual void SetComponentTickEnabled(bool bEnabled) override;
 	ENGINE_API virtual void OnComponentDestroyed(bool bDestroyingHierarchy) override;
 
+#if WITH_EDITOR
+	ENGINE_API virtual TStructOnScope<FActorComponentInstanceData> GetComponentInstanceData() const override;
+#endif // WITH_EDITOR
+
 	//Handle registering our end physics tick function
 	ENGINE_API virtual void RegisterEndPhysicsTick(bool bRegister);
 
diff --git a/Engine/Source/Runtime/Engine/Private/Components/SkeletalMeshComponent.cpp b/Engine/Source/Runtime/Engine/Private/Components/SkeletalMeshComponent.cpp
--- a/Engine/Source/Runtime/Engine/Private/Components/SkeletalMeshComponent.cpp
+++ b/Engine/Source/Runtime/Engine/Private/Components/SkeletalMeshComponent.cpp
@@ -46,6 +46,7 @@
 #if WITH_EDITOR
 #include "Engine/PoseWatch.h"
 #include "Settings/AnimBlueprintSettings.h"
+#include "SkeletalMeshComponentInstanceData.h"
 #endif
 
 #include UE_INLINE_GENERATED_CPP_BY_NAME(SkeletalMeshComponent)
@@ -530,6 +531,13 @@ void USkeletalMeshComponent::OnComponentDestroyed(bool bDestroyingHierarchy)
 	Super::OnComponentDestroyed(bDestroyingHierarchy);
 }
 
+#if WITH_EDITOR
+TStructOnScope<FActorComponentInstanceData> USkeletalMeshComponent::GetComponentInstanceData() const
+{
+	return MakeStructOnScope<FActorComponentInstanceData, FSkeletalMeshComponentInstanceData>(this);
+}
+#endif // WITH_EDITOR
+
 void USkeletalMeshComponent::RegisterEndPhysicsTick(bool bRegister)
 {
 	if (bRegister != EndPhysicsTickFunction.IsTickFunctionRegistered())
diff --git a/Engine/Source/Runtime/Engine/Private/Components/SkeletalMeshComponentInstanceData.cpp b/Engine/Source/Runtime/Engine/Private/Components/SkeletalMeshComponentInstanceData.cpp
new file mode 100644
--- /dev/null
+++ b/Engine/Source/Runtime/Engine/Private/Components/SkeletalMeshComponentInstanceData.cpp
@@ -0,0 +1,43 @@
+// Copyright Epic Games, Inc. All Rights Reserved.
+
+#include "SkeletalMeshComponentInstanceData.h"
+
+#include "Components/SkeletalMeshComponent.h"
+
+
+FSkeletalMeshComponentInstanceData::FSkeletalMeshComponentInstanceData(const USkeletalMeshComponent* SourceComponent)
+	: FSceneComponentInstanceData(SourceComponent)
+	, bUpdateAnimationInEditor(0)
+	, bUpdateClothInEditor(0)
+
+{
+#if WITH_EDITOR
+	// Only Blueprint components would reset tansient values when the construction script is re-run.
+	// Hence, we only need to apply instance cache for Blueprint created components.
+	const bool bIsBlueprintCreatedComponent = SourceComponent->CreationMethod == EComponentCreationMethod::SimpleConstructionScript
+		|| SourceComponent->CreationMethod == EComponentCreationMethod::UserConstructionScript;
+	if (bIsBlueprintCreatedComponent)
+	{
+		bUpdateAnimationInEditor = SourceComponent->GetUpdateAnimationInEditor();
+		bUpdateClothInEditor = SourceComponent->GetUpdateClothInEditor();
+	}
+#endif // WITH_EDITOR
+}
+
+bool FSkeletalMeshComponentInstanceData::ContainsData() const
+{
+	return bUpdateAnimationInEditor || bUpdateClothInEditor;
+}
+
+void FSkeletalMeshComponentInstanceData::ApplyToComponent(UActorComponent* Component, const ECacheApplyPhase CacheApplyPhase)
+{
+	Super::ApplyToComponent(Component, CacheApplyPhase);
+
+#if WITH_EDITOR
+	if (USkeletalMeshComponent * SkeletalMesh = Cast<USkeletalMeshComponent>(Component))
+	{
+		SkeletalMesh->SetUpdateAnimationInEditor(bUpdateAnimationInEditor);
+		SkeletalMesh->SetUpdateClothInEditor(bUpdateClothInEditor);
+	}
+#endif // WITH_EDITOR
+}
diff --git a/Engine/Source/Runtime/Engine/Private/Components/SkeletalMeshComponentInstanceData.h b/Engine/Source/Runtime/Engine/Private/Components/SkeletalMeshComponentInstanceData.h
new file mode 100644
--- /dev/null
+++ b/Engine/Source/Runtime/Engine/Private/Components/SkeletalMeshComponentInstanceData.h
@@ -0,0 +1,24 @@
+// Copyright Epic Games, Inc. All Rights Reserved.
+#pragma once
+
+#include "Components/SceneComponent.h"
+
+#include "SkeletalMeshComponentInstanceData.generated.h"
+
+class USkeletalMeshComponent;
+
+/** Saves internal SkeletalMesh (transient) state that gets lost at reconstruction for Blueprint created components. */
+USTRUCT()
+struct ENGINE_API FSkeletalMeshComponentInstanceData : public FSceneComponentInstanceData
+{
+	GENERATED_BODY()
+	
+	FSkeletalMeshComponentInstanceData() = default;
+	explicit FSkeletalMeshComponentInstanceData(const USkeletalMeshComponent* SourceComponent);
+	
+	virtual bool ContainsData() const override;
+	virtual void ApplyToComponent(UActorComponent* Component, const ECacheApplyPhase CacheApplyPhase) override;
+
+	uint8 bUpdateAnimationInEditor : 1;
+	uint8 bUpdateClothInEditor : 1;
+};
