Adds a LensDistortion show flag

diff --git a/Engine/Source/Runtime/Engine/Public/ShowFlagsValues.inl b/Engine/Source/Runtime/Engine/Public/ShowFlagsValues.inl
--- a/Engine/Source/Runtime/Engine/Public/ShowFlagsValues.inl
+++ b/Engine/Source/Runtime/Engine/Public/ShowFlagsValues.inl
@@ -42,6 +42,8 @@ SHOWFLAG_FIXED_IN_SHIPPING(0, VisualizeSkyLightIlluminance, SFG_Visualize, NSLOC
 SHOWFLAG_FIXED_IN_SHIPPING(0, VisualizeLocalExposure, SFG_Visualize, NSLOCTEXT("UnrealEd", "VisualizeLocalExposureSF", "Local Exposure"))
 /** Image based lens flares (Simulate artifact of reflections within a camera system) */
 SHOWFLAG_FIXED_IN_SHIPPING(1, LensFlares, SFG_PostProcess, NSLOCTEXT("UnrealEd", "LensFlaresSF", "Lens Flares"))
+/** Lens distortion (Simulate the distortion of light through the lens onto the filmback) */
+SHOWFLAG_FIXED_IN_SHIPPING(1, LensDistortion, SFG_PostProcess, NSLOCTEXT("UnrealEd", "LensDistortionSF", "Lens Distortion"))
 /** show indirect lighting component, for now SHOWFLAG_ALWAYS_ACCESSIBLE because it's needed by r.GBuffer */
 SHOWFLAG_ALWAYS_ACCESSIBLE(GlobalIllumination, SFG_LightingComponents, NSLOCTEXT("UnrealEd", "GlobalIlluminationSF", "Global Illumination"))
 /** Darkens the screen borders (Camera artifact and artistic effect) */
diff --git a/Engine/Source/Runtime/Renderer/Private/SceneRendering.cpp b/Engine/Source/Runtime/Renderer/Private/SceneRendering.cpp
--- a/Engine/Source/Runtime/Renderer/Private/SceneRendering.cpp
+++ b/Engine/Source/Runtime/Renderer/Private/SceneRendering.cpp
@@ -3767,7 +3767,7 @@ IVisibilityTaskData* FSceneRenderer::OnRenderBegin(FRDGBuilder& GraphBuilder)
 			}
 		}
 
-		if (FPaniniProjectionConfig::IsEnabledByCVars())
+		if (ViewFamily.EngineShowFlags.LensDistortion && FPaniniProjectionConfig::IsEnabledByCVars())
 		{
 			const FPaniniProjectionConfig PaniniProjection = FPaniniProjectionConfig::ReadCVars();
 
