World Partition - Prevent cooked data from pulling on HLOD layers & their dependencies

diff --git a/Engine/Source/Runtime/Engine/Private/WorldPartition/WorldPartitionRuntimeSpatialHash.cpp b/Engine/Source/Runtime/Engine/Private/WorldPartition/WorldPartitionRuntimeSpatialHash.cpp
--- a/Engine/Source/Runtime/Engine/Private/WorldPartition/WorldPartitionRuntimeSpatialHash.cpp
+++ b/Engine/Source/Runtime/Engine/Private/WorldPartition/WorldPartitionRuntimeSpatialHash.cpp
@@ -99,7 +99,6 @@ FSpatialHashStreamingGrid::FSpatialHashStreamingGrid()
 	, DebugColor(ForceInitToZero)
 	, WorldBounds(ForceInitToZero)
 	, bClientOnlyVisible(false)
-	, HLODLayer(nullptr)
 	, GridIndex(INDEX_NONE)
 	, GridHelper(nullptr)
 {
@@ -212,10 +211,6 @@ void FSpatialHashStreamingGrid::DumpStateLog(FHierarchicalLogArchive& Ar) const
 	Ar.Printf(TEXT("Block Slow Loading: %s"), bBlockOnSlowStreaming ? TEXT("Yes") : TEXT("No"));
 	Ar.Printf(TEXT(" ClientOnlyVisible: %s"), bClientOnlyVisible ? TEXT("Yes") : TEXT("No"));
 	Ar.Printf(TEXT(""));
-	if (HLODLayer)
-	{
-		Ar.Printf(TEXT("    HLOD Layer: %s"), *HLODLayer->GetName());
-	}
 
 	struct FGridLevelStats
 	{
@@ -1579,7 +1574,6 @@ bool UWorldPartitionRuntimeSpatialHash::CreateStreamingGrid(const FSpatialHashRu
 	CurrentStreamingGrid.Origin = FVector(RuntimeGrid.Origin, 0);
 	CurrentStreamingGrid.DebugColor = RuntimeGrid.DebugColor;
 	CurrentStreamingGrid.bClientOnlyVisible = RuntimeGrid.bClientOnlyVisible;
-	CurrentStreamingGrid.HLODLayer = RuntimeGrid.HLODLayer;
 	CurrentStreamingGrid.GridIndex = (StreamingGrids.Num() - 1);
 
 	// Move actors into the final streaming grids
@@ -1995,7 +1989,6 @@ FGuid UWorldPartitionRuntimeSpatialHash::RegisterWorldAssetStreaming(const UWorl
 		CurrentGrid.bBlockOnSlowStreaming = SourceGrid->bBlockOnSlowStreaming;
 		CurrentGrid.DebugColor = FColor::MakeRandomColor();
 		CurrentGrid.bClientOnlyVisible = SourceGrid->bClientOnlyVisible;
-		CurrentGrid.HLODLayer = SourceGrid->HLODLayer;
 		CurrentGrid.GridLevels.SetNum(CurrentGridHelper.Levels.Num());
 
 		if (bIsHLODPass && InParams.WorldAssetHLOD.IsNull())
diff --git a/Engine/Source/Runtime/Engine/Public/WorldPartition/RuntimeHashSet/WorldPartitionRuntimeHashSet.h b/Engine/Source/Runtime/Engine/Public/WorldPartition/RuntimeHashSet/WorldPartitionRuntimeHashSet.h
--- a/Engine/Source/Runtime/Engine/Public/WorldPartition/RuntimeHashSet/WorldPartitionRuntimeHashSet.h
+++ b/Engine/Source/Runtime/Engine/Public/WorldPartition/RuntimeHashSet/WorldPartitionRuntimeHashSet.h
@@ -23,6 +23,7 @@ struct FRuntimePartitionHLODSetup
 {
 	GENERATED_USTRUCT_BODY()
 
+#if WITH_EDITORONLY_DATA
 	/** Name for this HLOD layer setup */
 	UPROPERTY(EditAnywhere, Category = RuntimeSettings)
 	FName Name;
@@ -37,6 +38,7 @@ struct FRuntimePartitionHLODSetup
 
 	UPROPERTY(VisibleAnywhere, Category = RuntimeSettings, Instanced, Meta = (EditCondition = "bIsSpatiallyLoaded", HideEditConditionToggle, NoResetToDefault, TitleProperty = "Name"))
 	TObjectPtr<URuntimePartition> PartitionLayer;
+#endif
 };
 
 /** Holds settings for a runtime partition instance. */
@@ -57,9 +59,11 @@ struct FRuntimePartitionDesc
 	UPROPERTY(VisibleAnywhere, Category = RuntimeSettings, Instanced, Meta = (EditCondition = "Class != nullptr", HideEditConditionToggle, NoResetToDefault, TitleProperty = "Name"))
 	TObjectPtr<URuntimePartition> MainLayer;
 
+#if WITH_EDITORONLY_DATA
 	/** HLOD setups used by this partition, one for each layers in the hierarchy */
 	UPROPERTY(EditAnywhere, Category = RuntimeSettings, Meta = (EditCondition = "Class != nullptr", HideEditConditionToggle, ForceInlineRow))
 	TArray<FRuntimePartitionHLODSetup> HLODSetups;
+#endif
 
 #if WITH_EDITOR
 	void UpdateHLODPartitionLayers();
diff --git a/Engine/Source/Runtime/Engine/Public/WorldPartition/WorldPartitionRuntimeSpatialHash.h b/Engine/Source/Runtime/Engine/Public/WorldPartition/WorldPartitionRuntimeSpatialHash.h
--- a/Engine/Source/Runtime/Engine/Public/WorldPartition/WorldPartitionRuntimeSpatialHash.h
+++ b/Engine/Source/Runtime/Engine/Public/WorldPartition/WorldPartitionRuntimeSpatialHash.h
@@ -165,9 +165,6 @@ struct FSpatialHashStreamingGrid
 	UPROPERTY()
 	bool bClientOnlyVisible;
 
-	UPROPERTY()
-	TObjectPtr<const UHLODLayer> HLODLayer;
-
 	UPROPERTY()
 	int32 GridIndex;
 		
