Added support for blend mode in Montage Inertialization when used in conjunction with DeadBlending node.

diff --git a/Engine/Source/Runtime/AnimGraphRuntime/Private/AnimNodes/AnimNode_Slot.cpp b/Engine/Source/Runtime/AnimGraphRuntime/Private/AnimNodes/AnimNode_Slot.cpp
--- a/Engine/Source/Runtime/AnimGraphRuntime/Private/AnimNodes/AnimNode_Slot.cpp
+++ b/Engine/Source/Runtime/AnimGraphRuntime/Private/AnimNodes/AnimNode_Slot.cpp
@@ -43,21 +43,18 @@ void FAnimNode_Slot::Update_AnyThread(const FAnimationUpdateContext& Context)
 	// Update cache in AnimInstance.
 	Context.AnimInstanceProxy->UpdateSlotNodeWeight(SlotName, WeightData.SlotNodeWeight, Context.GetFinalBlendWeight());
 
-	UE::Anim::FSlotInertializationRequest InertializationRequest;
-	if (Context.AnimInstanceProxy->GetSlotInertializationRequest(SlotName, InertializationRequest))
+	FInertializationRequest InertializationRequest;
+	if (Context.AnimInstanceProxy->GetSlotInertializationRequestData(SlotName, InertializationRequest))
 	{
 		UE::Anim::IInertializationRequester* InertializationRequester = Context.GetMessage<UE::Anim::IInertializationRequester>();
 		if (InertializationRequester)
 		{
-			FInertializationRequest Request;
-			Request.Duration = InertializationRequest.Get<0>();
-			Request.BlendProfile = InertializationRequest.Get<1>();
 #if ANIM_TRACE_ENABLED
-			Request.NodeId = Context.GetCurrentNodeId();
-			Request.AnimInstance = Context.AnimInstanceProxy->GetAnimInstanceObject();
+			InertializationRequest.NodeId = Context.GetCurrentNodeId();
+			InertializationRequest.AnimInstance = Context.AnimInstanceProxy->GetAnimInstanceObject();
 #endif
 
-			InertializationRequester->RequestInertialization(Request);
+			InertializationRequester->RequestInertialization(InertializationRequest);
 		}
 		else
 		{
diff --git a/Engine/Source/Runtime/Engine/Classes/Animation/AnimInertializationRequest.h b/Engine/Source/Runtime/Engine/Classes/Animation/AnimInertializationRequest.h
new file mode 100644
--- /dev/null
+++ b/Engine/Source/Runtime/Engine/Classes/Animation/AnimInertializationRequest.h
@@ -0,0 +1,89 @@
+// Copyright Epic Games, Inc. All Rights Reserved.
+
+#pragma once
+
+#include "CoreMinimal.h"
+#include "AlphaBlend.h" // Required for EAlphaBlendOption
+
+#include "AnimInertializationRequest.generated.h"
+
+class UBlendProfile;
+class UCurveFloat;
+
+USTRUCT()
+struct FInertializationRequest
+{
+	GENERATED_BODY()
+
+	ENGINE_API FInertializationRequest();
+	ENGINE_API FInertializationRequest(float InDuration, const UBlendProfile* InBlendProfile);
+
+	// Note: We need to explicitly disable warnings on these constructors/operators for clang to be happy with deprecated variables
+	PRAGMA_DISABLE_DEPRECATION_WARNINGS
+	~FInertializationRequest() = default;
+	FInertializationRequest(const FInertializationRequest&) = default;
+	FInertializationRequest(FInertializationRequest&&) = default;
+	FInertializationRequest& operator=(const FInertializationRequest&) = default;
+	FInertializationRequest& operator=(FInertializationRequest&&) = default;
+	PRAGMA_ENABLE_DEPRECATION_WARNINGS
+
+	ENGINE_API void Clear();
+
+	// Comparison operator used to test for equality in the array of animation requests to that
+	// only unique requests are added. This does not take into account the properties that are 
+	// used only for debugging and only used when ANIM_TRACE_ENABLED
+	friend bool operator==(const FInertializationRequest& A, const FInertializationRequest& B)
+	{
+		return
+			(A.Duration == B.Duration) &&
+			(A.BlendProfile == B.BlendProfile) &&
+			(A.bUseBlendMode == B.bUseBlendMode) &&
+			(A.BlendMode == B.BlendMode) &&
+			(A.CustomBlendCurve == B.CustomBlendCurve);
+	}
+
+	friend bool operator!=(const FInertializationRequest& A, const FInertializationRequest& B)
+	{
+		return !(A == B);
+	}
+
+	// Blend duration of the inertialization request.
+	UPROPERTY(Transient)
+	float Duration = -1.0f;
+
+	// Blend profile to control per-joint blend times.
+	UPROPERTY(Transient)
+	TObjectPtr<const UBlendProfile> BlendProfile = nullptr;
+
+	// If to use the provided blend mode.
+	UPROPERTY(Transient)
+	bool bUseBlendMode = false;
+
+	// Blend mode to use.
+	UPROPERTY(Transient)
+	EAlphaBlendOption BlendMode = EAlphaBlendOption::Linear;
+
+	// Custom blend curve to use when use of the blend mode is active.
+	UPROPERTY(Transient)
+	TObjectPtr<UCurveFloat> CustomBlendCurve = nullptr;
+
+// if ANIM_TRACE_ENABLED - these properties are only used for debugging when ANIM_TRACE_ENABLED == 1
+	
+	UE_DEPRECATED(5.4, "Use DescriptionString instead.")
+	UPROPERTY(Transient, meta = (DeprecatedProperty, DeprecationMessage = "Use DescriptionString instead."))
+	FText Description_DEPRECATED;
+
+	// Description of the request
+	UPROPERTY(Transient)
+	FString DescriptionString;
+
+	// Node id from which this request was made.
+	UPROPERTY(Transient)
+	int32 NodeId = INDEX_NONE;
+
+	// Anim instance from which this request was made.
+	UPROPERTY(Transient)
+	TObjectPtr<UObject> AnimInstance = nullptr;
+
+// endif ANIM_TRACE_ENABLED
+};
diff --git a/Engine/Source/Runtime/Engine/Classes/Animation/AnimInstance.h b/Engine/Source/Runtime/Engine/Classes/Animation/AnimInstance.h
--- a/Engine/Source/Runtime/Engine/Classes/Animation/AnimInstance.h
+++ b/Engine/Source/Runtime/Engine/Classes/Animation/AnimInstance.h
@@ -11,6 +11,7 @@
 #include "Animation/AnimSubsystemInstance.h"
 #include "Animation/AnimSync.h"
 #include "Animation/AnimNotifies/AnimNotify.h"
+#include "Animation/AnimInertializationRequest.h"
 #include "AnimInstance.generated.h"
 
 // Post Compile Validation requires WITH_EDITOR
@@ -50,6 +51,7 @@ typedef TArray<FTransform> FTransformArrayA2;
 namespace UE::Anim
 {
 	struct FHeapAttributeContainer;
+	// DEPRECATED use FInertializationRequest instead
 	using FSlotInertializationRequest = TPair<float, const UBlendProfile*>;
 	struct FCurveFilterSettings;
 }	// namespace UE::Anim
@@ -883,10 +885,15 @@ private:
 protected:
 	/** Map between Active Montages and their FAnimMontageInstance */
 	TMap<class UAnimMontage*, struct FAnimMontageInstance*> ActiveMontagesMap;
+	
 
-	/**  Inertialization requests gathered this frame. Gets reset in UpdateMontageEvaluationData */
+	UE_DEPRECATED(5.5, "This property is deprecated. Please use SlotGroupInertializationRequestDataMap instead")
 	TMap<FName, UE::Anim::FSlotInertializationRequest> SlotGroupInertializationRequestMap;
 
+	/**  Inertialization requests gathered this frame. Gets reset in UpdateMontageEvaluationData */
+	UPROPERTY(Transient)
+	TMap<FName, FInertializationRequest> SlotGroupInertializationRequestDataMap;
+
 	/* StopAllMontagesByGroupName needs a BlendMode and BlendProfile to function properly if using non-default ones in your montages. If you want default BlendMode/BlendProfiles, you need to update the calling code to do so. */
 	UE_DEPRECATED(5.0, "Use StopAllMontagesByGroupName with other signature.")
 	ENGINE_API void StopAllMontagesByGroupName(FName InGroupName, const FAlphaBlend& BlendOut);
@@ -907,6 +914,9 @@ public:
 	/**  Builds an inertialization request from the montage's group, provided duration and optional blend profile*/
 	ENGINE_API void RequestMontageInertialization(const UAnimMontage* Montage, float Duration, const UBlendProfile* BlendProfile = nullptr);
 
+	/**  Makes an inertialization request from the montage's group. */
+	ENGINE_API void RequestMontageInertialization(const UAnimMontage* Montage, const FInertializationRequest& Request);
+
 	/**  Requests an inertial blend during the next anim graph update. Requires your anim graph to have a slot node belonging to the specified group name */
 	UFUNCTION(BlueprintCallable, Category = "Animation|Inertial Blending")
 	ENGINE_API void RequestSlotGroupInertialization(FName InSlotGroupName, float Duration, const UBlendProfile* BlendProfile = nullptr);
diff --git a/Engine/Source/Runtime/Engine/Classes/Animation/AnimNode_Inertialization.h b/Engine/Source/Runtime/Engine/Classes/Animation/AnimNode_Inertialization.h
--- a/Engine/Source/Runtime/Engine/Classes/Animation/AnimNode_Inertialization.h
+++ b/Engine/Source/Runtime/Engine/Classes/Animation/AnimNode_Inertialization.h
@@ -7,7 +7,7 @@
 #include "Animation/AnimNodeBase.h"
 #include "Animation/AnimCurveTypes.h"
 #include "Animation/AnimNodeMessages.h"
-#include "AlphaBlend.h" // Required for EAlphaBlendOption
+#include "Animation/AnimInertializationRequest.h"
 #include "Interfaces/Interface_BoneReferenceSkeletonProvider.h"
 
 #include "AnimNode_Inertialization.generated.h"
@@ -109,84 +109,6 @@ struct FInertializationCurve
 	}
 };
 
-USTRUCT()
-struct FInertializationRequest
-{
-	GENERATED_BODY()
-
-	ENGINE_API FInertializationRequest();
-	ENGINE_API FInertializationRequest(float InDuration, const UBlendProfile* InBlendProfile);
-
-	// Note: We need to explicitly disable warnings on these constructors/operators for clang to be happy with deprecated variables
-	PRAGMA_DISABLE_DEPRECATION_WARNINGS
-	~FInertializationRequest() = default;
-	FInertializationRequest(const FInertializationRequest&) = default;
-	FInertializationRequest(FInertializationRequest&&) = default;
-	FInertializationRequest& operator=(const FInertializationRequest&) = default;
-	FInertializationRequest& operator=(FInertializationRequest&&) = default;
-	PRAGMA_ENABLE_DEPRECATION_WARNINGS
-
-	ENGINE_API void Clear();
-
-	// Comparison operator used to test for equality in the array of animation requests to that
-	// only unique requests are added. This does not take into account the properties that are 
-	// used only for debugging and only used when ANIM_TRACE_ENABLED
-	friend bool operator==(const FInertializationRequest& A, const FInertializationRequest& B)
-	{
-		return
-			(A.Duration == B.Duration) &&
-			(A.BlendProfile == B.BlendProfile) &&
-			(A.bUseBlendMode == B.bUseBlendMode) &&
-			(A.BlendMode == B.BlendMode) &&
-			(A.CustomBlendCurve == B.CustomBlendCurve);
-	}
-
-	friend bool operator!=(const FInertializationRequest& A, const FInertializationRequest& B)
-	{
-		return !(A == B);
-	}
-
-	// Blend duration of the inertialization request.
-	UPROPERTY(Transient)
-	float Duration = -1.0f;
-
-	// Blend profile to control per-joint blend times.
-	UPROPERTY(Transient)
-	TObjectPtr<const UBlendProfile> BlendProfile = nullptr;
-
-	// If to use the provided blend mode.
-	UPROPERTY(Transient)
-	bool bUseBlendMode = false;
-
-	// Blend mode to use.
-	UPROPERTY(Transient)
-	EAlphaBlendOption BlendMode = EAlphaBlendOption::Linear;
-
-	// Custom blend curve to use when use of the blend mode is active.
-	UPROPERTY(Transient)
-	TObjectPtr<UCurveFloat> CustomBlendCurve = nullptr;
-
-// if ANIM_TRACE_ENABLED - these properties are only used for debugging when ANIM_TRACE_ENABLED == 1
-	
-	UE_DEPRECATED(5.4, "Use DescriptionString instead.")
-	UPROPERTY(Transient, meta = (DeprecatedProperty, DeprecationMessage = "Use DescriptionString instead."))
-	FText Description_DEPRECATED;
-
-	// Description of the request
-	UPROPERTY(Transient)
-	FString DescriptionString;
-
-	// Node id from which this request was made.
-	UPROPERTY(Transient)
-	int32 NodeId = INDEX_NONE;
-
-	// Anim instance from which this request was made.
-	UPROPERTY(Transient)
-	TObjectPtr<UObject> AnimInstance = nullptr;
-
-// endif ANIM_TRACE_ENABLED
-};
-
 USTRUCT()
 struct UE_DEPRECATED(5.4, "Internal private pose storage is now used by inertialization.") FInertializationPose
 {
diff --git a/Engine/Source/Runtime/Engine/Private/Animation/AnimInstance.cpp b/Engine/Source/Runtime/Engine/Private/Animation/AnimInstance.cpp
--- a/Engine/Source/Runtime/Engine/Private/Animation/AnimInstance.cpp
+++ b/Engine/Source/Runtime/Engine/Private/Animation/AnimInstance.cpp
@@ -359,7 +359,7 @@ void UAnimInstance::UninitializeAnimation()
 	ActiveAnimNotifyEventReference.Reset(); 
 	NotifyQueue.Reset(SkelMeshComp);
 
-	SlotGroupInertializationRequestMap.Reset();
+	SlotGroupInertializationRequestDataMap.Reset();
 	
 	// Cleanup layer nodes
 	if (IAnimClassInterface* AnimBlueprintClass = IAnimClassInterface::GetFromClass(GetClass()))
@@ -1892,21 +1892,34 @@ void UAnimInstance::Montage_Advance(float DeltaSeconds)
 
 void UAnimInstance::RequestSlotGroupInertialization(FName InSlotGroupName, float Duration, const UBlendProfile* BlendProfile)
 {
+	FInertializationRequest Request;
+	Request.Duration = Duration;
+	Request.BlendProfile = BlendProfile;
+
 	// Must add this on both the anim instance and proxy's map, as this could called after UAnimInstance::UpdateMontageEvaluationData.
-	SlotGroupInertializationRequestMap.FindOrAdd(InSlotGroupName) = UE::Anim::FSlotInertializationRequest(Duration, BlendProfile);
-	GetProxyOnAnyThread<FAnimInstanceProxy>().GetSlotGroupInertializationRequestMap().FindOrAdd(InSlotGroupName) = UE::Anim::FSlotInertializationRequest(Duration, BlendProfile);
+	SlotGroupInertializationRequestDataMap.FindOrAdd(InSlotGroupName) = Request;
+	GetProxyOnAnyThread<FAnimInstanceProxy>().GetSlotGroupInertializationRequestDataMap().FindOrAdd(InSlotGroupName) = Request;
 }
 
-void UAnimInstance::RequestMontageInertialization(const UAnimMontage* Montage, float Duration, const UBlendProfile* BlendProfile)
+void UAnimInstance::RequestMontageInertialization(const UAnimMontage* Montage, const FInertializationRequest& Request)
 {
 	if (Montage)
 	{
 		// Adds a new request or overwrites an existing one
 		// We always overwrite with the last request, instead of using the shortest one (differs from AnimNode_Inertialization), because we expect the last montage played/stopped to take precedence
-		SlotGroupInertializationRequestMap.FindOrAdd(Montage->GetGroupName()) = UE::Anim::FSlotInertializationRequest(Duration, BlendProfile);
+		SlotGroupInertializationRequestDataMap.FindOrAdd(Montage->GetGroupName()) = Request;
 	}
 }
 
+void UAnimInstance::RequestMontageInertialization(const UAnimMontage* Montage, float Duration, const UBlendProfile* BlendProfile)
+{
+	FInertializationRequest Request;
+	Request.Duration = Duration;
+	Request.BlendProfile = BlendProfile;
+	
+	RequestMontageInertialization(Montage, Request);
+}
+
 void UAnimInstance::QueueMontageBlendingOutEvent(const FQueuedMontageBlendingOutEvent& MontageBlendingOutEvent)
 {
 	if (bQueueMontageEvents)
@@ -3878,11 +3891,11 @@ void UAnimInstance::UpdateMontageEvaluationData()
 		}
 	}
 
-	Proxy.GetSlotGroupInertializationRequestMap() = SlotGroupInertializationRequestMap;
+	Proxy.GetSlotGroupInertializationRequestDataMap() = SlotGroupInertializationRequestDataMap;
 
 	// Reset inertialization requests every frame.
 	// If the request is missed by the graph (i.e. the slot node is not relevant), we assume what brought it back to relevancy will handle the blend instead.
-	SlotGroupInertializationRequestMap.Reset();
+	SlotGroupInertializationRequestDataMap.Reset();
 }
 
 float UAnimInstance::GetInstanceAssetPlayerLength(int32 AssetPlayerIndex)
diff --git a/Engine/Source/Runtime/Engine/Private/Animation/AnimInstanceProxy.cpp b/Engine/Source/Runtime/Engine/Private/Animation/AnimInstanceProxy.cpp
--- a/Engine/Source/Runtime/Engine/Private/Animation/AnimInstanceProxy.cpp
+++ b/Engine/Source/Runtime/Engine/Private/Animation/AnimInstanceProxy.cpp
@@ -110,11 +110,12 @@ FAnimInstanceProxy::FAnimInstanceProxy(UAnimInstance* Instance)
 	, bUseMainInstanceMontageEvaluationData(false)
 {
 }
-
+PRAGMA_DISABLE_DEPRECATION_WARNINGS
 FAnimInstanceProxy::FAnimInstanceProxy(const FAnimInstanceProxy&) = default;
 FAnimInstanceProxy& FAnimInstanceProxy::operator=(FAnimInstanceProxy&&) = default;
 FAnimInstanceProxy& FAnimInstanceProxy::operator=(const FAnimInstanceProxy&) = default;
 FAnimInstanceProxy::~FAnimInstanceProxy() = default;
+PRAGMA_ENABLE_DEPRECATION_WARNINGS
 
 void FAnimInstanceProxy::UpdateAnimationNode(const FAnimationUpdateContext& InContext)
 {
@@ -452,7 +453,7 @@ void FAnimInstanceProxy::Uninitialize(UAnimInstance* InAnimInstance)
 	DECLARE_SCOPE_HIERARCHICAL_COUNTER_FUNC()
 
 	MontageEvaluationData.Reset();
-	SlotGroupInertializationRequestMap.Reset();
+	SlotGroupInertializationRequestDataMap.Reset();
 	DefaultLinkedInstanceInputNode = nullptr;
 	ResetAnimationCurves();
 	MaterialParametersToClear.Reset();
@@ -1020,9 +1021,14 @@ void FAnimInstanceProxy::UpdateSlotNodeWeight(const FName& SlotNodeName, float I
 }
 
 bool FAnimInstanceProxy::GetSlotInertializationRequest(const FName& SlotName, UE::Anim::FSlotInertializationRequest& OutRequest)
+{
+	return false;
+}
+
+bool FAnimInstanceProxy::GetSlotInertializationRequestData(const FName& SlotName, FInertializationRequest& OutRequest)
 {
 	const FName GroupName = Skeleton ? Skeleton->GetSlotGroupName(SlotName) : NAME_None;
-	if (const UE::Anim::FSlotInertializationRequest* FoundRequest = GetSlotGroupInertializationRequestMap().Find(GroupName))
+	if (const FInertializationRequest* FoundRequest = GetSlotGroupInertializationRequestDataMap().Find(GroupName))
 	{
 		OutRequest = *FoundRequest;
 		return true;
@@ -2140,6 +2146,9 @@ const FMontageEvaluationState* FAnimInstanceProxy::GetActiveMontageEvaluationSta
 	return nullptr;
 }
 
+
+PRAGMA_DISABLE_DEPRECATION_WARNINGS
+
 TMap<FName, UE::Anim::FSlotInertializationRequest>& FAnimInstanceProxy::GetSlotGroupInertializationRequestMap()
 {
 	if (bUseMainInstanceMontageEvaluationData && GetMainInstanceProxy())
@@ -2150,6 +2159,18 @@ TMap<FName, UE::Anim::FSlotInertializationRequest>& FAnimInstanceProxy::GetSlotG
 	return SlotGroupInertializationRequestMap;
 }
 
+PRAGMA_ENABLE_DEPRECATION_WARNINGS
+
+TMap<FName, FInertializationRequest>& FAnimInstanceProxy::GetSlotGroupInertializationRequestDataMap()
+{
+	if (bUseMainInstanceMontageEvaluationData && GetMainInstanceProxy())
+	{
+		return GetMainInstanceProxy()->SlotGroupInertializationRequestDataMap;
+	}
+
+	return SlotGroupInertializationRequestDataMap;
+}
+
 void FAnimInstanceProxy::GatherDebugData(FNodeDebugData& DebugData)
 {
 	GatherDebugData_WithRoot(DebugData, RootNode, NAME_AnimGraph);
diff --git a/Engine/Source/Runtime/Engine/Private/Animation/AnimMontage.cpp b/Engine/Source/Runtime/Engine/Private/Animation/AnimMontage.cpp
--- a/Engine/Source/Runtime/Engine/Private/Animation/AnimMontage.cpp
+++ b/Engine/Source/Runtime/Engine/Private/Animation/AnimMontage.cpp
@@ -1477,10 +1477,16 @@ void FAnimMontageInstance::Play(float InPlayRate, const FMontageBlendSettings& B
 	FAlphaBlendArgs BlendInArgs = BlendInSettings.Blend;
 	if (AnimInstance.IsValid() && BlendInSettings.BlendMode == EMontageBlendMode::Inertialization)
 	{
-		const float InertialBlendDuration = BlendInArgs.BlendTime;
+		FInertializationRequest Request;
+		Request.Duration = BlendInArgs.BlendTime;
+		Request.BlendMode = BlendInArgs.BlendOption;
+		Request.CustomBlendCurve = BlendInArgs.CustomCurve;
+		Request.BlendProfile = BlendInSettings.BlendProfile;
+		Request.bUseBlendMode = (BlendInArgs.BlendOption != EAlphaBlendOption::Linear) || (BlendInArgs.CustomCurve != nullptr);
+
 		// Request new inertialization for new montage's group name
 		// If there is an existing inertialization request, we overwrite that here.
-		AnimInstance->RequestMontageInertialization(Montage, InertialBlendDuration, BlendInSettings.BlendProfile);
+		AnimInstance->RequestMontageInertialization(Montage, Request);
 
 		// When using inertialization, we need to instantly blend in.
 		BlendInArgs.BlendTime = 0.0f;
@@ -1547,8 +1553,15 @@ void FAnimMontageInstance::Stop(const FMontageBlendSettings& InBlendOutSettings,
 
 				if (bShouldInertialize)
 				{
+					FInertializationRequest Request;
+					Request.Duration = InBlendOutSettings.Blend.BlendTime;
+					Request.BlendMode = InBlendOutSettings.Blend.BlendOption;
+					Request.CustomBlendCurve = InBlendOutSettings.Blend.CustomCurve;
+					Request.BlendProfile = InBlendOutSettings.BlendProfile;
+					Request.bUseBlendMode = (InBlendOutSettings.Blend.BlendOption != EAlphaBlendOption::Linear) || (InBlendOutSettings.Blend.CustomCurve != nullptr);
+
 					// Send the inertial blend request to the anim instance
-					Inst->RequestMontageInertialization(Montage, InBlendOutSettings.Blend.BlendTime, InBlendOutSettings.BlendProfile);
+					Inst->RequestMontageInertialization(Montage, Request);
 				}
 			}
 		}
diff --git a/Engine/Source/Runtime/Engine/Public/Animation/AnimInstanceProxy.h b/Engine/Source/Runtime/Engine/Public/Animation/AnimInstanceProxy.h
--- a/Engine/Source/Runtime/Engine/Public/Animation/AnimInstanceProxy.h
+++ b/Engine/Source/Runtime/Engine/Public/Animation/AnimInstanceProxy.h
@@ -11,6 +11,7 @@
 #include "Animation/AnimClassInterface.h"
 #include "Animation/AnimTrace.h"
 #include "Animation/AnimSync.h"
+#include "Animation/AnimInertializationRequest.h"
 
 #include "AnimInstanceProxy.generated.h"
 
@@ -45,6 +46,7 @@ namespace UE::Anim
 	class FAnimSyncGroupScope;
 	class FActiveStateMachineScope;
 	struct FAnimSyncParams;
+	// DEPRECATED use FInertializationRequest instead
 	using FSlotInertializationRequest = TPair<float, const UBlendProfile*>;
 }
 
@@ -423,8 +425,11 @@ public:
 	// Allow slot nodes to store off their weight during ticking
 	ENGINE_API void UpdateSlotNodeWeight(const FName& SlotNodeName, float InLocalMontageWeight, float InNodeGlobalWeight);
 
+	UE_DEPRECATED(5.5, "Please use GetSlotInertializationRequestData instead.")
 	ENGINE_API bool GetSlotInertializationRequest(const FName& SlotName, UE::Anim::FSlotInertializationRequest& OutRequest);
 
+	ENGINE_API bool GetSlotInertializationRequestData(const FName& SlotName, FInertializationRequest& OutRequest);
+
 	/** Register a named slot */
 	ENGINE_API void RegisterSlotNodeWithAnimInstance(const FName& SlotNodeName);
 
@@ -772,7 +777,10 @@ protected:
 		Note that there might be multiple Active at the same time. This will only return the first active one it finds. **/
 	ENGINE_API const FMontageEvaluationState* GetActiveMontageEvaluationState() const;
 
+	UE_DEPRECATED(5.5, "Please use GetSlotGroupInertializationRequestDataMap instead.")
 	ENGINE_API TMap<FName, UE::Anim::FSlotInertializationRequest>& GetSlotGroupInertializationRequestMap();
+	
+	ENGINE_API TMap<FName, FInertializationRequest>& GetSlotGroupInertializationRequestDataMap();
 
 	/** Access montage array data */
 	ENGINE_API TArray<FMontageEvaluationState>& GetMontageEvaluationData();
@@ -1122,9 +1130,13 @@ private:
 	/** Copy of UAnimInstance::MontageInstances data used for update & evaluation */
 	TArray<FMontageEvaluationState> MontageEvaluationData;
 
-	// Inertialization request for each slot.
+	UE_DEPRECATED(5.5, "This property is deprecated. Please use SlotGroupInertializationRequestDataMap instead")
 	TMap<FName, UE::Anim::FSlotInertializationRequest> SlotGroupInertializationRequestMap;
 
+	// Inertialization request for each slot.
+	UPROPERTY(Transient)
+	TMap<FName, FInertializationRequest> SlotGroupInertializationRequestDataMap;
+
 	/** Delegate fired on the game thread before update occurs */
 	TArray<FAnimNode_Base*> GameThreadPreUpdateNodes;
 
