Update to Legacy Material Layering to support Substrate blend functionality.
- Currently disabled by the macro ENABLE_MATERIAL_LAYER_PROTOTYPE until we have a fully fleshed out system to enable.
- Implemented 2 new Substrate nodes: Get and Set attributes. Get allows the embedding of a Front Material with general Material Attributes for passing between functions, and Set splits the Material Attributes back into Front Material and Non-Substrate attributes. These are not strictly for use in Layering. This allows the Material Attributes including Substrate data to pass between Layers and Blends smoothly.
- Changed the autopopulation of Layer and Blend function material graphs to utilise the Get and Set attributes nodes. Non Substrate Attributes continue to pass through the legacy blend node, which only supports Horizontal blending. The Front Material is now separated and passes through the Substrate Horizontal Mixing node. This all occurs via a new engine level Blend Material Function MLB_Horizontal.
- Created update paths to convert legacy material layers to utilise the new Substrate structure (automatically converts to Horizontal mixing only)
- Created these structures using the "MaterialFunctionUsage" flags at the Function level -> This will ensure the ability to create contextual re-use of material functions regardless of whether they are marked layer or not.
- Altered the Preview functionality of Blend nodes so the expected output is clearer (uses the Input's embedded PreviewValue). Created a static default substrate slab creator function which supports this with any override color.
- Changed the colour of the Substrate utility nodes to match Substrate

#virtualized

diff --git a/Engine/Config/BaseEngine.ini b/Engine/Config/BaseEngine.ini
--- a/Engine/Config/BaseEngine.ini
+++ b/Engine/Config/BaseEngine.ini
@@ -1946,6 +1946,7 @@ AutoSaveIndex=0
 +PackagesToBeFullyLoadedAtStartup=/Engine/EngineMaterials/RemoveSurfaceMaterial
 +PackagesToBeFullyLoadedAtStartup=/Engine/EngineMaterials/WeightMapPlaceholderTexture
 +PackagesToBeFullyLoadedAtStartup=/Engine/EngineMaterials/NaniteHiddenSectionMaterial
++PackagesToBeFullyLoadedAtStartup=/Engine/EngineMaterials/Blends/MLB_Horizontal
 
 +PackagesToBeFullyLoadedAtStartup=/Engine/EngineFonts/SmallFont
 +PackagesToBeFullyLoadedAtStartup=/Engine/EngineFonts/TinyFont
diff --git a/Engine/Source/Editor/MaterialEditor/Private/MaterialEditor.cpp b/Engine/Source/Editor/MaterialEditor/Private/MaterialEditor.cpp
--- a/Engine/Source/Editor/MaterialEditor/Private/MaterialEditor.cpp
+++ b/Engine/Source/Editor/MaterialEditor/Private/MaterialEditor.cpp
@@ -753,6 +753,7 @@ void FMaterialEditor::InitMaterialEditor( const EToolkitMode::Type Mode, const T
 				else
 				{
 					Expression = CreateNewMaterialExpression(UMaterialExpressionMaterialLayerOutput::StaticClass(), OutputPlacement, false, true);
+					Expression->bCollapsed = true;
 					SetPreviewExpression(Expression);
 					// This shouldn't count as having dirtied the material, so reset the flag
 					bMaterialDirty = false;
@@ -760,55 +761,171 @@ void FMaterialEditor::InitMaterialEditor( const EToolkitMode::Type Mode, const T
 				// We can check the usage here and add the appropriate inputs too (e.g. Layer==1MA, Blend==2MA)
 				if (MaterialFunction->GetMaterialFunctionUsage() == EMaterialFunctionUsage::MaterialLayer)
 				{
-					UMaterialExpression* Input = CreateNewMaterialExpression(UMaterialExpressionFunctionInput::StaticClass(), FVector2D(-350, 300), false, true);
-					if (Input)
+#if ENABLE_MATERIAL_LAYER_PROTOTYPE
+					if (Substrate::IsSubstrateEnabled())
 					{
-						UMaterialExpressionFunctionInput* BaseAttributesInput = Cast<UMaterialExpressionFunctionInput>(Input);
-						BaseAttributesInput->InputType = FunctionInput_MaterialAttributes;
-						BaseAttributesInput->InputName = TEXT("Material Attributes");
-						BaseAttributesInput->bUsePreviewValueAsDefault = true;
+						UMaterialExpression* Input = CreateNewMaterialExpression(UMaterialExpressionFunctionInput::StaticClass(), FVector2D(-800, 300), false, true);
+						if (Input)
+						{
+							UMaterialExpressionFunctionInput* BaseAttributesInput = Cast<UMaterialExpressionFunctionInput>(Input);
+							BaseAttributesInput->InputType = FunctionInput_MaterialAttributes;
+							BaseAttributesInput->InputName = TEXT("Material Attributes");
+							BaseAttributesInput->bUsePreviewValueAsDefault = true;
+							BaseAttributesInput->PreviewValue = FLinearColor::MakeRandomColor();
+						}
+
+						UMaterialExpression* SetSubstrateAttributes = CreateNewMaterialExpression(UMaterialExpressionSubstrateSetAttributes::StaticClass(), FVector2D(-50, 300), false, true);
+						UMaterialExpression* GetSubstrateAttributes = CreateNewMaterialExpression(UMaterialExpressionSubstrateGetAttributes::StaticClass(), FVector2D(-450, 300), false, true);
+						if (Input && SetSubstrateAttributes && GetSubstrateAttributes)
+						{
+							UMaterialEditingLibrary::ConnectMaterialExpressions(Input, FString(), GetSubstrateAttributes, FString());
+						
+							UMaterialEditingLibrary::ConnectMaterialExpressions(GetSubstrateAttributes, FRONT_MATERIAL_ATTRIBUTES_TEXT, SetSubstrateAttributes, FRONT_MATERIAL_ATTRIBUTES_TEXT);
+							UMaterialEditingLibrary::ConnectMaterialExpressions(GetSubstrateAttributes, NON_SUBSTRATE_ATTRIBUTES_TEXT, SetSubstrateAttributes, NON_SUBSTRATE_ATTRIBUTES_TEXT);
+
+							UMaterialEditingLibrary::ConnectMaterialExpressions(SetSubstrateAttributes, FString(), Expression, FString());
+							bMaterialDirty = true;
+						}
 
 					}
-					if (GetDefault<UEditorExperimentalSettings>()->bExampleLayersAndBlends)
+					else
+#endif //ENABLE_MATERIAL_LAYER_PROTOTYPE
 					{
-						UMaterialExpression* SetMaterialAttributes = CreateNewMaterialExpression(UMaterialExpressionSetMaterialAttributes::StaticClass(), FVector2D(40, 300), false, true);
-						if (Input && SetMaterialAttributes)
+						UMaterialExpression* Input = CreateNewMaterialExpression(UMaterialExpressionFunctionInput::StaticClass(), FVector2D(-350, 300), false, true);
+						if (Input)
+						{
+							UMaterialExpressionFunctionInput* BaseAttributesInput = Cast<UMaterialExpressionFunctionInput>(Input);
+							BaseAttributesInput->InputType = FunctionInput_MaterialAttributes;
+							BaseAttributesInput->InputName = TEXT("Material Attributes");
+							BaseAttributesInput->bUsePreviewValueAsDefault = true;
+						}
+
+						if (GetDefault<UEditorExperimentalSettings>()->bExampleLayersAndBlends)
 						{
+							UMaterialExpression* SetMaterialAttributes = CreateNewMaterialExpression(UMaterialExpressionSetMaterialAttributes::StaticClass(), FVector2D(40, 300), false, true);
+							if (Input && SetMaterialAttributes)
+							{
 							UMaterialEditingLibrary::ConnectMaterialExpressions(Input, FString(), SetMaterialAttributes, FString());
-							UMaterialEditingLibrary::ConnectMaterialExpressions(SetMaterialAttributes, FString(), Expression, FString());
-							bMaterialDirty = true;
+								UMaterialEditingLibrary::ConnectMaterialExpressions(SetMaterialAttributes, FString(), Expression, FString());
+								bMaterialDirty = true;
+							}
 						}
 					}
 				}
 				else if (MaterialFunction->GetMaterialFunctionUsage() == EMaterialFunctionUsage::MaterialLayerBlend)
 				{
-					// "Top layer" should be below "bottom layer" on the graph, to align with B on blend nodes
-					UMaterialExpression* InputTop = CreateNewMaterialExpression(UMaterialExpressionFunctionInput::StaticClass(), FVector2D(-300, 400), false, true);
-					if (InputTop)
+#if ENABLE_MATERIAL_LAYER_PROTOTYPE
+					if (Substrate::IsSubstrateEnabled())
 					{
-						UMaterialExpressionFunctionInput* BaseAttributesInput = Cast<UMaterialExpressionFunctionInput>(InputTop);
-						BaseAttributesInput->InputType = FunctionInput_MaterialAttributes;
-						BaseAttributesInput->InputName = TEXT("Top Layer");
-						BaseAttributesInput->bUsePreviewValueAsDefault = true;
-					}
+						{
+							UMaterialExpression* InputBottom = CreateNewMaterialExpression(UMaterialExpressionFunctionInput::StaticClass(), FVector2D(-1100, 200), false, true);
+							if (InputBottom)
+							{
+								UMaterialExpressionFunctionInput* BaseAttributesInput = Cast<UMaterialExpressionFunctionInput>(InputBottom);
+								BaseAttributesInput->InputType = FunctionInput_MaterialAttributes;
+								BaseAttributesInput->InputName = TEXT("Background Layer");
+								BaseAttributesInput->bUsePreviewValueAsDefault = true;
+							}
 
-					UMaterialExpression* InputBottom = CreateNewMaterialExpression(UMaterialExpressionFunctionInput::StaticClass(), FVector2D(-300, 200), false, true);
-					if (InputBottom)
-					{
-						UMaterialExpressionFunctionInput* BaseAttributesInput = Cast<UMaterialExpressionFunctionInput>(InputBottom);
-						BaseAttributesInput->InputType = FunctionInput_MaterialAttributes;
-						BaseAttributesInput->InputName = TEXT("Bottom Layer");
-						BaseAttributesInput->bUsePreviewValueAsDefault = true;
+							// "Top layer" should be below "bottom layer" on the graph, to align with Foreground/B on blend nodes
+							UMaterialExpression* InputTop = CreateNewMaterialExpression(UMaterialExpressionFunctionInput::StaticClass(), FVector2D(-1100, 400), false, true);
+							if (InputTop)
+							{
+								UMaterialExpressionFunctionInput* BaseAttributesInput = Cast<UMaterialExpressionFunctionInput>(InputTop);
+								BaseAttributesInput->InputType = FunctionInput_MaterialAttributes;
+								BaseAttributesInput->InputName = TEXT("Foreground Layer");
+								BaseAttributesInput->bUsePreviewValueAsDefault = true;
+								BaseAttributesInput->PreviewValue = FLinearColor::White;
+							}
+
+							if(InputTop && InputBottom && GetDefault<UEditorExperimentalSettings>()->bExampleLayersAndBlends)
+							{
+								static TObjectPtr<UMaterialFunction> DefaultBlendFunction = FindObject<UMaterialFunction>(nullptr, DEFAULT_MATERIALLAYERBLEND_PATH);
+								if (!DefaultBlendFunction)
+								{
+									DefaultBlendFunction = LoadObject<UMaterialFunction>(nullptr, DEFAULT_MATERIALLAYERBLEND_PATH);
+								}
+							
+								if (DefaultBlendFunction)
+								{
+									if(UMaterialExpressionMaterialFunctionCall* BlendFunctionCall = Cast<UMaterialExpressionMaterialFunctionCall>(CreateNewMaterialExpression(UMaterialExpressionMaterialFunctionCall::StaticClass(), FVector2D(-100, 300), false, false)))
+									{
+										BlendFunctionCall->Function = MaterialFunction;
+										InputTop->MaterialExpressionEditorX = -500;
+										InputBottom->MaterialExpressionEditorX = -500;
+										if (BlendFunctionCall->SetMaterialFunction(DefaultBlendFunction))
+										{
+											if (BlendFunctionCall->FunctionInputs.Num() >= 2 && BlendFunctionCall->FunctionOutputs.Num() > 0)
+											{
+												BlendFunctionCall->FunctionInputs[0].Input.Connect(0, InputBottom);
+												BlendFunctionCall->FunctionInputs[1].Input.Connect(0, InputTop);
+												UMaterialEditingLibrary::ConnectMaterialExpressions(BlendFunctionCall, FString(), Expression, FString());
+												bMaterialDirty = true;
+											}
+										}
+									}
+								}
+
+								if (!bMaterialDirty)
+								{
+									UMaterialExpression* SetSubstrateAttributes = CreateNewMaterialExpression(UMaterialExpressionSubstrateSetAttributes::StaticClass(), FVector2D(-75, 300), false, true);
+									UMaterialExpression* HorizontalMixingNode = CreateNewMaterialExpression(UMaterialExpressionSubstrateHorizontalMixing::StaticClass(), FVector2D(-350, 200), false, true);
+									UMaterialExpression* GetBottomSubstrateAttributes = CreateNewMaterialExpression(UMaterialExpressionSubstrateGetAttributes::StaticClass(), FVector2D(-750, 200), false, true);
+
+									UMaterialExpression* LegacyBlendNode = CreateNewMaterialExpression(UMaterialExpressionBlendMaterialAttributes::StaticClass(), FVector2D(-350, 400), false, true);
+									UMaterialExpression* GetTopSubstrateAttributes = CreateNewMaterialExpression(UMaterialExpressionSubstrateGetAttributes::StaticClass(), FVector2D(-750, 400), false, true);
+
+									if (SetSubstrateAttributes && HorizontalMixingNode && GetBottomSubstrateAttributes && LegacyBlendNode && GetTopSubstrateAttributes)
+									{
+										UMaterialEditingLibrary::ConnectMaterialExpressions(InputBottom, FString(), GetBottomSubstrateAttributes, FString());
+										UMaterialEditingLibrary::ConnectMaterialExpressions(GetBottomSubstrateAttributes, FRONT_MATERIAL_ATTRIBUTES_TEXT, HorizontalMixingNode, FString("Background"));
+										UMaterialEditingLibrary::ConnectMaterialExpressions(GetBottomSubstrateAttributes, NON_SUBSTRATE_ATTRIBUTES_TEXT, LegacyBlendNode, FString("A"));
+
+										UMaterialEditingLibrary::ConnectMaterialExpressions(InputTop, FString(), GetTopSubstrateAttributes, FString());
+										UMaterialEditingLibrary::ConnectMaterialExpressions(GetTopSubstrateAttributes, FRONT_MATERIAL_ATTRIBUTES_TEXT, HorizontalMixingNode, FString("Foreground"));
+										UMaterialEditingLibrary::ConnectMaterialExpressions(GetTopSubstrateAttributes, NON_SUBSTRATE_ATTRIBUTES_TEXT, LegacyBlendNode, FString("B"));
+
+										UMaterialEditingLibrary::ConnectMaterialExpressions(HorizontalMixingNode, FString(), SetSubstrateAttributes, FRONT_MATERIAL_ATTRIBUTES_TEXT);
+										UMaterialEditingLibrary::ConnectMaterialExpressions(LegacyBlendNode, FString(), SetSubstrateAttributes, NON_SUBSTRATE_ATTRIBUTES_TEXT);
+
+										UMaterialEditingLibrary::ConnectMaterialExpressions(SetSubstrateAttributes, FString(), Expression, FString());
+										bMaterialDirty = true;
+									}
+								}
+							}
+						}
 					}
-					if (GetDefault<UEditorExperimentalSettings>()->bExampleLayersAndBlends)
+					else
+#endif //ENABLE_MATERIAL_LAYER_PROTOTYPE
 					{
-						UMaterialExpression* BlendMaterialAttributes = CreateNewMaterialExpression(UMaterialExpressionBlendMaterialAttributes::StaticClass(), FVector2D(40, 300), false, true);
-						if (InputTop && InputBottom && BlendMaterialAttributes)
+						// "Top layer" should be below "bottom layer" on the graph, to align with B on blend nodes
+						UMaterialExpression* InputTop = CreateNewMaterialExpression(UMaterialExpressionFunctionInput::StaticClass(), FVector2D(-300, 400), false, true);
+						if (InputTop)
 						{
-							UMaterialEditingLibrary::ConnectMaterialExpressions(InputBottom, FString(), BlendMaterialAttributes, FString(TEXT("A")));
-							UMaterialEditingLibrary::ConnectMaterialExpressions(InputTop, FString(), BlendMaterialAttributes, FString(TEXT("B")));
-							UMaterialEditingLibrary::ConnectMaterialExpressions(BlendMaterialAttributes, FString(), Expression, FString());
-							bMaterialDirty = true;
+							UMaterialExpressionFunctionInput* BaseAttributesInput = Cast<UMaterialExpressionFunctionInput>(InputTop);
+							BaseAttributesInput->InputType = FunctionInput_MaterialAttributes;
+							BaseAttributesInput->InputName = TEXT("Top Layer");
+							BaseAttributesInput->bUsePreviewValueAsDefault = true;
+						}
+
+						UMaterialExpression* InputBottom = CreateNewMaterialExpression(UMaterialExpressionFunctionInput::StaticClass(), FVector2D(-300, 200), false, true);
+						if (InputBottom)
+						{
+							UMaterialExpressionFunctionInput* BaseAttributesInput = Cast<UMaterialExpressionFunctionInput>(InputBottom);
+							BaseAttributesInput->InputType = FunctionInput_MaterialAttributes;
+							BaseAttributesInput->InputName = TEXT("Bottom Layer");
+							BaseAttributesInput->bUsePreviewValueAsDefault = true;
+						}
+						if (GetDefault<UEditorExperimentalSettings>()->bExampleLayersAndBlends)
+						{
+							UMaterialExpression* BlendMaterialAttributes = CreateNewMaterialExpression(UMaterialExpressionBlendMaterialAttributes::StaticClass(), FVector2D(40, 300), false, true);
+							if (InputTop && InputBottom && BlendMaterialAttributes)
+							{
+								UMaterialEditingLibrary::ConnectMaterialExpressions(InputBottom, FString(), BlendMaterialAttributes, FString(TEXT("A")));
+								UMaterialEditingLibrary::ConnectMaterialExpressions(InputTop, FString(), BlendMaterialAttributes, FString(TEXT("B")));
+								UMaterialEditingLibrary::ConnectMaterialExpressions(BlendMaterialAttributes, FString(), Expression, FString());
+								bMaterialDirty = true;
+							}
 						}
 					}
 				}
@@ -3003,7 +3120,11 @@ void FMaterialEditor::UpdateMaterialinfoList_Old()
 			int32 NumOutputs = 0;
 			// For Material Layers
 
-			if (MaterialFunction->GetMaterialFunctionUsage() == EMaterialFunctionUsage::MaterialLayer)
+			if (MaterialFunction->GetMaterialFunctionUsage() == EMaterialFunctionUsage::MaterialLayer
+#if ENABLE_MATERIAL_LAYER_PROTOTYPE
+			&& !Substrate::IsSubstrateEnabled()
+#endif
+			)
 			{
 				// Material layers must have a single MA input and output only
 				for (UMaterialExpression* Expression : MaterialFunction->GetExpressions())
@@ -3039,7 +3160,11 @@ void FMaterialEditor::UpdateMaterialinfoList_Old()
 					FailingExpression.Add(nullptr);
 				}
 			}
-			else if (MaterialFunction->GetMaterialFunctionUsage() == EMaterialFunctionUsage::MaterialLayerBlend)
+			else if (MaterialFunction->GetMaterialFunctionUsage() == EMaterialFunctionUsage::MaterialLayerBlend
+#if ENABLE_MATERIAL_LAYER_PROTOTYPE
+			&& !Substrate::IsSubstrateEnabled()
+#endif
+			)
 			{
 				// Material layer blends can have two MA inputs and single MA output only
 				for (UMaterialExpression* Expression : MaterialFunction->GetExpressions())
diff --git a/Engine/Source/Editor/UnrealEd/Private/MaterialGraphNode.cpp b/Engine/Source/Editor/UnrealEd/Private/MaterialGraphNode.cpp
--- a/Engine/Source/Editor/UnrealEd/Private/MaterialGraphNode.cpp
+++ b/Engine/Source/Editor/UnrealEd/Private/MaterialGraphNode.cpp
@@ -325,7 +325,8 @@ FLinearColor UMaterialGraphNode::GetNodeTitleColor() const
 			return FColor( 0, 128, 128 );
 		}
 	}
-	else if (const UMaterialExpressionSubstrateBSDF* SubstrateBSDF = Cast<UMaterialExpressionSubstrateBSDF>(MaterialExpression))
+	else if (MaterialExpression->IsA<UMaterialExpressionSubstrateBSDF>() 
+		|| MaterialExpression->IsA<UMaterialExpressionSubstrateUtilityBase>())
 	{
 		return FColor(181, 29, 230);
 	}
diff --git a/Engine/Source/Runtime/Engine/Internal/Materials/MaterialExpressionMaterialSample.h b/Engine/Source/Runtime/Engine/Internal/Materials/MaterialExpressionMaterialSample.h
--- a/Engine/Source/Runtime/Engine/Internal/Materials/MaterialExpressionMaterialSample.h
+++ b/Engine/Source/Runtime/Engine/Internal/Materials/MaterialExpressionMaterialSample.h
@@ -11,11 +11,6 @@
 
 #include "MaterialExpressionMaterialSample.generated.h"
 
-#ifndef ENABLE_MATERIAL_LAYER_PROTOTYPE
-//WIP Internal Only - Please do not enable + use this in any projects other than for iterating the new layering prototype system until further notice
-#define ENABLE_MATERIAL_LAYER_PROTOTYPE 0
-#endif
-
 class FMaterialCompiler;
 
 UCLASS(MinimalAPI)
diff --git a/Engine/Source/Runtime/Engine/Private/Materials/Material.cpp b/Engine/Source/Runtime/Engine/Private/Materials/Material.cpp
--- a/Engine/Source/Runtime/Engine/Private/Materials/Material.cpp
+++ b/Engine/Source/Runtime/Engine/Private/Materials/Material.cpp
@@ -3367,201 +3367,373 @@ void UMaterial::ConvertMaterialToSubstrateMaterial()
 	bool bRelinkCustomOutputNodes = false;
 	UMaterialExpressionSubstrateShadingModels* ConvertNode = nullptr;
 	// Connect all the legacy pin into the conversion node
-	if (bUseMaterialAttributes && EditorOnly->MaterialAttributes.Expression && !EditorOnly->FrontMaterial.IsConnected() && !EditorOnly->MaterialAttributes.Expression->IsResultSubstrateMaterial(EditorOnly->MaterialAttributes.OutputIndex)) // M_Rifle cause issues there
+#if ENABLE_MATERIAL_LAYER_PROTOTYPE
+	//First check if this material contains a layer node connection, as that requires different upgrade behavior to convert to the Substrate material layer system
+	bool bLayersMaterial = false;
+	if (bUseMaterialAttributes && EditorOnly->MaterialAttributes.Expression && !EditorOnly->FrontMaterial.IsConnected())
 	{
-		UMaterialExpressionSubstrateConvertMaterialAttributes* ConvertAttributeNode = NewObject<UMaterialExpressionSubstrateConvertMaterialAttributes>(this);
-		ConvertAttributeNode->Material = this;
-		SetPosXAndMoveReferenceToTheRight(ConvertAttributeNode);
-		ConvertAttributeNode->SubsurfaceProfile = bRequireNoSubsurfaceProfile ? nullptr : SubsurfaceProfile;
-
-		// * Copy the material attribute connection to the conversion node.
-		// * Leave the material attribute existing connection plugged to the root node, 
-		//   so that other input (PixelDepthOffset, WorldPositionOffset, ...) get pull 
-		//   from the material attributes node
-		ConnectionTo(EditorOnly->MaterialAttributes, ConvertAttributeNode, 0, SUBSTRATE_COPY_CONNECTION);
-
-		// Reconnect custom output to material attribute conversion node
+		UMaterialExpression* FinalNode = nullptr;
+		UMaterialExpressionMaterialAttributeLayers* LayersNode = nullptr;
+		for(UMaterialExpression* Expression : GetExpressions())
 		{
-			check(ConvertAttributeNode);
-			GatherCustomNodes();
-
-			if (SingleLayerWaterOutput)
+			if(Expression->IsA<UMaterialExpressionMaterialAttributeLayers>())
 			{
-				ConnectionTo(*SingleLayerWaterOutput->GetInput(0), ConvertAttributeNode, 1); // WaterScatteringCoefficients
-				ConnectionTo(*SingleLayerWaterOutput->GetInput(1), ConvertAttributeNode, 2); // WaterAbsorptionCoefficients
-				ConnectionTo(*SingleLayerWaterOutput->GetInput(2), ConvertAttributeNode, 3); // WaterPhaseG
-				ConnectionTo(*SingleLayerWaterOutput->GetInput(3), ConvertAttributeNode, 4); // ColorScaleBehindWater
+				LayersNode = Cast<UMaterialExpressionMaterialAttributeLayers>(Expression);
+				FinalNode = EditorOnly->MaterialAttributes.Expression;
+				bLayersMaterial = true;
+				break;
 			}
 		}
 
-		// Connect converted Substrate data to root node
-		EditorOnly->FrontMaterial.Connect(0, ConvertAttributeNode);
-
-		// Shading Model
-		// * either use the shader graph expression 
-		// * or add a constant shading model
-		if (ShadingModel == MSM_FromMaterialExpression)
+		if(FinalNode && LayersNode)
 		{
-			ConvertAttributeNode->ShadingModelOverride = MSM_FromMaterialExpression;
-		}
-		else
-		{
-			// Store Substrate shading model of the converted material. 
-			check(ShadingModels.CountShadingModels() == 1);
-			ConvertAttributeNode->ShadingModelOverride = ShadingModel;
-		}
+			//Separate the layers/blends MAs to Substrate/Non-Substrate respectively and connect to their relevant material output
+			UMaterialExpressionSubstrateGetAttributes* GetAttributesNode = NewObject<UMaterialExpressionSubstrateGetAttributes>(this);
+			GetAttributesNode->Material = this;
+			GetAttributesNode->MaterialAttributes.Connect(0, FinalNode);
+			EditorOnly->FrontMaterial.Connect(ESubstrateAttributeIndex::MSA_FrontMaterial, GetAttributesNode);
+			EditorOnly->MaterialAttributes.Connect(ESubstrateAttributeIndex::MSA_NonSubstrateAttributes, GetAttributesNode);
+			SetPosXAndMoveReferenceToTheRight(GetAttributesNode);
 
-		if (MaterialDomain == MD_DeferredDecal)
-		{
-			// For now we don't enforce shading model since it could be driven by expression and we don't have much 
-			// control on this, but only DefaultLit should be supported.
+			//If a non-Substrate default input is connected to the MLA node, convert it.
+			if (LayersNode->Input.IsConnected() && !LayersNode->Input.Expression->IsResultSubstrateMaterial(LayersNode->Input.OutputIndex))
+			{
+				UMaterialExpressionSubstrateConvertMaterialAttributes* ConvertAttributeNode = NewObject<UMaterialExpressionSubstrateConvertMaterialAttributes>(this);
+				ConvertAttributeNode->Material = this;
+				ConvertAttributeNode->MaterialAttributes.Connect(LayersNode->Input.OutputIndex, LayersNode->Input.Expression);
+				ConvertAttributeNode->ShadingModelOverride = MSM_DefaultLit;
 
-			// Now pass through the convert to decal node, which flag the material as SSM_Decal, which will set the domain to Decal.
-			UMaterialExpressionSubstrateConvertToDecal* ConvertToDecalNode = NewObject<UMaterialExpressionSubstrateConvertToDecal>(this);
-			ConvertAttributeNode->Material = this;
-			ReplaceNodeAndMoveToTheRight(ConvertAttributeNode, ConvertToDecalNode);
-			ConvertToDecalNode->DecalMaterial.Connect(0, ConvertAttributeNode);
+				UMaterialExpressionSubstrateSetAttributes* SetAttributesNode = NewObject<UMaterialExpressionSubstrateSetAttributes>(this);
+				SetAttributesNode->Material = this;
+				SetAttributesNode->FrontMaterial.Connect(0, ConvertAttributeNode);
+				SetAttributesNode->NonSubstrateAttributes.Connect(LayersNode->Input.OutputIndex, LayersNode->Input.Expression);
 
-			EditorOnly->FrontMaterial.Connect(0, ConvertToDecalNode);
-		}
+				LayersNode->Input.Connect(0, SetAttributesNode);
 
-		BlendMode = ConvertLegacyBlendMode(BlendMode, ShadingModels);
-		RefractionCoverageMode = RCM_CoverageIgnored;
-		bInvalidateShader = true;
+				SetPosXAndMoveReferenceToTheRight(GetAttributesNode);
+				ReplaceNodeAndMoveToTheRight(LayersNode, ConvertAttributeNode);
+				SetPosXAndMoveReferenceToTheRight(GetAttributesNode);
+				ReplaceNodeAndMoveToTheRight(LayersNode, SetAttributesNode);			
+			}
+
+			BlendMode = ConvertLegacyBlendMode(BlendMode, ShadingModels);
+			RefractionCoverageMode = RCM_CoverageIgnored;
+			bInvalidateShader = true;
+		}
 	}
-	else if (!bUseMaterialAttributes && !EditorOnly->FrontMaterial.IsConnected() && GetExpressions().IsEmpty())
+	
+	if(!bLayersMaterial)
+#endif //ENABLE_MATERIAL_LAYER_PROTOTYPE
 	{
-		// Empty material: Create by default a slab node
-		UMaterialFunction* DefaultMF = LoadObject<UMaterialFunction>(nullptr, TEXT("/Engine/Functions/Substrate/SMF_UE4Disney.SMF_UE4Disney"));
-		if (DefaultMF)
+		if (bUseMaterialAttributes && EditorOnly->MaterialAttributes.Expression && !EditorOnly->FrontMaterial.IsConnected() && !EditorOnly->MaterialAttributes.Expression->IsResultSubstrateMaterial(EditorOnly->MaterialAttributes.OutputIndex)) // M_Rifle cause issues there
 		{
-			DefaultMF->UpdateFromFunctionResource();
-			DefaultMF->PostEditChange();
-			DefaultMF->ConditionalPostLoad();
+			UMaterialExpressionSubstrateConvertMaterialAttributes* ConvertAttributeNode = NewObject<UMaterialExpressionSubstrateConvertMaterialAttributes>(this);
+			ConvertAttributeNode->Material = this;
+			SetPosXAndMoveReferenceToTheRight(ConvertAttributeNode);
+			ConvertAttributeNode->SubsurfaceProfile = bRequireNoSubsurfaceProfile ? nullptr : SubsurfaceProfile;
+
+			// * Copy the material attribute connection to the conversion node.
+			// * Leave the material attribute existing connection plugged to the root node, 
+			//   so that other input (PixelDepthOffset, WorldPositionOffset, ...) get pull 
+			//   from the material attributes node
+			ConnectionTo(EditorOnly->MaterialAttributes, ConvertAttributeNode, 0, SUBSTRATE_COPY_CONNECTION);
+
+			// Reconnect custom output to material attribute conversion node
+			{
+				check(ConvertAttributeNode);
+				GatherCustomNodes();
+
+				if (SingleLayerWaterOutput)
+				{
+					ConnectionTo(*SingleLayerWaterOutput->GetInput(0), ConvertAttributeNode, 1); // WaterScatteringCoefficients
+					ConnectionTo(*SingleLayerWaterOutput->GetInput(1), ConvertAttributeNode, 2); // WaterAbsorptionCoefficients
+					ConnectionTo(*SingleLayerWaterOutput->GetInput(2), ConvertAttributeNode, 3); // WaterPhaseG
+					ConnectionTo(*SingleLayerWaterOutput->GetInput(3), ConvertAttributeNode, 4); // ColorScaleBehindWater
+				}
+			}
+
+			// Connect converted Substrate data to root node
+			EditorOnly->FrontMaterial.Connect(0, ConvertAttributeNode);
 
-			UMaterialExpressionMaterialFunctionCall* MFCallNode = NewObject<UMaterialExpressionMaterialFunctionCall>(this);
-			if (MFCallNode->SetMaterialFunction(DefaultMF))
+			// Shading Model
+			// * either use the shader graph expression 
+			// * or add a constant shading model
+			if (ShadingModel == MSM_FromMaterialExpression)
+			{
+				ConvertAttributeNode->ShadingModelOverride = MSM_FromMaterialExpression;
+			}
+			else
 			{
-				// This is needed for input/output expressions to be set correctly, otherwise compilation will fail.
-				GetExpressionCollection().AddExpression(MFCallNode);
+				// Store Substrate shading model of the converted material. 
+				check(ShadingModels.CountShadingModels() == 1);
+				ConvertAttributeNode->ShadingModelOverride = ShadingModel;
+			}
 
-				SetPosXAndMoveReferenceToTheRight(MFCallNode);
-				EditorOnly->FrontMaterial.Connect(0, MFCallNode);
+			if (MaterialDomain == MD_DeferredDecal)
+			{
+				// For now we don't enforce shading model since it could be driven by expression and we don't have much 
+				// control on this, but only DefaultLit should be supported.
 
-				MFCallNode->UpdateFromFunctionResource();
-				MFCallNode->PostEditChange();
-				MFCallNode->ConditionalPostLoad();
+				// Now pass through the convert to decal node, which flag the material as SSM_Decal, which will set the domain to Decal.
+				UMaterialExpressionSubstrateConvertToDecal* ConvertToDecalNode = NewObject<UMaterialExpressionSubstrateConvertToDecal>(this);
+				ConvertAttributeNode->Material = this;
+				ReplaceNodeAndMoveToTheRight(ConvertAttributeNode, ConvertToDecalNode);
+				ConvertToDecalNode->DecalMaterial.Connect(0, ConvertAttributeNode);
 
-				ColorMatInputConnectionTo(EditorOnly->BaseColor,		MFCallNode, 0, MP_BaseColor);
-				ScalarMatInputConnectionTo(EditorOnly->Metallic,		MFCallNode, 1, MP_Metallic);
-				ScalarMatInputConnectionTo(EditorOnly->Specular,		MFCallNode, 2, MP_Specular);
-				ScalarMatInputConnectionTo(EditorOnly->Roughness,		MFCallNode, 3, MP_Roughness);
-				Vector3MatInputConnectionTo(EditorOnly->Normal,			MFCallNode, 4, MP_Normal);
-				ColorMatInputConnectionTo(EditorOnly->EmissiveColor,	MFCallNode, 5, MP_EmissiveColor);
-				ScalarMatInputConnectionTo(EditorOnly->Opacity,			MFCallNode, 6, MP_Opacity);
+				EditorOnly->FrontMaterial.Connect(0, ConvertToDecalNode);
 			}
+
+			BlendMode = ConvertLegacyBlendMode(BlendMode, ShadingModels);
+			RefractionCoverageMode = RCM_CoverageIgnored;
+			bInvalidateShader = true;
 		}
-		else
+		else if (!bUseMaterialAttributes && !EditorOnly->FrontMaterial.IsConnected() && GetExpressions().IsEmpty())
 		{
-			// Or if it cannot be found, a slab node
-			UMaterialExpressionSubstrateSlabBSDF* SlabNode = NewObject<UMaterialExpressionSubstrateSlabBSDF>(this);
-			SlabNode->Material = this;
-			SetPosXAndMoveReferenceToTheRight(SlabNode);
-			EditorOnly->FrontMaterial.Connect(0, SlabNode);
+			// Empty material: Create by default a slab node
+			UMaterialFunction* DefaultMF = LoadObject<UMaterialFunction>(nullptr, TEXT("/Engine/Functions/Substrate/SMF_UE4Disney.SMF_UE4Disney"));
+			if (DefaultMF)
+			{
+				DefaultMF->UpdateFromFunctionResource();
+				DefaultMF->PostEditChange();
+				DefaultMF->ConditionalPostLoad();
+
+				UMaterialExpressionMaterialFunctionCall* MFCallNode = NewObject<UMaterialExpressionMaterialFunctionCall>(this);
+				if (MFCallNode->SetMaterialFunction(DefaultMF))
+				{
+					// This is needed for input/output expressions to be set correctly, otherwise compilation will fail.
+					GetExpressionCollection().AddExpression(MFCallNode);
+
+					SetPosXAndMoveReferenceToTheRight(MFCallNode);
+					EditorOnly->FrontMaterial.Connect(0, MFCallNode);
+
+					MFCallNode->UpdateFromFunctionResource();
+					MFCallNode->PostEditChange();
+					MFCallNode->ConditionalPostLoad();
+
+					ColorMatInputConnectionTo(EditorOnly->BaseColor,		MFCallNode, 0, MP_BaseColor);
+					ScalarMatInputConnectionTo(EditorOnly->Metallic,		MFCallNode, 1, MP_Metallic);
+					ScalarMatInputConnectionTo(EditorOnly->Specular,		MFCallNode, 2, MP_Specular);
+					ScalarMatInputConnectionTo(EditorOnly->Roughness,		MFCallNode, 3, MP_Roughness);
+					Vector3MatInputConnectionTo(EditorOnly->Normal,			MFCallNode, 4, MP_Normal);
+					ColorMatInputConnectionTo(EditorOnly->EmissiveColor,	MFCallNode, 5, MP_EmissiveColor);
+					ScalarMatInputConnectionTo(EditorOnly->Opacity,			MFCallNode, 6, MP_Opacity);
+				}
+			}
+			else
+			{
+				// Or if it cannot be found, a slab node
+				UMaterialExpressionSubstrateSlabBSDF* SlabNode = NewObject<UMaterialExpressionSubstrateSlabBSDF>(this);
+				SlabNode->Material = this;
+				SetPosXAndMoveReferenceToTheRight(SlabNode);
+				EditorOnly->FrontMaterial.Connect(0, SlabNode);
+			}
+			bRelinkCustomOutputNodes = false;
+			bInvalidateShader = true;
 		}
-		bRelinkCustomOutputNodes = false;
-		bInvalidateShader = true;
-	}
-	else if (!bUseMaterialAttributes && !EditorOnly->FrontMaterial.IsConnected())
-	{
-		if (MaterialDomain == MD_Surface)
+		else if (!bUseMaterialAttributes && !EditorOnly->FrontMaterial.IsConnected())
 		{
-			bool bClearCoatConversionDone = false;
-			if (ShadingModel == MSM_ClearCoat)
+			if (MaterialDomain == MD_Surface)
 			{
-				GatherCustomNodes();
-				if (ClearCoatBottomNormalOutput)
+				bool bClearCoatConversionDone = false;
+				if (ShadingModel == MSM_ClearCoat)
 				{
-					// For this special case, using two slabs to create a clear coat material with separated top and bottom normal. 
-
-					// Create metalness to Slab parameterisation conveersion node
-					UMaterialExpressionSubstrateMetalnessToDiffuseAlbedoF0* SubstrateMetalnessToDiffuseAlbedoF0 = NewObject<UMaterialExpressionSubstrateMetalnessToDiffuseAlbedoF0>(this);
-					SetPosXAndMoveReferenceToTheRight(SubstrateMetalnessToDiffuseAlbedoF0);
-					ColorMatInputConnectionTo(EditorOnly->BaseColor, SubstrateMetalnessToDiffuseAlbedoF0, 0, MP_BaseColor);
-					ScalarMatInputConnectionTo(EditorOnly->Metallic, SubstrateMetalnessToDiffuseAlbedoF0, 1, MP_Metallic);
-					ScalarMatInputConnectionTo(EditorOnly->Specular, SubstrateMetalnessToDiffuseAlbedoF0, 2, MP_Specular);
+					GatherCustomNodes();
+					if (ClearCoatBottomNormalOutput)
+					{
+						// For this special case, using two slabs to create a clear coat material with separated top and bottom normal. 
+
+						// Create metalness to Slab parameterisation conveersion node
+						UMaterialExpressionSubstrateMetalnessToDiffuseAlbedoF0* SubstrateMetalnessToDiffuseAlbedoF0 = NewObject<UMaterialExpressionSubstrateMetalnessToDiffuseAlbedoF0>(this);
+						SetPosXAndMoveReferenceToTheRight(SubstrateMetalnessToDiffuseAlbedoF0);
+						ColorMatInputConnectionTo(EditorOnly->BaseColor, SubstrateMetalnessToDiffuseAlbedoF0, 0, MP_BaseColor);
+						ScalarMatInputConnectionTo(EditorOnly->Metallic, SubstrateMetalnessToDiffuseAlbedoF0, 1, MP_Metallic);
+						ScalarMatInputConnectionTo(EditorOnly->Specular, SubstrateMetalnessToDiffuseAlbedoF0, 2, MP_Specular);
 					
-					// Top slab BSDF as a simple Disney material
-					UMaterialExpressionSubstrateSlabBSDF* BottomSlabBSDF = NewObject<UMaterialExpressionSubstrateSlabBSDF>(this);
-					BottomSlabBSDF->Material = this;
-					SetPosXAndMoveReferenceToTheRight(BottomSlabBSDF);
-					BottomSlabBSDF->GetInput(0)->Connect(0, SubstrateMetalnessToDiffuseAlbedoF0);
-					BottomSlabBSDF->GetInput(1)->Connect(1, SubstrateMetalnessToDiffuseAlbedoF0);
-					BottomSlabBSDF->GetInput(2)->Connect(2, SubstrateMetalnessToDiffuseAlbedoF0);
-					ScalarMatInputConnectionTo(EditorOnly->Roughness, BottomSlabBSDF, 3, MP_Roughness);
-					ScalarMatInputConnectionTo(EditorOnly->Anisotropy, BottomSlabBSDF, 4, MP_Anisotropy, SUBSTRATE_COPY_CONNECTION);
-					Vector3MatInputConnectionTo(EditorOnly->Tangent, BottomSlabBSDF, 6, MP_Tangent);
-
-					check(ClearCoatBottomNormalOutput);
-					ConnectionTo(*ClearCoatBottomNormalOutput->GetInput(0), BottomSlabBSDF, 5, SUBSTRATE_COPY_CONNECTION);// ClearColorBottomNormal -> BottomSlabBSDF.Normal
-
-					// Now weight the top base material by opacity.
-					UMaterialExpressionSubstrateSlabBSDF* TopSlabBSDF = NewObject<UMaterialExpressionSubstrateSlabBSDF>(this);
-					TopSlabBSDF->Material = this;
-					TopSlabBSDF->MaterialExpressionEditorX = BottomSlabBSDF->MaterialExpressionEditorX;
-					TopSlabBSDF->MaterialExpressionEditorY = BottomSlabBSDF->MaterialExpressionEditorY + 650;
-					ColorMatInputConnectionTo(EditorOnly->EmissiveColor, TopSlabBSDF, 10, MP_EmissiveColor);
-					ScalarMatInputConnectionTo(EditorOnly->ClearCoatRoughness, TopSlabBSDF, 3, MP_CustomData0);	// ClearCoatRoughness => Roughness
-					Vector3MatInputConnectionTo(EditorOnly->Normal, TopSlabBSDF, 5, MP_Normal);
-
-					//  The top layer has a hard coded specular value of 0.5 (F0 = 0.04)
-					UMaterialExpressionConstant* ConstantHalf = NewObject<UMaterialExpressionConstant>(this);
-					ReplaceNodeAndMoveToTheRight(TopSlabBSDF, ConstantHalf);
-					ConstantHalf->R = 0.5f * 0.08f;
-					TopSlabBSDF->GetInput(1)->Connect(0, ConstantHalf);
-
-					// The original clear coat is a complex assemblage of arbitrary functions that do not always make sense.
-					// To simplify things, we set the top slab BSDF as having a constant Grey scale transmittance.
-					// As for the original, this is achieved with coverage so both transmittance and specular contribution vanishes
-					UMaterialExpressionConstant* ConstantZero = NewObject<UMaterialExpressionConstant>(this);
-					ReplaceNodeAndMoveToTheRight(TopSlabBSDF, ConstantZero);
-					ConstantZero->R = 0.0f;
-					TopSlabBSDF->GetInput(0)->Connect(0, ConstantZero);							// BaseColor = 0 to only feature absorption, no scattering
-
-					// Now setup the mean free path with a hard coded transmittance of 0.75 when viewing the surface perpendicularly
-					UMaterialExpressionConstant* Constant075 = NewObject<UMaterialExpressionConstant>(this);
-					ReplaceNodeAndMoveToTheRight(TopSlabBSDF, Constant075);
-					Constant075->R = 0.75f;
-					UMaterialExpressionSubstrateTransmittanceToMFP* TransToMDFP = NewObject<UMaterialExpressionSubstrateTransmittanceToMFP>(this);
-					ReplaceNodeAndMoveToTheRight(TopSlabBSDF, TransToMDFP);
-					TransToMDFP->GetInput(0)->Connect(0, Constant075);
-					TopSlabBSDF->GetInput(7)->Connect(0, TransToMDFP);							// MFP -> MFP
-					TopSlabBSDF->GetInput(13)->Connect(1, TransToMDFP);							// Thickness -> Thickness
-
-					// Now weight the top base material by ClearCoat
-					UMaterialExpressionSubstrateWeight* TopSlabBSDFWithCoverage = NewObject<UMaterialExpressionSubstrateWeight>(this);
-					SetPosXAndMoveReferenceToTheRight(TopSlabBSDFWithCoverage);
-					TopSlabBSDFWithCoverage->GetInput(0)->Connect(0, TopSlabBSDF);												// TopSlabBSDF -> A
-					ScalarMatInputConnectionTo(EditorOnly->ClearCoat, TopSlabBSDFWithCoverage, 1, MP_CustomData0);				// ClearCoat -> Weight
-					ScalarMatInputConnectionTo(EditorOnly->ClearCoatRoughness, TopSlabBSDFWithCoverage, 1, MP_CustomData1);		// ClearCoat -> Weight
-
-					UMaterialExpressionSubstrateVerticalLayering* VerticalLayering = NewObject<UMaterialExpressionSubstrateVerticalLayering>(this);
-					SetPosXAndMoveReferenceToTheRight(VerticalLayering);
-					VerticalLayering->GetInput(0)->Connect(0, TopSlabBSDFWithCoverage);			// Top -> Top
-					VerticalLayering->GetInput(1)->Connect(0, BottomSlabBSDF);					// Bottom -> Base
-
-					EditorOnly->FrontMaterial.Connect(0, VerticalLayering);
-					bClearCoatConversionDone = true;
-					bRelinkCustomOutputNodes = false;	// We do not want that to happen in this case
+						// Top slab BSDF as a simple Disney material
+						UMaterialExpressionSubstrateSlabBSDF* BottomSlabBSDF = NewObject<UMaterialExpressionSubstrateSlabBSDF>(this);
+						BottomSlabBSDF->Material = this;
+						SetPosXAndMoveReferenceToTheRight(BottomSlabBSDF);
+						BottomSlabBSDF->GetInput(0)->Connect(0, SubstrateMetalnessToDiffuseAlbedoF0);
+						BottomSlabBSDF->GetInput(1)->Connect(1, SubstrateMetalnessToDiffuseAlbedoF0);
+						BottomSlabBSDF->GetInput(2)->Connect(2, SubstrateMetalnessToDiffuseAlbedoF0);
+						ScalarMatInputConnectionTo(EditorOnly->Roughness, BottomSlabBSDF, 3, MP_Roughness);
+						ScalarMatInputConnectionTo(EditorOnly->Anisotropy, BottomSlabBSDF, 4, MP_Anisotropy, SUBSTRATE_COPY_CONNECTION);
+						Vector3MatInputConnectionTo(EditorOnly->Tangent, BottomSlabBSDF, 6, MP_Tangent);
+
+						check(ClearCoatBottomNormalOutput);
+						ConnectionTo(*ClearCoatBottomNormalOutput->GetInput(0), BottomSlabBSDF, 5, SUBSTRATE_COPY_CONNECTION);// ClearColorBottomNormal -> BottomSlabBSDF.Normal
+
+						// Now weight the top base material by opacity.
+						UMaterialExpressionSubstrateSlabBSDF* TopSlabBSDF = NewObject<UMaterialExpressionSubstrateSlabBSDF>(this);
+						TopSlabBSDF->Material = this;
+						TopSlabBSDF->MaterialExpressionEditorX = BottomSlabBSDF->MaterialExpressionEditorX;
+						TopSlabBSDF->MaterialExpressionEditorY = BottomSlabBSDF->MaterialExpressionEditorY + 650;
+						ColorMatInputConnectionTo(EditorOnly->EmissiveColor, TopSlabBSDF, 10, MP_EmissiveColor);
+						ScalarMatInputConnectionTo(EditorOnly->ClearCoatRoughness, TopSlabBSDF, 3, MP_CustomData0);	// ClearCoatRoughness => Roughness
+						Vector3MatInputConnectionTo(EditorOnly->Normal, TopSlabBSDF, 5, MP_Normal);
+
+						//  The top layer has a hard coded specular value of 0.5 (F0 = 0.04)
+						UMaterialExpressionConstant* ConstantHalf = NewObject<UMaterialExpressionConstant>(this);
+						ReplaceNodeAndMoveToTheRight(TopSlabBSDF, ConstantHalf);
+						ConstantHalf->R = 0.5f * 0.08f;
+						TopSlabBSDF->GetInput(1)->Connect(0, ConstantHalf);
+
+						// The original clear coat is a complex assemblage of arbitrary functions that do not always make sense.
+						// To simplify things, we set the top slab BSDF as having a constant Grey scale transmittance.
+						// As for the original, this is achieved with coverage so both transmittance and specular contribution vanishes
+						UMaterialExpressionConstant* ConstantZero = NewObject<UMaterialExpressionConstant>(this);
+						ReplaceNodeAndMoveToTheRight(TopSlabBSDF, ConstantZero);
+						ConstantZero->R = 0.0f;
+						TopSlabBSDF->GetInput(0)->Connect(0, ConstantZero);							// BaseColor = 0 to only feature absorption, no scattering
+
+						// Now setup the mean free path with a hard coded transmittance of 0.75 when viewing the surface perpendicularly
+						UMaterialExpressionConstant* Constant075 = NewObject<UMaterialExpressionConstant>(this);
+						ReplaceNodeAndMoveToTheRight(TopSlabBSDF, Constant075);
+						Constant075->R = 0.75f;
+						UMaterialExpressionSubstrateTransmittanceToMFP* TransToMDFP = NewObject<UMaterialExpressionSubstrateTransmittanceToMFP>(this);
+						ReplaceNodeAndMoveToTheRight(TopSlabBSDF, TransToMDFP);
+						TransToMDFP->GetInput(0)->Connect(0, Constant075);
+						TopSlabBSDF->GetInput(7)->Connect(0, TransToMDFP);							// MFP -> MFP
+						TopSlabBSDF->GetInput(13)->Connect(1, TransToMDFP);							// Thickness -> Thickness
+
+						// Now weight the top base material by ClearCoat
+						UMaterialExpressionSubstrateWeight* TopSlabBSDFWithCoverage = NewObject<UMaterialExpressionSubstrateWeight>(this);
+						SetPosXAndMoveReferenceToTheRight(TopSlabBSDFWithCoverage);
+						TopSlabBSDFWithCoverage->GetInput(0)->Connect(0, TopSlabBSDF);												// TopSlabBSDF -> A
+						ScalarMatInputConnectionTo(EditorOnly->ClearCoat, TopSlabBSDFWithCoverage, 1, MP_CustomData0);				// ClearCoat -> Weight
+						ScalarMatInputConnectionTo(EditorOnly->ClearCoatRoughness, TopSlabBSDFWithCoverage, 1, MP_CustomData1);		// ClearCoat -> Weight
+
+						UMaterialExpressionSubstrateVerticalLayering* VerticalLayering = NewObject<UMaterialExpressionSubstrateVerticalLayering>(this);
+						SetPosXAndMoveReferenceToTheRight(VerticalLayering);
+						VerticalLayering->GetInput(0)->Connect(0, TopSlabBSDFWithCoverage);			// Top -> Top
+						VerticalLayering->GetInput(1)->Connect(0, BottomSlabBSDF);					// Bottom -> Base
+
+						EditorOnly->FrontMaterial.Connect(0, VerticalLayering);
+						bClearCoatConversionDone = true;
+						bRelinkCustomOutputNodes = false;	// We do not want that to happen in this case
+					}
 				}
-			}
 			
-			if (!bClearCoatConversionDone)
+				if (!bClearCoatConversionDone)
+				{
+					ConvertNode = NewObject<UMaterialExpressionSubstrateShadingModels>(this);
+					ConvertNode->Material = this;
+					SetPosXAndMoveReferenceToTheRight(ConvertNode);
+					ConvertNode->SubsurfaceProfile = bRequireNoSubsurfaceProfile ? nullptr : SubsurfaceProfile;
+					ColorMatInputConnectionTo(EditorOnly->BaseColor, ConvertNode, 0, MP_BaseColor);
+					ScalarMatInputConnectionTo(EditorOnly->Metallic, ConvertNode, 1, MP_Metallic);
+					ScalarMatInputConnectionTo(EditorOnly->Specular, ConvertNode, 2, MP_Specular);
+					ScalarMatInputConnectionTo(EditorOnly->Roughness, ConvertNode, 3, MP_Roughness);
+					ScalarMatInputConnectionTo(EditorOnly->Anisotropy, ConvertNode, 4, MP_Anisotropy);
+					ColorMatInputConnectionTo(EditorOnly->EmissiveColor, ConvertNode, 5, MP_EmissiveColor);
+					Vector3MatInputConnectionTo(EditorOnly->Normal, ConvertNode, 6, MP_Normal, SUBSTRATE_COPY_CONNECTION);
+					Vector3MatInputConnectionTo(EditorOnly->Tangent, ConvertNode, 7, MP_Tangent);
+					ColorMatInputConnectionTo(EditorOnly->SubsurfaceColor, ConvertNode, 8, MP_SubsurfaceColor);
+					ScalarMatInputConnectionTo(EditorOnly->ClearCoat, ConvertNode, 9, MP_CustomData0);
+					ScalarMatInputConnectionTo(EditorOnly->ClearCoatRoughness, ConvertNode, 10, MP_CustomData1);
+					ScalarMatInputConnectionTo(EditorOnly->Opacity, ConvertNode, 11, MP_Opacity, SUBSTRATE_COPY_CONNECTION);	// We only copy, to keep Opacity on the root node in case BLEND_AlphaComposite is selected.
+					bRelinkCustomOutputNodes = true;
+			
+					// Shading Model
+					// * either use the shader graph expression 
+					// * or add a constant shading model
+					if (ShadingModel == MSM_FromMaterialExpression)
+					{
+						if (!EditorOnly->ShadingModelFromMaterialExpression.IsConnected())
+						{
+							ConvertNode->ShadingModelOverride = MSM_DefaultLit;
+						}
+						else
+						{
+							// Reconnect the shading model expression. 
+							// Note: assign the expression directly, as using ConvertNode->GetInput(19)->Connect(..) causes the expression to not be assigned
+							ConvertNode->ShadingModel.Connect(EditorOnly->ShadingModelFromMaterialExpression.OutputIndex, EditorOnly->ShadingModelFromMaterialExpression.Expression);
+						}
+
+						// Store Substrate shading model of the converted material. 
+						GatherCustomNodes();
+						if (SingleLayerWaterOutput)
+						{
+							ShadingModels.AddShadingModel(MSM_SingleLayerWater);
+						}
+
+						check(ShadingModels.CountShadingModels() >= 1);
+					}
+					else
+					{
+						ConvertNode->ShadingModelOverride = ShadingModel;
+						check(ShadingModels.CountShadingModels() == 1);
+					}
+
+					EditorOnly->FrontMaterial.Connect(0, ConvertNode);
+				}
+
+				bInvalidateShader = true;
+			}
+			else if (MaterialDomain == MD_Volume)
+			{
+				UMaterialExpressionSubstrateVolumetricFogCloudBSDF* VolBSDF = NewObject<UMaterialExpressionSubstrateVolumetricFogCloudBSDF>(this);
+				VolBSDF->Material = this;
+				SetPosXAndMoveReferenceToTheRight(VolBSDF);
+				ColorMatInputConnectionTo(EditorOnly->BaseColor, VolBSDF, 0, MP_BaseColor);	
+				ColorMatInputConnectionTo(EditorOnly->SubsurfaceColor, VolBSDF, 1, MP_SubsurfaceColor);
+				ColorMatInputConnectionTo(EditorOnly->EmissiveColor, VolBSDF, 2, MP_EmissiveColor);	
+				ScalarMatInputConnectionTo(EditorOnly->AmbientOcclusion, VolBSDF, 3, MP_AmbientOcclusion);
+
+				VolBSDF->bEmissiveOnly = ShadingModel == MSM_Unlit;
+
+				// SUBSTRATE_TODO remove the VolumetricAdvancedOutput node and add the input onto FogCloudBSDF even if only used by the cloud renderer?
+				EditorOnly->FrontMaterial.Connect(0, VolBSDF);
+				bInvalidateShader = true;
+			}
+			else if (MaterialDomain == MD_LightFunction)
+			{
+				// Some materials don't have their shading mode set correctly to Unlit. Since only Unlit is supported, forcing it here.
+				ShadingModel = MSM_Unlit;
+				ShadingModels.ClearShadingModels();
+				ShadingModels.AddShadingModel(MSM_Unlit);
+
+				// Only Emissive & Opacity are valid input for PostProcess material
+				UMaterialExpressionSubstrateLightFunction* LightFunctionNode = NewObject<UMaterialExpressionSubstrateLightFunction>(this);
+				LightFunctionNode->Material = this;
+				SetPosXAndMoveReferenceToTheRight(LightFunctionNode);
+				ColorMatInputConnectionTo(EditorOnly->EmissiveColor, LightFunctionNode, 0, MP_EmissiveColor);
+
+				EditorOnly->FrontMaterial.Connect(0, LightFunctionNode);
+				bInvalidateShader = true;
+			}
+			else if (MaterialDomain == MD_PostProcess)
+			{
+				// Some materials don't have their shading mode set correctly to Unlit. Since only Unlit is supported, forcing it here.
+				ShadingModel = MSM_Unlit;
+				ShadingModels.ClearShadingModels();
+				ShadingModels.AddShadingModel(MSM_Unlit);
+
+				if (MaterialDomain == MD_PostProcess && !IsPostProcessMaterialOutputingAlpha())
+				{
+					BlendMode = BLEND_Opaque;
+				}
+
+				UMaterialExpressionSubstratePostProcess* PostProcNode = NewObject<UMaterialExpressionSubstratePostProcess>(this);
+				PostProcNode->Material = this;
+				SetPosXAndMoveReferenceToTheRight(PostProcNode);
+
+				ColorMatInputConnectionTo(EditorOnly->EmissiveColor, PostProcNode, 0, MP_EmissiveColor);
+				ScalarMatInputConnectionTo(EditorOnly->Opacity, PostProcNode, 1, MP_Opacity, SUBSTRATE_COPY_CONNECTION);	// We only copy, to keep Opacity on the root node in case BLEND_AlphaComposite is selected.
+
+				EditorOnly->FrontMaterial.Connect(0, PostProcNode);
+				bInvalidateShader = true;
+			}
+			else if (MaterialDomain == MD_DeferredDecal)
 			{
+				// Some decal materials don't have their shading mode set correctly to DefaultLit. Since only DefaultLit is supported, forcing it here.
+				ShadingModel = MSM_DefaultLit;
+				ShadingModels.ClearShadingModels();
+				ShadingModels.AddShadingModel(MSM_DefaultLit);
+
 				ConvertNode = NewObject<UMaterialExpressionSubstrateShadingModels>(this);
 				ConvertNode->Material = this;
 				SetPosXAndMoveReferenceToTheRight(ConvertNode);
-				ConvertNode->SubsurfaceProfile = bRequireNoSubsurfaceProfile ? nullptr : SubsurfaceProfile;
 				ColorMatInputConnectionTo(EditorOnly->BaseColor, ConvertNode, 0, MP_BaseColor);
 				ScalarMatInputConnectionTo(EditorOnly->Metallic, ConvertNode, 1, MP_Metallic);
 				ScalarMatInputConnectionTo(EditorOnly->Specular, ConvertNode, 2, MP_Specular);
@@ -3574,153 +3746,40 @@ void UMaterial::ConvertMaterialToSubstrateMaterial()
 				ScalarMatInputConnectionTo(EditorOnly->ClearCoat, ConvertNode, 9, MP_CustomData0);
 				ScalarMatInputConnectionTo(EditorOnly->ClearCoatRoughness, ConvertNode, 10, MP_CustomData1);
 				ScalarMatInputConnectionTo(EditorOnly->Opacity, ConvertNode, 11, MP_Opacity, SUBSTRATE_COPY_CONNECTION);	// We only copy, to keep Opacity on the root node in case BLEND_AlphaComposite is selected.
-				bRelinkCustomOutputNodes = true;
-			
-				// Shading Model
-				// * either use the shader graph expression 
-				// * or add a constant shading model
-				if (ShadingModel == MSM_FromMaterialExpression)
-				{
-					if (!EditorOnly->ShadingModelFromMaterialExpression.IsConnected())
-					{
-						ConvertNode->ShadingModelOverride = MSM_DefaultLit;
-					}
-					else
-					{
-						// Reconnect the shading model expression. 
-						// Note: assign the expression directly, as using ConvertNode->GetInput(19)->Connect(..) causes the expression to not be assigned
-						ConvertNode->ShadingModel.Connect(EditorOnly->ShadingModelFromMaterialExpression.OutputIndex, EditorOnly->ShadingModelFromMaterialExpression.Expression);
-					}
 
-					// Store Substrate shading model of the converted material. 
-					GatherCustomNodes();
-					if (SingleLayerWaterOutput)
-					{
-						ShadingModels.AddShadingModel(MSM_SingleLayerWater);
-					}
+				// Add constant for the Unlit shading model
+				ConvertNode->ShadingModelOverride = ShadingModel;
+				check(ShadingModels.CountShadingModels() == 1);
 
-					check(ShadingModels.CountShadingModels() >= 1);
-				}
-				else
-				{
-					ConvertNode->ShadingModelOverride = ShadingModel;
-					check(ShadingModels.CountShadingModels() == 1);
-				}
+				// Now pass through the convert to decal node, which flag the material as SSM_Decal, which will set the domain to Decal.
+				UMaterialExpressionSubstrateConvertToDecal* ConvertToDecalNode= NewObject<UMaterialExpressionSubstrateConvertToDecal>(this);
+				ConvertToDecalNode->Material = this;
+				ReplaceNodeAndMoveToTheRight(ConvertNode, ConvertToDecalNode);
+				ConvertToDecalNode->DecalMaterial.Connect(0, ConvertNode);
 
-				EditorOnly->FrontMaterial.Connect(0, ConvertNode);
+				EditorOnly->FrontMaterial.Connect(0, ConvertToDecalNode);
+				bInvalidateShader = true;
 			}
-
-			bInvalidateShader = true;
-		}
-		else if (MaterialDomain == MD_Volume)
-		{
-			UMaterialExpressionSubstrateVolumetricFogCloudBSDF* VolBSDF = NewObject<UMaterialExpressionSubstrateVolumetricFogCloudBSDF>(this);
-			VolBSDF->Material = this;
-			SetPosXAndMoveReferenceToTheRight(VolBSDF);
-			ColorMatInputConnectionTo(EditorOnly->BaseColor, VolBSDF, 0, MP_BaseColor);	
-			ColorMatInputConnectionTo(EditorOnly->SubsurfaceColor, VolBSDF, 1, MP_SubsurfaceColor);
-			ColorMatInputConnectionTo(EditorOnly->EmissiveColor, VolBSDF, 2, MP_EmissiveColor);	
-			ScalarMatInputConnectionTo(EditorOnly->AmbientOcclusion, VolBSDF, 3, MP_AmbientOcclusion);
-
-			VolBSDF->bEmissiveOnly = ShadingModel == MSM_Unlit;
-
-			// SUBSTRATE_TODO remove the VolumetricAdvancedOutput node and add the input onto FogCloudBSDF even if only used by the cloud renderer?
-			EditorOnly->FrontMaterial.Connect(0, VolBSDF);
-			bInvalidateShader = true;
-		}
-		else if (MaterialDomain == MD_LightFunction)
-		{
-			// Some materials don't have their shading mode set correctly to Unlit. Since only Unlit is supported, forcing it here.
-			ShadingModel = MSM_Unlit;
-			ShadingModels.ClearShadingModels();
-			ShadingModels.AddShadingModel(MSM_Unlit);
-
-			// Only Emissive & Opacity are valid input for PostProcess material
-			UMaterialExpressionSubstrateLightFunction* LightFunctionNode = NewObject<UMaterialExpressionSubstrateLightFunction>(this);
-			LightFunctionNode->Material = this;
-			SetPosXAndMoveReferenceToTheRight(LightFunctionNode);
-			ColorMatInputConnectionTo(EditorOnly->EmissiveColor, LightFunctionNode, 0, MP_EmissiveColor);
-
-			EditorOnly->FrontMaterial.Connect(0, LightFunctionNode);
-			bInvalidateShader = true;
-		}
-		else if (MaterialDomain == MD_PostProcess)
-		{
-			// Some materials don't have their shading mode set correctly to Unlit. Since only Unlit is supported, forcing it here.
-			ShadingModel = MSM_Unlit;
-			ShadingModels.ClearShadingModels();
-			ShadingModels.AddShadingModel(MSM_Unlit);
-
-			if (MaterialDomain == MD_PostProcess && !IsPostProcessMaterialOutputingAlpha())
+			else if (MaterialDomain == MD_UI)
 			{
-				BlendMode = BLEND_Opaque;
-			}
-
-			UMaterialExpressionSubstratePostProcess* PostProcNode = NewObject<UMaterialExpressionSubstratePostProcess>(this);
-			PostProcNode->Material = this;
-			SetPosXAndMoveReferenceToTheRight(PostProcNode);
-
-			ColorMatInputConnectionTo(EditorOnly->EmissiveColor, PostProcNode, 0, MP_EmissiveColor);
-			ScalarMatInputConnectionTo(EditorOnly->Opacity, PostProcNode, 1, MP_Opacity, SUBSTRATE_COPY_CONNECTION);	// We only copy, to keep Opacity on the root node in case BLEND_AlphaComposite is selected.
-
-			EditorOnly->FrontMaterial.Connect(0, PostProcNode);
-			bInvalidateShader = true;
-		}
-		else if (MaterialDomain == MD_DeferredDecal)
-		{
-			// Some decal materials don't have their shading mode set correctly to DefaultLit. Since only DefaultLit is supported, forcing it here.
-			ShadingModel = MSM_DefaultLit;
-			ShadingModels.ClearShadingModels();
-			ShadingModels.AddShadingModel(MSM_DefaultLit);
+				// Some materials don't have their shading mode set correctly to Unlit. Since only Unlit is supported, forcing it here.
+				ShadingModel = MSM_Unlit;
+				ShadingModels.ClearShadingModels();
+				ShadingModels.AddShadingModel(MSM_Unlit);
 
-			ConvertNode = NewObject<UMaterialExpressionSubstrateShadingModels>(this);
-			ConvertNode->Material = this;
-			SetPosXAndMoveReferenceToTheRight(ConvertNode);
-			ColorMatInputConnectionTo(EditorOnly->BaseColor, ConvertNode, 0, MP_BaseColor);
-			ScalarMatInputConnectionTo(EditorOnly->Metallic, ConvertNode, 1, MP_Metallic);
-			ScalarMatInputConnectionTo(EditorOnly->Specular, ConvertNode, 2, MP_Specular);
-			ScalarMatInputConnectionTo(EditorOnly->Roughness, ConvertNode, 3, MP_Roughness);
-			ScalarMatInputConnectionTo(EditorOnly->Anisotropy, ConvertNode, 4, MP_Anisotropy);
-			ColorMatInputConnectionTo(EditorOnly->EmissiveColor, ConvertNode, 5, MP_EmissiveColor);
-			Vector3MatInputConnectionTo(EditorOnly->Normal, ConvertNode, 6, MP_Normal, SUBSTRATE_COPY_CONNECTION);
-			Vector3MatInputConnectionTo(EditorOnly->Tangent, ConvertNode, 7, MP_Tangent);
-			ColorMatInputConnectionTo(EditorOnly->SubsurfaceColor, ConvertNode, 8, MP_SubsurfaceColor);
-			ScalarMatInputConnectionTo(EditorOnly->ClearCoat, ConvertNode, 9, MP_CustomData0);
-			ScalarMatInputConnectionTo(EditorOnly->ClearCoatRoughness, ConvertNode, 10, MP_CustomData1);
-			ScalarMatInputConnectionTo(EditorOnly->Opacity, ConvertNode, 11, MP_Opacity, SUBSTRATE_COPY_CONNECTION);	// We only copy, to keep Opacity on the root node in case BLEND_AlphaComposite is selected.
-
-			// Add constant for the Unlit shading model
-			ConvertNode->ShadingModelOverride = ShadingModel;
-			check(ShadingModels.CountShadingModels() == 1);
-
-			// Now pass through the convert to decal node, which flag the material as SSM_Decal, which will set the domain to Decal.
-			UMaterialExpressionSubstrateConvertToDecal* ConvertToDecalNode= NewObject<UMaterialExpressionSubstrateConvertToDecal>(this);
-			ConvertToDecalNode->Material = this;
-			ReplaceNodeAndMoveToTheRight(ConvertNode, ConvertToDecalNode);
-			ConvertToDecalNode->DecalMaterial.Connect(0, ConvertNode);
-
-			EditorOnly->FrontMaterial.Connect(0, ConvertToDecalNode);
-			bInvalidateShader = true;
-		}
-		else if (MaterialDomain == MD_UI)
-		{
-			// Some materials don't have their shading mode set correctly to Unlit. Since only Unlit is supported, forcing it here.
-			ShadingModel = MSM_Unlit;
-			ShadingModels.ClearShadingModels();
-			ShadingModels.AddShadingModel(MSM_Unlit);
+				UMaterialExpressionSubstrateUI* UINode = NewObject<UMaterialExpressionSubstrateUI>(this);
+				UINode->Material = this;
+				SetPosXAndMoveReferenceToTheRight(UINode);
+				ColorMatInputConnectionTo(EditorOnly->EmissiveColor, UINode, 0, MP_EmissiveColor);
+				ScalarMatInputConnectionTo(EditorOnly->Opacity, UINode, 1, MP_Opacity, SUBSTRATE_COPY_CONNECTION);	// We only copy, to keep Opacity on the root node in case BLEND_AlphaComposite is selected.
 
-			UMaterialExpressionSubstrateUI* UINode = NewObject<UMaterialExpressionSubstrateUI>(this);
-			UINode->Material = this;
-			SetPosXAndMoveReferenceToTheRight(UINode);
-			ColorMatInputConnectionTo(EditorOnly->EmissiveColor, UINode, 0, MP_EmissiveColor);
-			ScalarMatInputConnectionTo(EditorOnly->Opacity, UINode, 1, MP_Opacity, SUBSTRATE_COPY_CONNECTION);	// We only copy, to keep Opacity on the root node in case BLEND_AlphaComposite is selected.
+				EditorOnly->FrontMaterial.Connect(0, UINode);
+				bInvalidateShader = true;
+			}
 
-			EditorOnly->FrontMaterial.Connect(0, UINode);
-			bInvalidateShader = true;
+			BlendMode = ConvertLegacyBlendMode(BlendMode, ShadingModels);
+			RefractionCoverageMode = RCM_CoverageIgnored;
 		}
-
-		BlendMode = ConvertLegacyBlendMode(BlendMode, ShadingModels);
-		RefractionCoverageMode = RCM_CoverageIgnored;
 	}
 
 	if (bRelinkCustomOutputNodes)
diff --git a/Engine/Source/Runtime/Engine/Private/Materials/MaterialExpressions.cpp b/Engine/Source/Runtime/Engine/Private/Materials/MaterialExpressions.cpp
--- a/Engine/Source/Runtime/Engine/Private/Materials/MaterialExpressions.cpp
+++ b/Engine/Source/Runtime/Engine/Private/Materials/MaterialExpressions.cpp
@@ -2146,6 +2146,7 @@ void UMaterialExpression::ConnectToPreviewMaterial(UMaterial* InMaterial, int32
 				UMaterialExpressionSubstrateConvertMaterialAttributes* ConvertAttributeNode = NewObject<UMaterialExpressionSubstrateConvertMaterialAttributes>(this);
 				ConvertAttributeNode->Material = InMaterial;
 				ConvertAttributeNode->MaterialAttributes.Connect(OutputIndex, this);
+				ConvertAttributeNode->ShadingModelOverride = MSM_DefaultLit;
 
 				// Connect substrate data into material FrontMaterial input
 				if (UMaterialEditorOnlyData* MaterialEditorOnlyData = InMaterial->GetEditorOnlyData())
@@ -7365,27 +7366,42 @@ uint32 UMaterialExpressionGetMaterialAttributes::GetOutputType(int32 OutputIndex
 }
 
 bool UMaterialExpressionGetMaterialAttributes::IsResultSubstrateMaterial(int32 OutputIndex)
-{
+{	
+#if ENABLE_MATERIAL_LAYER_PROTOTYPE
+	if (OutputIndex >= 1 && OutputIndex <= AttributeGetTypes.Num())
+	{
+		return FMaterialAttributeDefinitionMap::GetValueType(AttributeGetTypes[OutputIndex - 1]) == MCT_Substrate;
+	}
+#endif
+
 	if (MaterialAttributes.Expression)
 	{
-		return MaterialAttributes.Expression->IsResultSubstrateMaterial(0);
+		return MaterialAttributes.Expression->IsResultSubstrateMaterial(MaterialAttributes.OutputIndex);
 	}
 	return false;
 }
 
 void UMaterialExpressionGetMaterialAttributes::GatherSubstrateMaterialInfo(FSubstrateMaterialInfo& SubstrateMaterialInfo, int32 OutputIndex)
 {
+#if ENABLE_MATERIAL_LAYER_PROTOTYPE
+	if (IsResultSubstrateMaterial(OutputIndex))
+#else
 	if (MaterialAttributes.Expression)
+#endif
 	{
-		MaterialAttributes.Expression->GatherSubstrateMaterialInfo(SubstrateMaterialInfo, 0);
+		MaterialAttributes.Expression->GatherSubstrateMaterialInfo(SubstrateMaterialInfo, MaterialAttributes.OutputIndex);
 	}
 }
 
 FSubstrateOperator* UMaterialExpressionGetMaterialAttributes::SubstrateGenerateMaterialTopologyTree(class FMaterialCompiler* Compiler, class UMaterialExpression* Parent, int32 OutputIndex)
 {
+#if ENABLE_MATERIAL_LAYER_PROTOTYPE
+	if (IsResultSubstrateMaterial(OutputIndex))
+#else
 	if (MaterialAttributes.Expression)
+#endif
 	{
-		return MaterialAttributes.Expression->SubstrateGenerateMaterialTopologyTree(Compiler, Parent, 0);
+		return MaterialAttributes.Expression->SubstrateGenerateMaterialTopologyTree(Compiler, Parent, MaterialAttributes.OutputIndex);
 	}
 	return nullptr;
 }
@@ -7699,6 +7715,29 @@ uint32 UMaterialExpressionSetMaterialAttributes::GetInputType(int32 InputIndex)
 	return InputType;
 }
 
+bool UMaterialExpressionSetMaterialAttributes::IsResultSubstrateMaterial(int32 OutputIndex)
+{
+#if ENABLE_MATERIAL_LAYER_PROTOTYPE
+	if (Inputs[0].IsConnected())
+	{
+		Inputs[0].Expression->IsResultSubstrateMaterial(Inputs[0].OutputIndex);
+	}
+#endif
+
+	return false;
+}
+
+FSubstrateOperator* UMaterialExpressionSetMaterialAttributes::SubstrateGenerateMaterialTopologyTree(FMaterialCompiler* Compiler, UMaterialExpression* Parent, int32 OutputIndex)
+{
+#if ENABLE_MATERIAL_LAYER_PROTOTYPE
+	if(IsResultSubstrateMaterial(OutputIndex))
+	{
+		Inputs[0].Expression->SubstrateGenerateMaterialTopologyTree(Compiler, Parent, Inputs[0].OutputIndex);
+	}
+#endif
+	return nullptr;
+}
+
 void UMaterialExpressionSetMaterialAttributes::GetExpressionToolTip(TArray<FString>& OutToolTip)
 {
 	ConvertToMultilineToolTip(TEXT("Allows assigning values to specific inputs on a material attributes pin. Any unconnected inputs will be unchanged."), 40, OutToolTip);
@@ -8407,6 +8446,58 @@ int32 UMaterialExpressionMaterialAttributeLayers::Compile(FMaterialCompiler* Com
 	return Result;
 }
 
+#if ENABLE_MATERIAL_LAYER_PROTOTYPE
+bool UMaterialExpressionMaterialAttributeLayers::IsResultSubstrateMaterial(int32 OutputIndex)
+{
+	return Substrate::IsSubstrateEnabled();
+}
+
+FSubstrateOperator* UMaterialExpressionMaterialAttributeLayers::SubstrateGenerateMaterialTopologyTree(FMaterialCompiler* Compiler, UMaterialExpression* Parent, int32 OutputIndex)
+{
+	FSubstrateOperator* OutOperator = nullptr;
+
+	const FMaterialLayersFunctions* OverrideLayers = Compiler->GetMaterialLayers();
+	OverrideLayerGraph(OverrideLayers);
+
+	//Behaviour mirrors the behaviour of the existing compile function
+	if (ValidateLayerConfiguration(Compiler, true) && bIsLayerGraphBuilt)
+	{
+		if (NumActiveBlendCallers > 0 && BlendCallers[NumActiveBlendCallers - 1]->MaterialFunction)
+		{
+			// Multiple blended layers
+			OutOperator = BlendCallers[NumActiveBlendCallers - 1]->SubstrateGenerateMaterialTopologyTree(Compiler, Parent, 0);
+		}
+		else if (NumActiveLayerCallers > 0 && LayerCallers[NumActiveLayerCallers - 1]->MaterialFunction)
+		{
+			// Single layer
+			OutOperator = LayerCallers[NumActiveLayerCallers - 1]->SubstrateGenerateMaterialTopologyTree(Compiler, Parent, 0);
+		}
+		else if (NumActiveLayerCallers == 0)
+		{
+			// Pass-through
+			const FGuid AttributeID = Compiler->GetMaterialAttribute();
+			if (Input.GetTracedInput().Expression)
+			{
+				OutOperator = Input.Expression->SubstrateGenerateMaterialTopologyTree(Compiler, Parent, 0);
+			}
+		}
+
+		if (!OutOperator)
+		{
+			//If this is reached, compile a default operator to avoid a crash
+			FSubstrateOperator& DefaultOperator = Compiler->SubstrateCompilationRegisterOperator(SUBSTRATE_OPERATOR_BSDF, Compiler->SubstrateTreeStackGetPathUniqueId(), this, Parent, Compiler->SubstrateTreeStackGetParentPathUniqueId());
+			DefaultOperator.BSDFType = SUBSTRATE_BSDF_TYPE_SLAB;
+			DefaultOperator.ThicknessIndex = Compiler->SubstrateThicknessStackGetThicknessIndex();
+			OutOperator = &DefaultOperator;
+		}
+	}
+
+	OverrideLayerGraph(nullptr);
+
+	return OutOperator;
+}
+#endif //ENABLE_MATERIAL_LAYER_PROTOTYPE
+
 void UMaterialExpressionMaterialAttributeLayers::GetCaption(TArray<FString>& OutCaptions) const
 {
 	OutCaptions.Add(TEXT("Material Attribute Layers"));
@@ -10015,7 +10106,9 @@ bool UMaterialExpressionFeatureLevelSwitch::IsInputConnectionRequired(int32 Inpu
 
 bool UMaterialExpressionFeatureLevelSwitch::IsResultMaterialAttributes(int32 OutputIndex)
 {
+#if !ENABLE_MATERIAL_LAYER_PROTOTYPE
 	check(OutputIndex == 0);
+#endif
 	for (FExpressionInputIterator It{ this }; It; ++It)
 	{
 		// If there is a loop anywhere in this expression's inputs then we can't risk checking them
@@ -15461,6 +15554,8 @@ UMaterialInterface* UMaterialFunction::GetPreviewMaterial()
 		PreviewMaterial->bIsPreviewMaterial = true;
 
 		PreviewMaterial->AssignExpressionCollection(GetExpressionCollection());
+		//Update cached expression data to ensure function calls are populated for resolving the preview
+		PreviewMaterial->UpdateCachedExpressionData();
 
 		//Find the first output expression and use that. 
 		for (UMaterialExpression* Expression : GetExpressions())
@@ -15605,6 +15700,303 @@ void UMaterialFunction::Serialize(FArchive& Ar)
 #endif // #if WITH_EDITOR
 }
 
+#if WITH_EDITORONLY_DATA
+#if ENABLE_MATERIAL_LAYER_PROTOTYPE
+void UMaterialFunction::GetAllCustomOutputExpressions(TArray<class UMaterialExpressionCustomOutput*>& OutCustomOutputs) const
+{
+	for (UMaterialExpression* Expression : GetExpressions())
+	{
+		UMaterialExpressionCustomOutput* CustomOutput = Cast<UMaterialExpressionCustomOutput>(Expression);
+		if (CustomOutput)
+		{
+			OutCustomOutputs.Add(CustomOutput);
+		}
+	}
+}
+
+void UMaterialFunction::ConvertExpressionsBetweenLegacyAndSubstrate()
+{
+	if (!Substrate::IsSubstrateEnabled() || GetMaterialFunctionUsage() == EMaterialFunctionUsage::Default)
+	{
+		return;
+	}
+
+	UMaterialFunctionEditorOnlyData* EditorOnly = GetEditorOnlyData();
+	if (!EditorOnly)
+	{
+		return;
+	}
+
+	TArray<TObjectPtr<UMaterialExpression>>& Expressions = EditorOnly->ExpressionCollection.Expressions;
+	if (Expressions.Num() < 1)
+	{
+		return;
+	}
+
+	UMaterialExpressionFunctionOutput* OutputNode = Cast<UMaterialExpressionFunctionOutput>(Expressions[0]);
+	if (!OutputNode || !OutputNode->A.IsConnected() || !OutputNode->A.Expression->IsResultMaterialAttributes(OutputNode->A.OutputIndex))
+	{
+		return;
+	}
+
+	//Perform a full sweep to check for embedded Substrate expressions, checking if conversion is necessary
+	TArray<UMaterialFunctionInterface*> AllFunctions;
+	GetDependentFunctions(AllFunctions);
+	AllFunctions.Add(this);	
+	for (UMaterialFunctionInterface* Function : AllFunctions)
+	{
+		for (UMaterialExpression* Expression : Function->GetExpressions())
+		{
+			if (Expression->IsA<UMaterialExpressionSubstrateBSDF>()
+				|| Expression->IsA<UMaterialExpressionSubstrateUtilityBase>())
+			{
+				//If Substrate nodes are present, do not convert this function, assume previously converted or Substrate MF
+				return;
+			}
+		}
+	}
+
+	auto MoveNodeInHorizonalAxis = [](UMaterialExpression* NodeToMove, int32 HorizontalOffset = 0)
+	{
+		NodeToMove->MaterialExpressionEditorX += HorizontalOffset;
+	};
+
+	auto PlaceNodeInLocation = [](UMaterialExpression* ExistingNode, UMaterialExpression* NewNode)
+	{
+		NewNode->MaterialExpressionEditorX = ExistingNode->MaterialExpressionEditorX;
+		NewNode->MaterialExpressionEditorY = ExistingNode->MaterialExpressionEditorY;
+	};
+
+	auto PlaceBelowNode= [](UMaterialExpression* ExistingNode, UMaterialExpression* NewNode, int32 DownOffset = 100)
+	{
+		NewNode->MaterialExpressionEditorX = ExistingNode->MaterialExpressionEditorX;
+		NewNode->MaterialExpressionEditorY = ExistingNode->MaterialExpressionEditorY + DownOffset;
+	};
+
+	auto ReplaceNodeAndMoveToTheRight = [](UMaterialExpression* ExistingNode, UMaterialExpression* NewNode, int32 RightOffset = 350)
+	{
+		NewNode->MaterialExpressionEditorX = ExistingNode->MaterialExpressionEditorX;
+		NewNode->MaterialExpressionEditorY = ExistingNode->MaterialExpressionEditorY;
+		ExistingNode->MaterialExpressionEditorX = NewNode->MaterialExpressionEditorX + RightOffset;
+	};
+
+	bool bBlendConverted = false;
+	if (GetMaterialFunctionUsage() == EMaterialFunctionUsage::MaterialLayerBlend)
+	{
+		UMaterialExpressionFunctionInput* BottomInput = nullptr;
+		UMaterialExpressionFunctionInput* TopInput = nullptr;
+
+		for (UMaterialExpression* Expression : Expressions)
+		{
+			if (UMaterialExpressionFunctionInput* InputNode = Cast<UMaterialExpressionFunctionInput>(Expression))
+			{
+				if (!BottomInput && InputNode->InputName == TEXT("Bottom Layer"))
+				{
+					BottomInput = InputNode;
+				}
+				else if (!TopInput && InputNode->InputName == TEXT("Top Layer"))
+				{
+					TopInput = InputNode;
+				}
+			}
+		}
+	
+		//The two legacy inputs are required at a minimum.
+		if (!TopInput || !BottomInput)
+		{
+			return;
+		}
+
+		BottomInput->InputName = TEXT("Background Layer");
+		TopInput->InputName = TEXT("Foreground Layer");
+		if (BottomInput->PreviewValue == FVector4f(FLinearColor::Black) && TopInput->PreviewValue == FVector4f(FLinearColor::Black))
+		{
+			//Only occurs during the update process and if inputs are set to black default so the blend has distinct preview, further user defined colours will not be altered.
+			TopInput->PreviewValue = FLinearColor::White;
+		}
+
+		/**
+		 * Blend node conversion occurs recursively. We can't resolve which nodes could have multiple inputs already connected to a node's output from the expression itself,
+		 * so we have to traverse from the output function to map the existing graph correctly.
+		 */
+		TMap<UMaterialExpression*, UMaterialExpression*> ReplacementNodeMapping;
+		TArray<UMaterialExpression*> RecursedExpressions;
+		auto ConvertBlendExpression = [&](FExpressionInput* Input, UMaterialExpression* InputExpression) -> void
+		{
+			if(!Input || !InputExpression)
+			{
+				return;
+			}
+		
+			if (UMaterialExpressionBlendMaterialAttributes* BlendNode = Cast<UMaterialExpressionBlendMaterialAttributes>(InputExpression))
+			{
+				//Load the existing substrate default blend function which will replace the legacy blend node.
+				static TObjectPtr<UMaterialFunction> DefaultBlendFunction = FindObject<UMaterialFunction>(GetTransientPackage(), DEFAULT_MATERIALLAYERBLEND_PATH);
+				if (!DefaultBlendFunction)
+				{
+					DefaultBlendFunction = LoadObject<UMaterialFunction>(GetTransientPackage(), DEFAULT_MATERIALLAYERBLEND_PATH);
+				}
+
+				if(DefaultBlendFunction)
+				{
+					UMaterialExpressionMaterialFunctionCall* BlendFunctionCall = ReplacementNodeMapping.Contains(BlendNode) ? Cast<UMaterialExpressionMaterialFunctionCall>(*ReplacementNodeMapping.Find(BlendNode)) : nullptr;
+					if (!BlendFunctionCall)
+					{
+						//Store the mapping of a blend node to it's replacement function call, which means we only create a call once per individual blend node.
+						BlendFunctionCall = Cast<UMaterialExpressionMaterialFunctionCall>(ReplacementNodeMapping.Add(BlendNode, NewObject<UMaterialExpressionMaterialFunctionCall>(this)));
+						BlendFunctionCall->Function = this;
+						BlendFunctionCall->SetMaterialFunction(DefaultBlendFunction);
+
+						if (BlendFunctionCall && BlendFunctionCall->FunctionInputs.Num() >= 3)
+						{
+							//These should match the sort priority of the Default MLB
+							BlendFunctionCall->FunctionInputs[0].Input.Connect(BlendNode->A.OutputIndex, BlendNode->A.Expression); //Background
+							BlendFunctionCall->FunctionInputs[1].Input.Connect(BlendNode->B.OutputIndex, BlendNode->B.Expression); //Foreground
+							BlendFunctionCall->FunctionInputs[2].Input.Connect(BlendNode->Alpha.OutputIndex, BlendNode->Alpha.Expression);
+						}
+
+						PlaceNodeInLocation(BlendNode, BlendFunctionCall);
+						InputExpression = BlendFunctionCall;
+						Expressions.Add(BlendFunctionCall);
+						RecursedExpressions.Add(BlendFunctionCall);
+					}
+
+					if (BlendFunctionCall)
+					{
+						Input->Connect(0, BlendFunctionCall);
+						bBlendConverted = true;
+					}
+				}
+			}
+			else if (InputExpression->IsA<UMaterialExpressionSetMaterialAttributes>() || InputExpression->IsA<UMaterialExpressionMakeMaterialAttributes>())
+			{
+				//Same behaviour as above, but for now extends the Make/Set nodes with conversion, so we don't remove the existing node in this case.
+				UMaterialExpressionSubstrateSetAttributes* SetSubstrateNode = ReplacementNodeMapping.Contains(InputExpression) ? Cast<UMaterialExpressionSubstrateSetAttributes>(*ReplacementNodeMapping.Find(InputExpression)) : nullptr;
+				if(!SetSubstrateNode)
+				{
+					SetSubstrateNode = Cast<UMaterialExpressionSubstrateSetAttributes>(ReplacementNodeMapping.Add(InputExpression, NewObject<UMaterialExpressionSubstrateSetAttributes>(this)));
+					SetSubstrateNode->Function = this;
+					SetSubstrateNode->NonSubstrateAttributes.Connect(0, InputExpression);
+					Expressions.Add(SetSubstrateNode);
+					RecursedExpressions.Add(SetSubstrateNode);
+
+					UMaterialExpressionSubstrateConvertMaterialAttributes* ConvertNode = NewObject<UMaterialExpressionSubstrateConvertMaterialAttributes>(this);
+					ConvertNode->Function = this;
+					ConvertNode->MaterialAttributes.Connect(0, InputExpression);
+					ConvertNode->ShadingModelOverride = MSM_DefaultLit;
+					Expressions.Add(ConvertNode);
+					RecursedExpressions.Add(ConvertNode);
+
+					SetSubstrateNode->FrontMaterial.Connect(0, ConvertNode);
+
+					PlaceBelowNode(InputExpression, ConvertNode);
+					PlaceBelowNode(ConvertNode, SetSubstrateNode, 250);
+				}
+
+				if(SetSubstrateNode)
+				{
+					Input->Connect(0, SetSubstrateNode);
+					bBlendConverted = true;
+				}
+			}
+
+			return;			
+		};
+
+		auto RecurseBlendFunction = [&](UMaterialExpression* Expression, auto&& RecurseBlendFunction) -> void
+		{
+			//If null or if the expression has already been recursed, skip.
+			if(!Expression || RecursedExpressions.Contains(Expression))
+			{
+				return;
+			}
+			RecursedExpressions.Add(Expression);
+
+			//If the expression is valid, iterate the connected expressions for upgrade to Substrate, then recurse via each valid expression
+			for (FExpressionInputIterator It{ Expression }; It; ++It)
+			{
+				ConvertBlendExpression(It.Input, It.Input->Expression);
+				RecurseBlendFunction(It.Input->Expression, RecurseBlendFunction);	
+			}
+		};
+		//Begin recursion process from the output node.
+		RecurseBlendFunction(OutputNode, RecurseBlendFunction);
+
+		//Once the nodes have been replaced and all connections remapped, we can remove the replaced nodes from the function.
+		for (auto& Mapping : ReplacementNodeMapping)
+		{
+			if(UMaterialExpressionBlendMaterialAttributes* BlendExpression = Cast<UMaterialExpressionBlendMaterialAttributes>(Mapping.Key))
+			{
+				for (FExpressionInputIterator It{ BlendExpression }; It; ++It)
+				{
+					It.Input->Expression = nullptr;
+				}
+				Expressions.Remove(BlendExpression);
+			}
+		}
+	}
+
+	//If we are converting a standard layer function, or if the blend node has not had blend functionality replaced directly, add basic conversion logic to the node.
+	if (!bBlendConverted || GetMaterialFunctionUsage() == EMaterialFunctionUsage::MaterialLayer)
+	{
+		//Layer assets are simpler than their blend counterparts because we only need to convert the legacy MAs, similar to the behaviour of a material upgrade.
+		UMaterialExpressionSubstrateConvertMaterialAttributes* ConvertAttributesNode = NewObject<UMaterialExpressionSubstrateConvertMaterialAttributes>(this);
+		ConvertAttributesNode->Function = this;
+		ConvertAttributesNode->MaterialAttributes.Connect(0, OutputNode->A.Expression);
+		ConvertAttributesNode->ShadingModelOverride = MSM_DefaultLit;
+		ReplaceNodeAndMoveToTheRight(OutputNode, ConvertAttributesNode);
+
+		//Add Custom logic connections
+		TArray<class UMaterialExpressionCustomOutput*> CustomOutputExpressions;
+		GetAllCustomOutputExpressions(CustomOutputExpressions);
+		check(ConvertAttributesNode);
+		for (UMaterialExpressionCustomOutput* Expression : CustomOutputExpressions)
+		{
+			// Gather custom output for single layer water
+			if (UMaterialExpressionSingleLayerWaterMaterialOutput* SingleLayerWaterOutput = Cast<UMaterialExpressionSingleLayerWaterMaterialOutput>(Expression))
+			{	
+				if (SingleLayerWaterOutput->ScatteringCoefficients.Expression)
+				{
+					ConvertAttributesNode->WaterScatteringCoefficients.Connect(SingleLayerWaterOutput->ScatteringCoefficients.OutputIndex, SingleLayerWaterOutput->ScatteringCoefficients.Expression);
+				}
+				if (SingleLayerWaterOutput->AbsorptionCoefficients.Expression)
+				{
+					ConvertAttributesNode->WaterAbsorptionCoefficients.Connect(SingleLayerWaterOutput->AbsorptionCoefficients.OutputIndex, SingleLayerWaterOutput->AbsorptionCoefficients.Expression);
+				}
+				if (SingleLayerWaterOutput->PhaseG.Expression)
+				{
+					ConvertAttributesNode->WaterPhaseG.Connect(SingleLayerWaterOutput->PhaseG.OutputIndex, SingleLayerWaterOutput->PhaseG.Expression);
+				}
+				if (SingleLayerWaterOutput->ColorScaleBehindWater.Expression)
+				{
+					ConvertAttributesNode->ColorScaleBehindWater.Connect(SingleLayerWaterOutput->ColorScaleBehindWater.OutputIndex, SingleLayerWaterOutput->ColorScaleBehindWater.Expression);
+				}
+
+				break;
+			}
+		}
+
+		/**
+		 * Layer MFs differ from materials in that they only have 1 output,
+		 * so we use the Set Substrate Attributes node to collect the Front Material and MAs to pass to the next function in the layer stack.
+		 */
+		UMaterialExpressionSubstrateSetAttributes* SetAttributesNode = NewObject<UMaterialExpressionSubstrateSetAttributes>(this);
+		SetAttributesNode->Function = this;
+		SetAttributesNode->NonSubstrateAttributes.Connect(0, OutputNode->A.Expression);
+		SetAttributesNode->FrontMaterial.Connect(0, ConvertAttributesNode);
+		ReplaceNodeAndMoveToTheRight(OutputNode, SetAttributesNode);
+
+		Expressions.EmplaceAt(1, SetAttributesNode);
+		Expressions.EmplaceAt(2, ConvertAttributesNode);
+
+		OutputNode->A.Connect(0, SetAttributesNode);
+	}
+	MoveNodeInHorizonalAxis(OutputNode, 50);
+	OutputNode->bCollapsed = true;
+}
+#endif //ENABLE_MATERIAL_LAYER_PROTOTYPE
+#endif //WITH_EDITORONLY_DATA
+
 void UMaterialFunction::PostLoad()
 {
 	LLM_SCOPE(ELLMTag::Materials);
@@ -15637,6 +16029,9 @@ void UMaterialFunction::PostLoad()
 		ensure(!EditorOnly->ExpressionCollection.ExpressionExecEnd);
 		EditorOnly->ExpressionCollection.ExpressionExecEnd = MoveTemp(ExpressionExecEnd_DEPRECATED);
 	}
+#if ENABLE_MATERIAL_LAYER_PROTOTYPE
+	ConvertExpressionsBetweenLegacyAndSubstrate();
+#endif //ENABLE_MATERIAL_LAYER_PROTOTYPE
 
 	if (EditorOnly)
 	{
@@ -15955,74 +16350,171 @@ bool UMaterialFunction::ValidateFunctionUsage(FMaterialCompiler* Compiler, const
 #if WITH_EDITOR
 	if (GetMaterialFunctionUsage() == EMaterialFunctionUsage::MaterialLayer)
 	{
-		// Material layers must have a single MA input and output only
-		for (UMaterialExpression* Expression : GetExpressions())
+#if ENABLE_MATERIAL_LAYER_PROTOTYPE
+		if (Substrate::IsSubstrateEnabled())
 		{
-			if (UMaterialExpressionFunctionInput* InputExpression = Cast<UMaterialExpressionFunctionInput>(Expression))
+			// Material layers must have a single MA input and output only
+			for (UMaterialExpression* Expression : GetExpressions())
 			{
-				++NumInputs;
-				if (NumInputs > 1 || !InputExpression->IsResultMaterialAttributes(0))
+				if (UMaterialExpressionFunctionInput* InputExpression = Cast<UMaterialExpressionFunctionInput>(Expression))
 				{
-					Compiler->Errorf(TEXT("Layer graphs only support a single material attributes input."));
-					bHasValidOutput = false;
+					if (InputExpression->IsResultMaterialAttributes(0) || InputExpression->IsResultSubstrateMaterial(0))
+					{
+						++NumInputs;
+					}
+
+					if (NumInputs > 1)
+					{
+						Compiler->Errorf(TEXT("Layer graphs only support a single material attributes or Substrate inputs currently."));
+						bHasValidOutput = false;
+					}
 				}
-			}
-			else if (UMaterialExpressionFunctionOutput* OutputExpression = Cast<UMaterialExpressionFunctionOutput>(Expression))
-			{
-				++NumOutputs;
-				if (NumOutputs > 1 || !OutputExpression->IsResultMaterialAttributes(0))
+				else if (UMaterialExpressionFunctionOutput* OutputExpression = Cast<UMaterialExpressionFunctionOutput>(Expression))
+				{
+					if (OutputExpression->IsResultMaterialAttributes(0) || OutputExpression->IsResultSubstrateMaterial(0))
+					{
+						++NumOutputs;
+					}
+
+					if (NumOutputs > 1)
+					{
+						Compiler->Errorf(TEXT("Layer graphs only support a single material attributes or Substrate outputs currently."));
+						bHasValidOutput = false;
+					}
+				}
+				else if (UMaterialExpressionMaterialAttributeLayers* RecursiveLayer = Cast<UMaterialExpressionMaterialAttributeLayers>(Expression))
 				{
-					Compiler->Errorf(TEXT("Layer graphs only support a single material attributes output."));
+					Compiler->Errorf(TEXT("Layer graphs do not support layers within layers."));
 					bHasValidOutput = false;
 				}
 			}
-			else if (UMaterialExpressionMaterialAttributeLayers* RecursiveLayer = Cast<UMaterialExpressionMaterialAttributeLayers>(Expression))
+
+			if (NumInputs > 1 || NumOutputs < 1)
 			{
-				Compiler->Errorf(TEXT("Layer graphs do not support layers within layers."));
-					bHasValidOutput = false;
+				Compiler->Errorf(TEXT("Layer graphs require a single material attributes or Substrate output and optionally, a single material attributes or Substrate input."));
+				bHasValidOutput = false;
 			}
 		}
-
-		if ( NumInputs > 1 || NumOutputs < 1)
+		else
+#endif //ENABLE_MATERIAL_LAYER_PROTOTYPE
 		{
-			Compiler->Errorf(TEXT("Layer graphs require a single material attributes output and optionally, a single material attributes input."));
-			bHasValidOutput = false;
+			// Material layers must have a single MA input and output only
+			for (UMaterialExpression* Expression : GetExpressions())
+			{
+				if (UMaterialExpressionFunctionInput* InputExpression = Cast<UMaterialExpressionFunctionInput>(Expression))
+				{
+					++NumInputs;
+					if (NumInputs > 1 || !InputExpression->IsResultMaterialAttributes(0))
+					{
+						Compiler->Errorf(TEXT("Layer graphs only support a single material attributes input."));
+						bHasValidOutput = false;
+					}
+				}
+				else if (UMaterialExpressionFunctionOutput* OutputExpression = Cast<UMaterialExpressionFunctionOutput>(Expression))
+				{
+					++NumOutputs;
+					if (NumOutputs > 1 || !OutputExpression->IsResultMaterialAttributes(0))
+					{
+						Compiler->Errorf(TEXT("Layer graphs only support a single material attributes output."));
+						bHasValidOutput = false;
+					}
+				}
+				else if (UMaterialExpressionMaterialAttributeLayers* RecursiveLayer = Cast<UMaterialExpressionMaterialAttributeLayers>(Expression))
+				{
+					Compiler->Errorf(TEXT("Layer graphs do not support layers within layers."));
+						bHasValidOutput = false;
+				}
+			}
+
+			if ( NumInputs > 1 || NumOutputs < 1)
+			{
+				Compiler->Errorf(TEXT("Layer graphs require a single material attributes output and optionally, a single material attributes input."));
+				bHasValidOutput = false;
+			}
 		}
 	}
 	else if (GetMaterialFunctionUsage() == EMaterialFunctionUsage::MaterialLayerBlend)
 	{
-		// Material layer blends can have up to two MA inputs and single MA output only
-		for (UMaterialExpression* Expression : GetExpressions())
+#if ENABLE_MATERIAL_LAYER_PROTOTYPE
+		if(Substrate::IsSubstrateEnabled())
 		{
-			if (UMaterialExpressionFunctionInput* InputExpression = Cast<UMaterialExpressionFunctionInput>(Expression))
+			// Material layer blends can have up to two MA inputs and single MA output only
+			for (UMaterialExpression* Expression : GetExpressions())
 			{
-				++NumInputs;
-				if (NumInputs > 2 || !InputExpression->IsResultMaterialAttributes(0))
+				if (UMaterialExpressionFunctionInput* InputExpression = Cast<UMaterialExpressionFunctionInput>(Expression))
 				{
-					Compiler->Errorf(TEXT("Layer blend graphs only support two material attributes inputs."));
-					bHasValidOutput = false;
+					if(InputExpression->IsResultMaterialAttributes(0) || InputExpression->IsResultSubstrateMaterial(0))
+					{
+						++NumInputs;
+					}
+
+					if (NumInputs > 2)
+					{
+						Compiler->Errorf(TEXT("Layer blend graphs only support two material attributes or Substrate inputs currently."));
+						bHasValidOutput = false;
+					}
 				}
-			}
-			else if (UMaterialExpressionFunctionOutput* OutputExpression = Cast<UMaterialExpressionFunctionOutput>(Expression))
-			{
-				++NumOutputs;
-				if (NumOutputs > 1 || !OutputExpression->IsResultMaterialAttributes(0))
+				else if (UMaterialExpressionFunctionOutput* OutputExpression = Cast<UMaterialExpressionFunctionOutput>(Expression))
 				{
-					Compiler->Errorf(TEXT("Layer blend graphs only support a single material attributes output."));
+					if (OutputExpression->IsResultMaterialAttributes(0) || OutputExpression->IsResultSubstrateMaterial(0))
+					{
+						++NumOutputs;
+					}
+					if (NumOutputs > 1)
+					{
+						Compiler->Errorf(TEXT("Layer blend graphs only support a single MA or Substrate output currently."));
+						bHasValidOutput = false;
+					}
+				}
+				else if (UMaterialExpressionMaterialAttributeLayers* RecursiveLayer = Cast<UMaterialExpressionMaterialAttributeLayers>(Expression))
+				{
+					Compiler->Errorf(TEXT("Layer blend graphs do not support layers within layers."));
 					bHasValidOutput = false;
 				}
 			}
-			else if (UMaterialExpressionMaterialAttributeLayers* RecursiveLayer = Cast<UMaterialExpressionMaterialAttributeLayers>(Expression))
+
+			if (NumOutputs < 1)
 			{
-				Compiler->Errorf(TEXT("Layer blend graphs do not support layers within layers."));
-					bHasValidOutput = false;
+				Compiler->Errorf(TEXT("Layer blend graphs must have a only a single MA or Substrate output currently."));
+				bHasValidOutput = false;
 			}
 		}
-
-		if (NumOutputs < 1)
+		else
+#endif //ENABLE_MATERIAL_LAYER_PROTOTYPE
 		{
-			Compiler->Errorf(TEXT("Layer blend graphs can have up to two material attributes inputs and a single output."));
-			bHasValidOutput = false;
+			// Material layer blends can have up to two MA inputs and single MA output only
+			for (UMaterialExpression* Expression : GetExpressions())
+			{
+				if (UMaterialExpressionFunctionInput* InputExpression = Cast<UMaterialExpressionFunctionInput>(Expression))
+				{
+					++NumInputs;
+					if (NumInputs > 2 || !InputExpression->IsResultMaterialAttributes(0))
+					{
+						Compiler->Errorf(TEXT("Layer blend graphs only support two material attributes inputs."));
+						bHasValidOutput = false;
+					}
+				}
+				else if (UMaterialExpressionFunctionOutput* OutputExpression = Cast<UMaterialExpressionFunctionOutput>(Expression))
+				{
+					++NumOutputs;
+					if (NumOutputs > 1 || !OutputExpression->IsResultMaterialAttributes(0))
+					{
+						Compiler->Errorf(TEXT("Layer blend graphs only support a single material attributes output."));
+						bHasValidOutput = false;
+					}
+				}
+				else if (UMaterialExpressionMaterialAttributeLayers* RecursiveLayer = Cast<UMaterialExpressionMaterialAttributeLayers>(Expression))
+				{
+					Compiler->Errorf(TEXT("Layer blend graphs do not support layers within layers."));
+						bHasValidOutput = false;
+				}
+			}
+
+			if (NumOutputs < 1)
+			{
+				Compiler->Errorf(TEXT("Layer blend graphs can have up to two material attributes inputs and a single output."));
+				bHasValidOutput = false;
+			}
 		}
 	}
 #endif
@@ -18219,15 +18711,32 @@ int32 UMaterialExpressionFunctionInput::CompilePreviewValue(FMaterialCompiler* C
 			return Compiler->Constant3(PreviewValue.X, PreviewValue.Y, PreviewValue.Z);
 		case FunctionInput_Vector4:
 			return Compiler->Constant4(PreviewValue.X, PreviewValue.Y, PreviewValue.Z, PreviewValue.W);
-		case FunctionInput_MaterialAttributes:		
-			return FMaterialAttributeDefinitionMap::CompileDefaultExpression(Compiler, AttributeID);
+		case FunctionInput_MaterialAttributes:
+		{
+#if ENABLE_MATERIAL_LAYER_PROTOTYPE
+			if (AttributeID == FMaterialAttributeDefinitionMap::GetID(MP_EmissiveColor))
+			{
+				return Compiler->Constant3(PreviewValue.X, PreviewValue.Y, PreviewValue.Z);
+			}
+
+			if (!Substrate::IsSubstrateEnabled() || AttributeID != FMaterialAttributeDefinitionMap::GetID(MP_FrontMaterial))
+#endif //ENABLE_MATERIAL_LAYER_PROTOTYPE
+			{
+				return FMaterialAttributeDefinitionMap::CompileDefaultExpression(Compiler, AttributeID);
+			}
+		}
+		case FunctionInput_Substrate:
+#if ENABLE_MATERIAL_LAYER_PROTOTYPE
+		{
+			return UMaterialExpressionSubstrateSlabBSDF::CompileDefaultSlab(Compiler, FVector3f(PreviewValue.X, PreviewValue.Y, PreviewValue.Z));
+		}
+#endif //ENABLE_MATERIAL_LAYER_PROTOTYPE
 		case FunctionInput_Texture2D:
 		case FunctionInput_TextureCube:
 		case FunctionInput_Texture2DArray:
 		case FunctionInput_TextureExternal:
 		case FunctionInput_StaticBool:
 		case FunctionInput_Bool:
-		case FunctionInput_Substrate:
 			return Compiler->Errorf(TEXT("Missing Preview connection for function input '%s'"), *InputName.ToString());
 		default:
 			return Compiler->Errorf(TEXT("Unknown input type"));
@@ -18399,29 +18908,45 @@ bool UMaterialExpressionFunctionInput::IsResultSubstrateMaterial(int32 OutputInd
 	}
 	else
 	{
+#if ENABLE_MATERIAL_LAYER_PROTOTYPE
+		//This ensures the default Substrate operator is always generated for layering, reflecting the behaviour in SubstrateGenerateMaterialTopologyTree
+		return Substrate::IsSubstrateEnabled() && IsResultMaterialAttributes(OutputIndex);
+#else
 		return false;
+#endif
 	}
 }
 
 void UMaterialExpressionFunctionInput::GatherSubstrateMaterialInfo(FSubstrateMaterialInfo& SubstrateMaterialInfo, int32 OutputIndex)
 {
+#if ENABLE_MATERIAL_LAYER_PROTOTYPE
+	if (IsResultSubstrateMaterial(OutputIndex))
+#else
 	if (FunctionInput_Substrate == InputType)
+#endif
 	{
 		FExpressionInput EffectivePreviewDuringCompileTracedInput = EffectivePreviewDuringCompile.GetTracedInput();
-		int32 ExpressionResult = INDEX_NONE;
 		if (EffectivePreviewDuringCompileTracedInput.GetTracedInput().Expression)
 		{
 			EffectivePreviewDuringCompileTracedInput.Expression->GatherSubstrateMaterialInfo(SubstrateMaterialInfo, EffectivePreviewDuringCompileTracedInput.OutputIndex);
+			return;
 		}
 	}
+#if ENABLE_MATERIAL_LAYER_PROTOTYPE
+	SubstrateMaterialInfo.AddShadingModel(SSM_DefaultLit);
+	SubstrateMaterialInfo.AddGuid(MaterialExpressionGuid);
+#endif
 }
 
 FSubstrateOperator* UMaterialExpressionFunctionInput::SubstrateGenerateMaterialTopologyTree(class FMaterialCompiler* Compiler, class UMaterialExpression* Parent, int32 OutputIndex)
 {
+#if ENABLE_MATERIAL_LAYER_PROTOTYPE
+	if (IsResultSubstrateMaterial(OutputIndex))
+#else
 	if (FunctionInput_Substrate == InputType)
+#endif
 	{
 		FExpressionInput EffectivePreviewDuringCompileTracedInput = EffectivePreviewDuringCompile.GetTracedInput();
-		int32 ExpressionResult = INDEX_NONE;
 		if (EffectivePreviewDuringCompileTracedInput.GetTracedInput().Expression)
 		{
 			return EffectivePreviewDuringCompileTracedInput.Expression->SubstrateGenerateMaterialTopologyTree(Compiler, Parent, EffectivePreviewDuringCompileTracedInput.OutputIndex);
@@ -25272,6 +25797,76 @@ int32 UMaterialExpressionSubstrateSlabBSDF::Compile(class FMaterialCompiler* Com
 	return OutputCodeChunk;
 }
 
+#if ENABLE_MATERIAL_LAYER_PROTOTYPE
+int32 UMaterialExpressionSubstrateSlabBSDF::CompileDefaultSlab(FMaterialCompiler* Compiler, FVector3f DiffuseAlbedoOverride)
+{
+	/**
+	 * This creates a default slab, primarily for material function previews where the input may be optional/missing.
+	 */
+	if (!Substrate::IsSubstrateEnabled())
+	{
+		return Compiler->SubstrateCreateAndRegisterNullMaterial();
+	}
+
+	FGuid PathUniqueId = Compiler->SubstrateTreeStackGetPathUniqueId();
+	FSubstrateOperator& SubstrateOperator = Compiler->SubstrateCompilationGetOperator(PathUniqueId);
+
+	int32 SSSProfileCodeChunk = INDEX_NONE;
+	int32 SpecularProfileCodeChunk = INDEX_NONE;	
+	int32 TangentCodeChunk = INDEX_NONE;
+	int32 NormalCodeChunk = Compiler->VertexNormal();
+	int32 ThicknessCodeChunk = Compiler->Constant(SUBSTRATE_LAYER_DEFAULT_THICKNESS_CM);
+
+	const FSubstrateRegisteredSharedLocalBasis NewRegisteredSharedLocalBasis = SubstrateCompilationInfoCreateSharedLocalBasis(Compiler, NormalCodeChunk, TangentCodeChunk);
+	SubstrateOperator.BSDFRegisteredSharedLocalBasis = NewRegisteredSharedLocalBasis;
+
+	int32 DiffuseAlbedoCodeChunk = Compiler->Constant3(DiffuseAlbedoOverride.X, DiffuseAlbedoOverride.Y, DiffuseAlbedoOverride.Z);
+	const float DefaultF0 = 0.04f;
+	int32 F0CodeChunk = Compiler->Constant3(DefaultF0, DefaultF0, DefaultF0);
+	int32 RoughnessCodeChunk = Compiler->Constant(0.5f);
+	int32 AnisotropyCodeChunk = Compiler->Constant(0.0f);
+	int32 F90CodeChunk = Compiler->Constant3(1.0f, 1.0f, 1.0f);
+	int32 SSSMFPCodeChunk = Compiler->Constant3(0.0f, 0.0f, 0.0f);
+	int32 SSSMFPScaleCodeChunk = Compiler->Constant(1.0f);
+	int32 SSSPhaseAnisotropyCodeChunk = Compiler->Constant(0.0f);
+	int32 SecondRoughnessCodeChunk = Compiler->Constant(0.0f);
+	int32 SecondRoughnessWeightCodeChunk = Compiler->Constant(0.0f);
+	int32 FuzzAmountCodeChunk = Compiler->Constant(0.0f);
+	int32 FuzzColorCodeChunk = Compiler->Constant3(0.0f, 0.0f, 0.0f);
+	int32 FuzzRoughnessCodeChunk = Compiler->Constant(0.5f);
+	int32 GlintValueCodeChunk = Compiler->Constant(0.0f);
+	int32 GlintUVCodeChunk = Compiler->Constant2(0.0f, 0.0f);
+
+	return Compiler->SubstrateSlabBSDF(
+		DiffuseAlbedoCodeChunk,
+		F0CodeChunk,
+		F90CodeChunk,
+		RoughnessCodeChunk,
+		AnisotropyCodeChunk,
+		Compiler->Constant(0.0f),
+		SSSMFPCodeChunk,
+		SSSMFPScaleCodeChunk,
+		SSSPhaseAnisotropyCodeChunk,
+		Compiler->Constant(0.0f),
+		Compiler->Constant3(0.0f, 0.0f, 0.0f),
+		SecondRoughnessCodeChunk,
+		SecondRoughnessWeightCodeChunk,
+		Compiler->Constant(0.0f),										// SecondRoughnessAsSimpleClearCoat
+		FuzzAmountCodeChunk,
+		FuzzColorCodeChunk,
+		FuzzRoughnessCodeChunk,
+		ThicknessCodeChunk,
+		GlintValueCodeChunk,
+		GlintUVCodeChunk,
+		Compiler->Constant(0.0f),
+		false,
+		NormalCodeChunk,
+		TangentCodeChunk,
+		Compiler->GetSubstrateSharedLocalBasisIndexMacro(NewRegisteredSharedLocalBasis),
+		(SubstrateOperator.Index != INDEX_NONE && SubstrateOperator.BSDFIndex != INDEX_NONE) && (!SubstrateOperator.bUseParameterBlending || (SubstrateOperator.bUseParameterBlending && SubstrateOperator.bRootOfParameterBlendingSubTree)) ? &SubstrateOperator : nullptr);
+}
+#endif //ENABLE_MATERIAL_LAYER_PROTOTYPE
+
 FSubstrateMaterialComplexity UMaterialExpressionSubstrateSlabBSDF::GetComplexity() const
 {
 	FSubstrateMaterialComplexity Out;
@@ -28272,6 +28867,252 @@ bool UMaterialExpressionSubstrateConvertMaterialAttributes::HasSSS() const
 
 #endif // WITH_EDITOR
 
+// Substrate Attributes Start -----
+UMaterialExpressionSubstrateGetAttributes::UMaterialExpressionSubstrateGetAttributes(const FObjectInitializer& ObjectInitializer)
+	: Super(ObjectInitializer)
+{
+#if ENABLE_MATERIAL_LAYER_PROTOTYPE
+#if WITH_EDITORONLY_DATA
+	// Structure to hold one-time initialization
+	struct FConstructorStatics
+	{
+		FText NAME_SubstrateAttributes;
+		FConstructorStatics()
+			: NAME_SubstrateAttributes(LOCTEXT("SubstrateAttributes", "Substrate Attributes"))
+		{
+		}
+	};
+
+	static FConstructorStatics ConstructorStatics;
+
+	bShowOutputNameOnPin = true;
+	bShowMaskColorsOnPin = false;
+
+	MenuCategories.Add(ConstructorStatics.NAME_SubstrateAttributes); //Check this setup for other Substrate nodes
+
+	Outputs.Reset();
+	Outputs.SetNum(ESubstrateAttributeIndex::MSA_MAX);
+	Outputs[ESubstrateAttributeIndex::MSA_NonSubstrateAttributes] = FExpressionOutput(NON_SUBSTRATE_ATTRIBUTES_TEXT);
+	Outputs[ESubstrateAttributeIndex::MSA_FrontMaterial] = FExpressionOutput(FRONT_MATERIAL_ATTRIBUTES_TEXT);
+#endif
+
+#if WITH_EDITOR
+	CachedInputs.Empty();
+	CachedInputs.Add(&MaterialAttributes);
+#endif
+#endif //ENABLE_MATERIAL_LAYER_PROTOTYPE
+}
+
+#if ENABLE_MATERIAL_LAYER_PROTOTYPE
+#if WITH_EDITOR
+void UMaterialExpressionSubstrateGetAttributes::PostEditChangeProperty(FPropertyChangedEvent& PropertyChangedEvent)
+{
+	if (GraphNode && PropertyChangedEvent.Property != nullptr)
+	{
+		GraphNode->ReconstructNode();
+	}
+	Super::PostEditChangeProperty(PropertyChangedEvent);
+}
+
+int32 UMaterialExpressionSubstrateGetAttributes::Compile(class FMaterialCompiler* Compiler, int32 OutputIndex)
+{
+	int32 Ret = INDEX_NONE;
+
+	FGuid AttributeId = Compiler->GetMaterialAttribute();
+	EMaterialProperty Property = FMaterialAttributeDefinitionMap::GetProperty(AttributeId);
+
+	if ((OutputIndex == ESubstrateAttributeIndex::MSA_NonSubstrateAttributes && Property != MP_FrontMaterial)
+		|| (OutputIndex == ESubstrateAttributeIndex::MSA_FrontMaterial && Property == MP_FrontMaterial))
+	{
+		Ret = MaterialAttributes.CompileWithDefault(Compiler, AttributeId);
+	}
+
+	return Ret;
+}
+
+void UMaterialExpressionSubstrateGetAttributes::GatherSubstrateMaterialInfo(FSubstrateMaterialInfo& SubstrateMaterialInfo, int32 OutputIndex)
+{
+	if (MaterialAttributes.Expression && IsResultSubstrateMaterial(OutputIndex))
+	{
+		MaterialAttributes.Expression->GatherSubstrateMaterialInfo(SubstrateMaterialInfo, MaterialAttributes.OutputIndex);
+	}
+}
+
+FSubstrateOperator* UMaterialExpressionSubstrateGetAttributes::SubstrateGenerateMaterialTopologyTree(class FMaterialCompiler* Compiler, class UMaterialExpression* Parent, int32 OutputIndex)
+{
+	if (MaterialAttributes.Expression && IsResultSubstrateMaterial(OutputIndex))
+	{
+		return MaterialAttributes.Expression->SubstrateGenerateMaterialTopologyTree(Compiler, Parent, MaterialAttributes.OutputIndex);
+	}
+
+	return nullptr;
+}
+
+void UMaterialExpressionSubstrateGetAttributes::GetCaption(TArray<FString>& OutCaptions) const
+{
+	OutCaptions.Add(TEXT("Get Substrate Attributes"));
+}
+
+
+FName UMaterialExpressionSubstrateGetAttributes::GetInputName(int32 InputIndex) const
+{
+	if (0 == InputIndex)
+	{
+		return *NSLOCTEXT("GetSubstrateAttributes", "InputName", "Material Attributes").ToString();
+	}
+	return NAME_None;
+}
+
+bool UMaterialExpressionSubstrateGetAttributes::IsInputConnectionRequired(int32 InputIndex) const
+{
+	return true;
+}
+
+uint32 UMaterialExpressionSubstrateGetAttributes::GetInputType(int32 InputIndex)
+{
+	return MCT_MaterialAttributes;
+}
+
+uint32 UMaterialExpressionSubstrateGetAttributes::GetOutputType(int32 OutputIndex)
+{
+	switch (OutputIndex)
+	{
+	case ESubstrateAttributeIndex::MSA_FrontMaterial: return MCT_Substrate;
+	case ESubstrateAttributeIndex::MSA_NonSubstrateAttributes: return MCT_MaterialAttributes;
+	}
+	check(false);
+	return MCT_Float1;
+}
+#endif // WITH_EDITOR
+#endif //ENABLE_MATERIAL_LAYER_PROTOTYPE
+
+UMaterialExpressionSubstrateSetAttributes::UMaterialExpressionSubstrateSetAttributes(const FObjectInitializer& ObjectInitializer)
+	: Super(ObjectInitializer)
+{
+#if ENABLE_MATERIAL_LAYER_PROTOTYPE
+#if WITH_EDITORONLY_DATA
+	// Structure to hold one-time initialization
+	struct FConstructorStatics
+	{
+		FText NAME_SubstrateAttributes;
+		FConstructorStatics()
+			: NAME_SubstrateAttributes(LOCTEXT("SubstrateAttributes", "Substrate Attributes"))
+		{
+		}
+	};
+
+	static FConstructorStatics ConstructorStatics;
+
+	bShowOutputNameOnPin = true;
+	bShowMaskColorsOnPin = false;
+
+	MenuCategories.Add(ConstructorStatics.NAME_SubstrateAttributes);
+
+	Outputs.Reset();
+	Outputs.Add(FExpressionOutput(TEXT("Material Attributes")));
+#endif
+
+#if WITH_EDITOR
+	CachedInputs.Empty();
+	CachedInputs.Add(&FrontMaterial);
+	CachedInputs.Add(&NonSubstrateAttributes);
+#endif
+#endif //ENABLE_MATERIAL_LAYER_PROTOTYPE
+}
+
+#if ENABLE_MATERIAL_LAYER_PROTOTYPE
+#if WITH_EDITOR
+void UMaterialExpressionSubstrateSetAttributes::PostEditChangeProperty(FPropertyChangedEvent& PropertyChangedEvent)
+{
+	if (GraphNode && PropertyChangedEvent.Property != nullptr)
+	{
+		GraphNode->ReconstructNode();
+	}
+	Super::PostEditChangeProperty(PropertyChangedEvent);
+}
+
+int32 UMaterialExpressionSubstrateSetAttributes::Compile(class FMaterialCompiler* Compiler, int32 OutputIndex)
+{
+	int32 Ret = INDEX_NONE;
+	UMaterialExpression* Expression = nullptr;
+
+	EMaterialProperty Property = FMaterialAttributeDefinitionMap::GetProperty(Compiler->GetMaterialAttribute());
+	if (Property == MP_FrontMaterial)
+	{
+		Expression = FrontMaterial.Expression;
+	}
+	else
+	{
+		Expression = NonSubstrateAttributes.Expression;
+	}
+
+	if (Expression)
+	{
+		Ret = Expression->Compile(Compiler, OutputIndex);
+	}
+
+	return Ret;
+}
+
+void UMaterialExpressionSubstrateSetAttributes::GatherSubstrateMaterialInfo(FSubstrateMaterialInfo& SubstrateMaterialInfo, int32 OutputIndex)
+{
+	if (FrontMaterial.Expression)
+	{
+		FrontMaterial.Expression->GatherSubstrateMaterialInfo(SubstrateMaterialInfo, FrontMaterial.OutputIndex);
+	}
+}
+
+FSubstrateOperator* UMaterialExpressionSubstrateSetAttributes::SubstrateGenerateMaterialTopologyTree(class FMaterialCompiler* Compiler, class UMaterialExpression* Parent, int32 OutputIndex)
+{
+	if (FrontMaterial.Expression)
+	{
+		return FrontMaterial.Expression->SubstrateGenerateMaterialTopologyTree(Compiler, Parent, FrontMaterial.OutputIndex);
+	}
+
+	return nullptr;
+}
+
+void UMaterialExpressionSubstrateSetAttributes::GetCaption(TArray<FString>& OutCaptions) const
+{
+	OutCaptions.Add(TEXT("Set Substrate Attributes"));
+}
+
+
+FName UMaterialExpressionSubstrateSetAttributes::GetInputName(int32 InputIndex) const
+{
+	if (InputIndex == ESubstrateAttributeIndex::MSA_NonSubstrateAttributes)
+	{
+		return *NSLOCTEXT("SetSubstrateAttributes", "InputName", NON_SUBSTRATE_ATTRIBUTES_TEXT).ToString();
+	}
+	if (InputIndex == ESubstrateAttributeIndex::MSA_FrontMaterial)
+	{
+		return *NSLOCTEXT("SetSubstrateAttributes", "InputName", FRONT_MATERIAL_ATTRIBUTES_TEXT).ToString();
+	}
+	return NAME_None;
+}
+
+uint32 UMaterialExpressionSubstrateSetAttributes::GetInputType(int32 InputIndex)
+{
+	if (InputIndex == ESubstrateAttributeIndex::MSA_NonSubstrateAttributes)
+	{
+		return MCT_MaterialAttributes;
+	}
+	if (InputIndex == ESubstrateAttributeIndex::MSA_FrontMaterial)
+	{
+		return MCT_Substrate;
+	}
+	return Super::GetInputType(InputIndex);
+}
+
+bool UMaterialExpressionSubstrateSetAttributes::IsInputConnectionRequired(int32 InputIndex) const
+{
+	return true;
+}
+#endif // WITH_EDITOR
+#endif // ENABLE_MATERIAL_LAYER_PROTOTYPE
+
+// Substrate Attributes End -----
+
 UMaterialExpressionExecBegin::UMaterialExpressionExecBegin(const FObjectInitializer& ObjectInitializer)
 	: Super(ObjectInitializer)
 {
diff --git a/Engine/Source/Runtime/Engine/Public/Materials/MaterialExpression.h b/Engine/Source/Runtime/Engine/Public/Materials/MaterialExpression.h
--- a/Engine/Source/Runtime/Engine/Public/Materials/MaterialExpression.h
+++ b/Engine/Source/Runtime/Engine/Public/Materials/MaterialExpression.h
@@ -697,3 +697,9 @@ struct FExpressionInputIterator
 };
 
 #endif // WITH_EDITOR
+
+#ifndef ENABLE_MATERIAL_LAYER_PROTOTYPE
+//WIP Internal Only - Please do not enable + use this in any projects other than for iterating the new layering prototype system until further notice
+//Temporarily placed here for the best possible catch all position when in use - will be removed for 5.5
+#define ENABLE_MATERIAL_LAYER_PROTOTYPE 0
+#endif
\ No newline at end of file
diff --git a/Engine/Source/Runtime/Engine/Public/Materials/MaterialExpressionMaterialAttributeLayers.h b/Engine/Source/Runtime/Engine/Public/Materials/MaterialExpressionMaterialAttributeLayers.h
--- a/Engine/Source/Runtime/Engine/Public/Materials/MaterialExpressionMaterialAttributeLayers.h
+++ b/Engine/Source/Runtime/Engine/Public/Materials/MaterialExpressionMaterialAttributeLayers.h
@@ -109,6 +109,10 @@ class UMaterialExpressionMaterialAttributeLayers : public UMaterialExpression
 	virtual bool IsInputConnectionRequired(int32 InputIndex) const override {return false;}
 	virtual uint32 GetInputType(int32 InputIndex) override;
 	virtual bool IsResultMaterialAttributes(int32 OutputIndex) override {return true;}
+#if ENABLE_MATERIAL_LAYER_PROTOTYPE
+	virtual bool IsResultSubstrateMaterial(int32 OutputIndex) override;
+	virtual FSubstrateOperator* SubstrateGenerateMaterialTopologyTree(class FMaterialCompiler* Compiler, class UMaterialExpression* Parent, int32 OutputIndex) override;
+#endif //ENABLE_MATERIAL_LAYER_PROTOTYPE
 	virtual bool GenerateHLSLExpression(FMaterialHLSLGenerator& Generator, UE::HLSLTree::FScope& Scope, int32 OutputIndex, UE::HLSLTree::FExpression const*& OutExpression) const override;
 #endif
 	//~ End UMaterialExpression Interface
diff --git a/Engine/Source/Runtime/Engine/Public/Materials/MaterialExpressionSetMaterialAttributes.h b/Engine/Source/Runtime/Engine/Public/Materials/MaterialExpressionSetMaterialAttributes.h
--- a/Engine/Source/Runtime/Engine/Public/Materials/MaterialExpressionSetMaterialAttributes.h
+++ b/Engine/Source/Runtime/Engine/Public/Materials/MaterialExpressionSetMaterialAttributes.h
@@ -42,6 +42,8 @@ class UMaterialExpressionSetMaterialAttributes : public UMaterialExpression
 	virtual uint32 GetInputType(int32 InputIndex) override;
 	virtual bool IsInputConnectionRequired(int32 InputIndex) const override {return true;}
 	virtual bool IsResultMaterialAttributes(int32 OutputIndex) override {return true;}
+	virtual bool IsResultSubstrateMaterial(int32 OutputIndex) override;
+	virtual FSubstrateOperator* SubstrateGenerateMaterialTopologyTree(FMaterialCompiler* Compiler, UMaterialExpression* Parent, int32 OutputIndex) override;
 	virtual void GetExpressionToolTip(TArray<FString>& OutToolTip) override;
 
 	virtual bool GenerateHLSLExpression(FMaterialHLSLGenerator& Generator, UE::HLSLTree::FScope& Scope, int32 OutputIndex, UE::HLSLTree::FExpression const*& OutExpression) const override;
diff --git a/Engine/Source/Runtime/Engine/Public/Materials/MaterialExpressionSubstrate.h b/Engine/Source/Runtime/Engine/Public/Materials/MaterialExpressionSubstrate.h
--- a/Engine/Source/Runtime/Engine/Public/Materials/MaterialExpressionSubstrate.h
+++ b/Engine/Source/Runtime/Engine/Public/Materials/MaterialExpressionSubstrate.h
@@ -330,6 +330,9 @@ class UMaterialExpressionSubstrateSlabBSDF : public UMaterialExpressionSubstrate
 	//~ Begin UMaterialExpression Interface
 #if WITH_EDITOR
 	virtual int32 Compile(class FMaterialCompiler* Compiler, int32 OutputIndex) override;
+#if ENABLE_MATERIAL_LAYER_PROTOTYPE
+	static int32 CompileDefaultSlab(class FMaterialCompiler* Compiler, FVector3f DiffuseAlbedoOverride = {0.18f, 0.18f, 0.18f});
+#endif //ENABLE_MATERIAL_LAYER_PROTOTYPE
 	virtual void GetCaption(TArray<FString>& OutCaptions) const override;
 	virtual uint32 GetOutputType(int32 OutputIndex) override;
 	virtual uint32 GetInputType(int32 InputIndex) override;
@@ -1192,3 +1195,77 @@ class UMaterialExpressionSubstrateThinFilm : public UMaterialExpressionSubstrate
 #endif
 	//~ End UMaterialExpression Interface
 };
+
+///////////////////////////////////////////////////////////////////////////////
+// Attributes
+
+#if ENABLE_MATERIAL_LAYER_PROTOTYPE
+enum ESubstrateAttributeIndex : int
+{
+	MSA_FrontMaterial,
+	MSA_NonSubstrateAttributes,
+	MSA_MAX
+};
+
+#define FRONT_MATERIAL_ATTRIBUTES_TEXT "Front Material"
+#define NON_SUBSTRATE_ATTRIBUTES_TEXT "Non-Substrate Attributes"
+
+inline bool CheckIfNonSubstratePropertyType(EMaterialProperty Property)
+{
+	return Property != MP_FrontMaterial;
+}
+
+#endif //#if ENABLE_MATERIAL_LAYER_PROTOTYPE
+UCLASS(MinimalAPI, collapsecategories, hidecategories = Object, DisplayName = "Substrate Get Attributes")
+class UMaterialExpressionSubstrateGetAttributes : public UMaterialExpressionSubstrateUtilityBase
+{
+	GENERATED_UCLASS_BODY()
+
+	UPROPERTY()
+	FMaterialAttributesInput MaterialAttributes;
+
+#if ENABLE_MATERIAL_LAYER_PROTOTYPE
+#if WITH_EDITOR
+	virtual void PostEditChangeProperty(FPropertyChangedEvent& PropertyChangedEvent) override;
+	virtual int32 Compile(class FMaterialCompiler* Compiler, int32 OutputIndex) override;
+	virtual void GetCaption(TArray<FString>& OutCaptions) const override;
+	virtual FName GetInputName(int32 InputIndex) const override;
+	virtual bool IsInputConnectionRequired(int32 InputIndex) const override;
+	virtual uint32 GetInputType(int32 InputIndex) override;
+	virtual uint32 GetOutputType(int32 OutputIndex) override;
+
+	virtual bool IsResultMaterialAttributes(int32 OutputIndex) override { return OutputIndex == MSA_NonSubstrateAttributes; }
+	virtual bool IsResultSubstrateMaterial(int32 OutputIndex) override { return OutputIndex == MSA_FrontMaterial; }
+	virtual void GatherSubstrateMaterialInfo(FSubstrateMaterialInfo& SubstrateMaterialInfo, int32 OutputIndex) override;
+	virtual FSubstrateOperator* SubstrateGenerateMaterialTopologyTree(class FMaterialCompiler* Compiler, class UMaterialExpression* Parent, int32 OutputIndex) override;
+#endif // WITH_EDITOR
+#endif //#if ENABLE_MATERIAL_LAYER_PROTOTYPE
+};
+
+UCLASS(MinimalAPI, collapsecategories, hidecategories = Object, DisplayName = "Substrate Set Attributes")
+class UMaterialExpressionSubstrateSetAttributes : public UMaterialExpressionSubstrateUtilityBase
+{
+	GENERATED_UCLASS_BODY()
+
+	UPROPERTY()
+	FExpressionInput FrontMaterial;
+
+	UPROPERTY()
+ 	FMaterialAttributesInput NonSubstrateAttributes;
+
+#if ENABLE_MATERIAL_LAYER_PROTOTYPE
+#if WITH_EDITOR
+	virtual void PostEditChangeProperty(FPropertyChangedEvent& PropertyChangedEvent) override;
+	virtual int32 Compile(class FMaterialCompiler* Compiler, int32 OutputIndex) override;
+	virtual void GetCaption(TArray<FString>& OutCaptions) const override;
+	virtual FName GetInputName(int32 InputIndex) const override;
+	virtual bool IsInputConnectionRequired(int32 InputIndex) const override;
+	virtual uint32 GetInputType(int32 InputIndex) override;
+
+	virtual bool IsResultMaterialAttributes(int32 OutputIndex) override { return true; }
+	virtual bool IsResultSubstrateMaterial(int32 OutputIndex) override { return FrontMaterial.IsConnected(); }
+	virtual void GatherSubstrateMaterialInfo(FSubstrateMaterialInfo& SubstrateMaterialInfo, int32 OutputIndex) override;
+	virtual FSubstrateOperator* SubstrateGenerateMaterialTopologyTree(class FMaterialCompiler* Compiler, class UMaterialExpression* Parent, int32 OutputIndex) override;
+#endif // WITH_EDITOR
+#endif //ENABLE_MATERIAL_LAYER_PROTOTYPE
+};
\ No newline at end of file
diff --git a/Engine/Source/Runtime/Engine/Public/Materials/MaterialFunction.h b/Engine/Source/Runtime/Engine/Public/Materials/MaterialFunction.h
--- a/Engine/Source/Runtime/Engine/Public/Materials/MaterialFunction.h
+++ b/Engine/Source/Runtime/Engine/Public/Materials/MaterialFunction.h
@@ -166,6 +166,13 @@ public:
 
 	/** Returns If returns an empty string, use the default class name for the material function. Otherwise, the string will be the name shown when the function is exposed to users in the material graph as a node, or from the contextual menu when searching for nodes. */
 	virtual FString GetUserExposedCaption() const override { return UserExposedCaption; }
+#if ENABLE_MATERIAL_LAYER_PROTOTYPE
+	/** Converts legacy Material Layer Functions to utilise Substrate */
+	void ConvertExpressionsBetweenLegacyAndSubstrate();
+
+	/** Retrieves an array of all custom output expressions in a Material Function */
+	void GetAllCustomOutputExpressions(TArray<class UMaterialExpressionCustomOutput*>& OutCustomOutputs) const;
+#endif //ENABLE_MATERIAL_LAYER_PROTOTYPE
 #endif // WITH_EDITORONLY_DATA
 
 #if WITH_EDITOR
diff --git a/Engine/Source/Runtime/Engine/Public/Materials/MaterialFunctionInterface.h b/Engine/Source/Runtime/Engine/Public/Materials/MaterialFunctionInterface.h
--- a/Engine/Source/Runtime/Engine/Public/Materials/MaterialFunctionInterface.h
+++ b/Engine/Source/Runtime/Engine/Public/Materials/MaterialFunctionInterface.h
@@ -35,6 +35,9 @@ enum class EMaterialFunctionUsage : uint8
 	MaterialLayer,
 	MaterialLayerBlend
 };
+#if ENABLE_MATERIAL_LAYER_PROTOTYPE
+#define DEFAULT_MATERIALLAYERBLEND_PATH TEXT("/Engine/EngineMaterials/Blends/MLB_Horizontal")
+#endif
 
 using FMFRecursionGuard = TMaterialRecursionGuard<class UMaterialFunctionInterface>;
 
