Fix compilation errors when compiling without Iris modules

#rnx

diff --git a/Engine/Source/Runtime/Engine/Private/DataChannel.cpp b/Engine/Source/Runtime/Engine/Private/DataChannel.cpp
--- a/Engine/Source/Runtime/Engine/Private/DataChannel.cpp
+++ b/Engine/Source/Runtime/Engine/Private/DataChannel.cpp
@@ -11,7 +11,9 @@
 #include "Engine/Level.h"
 #include "Engine/World.h"
 #include "GameFramework/WorldSettings.h"
+#if UE_WITH_IRIS
 #include "Iris/ReplicationSystem/ReplicationSystem.h"
+#endif
 #include "Misc/MemStack.h"
 #include "Misc/ScopeExit.h"
 #include "Net/Core/Trace/Private/NetTraceInternal.h"
diff --git a/Engine/Source/Runtime/Engine/Private/Net/Experimental/Iris/DataStreamChannel.cpp b/Engine/Source/Runtime/Engine/Private/Net/Experimental/Iris/DataStreamChannel.cpp
--- a/Engine/Source/Runtime/Engine/Private/Net/Experimental/Iris/DataStreamChannel.cpp
+++ b/Engine/Source/Runtime/Engine/Private/Net/Experimental/Iris/DataStreamChannel.cpp
@@ -469,12 +469,14 @@ void UDataStreamChannel::ReceivedNak(int32 PacketId)
 // Some DataStreams require perfect acking. If the ack sequence window is full we would get NAKs for packets thay may have been received.
 bool UDataStreamChannel::IsPacketWindowFull() const
 {
+#if UE_WITH_IRIS
 	const uint32 IrisPacketSequenceSafetyMarginUnsigned = static_cast<uint32>(FPlatformMath::Max(0, UE::Net::Private::IrisPacketSequenceSafetyMargin));
 	if (Connection->IsPacketSequenceWindowFull(IrisPacketSequenceSafetyMarginUnsigned))
 	{
 		UE_LOG(LogIris, Verbose, TEXT("Packet window full."));
 		return true;
 	}
+#endif // UE_WITH_IRIS
 
 	return WriteRecords.Count() >= WriteRecords.AllocatedCapacity();
 }
diff --git a/Engine/Source/Runtime/Engine/Private/NetDriver.cpp b/Engine/Source/Runtime/Engine/Private/NetDriver.cpp
--- a/Engine/Source/Runtime/Engine/Private/NetDriver.cpp
+++ b/Engine/Source/Runtime/Engine/Private/NetDriver.cpp
@@ -248,7 +248,7 @@ namespace UE::Net::Private
 
 	static bool bIgnoreStaticActorDestruction = false;
 
-
+#if UE_WITH_IRIS
 	void ApplyReplicationSystemConfig(const FNetDriverReplicationSystemConfig& ReplicationSystemConfig, UReplicationSystem::FReplicationSystemParams& OutParams)
 	{
 		if (ReplicationSystemConfig.MaxReplicatedObjectCount != 0)
@@ -286,6 +286,7 @@ namespace UE::Net::Private
 			OutParams.MaxNetObjectGroupCount = ReplicationSystemConfig.MaxNetObjectGroupCount;
 		}
 	}
+#endif // UE_WITH_IRIS
 
 	bool IsGuidInOuterChain(const FNetGUIDCache& GuidCache, const FNetGuidCacheObject* CacheObj, FNetworkGUID GuidMatch)
 	{
diff --git a/Engine/Source/Runtime/Engine/Private/World.cpp b/Engine/Source/Runtime/Engine/Private/World.cpp
--- a/Engine/Source/Runtime/Engine/Private/World.cpp
+++ b/Engine/Source/Runtime/Engine/Private/World.cpp
@@ -7895,9 +7895,15 @@ UWorld* FSeamlessTravelHandler::Tick()
 				// Warn if we loaded a game mode that wanted a different replication system from the previous mode.
 				if (AGameModeBase* GameMode = LoadedWorld->GetAuthGameMode())
 				{
-					EReplicationSystem LoadedGameModeRepSystem = GameMode->GetGameNetDriverReplicationSystem();
+					const EReplicationSystem LoadedGameModeRepSystem = GameMode->GetGameNetDriverReplicationSystem();
+#if UE_WITH_IRIS
+					const EReplicationSystem IrisCmdlineRepSystem = UE::Net::GetUseIrisReplicationCmdlineValue();
+#else
+					const EReplicationSystem IrisCmdlineRepSystem = EReplicationSystem::Default;
+#endif
+
 					const bool bIsNetDriverCompatible = LoadedGameModeRepSystem == EReplicationSystem::Default || 
-														UE::Net::GetUseIrisReplicationCmdlineValue() != EReplicationSystem::Default ||
+														IrisCmdlineRepSystem != EReplicationSystem::Default ||
 														(LoadedGameModeRepSystem == EReplicationSystem::Iris && NetDriver->IsUsingIrisReplication());
 					ensureMsgf(bIsNetDriverCompatible, TEXT("Seamless travel loaded game mode %s that wants a different replication system than the current NetDriver uses."), *GetNameSafe(GameMode));
 				}
