- Added UE_FORCE_CRASH, and replaced some hardcoded crash calls with it

diff --git a/Engine/Source/Runtime/Core/Private/IOS/IOSPlatformMemory.cpp b/Engine/Source/Runtime/Core/Private/IOS/IOSPlatformMemory.cpp
--- a/Engine/Source/Runtime/Core/Private/IOS/IOSPlatformMemory.cpp
+++ b/Engine/Source/Runtime/Core/Private/IOS/IOSPlatformMemory.cpp
@@ -55,8 +55,7 @@ void FIOSPlatformMemory::OnOutOfMemory(uint64 Size, uint32 Alignment)
 		UE_LOG(LogMemory, Warning, TEXT("Ran out of memory allocating %llu bytes with alignment %u. Last error msg: %s."), Size, Alignment, ErrorMsg);
     
 		// make this a fatal error that ends here not in the log
-		// changed to 3 from NULL because clang noticed writing to NULL and warned about it
-		*(int32 *)3 = 123;
+		UE_FORCE_CRASH();
 	};
 	
 	UE_CALL_ONCE(HandleOOM);
diff --git a/Engine/Source/Runtime/Core/Private/IOS/IOSPlatformMisc.cpp b/Engine/Source/Runtime/Core/Private/IOS/IOSPlatformMisc.cpp
--- a/Engine/Source/Runtime/Core/Private/IOS/IOSPlatformMisc.cpp
+++ b/Engine/Source/Runtime/Core/Private/IOS/IOSPlatformMisc.cpp
@@ -1916,14 +1916,14 @@ void FIOSPlatformMisc::GPUAssert()
 {
     // make this a fatal error that ends here not in the log
     // changed to 3 from NULL because clang noticed writing to NULL and warned about it
-    *(int32 *)13 = 123;
+	UE_FORCE_CRASH_AT_OFFSET(13);
 }
 
 void FIOSPlatformMisc::MetalAssert()
 {
     // make this a fatal error that ends here not in the log
     // changed to 3 from NULL because clang noticed writing to NULL and warned about it
-    *(int32 *)7 = 123;
+	UE_FORCE_CRASH_AT_OFFSET(7);
 }
 
 struct FCPUFeatures
diff --git a/Engine/Source/Runtime/Core/Public/GenericPlatform/GenericPlatformMisc.h b/Engine/Source/Runtime/Core/Public/GenericPlatform/GenericPlatformMisc.h
--- a/Engine/Source/Runtime/Core/Public/GenericPlatform/GenericPlatformMisc.h
+++ b/Engine/Source/Runtime/Core/Public/GenericPlatform/GenericPlatformMisc.h
@@ -1408,9 +1408,9 @@ public:
 #	endif
 #else
 		// any other architecture must implement this function so this is not called, but we can't use any convenience functions 
-		// like unimplemented() or checkf(), so just crash with a comment that if you get here, implement Prefetch!
-		int* A = (int*)3;
-		*A = 13;
+		// like unimplemented() or checkf(), so just crash with a comment that if you get here, implement Prefetch and ALL Prefecth calls below
+		// this is because the Prefetch calls below will end up calling _this one_ not the platform specific functions
+		UE_FORCE_CRASH();
 #endif
 	}
 
diff --git a/Engine/Source/Runtime/Core/Public/HAL/Platform.h b/Engine/Source/Runtime/Core/Public/HAL/Platform.h
--- a/Engine/Source/Runtime/Core/Public/HAL/Platform.h
+++ b/Engine/Source/Runtime/Core/Public/HAL/Platform.h
@@ -957,6 +957,17 @@ int32 main(int32 ArgC, ANSICHAR* Utf8ArgV[]) \
 int32 tchar_main(int32 ArgC, TCHAR* ArgV[])
 #endif
 
+// perform an unaligned write to almost-zero (writing to 0 will throw warnings in some compilers, plus it might actually not crash)
+// this allow some platforms to change how they force a crash, in case this isn't enough
+// a platform can likely just #define UE_FORCE_CRASH_AT_OFFSET to override the crash behavior
+#ifndef UE_FORCE_CRASH_AT_OFFSET
+	#define UE_FORCE_CRASH_AT_OFFSET(x) *(int32 *)x = 123
+#endif
+
+#ifndef UE_FORCE_CRASH
+	#define UE_FORCE_CRASH() UE_FORCE_CRASH_AT_OFFSET(3)
+#endif
+
 //--------------------------------------------------------------------------------------------
 // POD types refactor for porting old code:
 //
diff --git a/Engine/Source/Runtime/Engine/Private/UnrealEngine.cpp b/Engine/Source/Runtime/Engine/Private/UnrealEngine.cpp
--- a/Engine/Source/Runtime/Engine/Private/UnrealEngine.cpp
+++ b/Engine/Source/Runtime/Engine/Private/UnrealEngine.cpp
@@ -10499,7 +10499,7 @@ bool UEngine::PerformError(const TCHAR* Cmd, FOutputDevice& Ar)
 			{
 				UE_LOG(LogEngine, Warning, TEXT("Printed warning to log."));
 				FGenericCrashContext::SetCrashTrigger(ECrashTrigger::Debug);
-				*(int32 *)3 = 123;
+				UE_FORCE_CRASH();
 			}
 		};
 		ENQUEUE_RENDER_COMMAND(CauseRenderThreadCrash)(&FRender::GPF);
@@ -10617,7 +10617,7 @@ bool UEngine::PerformError(const TCHAR* Cmd, FOutputDevice& Ar)
 			{
 				UE_LOG(LogEngine, Warning, TEXT("Printed warning to log."));
 				FGenericCrashContext::SetCrashTrigger(ECrashTrigger::Debug);
-				*(int32 *)3 = 123;
+				UE_FORCE_CRASH();
 			}
 		};
 
@@ -10725,7 +10725,7 @@ bool UEngine::PerformError(const TCHAR* Cmd, FOutputDevice& Ar)
 					{
 						UE_LOG(LogEngine, Warning, TEXT("Printed warning to log."));
 						FGenericCrashContext::SetCrashTrigger(ECrashTrigger::Debug);
-						*(int32 *)3 = 123;
+						UE_FORCE_CRASH();
 						break;
 					}
 					else
@@ -10827,8 +10827,7 @@ bool UEngine::PerformError(const TCHAR* Cmd, FOutputDevice& Ar)
 		UE_LOG(LogEngine, Warning, TEXT("Printed warning to log."));
 		Ar.Log(TEXT("Crashing with voluntary GPF"));
 		FGenericCrashContext::SetCrashTrigger(ECrashTrigger::Debug);
-		// changed to 3 from NULL because clang noticed writing to NULL and warned about it
-		*(int32 *)3 = 123;
+		UE_FORCE_CRASH();
 		return true;
 	}
 	else if (FParse::Command(&Cmd, TEXT("ENSURE")))
@@ -11166,8 +11165,8 @@ bool UEngine::PerformError(const TCHAR* Cmd, FOutputDevice& Ar)
 		struct FAudio
 		{
 			static void GPF()
-		{
-			*(int32 *)3 = 123;
+			{
+				UE_FORCE_CRASH();
 			}
 		};
 		FAudioThread::RunCommandOnAudioThread(&FAudio::GPF, TStatId());
@@ -14801,7 +14800,7 @@ EBrowseReturnVal::Type UEngine::Browse( FWorldContext& WorldContext, FURL URL, F
 			HandleBrowseToDefaultMapFailure(WorldContext, DefaultURL.ToString(), Error);
 			return EBrowseReturnVal::Failure;
 		}
-
+		
 		CollectGarbage( GARBAGE_COLLECTION_KEEPFLAGS );
 
 		// now remove "failed" and "closed" options from LastURL so it doesn't get copied on to future URLs
@@ -15199,7 +15198,7 @@ bool UEngine::LoadMap( FWorldContext& WorldContext, FURL URL, class UPendingNetG
 	CleanupPackagesToFullyLoad(WorldContext, FULLYLOAD_Game_PreLoadClass, TEXT(""));
 	CleanupPackagesToFullyLoad(WorldContext, FULLYLOAD_Game_PostLoadClass, TEXT(""));
 	CleanupPackagesToFullyLoad(WorldContext, FULLYLOAD_Mutator, TEXT(""));
-
+	
 
 	// Cancel any pending async map changes after flushing async loading. We flush async loading before canceling the map change
 	// to avoid completion after cancellation to not leave references to the "to be changed to" level around. Async loading is
