Actor
* Ensure if we detect that the EndPlay virtual functions did not end up calling AActor::EndPlay.
* EndPlay must always be called otherwise replicated actors and their subobjects might leak and be forever referenced by the ReplicationSystem

diff --git a/Engine/Source/Runtime/Engine/Private/Actor.cpp b/Engine/Source/Runtime/Engine/Private/Actor.cpp
--- a/Engine/Source/Runtime/Engine/Private/Actor.cpp
+++ b/Engine/Source/Runtime/Engine/Private/Actor.cpp
@@ -2759,6 +2759,7 @@ void AActor::RouteEndPlay(const EEndPlayReason::Type EndPlayReason)
 		if (ActorHasBegunPlay == EActorBeginPlayState::HasBegunPlay)
 		{
 			EndPlay(EndPlayReason);
+			ensureMsgf(ActorHasBegunPlay == EActorBeginPlayState::HasNotBegunPlay, TEXT("EndPlay on %s failed. Make sure to call Super::EndPlay() in your override function."), *GetName());
 		}
 
 		// Behaviors specific to an actor being unloaded due to a streaming level removal
@@ -2772,7 +2773,7 @@ void AActor::RouteEndPlay(const EEndPlayReason::Type EndPlayReason)
 				World->RemoveNetworkActor(this);
 #if UE_WITH_IRIS
 				EndReplication(EndPlayReason);
-#endif // UE_WITH_IRIS
+#endif
 			}
 		}
 
@@ -2795,8 +2796,9 @@ void AActor::EndPlay(const EEndPlayReason::Type EndPlayReason)
 		ActorHasBegunPlay = EActorBeginPlayState::HasNotBegunPlay;
 
 #if UE_WITH_IRIS
+		// This must be called otherwise the ReplicationSystem will keep a reference to the actor forever.
 		EndReplication(EndPlayReason);
-#endif // UE_WITH_IRIS
+#endif
 
 		// Dispatch the blueprint events
 		ReceiveEndPlay(EndPlayReason);
