Allow Blueprints to add UserDefinedStruct DataTable rows via UDataTableFunctionLibrary::AddDataTableRow

diff --git a/Engine/Source/Runtime/Engine/Classes/Engine/CompositeDataTable.h b/Engine/Source/Runtime/Engine/Classes/Engine/CompositeDataTable.h
--- a/Engine/Source/Runtime/Engine/Classes/Engine/CompositeDataTable.h
+++ b/Engine/Source/Runtime/Engine/Classes/Engine/CompositeDataTable.h
@@ -44,6 +44,7 @@ class UCompositeDataTable
 	ENGINE_API virtual void EmptyTable() override;
 	ENGINE_API virtual void RemoveRow(FName RowName) override;
 	ENGINE_API virtual void AddRow(FName RowName, const FTableRowBase& RowData) override;
+	ENGINE_API virtual void AddRow(FName RowName, const uint8* RowData, const UScriptStruct* RowType) override;
 
 #if WITH_EDITOR
 	ENGINE_API virtual void CleanBeforeStructChange() override;
diff --git a/Engine/Source/Runtime/Engine/Classes/Engine/DataTable.h b/Engine/Source/Runtime/Engine/Classes/Engine/DataTable.h
--- a/Engine/Source/Runtime/Engine/Classes/Engine/DataTable.h
+++ b/Engine/Source/Runtime/Engine/Classes/Engine/DataTable.h
@@ -323,6 +323,7 @@ public:
 
 	/** Copies RowData into table. That is: create Row if not found and copy data into the RowMap based on RowData. This is a "copy in" operation, so changing the passed in RowData after the fact does nothing. */
 	ENGINE_API virtual void AddRow(FName RowName, const FTableRowBase& RowData);
+	ENGINE_API virtual void AddRow(FName RowName, const uint8* RowData, const UScriptStruct* RowType);
 
 #if WITH_EDITOR
 	ENGINE_API virtual void CleanBeforeStructChange();
diff --git a/Engine/Source/Runtime/Engine/Classes/Kismet/DataTableFunctionLibrary.h b/Engine/Source/Runtime/Engine/Classes/Kismet/DataTableFunctionLibrary.h
--- a/Engine/Source/Runtime/Engine/Classes/Kismet/DataTableFunctionLibrary.h
+++ b/Engine/Source/Runtime/Engine/Classes/Kismet/DataTableFunctionLibrary.h
@@ -9,7 +9,6 @@
 #include "UObject/ScriptMacros.h"
 #include "Kismet/BlueprintFunctionLibrary.h"
 #include "Engine/DataTable.h"
-#include "UObject/Class.h" // for FStructUtils
 #include "Blueprint/BlueprintExceptionInfo.h"
 #include "DataTableFunctionLibrary.generated.h"
 
@@ -79,62 +78,8 @@ class UDataTableFunctionLibrary : public UBlueprintFunctionLibrary
     /** Get a Row from a DataTable given a RowName */
     UFUNCTION(BlueprintCallable, CustomThunk, Category = "DataTable", meta=(CustomStructureParam = "OutRow", BlueprintInternalUseOnly="true"))
     static ENGINE_API bool GetDataTableRowFromName(UDataTable* Table, FName RowName, FTableRowBase& OutRow);
-    
 	static ENGINE_API bool Generic_GetDataTableRowFromName(const UDataTable* Table, FName RowName, void* OutRowPtr);
-
-    /** Based on UDataTableFunctionLibrary::GetDataTableRow */
-    DECLARE_FUNCTION(execGetDataTableRowFromName)
-    {
-        P_GET_OBJECT(UDataTable, Table);
-        P_GET_PROPERTY(FNameProperty, RowName);
-        
-        Stack.StepCompiledIn<FStructProperty>(NULL);
-        void* OutRowPtr = Stack.MostRecentPropertyAddress;
-
-		P_FINISH;
-		bool bSuccess = false;
-		
-		FStructProperty* StructProp = CastField<FStructProperty>(Stack.MostRecentProperty);
-		if (!Table)
-		{
-			FBlueprintExceptionInfo ExceptionInfo(
-				EBlueprintExceptionType::AccessViolation,
-				NSLOCTEXT("GetDataTableRow", "MissingTableInput", "Failed to resolve the table input. Be sure the DataTable is valid.")
-			);
-			FBlueprintCoreDelegates::ThrowScriptException(P_THIS, Stack, ExceptionInfo);
-		}
-		else if(StructProp && OutRowPtr)
-		{
-			UScriptStruct* OutputType = StructProp->Struct;
-			const UScriptStruct* TableType  = Table->GetRowStruct();
-		
-			const bool bCompatible = (OutputType == TableType) || 
-				(OutputType->IsChildOf(TableType) && FStructUtils::TheSameLayout(OutputType, TableType));
-			if (bCompatible)
-			{
-				P_NATIVE_BEGIN;
-				bSuccess = Generic_GetDataTableRowFromName(Table, RowName, OutRowPtr);
-				P_NATIVE_END;
-			}
-			else
-			{
-				FBlueprintExceptionInfo ExceptionInfo(
-					EBlueprintExceptionType::AccessViolation,
-					NSLOCTEXT("GetDataTableRow", "IncompatibleProperty", "Incompatible output parameter; the data table's type is not the same as the return type.")
-					);
-				FBlueprintCoreDelegates::ThrowScriptException(P_THIS, Stack, ExceptionInfo);
-			}
-		}
-		else
-		{
-			FBlueprintExceptionInfo ExceptionInfo(
-				EBlueprintExceptionType::AccessViolation,
-				NSLOCTEXT("GetDataTableRow", "MissingOutputProperty", "Failed to resolve the output parameter for GetDataTableRow.")
-			);
-			FBlueprintCoreDelegates::ThrowScriptException(P_THIS, Stack, ExceptionInfo);
-		}
-		*(bool*)RESULT_PARAM = bSuccess;
-    }
+	DECLARE_FUNCTION(execGetDataTableRowFromName);
 
 #if WITH_EDITOR
 	/** 
diff --git a/Engine/Source/Runtime/Engine/Private/CompositeDataTable.cpp b/Engine/Source/Runtime/Engine/Private/CompositeDataTable.cpp
--- a/Engine/Source/Runtime/Engine/Private/CompositeDataTable.cpp
+++ b/Engine/Source/Runtime/Engine/Private/CompositeDataTable.cpp
@@ -198,6 +198,11 @@ void UCompositeDataTable::AddRow(FName RowName, const FTableRowBase& RowData)
 	// do nothing
 }
 
+void UCompositeDataTable::AddRow(FName RowName, const uint8* RowData, const UScriptStruct* RowType)
+{
+	// do nothing
+}
+
 void UCompositeDataTable::Serialize(FArchive& Ar)
 {
 	if (Ar.IsLoading())
diff --git a/Engine/Source/Runtime/Engine/Private/DataTable.cpp b/Engine/Source/Runtime/Engine/Private/DataTable.cpp
--- a/Engine/Source/Runtime/Engine/Private/DataTable.cpp
+++ b/Engine/Source/Runtime/Engine/Private/DataTable.cpp
@@ -470,6 +470,26 @@ void UDataTable::AddRow(FName RowName, const FTableRowBase& RowData)
 	AddRowInternal(RowName, NewRawRowData);
 }
 
+void UDataTable::AddRow(FName RowName, const uint8* RowData, const UScriptStruct* RowType)
+{
+	DATATABLE_CHANGE_SCOPE();
+
+	UScriptStruct& EmptyUsingStruct = GetEmptyUsingStruct();
+
+	checkf(RowType == &EmptyUsingStruct, TEXT("AddRow called with an incompatible row type! Got '%s', but expected '%s'"), *RowType->GetPathName(), *EmptyUsingStruct.GetPathName());
+
+	// We want to delete the row memory even for child classes that override remove
+	RemoveRowInternal(RowName);
+
+	uint8* NewRawRowData = (uint8*)FMemory::Malloc(EmptyUsingStruct.GetStructureSize());
+
+	EmptyUsingStruct.InitializeStruct(NewRawRowData);
+	EmptyUsingStruct.CopyScriptStruct(NewRawRowData, RowData);
+
+	// Add to map
+	AddRowInternal(RowName, NewRawRowData);
+}
+
 void UDataTable::AddRowInternal(FName RowName, uint8* RowData)
 {
 	RowMap.Add(RowName, RowData);
diff --git a/Engine/Source/Runtime/Engine/Private/DataTableFunctionLibrary.cpp b/Engine/Source/Runtime/Engine/Private/DataTableFunctionLibrary.cpp
--- a/Engine/Source/Runtime/Engine/Private/DataTableFunctionLibrary.cpp
+++ b/Engine/Source/Runtime/Engine/Private/DataTableFunctionLibrary.cpp
@@ -95,6 +95,57 @@ bool UDataTableFunctionLibrary::GetDataTableRowFromName(UDataTable* Table, FName
 	return false;
 }
 
+DEFINE_FUNCTION(UDataTableFunctionLibrary::execGetDataTableRowFromName)
+{
+    P_GET_OBJECT(UDataTable, Table);
+    P_GET_PROPERTY(FNameProperty, RowName);
+        
+    Stack.StepCompiledIn<FStructProperty>(NULL);
+    void* OutRowPtr = Stack.MostRecentPropertyAddress;
+
+	P_FINISH;
+	bool bSuccess = false;
+		
+	FStructProperty* StructProp = CastField<FStructProperty>(Stack.MostRecentProperty);
+	if (!Table)
+	{
+		FBlueprintExceptionInfo ExceptionInfo(
+			EBlueprintExceptionType::AccessViolation,
+			NSLOCTEXT("GetDataTableRow", "MissingTableInput", "Failed to resolve the table input. Be sure the DataTable is valid.")
+		);
+		FBlueprintCoreDelegates::ThrowScriptException(P_THIS, Stack, ExceptionInfo);
+	}
+	else if(StructProp && OutRowPtr)
+	{
+		UScriptStruct* OutputType = StructProp->Struct;
+		const UScriptStruct* TableType  = Table->GetRowStruct();
+		
+		if (OutputType == TableType)
+		{
+			P_NATIVE_BEGIN;
+			bSuccess = Generic_GetDataTableRowFromName(Table, RowName, OutRowPtr);
+			P_NATIVE_END;
+		}
+		else
+		{
+			FBlueprintExceptionInfo ExceptionInfo(
+				EBlueprintExceptionType::AccessViolation,
+				FText::Format(NSLOCTEXT("GetDataTableRow", "IncompatibleProperty", "Incompatible output parameter; the data table's type ({0}) is not the same as the return type ({1})."), FText::AsCultureInvariant(GetPathNameSafe(TableType)), FText::AsCultureInvariant(OutputType->GetPathName()))
+				);
+			FBlueprintCoreDelegates::ThrowScriptException(P_THIS, Stack, ExceptionInfo);
+		}
+	}
+	else
+	{
+		FBlueprintExceptionInfo ExceptionInfo(
+			EBlueprintExceptionType::AccessViolation,
+			NSLOCTEXT("GetDataTableRow", "MissingOutputProperty", "Failed to resolve the output parameter for GetDataTableRow.")
+		);
+		FBlueprintCoreDelegates::ThrowScriptException(P_THIS, Stack, ExceptionInfo);
+	}
+	*(bool*)RESULT_PARAM = bSuccess;
+}
+
 void UDataTableFunctionLibrary::GetDataTableRowNames(const UDataTable* Table, TArray<FName>& OutRowNames)
 {
 	if (Table)
@@ -326,6 +377,7 @@ void UDataTableFunctionLibrary::AddDataTableRow(UDataTable* const DataTable, con
 		return;
 	}
 
+	DataTable->Modify();
 	DataTable->AddRow(RowName, RowData);
 }
 
@@ -338,7 +390,7 @@ DEFINE_FUNCTION(UDataTableFunctionLibrary::execAddDataTableRow)
 	Stack.MostRecentPropertyContainer = nullptr;
 	Stack.StepCompiledIn<FStructProperty>(nullptr);
 
-	const FTableRowBase* const RowData = reinterpret_cast<FTableRowBase*>(Stack.MostRecentPropertyAddress);
+	const uint8* const RowData = Stack.MostRecentPropertyAddress;
 	const FStructProperty* const StructProp = CastField<FStructProperty>(Stack.MostRecentProperty);
 
 	P_FINISH;
@@ -357,19 +409,18 @@ DEFINE_FUNCTION(UDataTableFunctionLibrary::execAddDataTableRow)
 		const UScriptStruct* const TableType = DataTable->GetRowStruct();
 
 		// If the row type is compatible with the table type ...
-		const bool bIsTableRow = RowType->IsChildOf(FTableRowBase::StaticStruct());
-		const bool bMatchesTableType = RowType == TableType;
-		if (bIsTableRow && (bMatchesTableType || (RowType->IsChildOf(TableType) && FStructUtils::TheSameLayout(RowType, TableType))))
+		if (RowType == TableType)
 		{
 			P_NATIVE_BEGIN;
-			AddDataTableRow(DataTable, RowName, *RowData);
+			DataTable->Modify();
+			DataTable->AddRow(RowName, RowData, RowType);
 			P_NATIVE_END;
 		}
 		else
 		{
 			FBlueprintExceptionInfo ExceptionInfo(
 				EBlueprintExceptionType::AccessViolation,
-				NSLOCTEXT("AddDataTableRow", "IncompatibleProperty", "The data table type is incompatible with the RowData parameter.")
+				FText::Format(NSLOCTEXT("AddDataTableRow", "IncompatibleProperty", "Incompatible RowData parameter; the data table's type ({0}) is not the same as the RowData type ({1})."), FText::AsCultureInvariant(GetPathNameSafe(TableType)), FText::AsCultureInvariant(RowType->GetPathName()))
 			);
 			FBlueprintCoreDelegates::ThrowScriptException(P_THIS, Stack, ExceptionInfo);
 		}
