Iris
* Clarified a few logs and ensures

#rnx

diff --git a/Engine/Source/Runtime/Engine/Private/Net/Iris/ReplicationSystem/ActorReplicationBridge.cpp b/Engine/Source/Runtime/Engine/Private/Net/Iris/ReplicationSystem/ActorReplicationBridge.cpp
--- a/Engine/Source/Runtime/Engine/Private/Net/Iris/ReplicationSystem/ActorReplicationBridge.cpp
+++ b/Engine/Source/Runtime/Engine/Private/Net/Iris/ReplicationSystem/ActorReplicationBridge.cpp
@@ -422,8 +422,8 @@ void UActorReplicationBridge::EndReplication(AActor* Actor, EEndPlayReason::Type
 	FNetRefHandle RefHandle = GetReplicatedRefHandle(Actor, EGetRefHandleFlags::EvenIfGarbage);
 	if (RefHandle.IsValid())
 	{
-		UE_LOG(LogIrisBridge, Verbose, TEXT("EndReplication for %s %s. Reason %s "), *GetNameSafe(Actor), *RefHandle.ToString(), *UEnum::GetValueAsString(TEXT("Engine.EEndPlayReason"), EndPlayReason));
-		ensureMsgf(IsValid(Actor), TEXT("Calling EndReplication for Invalid Object for %s %s."), *GetNameSafe(Actor), *RefHandle.ToString());
+		UE_LOG(LogIrisBridge, Verbose, TEXT("EndReplication for %s. Reason %s "), *PrintObjectFromNetRefHandle(RefHandle), *UEnum::GetValueAsString(TEXT("Engine.EEndPlayReason"), EndPlayReason));
+		ensureMsgf(IsValid(Actor), TEXT("Calling EndReplication for Invalid Object: %s."), *PrintObjectFromNetRefHandle(RefHandle));
 	
 		EEndReplicationFlags Flags = EEndReplicationFlags::None;
 		const bool bShouldDestroyObject = EndPlayReason == EEndPlayReason::Destroyed;		
diff --git a/Engine/Source/Runtime/Engine/Private/Net/Iris/ReplicationSystem/ReplicationSystemUtil.cpp b/Engine/Source/Runtime/Engine/Private/Net/Iris/ReplicationSystem/ReplicationSystemUtil.cpp
--- a/Engine/Source/Runtime/Engine/Private/Net/Iris/ReplicationSystem/ReplicationSystemUtil.cpp
+++ b/Engine/Source/Runtime/Engine/Private/Net/Iris/ReplicationSystem/ReplicationSystemUtil.cpp
@@ -318,7 +318,7 @@ void FReplicationSystemUtil::AddDependentActor(const AActor* Parent, AActor* Chi
 {
 	// Can only add dependent actors on already replicating actors
 	FNetHandle ParentHandle = GetNetHandle(Parent);
-	if (!ensureMsgf(ParentHandle.IsValid(), TEXT("FReplicationSystemUtil::AddDependentActor Parent %s must be replicated"), *GetPathNameSafe(Parent)))
+	if (!ensureMsgf(ParentHandle.IsValid(), TEXT("FReplicationSystemUtil::AddDependentActor Parent %s is not replicated. Cannot attach child %s as dependent"), *GetPathNameSafe(Parent), *GetNameSafe(Child)))
 	{
 		return;
 	}
@@ -340,7 +340,7 @@ void FReplicationSystemUtil::AddDependentActor(const AActor* Parent, AActor* Chi
 							const FActorBeginReplicationParams BeginReplicationParams;
 							ChildRefHandle = Bridge->BeginReplication(Child, BeginReplicationParams);
 						}
-						if (ensureMsgf(ChildRefHandle.IsValid(), TEXT("FReplicationSystemUtil::AddDependentActor Child %s must be replicated"), *GetPathNameSafe(Child)))
+						if (ensureMsgf(ChildRefHandle.IsValid(), TEXT("FReplicationSystemUtil::AddDependentActor Child %s is not replicated"), *GetPathNameSafe(Child)))
 						{
 							Bridge->AddDependentObject(ParentRefHandle, ChildRefHandle, SchedulingHint);
 						}
diff --git a/Engine/Source/Runtime/Experimental/Iris/Core/Private/Iris/ReplicationSystem/NetRefHandleManager.cpp b/Engine/Source/Runtime/Experimental/Iris/Core/Private/Iris/ReplicationSystem/NetRefHandleManager.cpp
--- a/Engine/Source/Runtime/Experimental/Iris/Core/Private/Iris/ReplicationSystem/NetRefHandleManager.cpp
+++ b/Engine/Source/Runtime/Experimental/Iris/Core/Private/Iris/ReplicationSystem/NetRefHandleManager.cpp
@@ -784,7 +784,8 @@ bool FNetRefHandleManager::AddDependentObject(FNetRefHandle ParentRefHandle, FNe
 	{
 		// If this gets out of sync something is messed up
 		check(bDependentHadParentAlready == bDependentIsAlreadyDependant);
-		ensureMsgf(false, TEXT("FNetRefHandleManager::AddDependentObject %s already is a child of %s parent %s"), *PrintObjectFromNetRefHandle(DependentObjectRefHandle), *PrintObjectFromIndex(DependentObjectData.SubObjectRootIndex), *PrintObjectFromNetRefHandle(ParentRefHandle));
+		ensureMsgf(false, TEXT("FNetRefHandleManager::AddDependentObject: %s already is a child of parent: %s. Cannot be added to: %s"), 
+			*PrintObjectFromNetRefHandle(DependentObjectRefHandle), *PrintObjectFromIndex(DependentObjectData.SubObjectRootIndex), *PrintObjectFromNetRefHandle(ParentRefHandle));
 	}
 
 	return true;
