PR #11601: Fix compile error if STREAMING_LOG_VIEWCHANGE is set to 1

#rnx

diff --git a/Engine/Source/Runtime/Engine/Private/ContentStreaming.cpp b/Engine/Source/Runtime/Engine/Private/ContentStreaming.cpp
--- a/Engine/Source/Runtime/Engine/Private/ContentStreaming.cpp
+++ b/Engine/Source/Runtime/Engine/Private/ContentStreaming.cpp
@@ -573,11 +573,11 @@ void IStreamingManager::SetupViewInfos( float DeltaTime )
 		// Check if we're adding any new locations.
 		for ( int32 ViewIndex = 0; ViewIndex < CurrentViewInfos.Num(); ++ViewIndex )
 		{
-			FStreamingViewInfo& ViewInfo = CurrentViewInfos( ViewIndex );
+			FStreamingViewInfo& ViewInfo = CurrentViewInfos[ViewIndex];
 			bool bFound = false;
 			for ( int32 PrevView=0; PrevView < GPrevViewLocations.Num(); ++PrevView )
 			{
-				if ( (ViewInfo.ViewOrigin - GPrevViewLocations(PrevView).ViewOrigin).SizeSquared() < 10000.0f )
+				if ( (ViewInfo.ViewOrigin - GPrevViewLocations[PrevView].ViewOrigin).SizeSquared() < 10000.0f )
 				{
 					bFound = true;
 					break;
@@ -595,8 +595,8 @@ void IStreamingManager::SetupViewInfos( float DeltaTime )
 			bool bFound = false;
 			for ( int32 ViewIndex = 0; ViewIndex < CurrentViewInfos.Num(); ++ViewIndex )
 			{
-				FStreamingViewInfo& ViewInfo = CurrentViewInfos( ViewIndex );
-				if ( (ViewInfo.ViewOrigin - GPrevViewLocations(PrevView).ViewOrigin).SizeSquared() < 10000.0f )
+				FStreamingViewInfo& ViewInfo = CurrentViewInfos[ViewIndex];
+				if ( (ViewInfo.ViewOrigin - GPrevViewLocations[PrevView].ViewOrigin).SizeSquared() < 10000.0f )
 				{
 					bFound = true;
 					break;
@@ -604,7 +604,7 @@ void IStreamingManager::SetupViewInfos( float DeltaTime )
 			}
 			if ( !bFound )
 			{
-				FStreamingViewInfo& PrevViewInfo = GPrevViewLocations(PrevView);
+				FStreamingViewInfo& PrevViewInfo = GPrevViewLocations[PrevView];
 				UE_LOG(LogContentStreaming, Log, TEXT("Removing location: X=%.1f, Y=%.1f, Z=%.1f (override=%d, boost=%.1f)"), PrevViewInfo.ViewOrigin.X, PrevViewInfo.ViewOrigin.Y, PrevViewInfo.ViewOrigin.Z, PrevViewInfo.bOverrideLocation, PrevViewInfo.BoostFactor );
 			}
 		}
@@ -613,7 +613,7 @@ void IStreamingManager::SetupViewInfos( float DeltaTime )
 		GPrevViewLocations.Empty(CurrentViewInfos.Num());
 		for ( int32 ViewIndex = 0; ViewIndex < CurrentViewInfos.Num(); ++ViewIndex )
 		{
-			FStreamingViewInfo& ViewInfo = CurrentViewInfos( ViewIndex );
+			FStreamingViewInfo& ViewInfo = CurrentViewInfos[ ViewIndex ];
 			GPrevViewLocations.Add( ViewInfo );
 		}
 	}
