Remove localisation in layer prototype to avoid warnings in localisation builds.

diff --git a/Engine/Source/Runtime/Engine/Private/Materials/MaterialExpressions.cpp b/Engine/Source/Runtime/Engine/Private/Materials/MaterialExpressions.cpp
--- a/Engine/Source/Runtime/Engine/Private/Materials/MaterialExpressions.cpp
+++ b/Engine/Source/Runtime/Engine/Private/Materials/MaterialExpressions.cpp
@@ -29082,11 +29082,11 @@ FName UMaterialExpressionSubstrateSetAttributes::GetInputName(int32 InputIndex)
 {
 	if (InputIndex == ESubstrateAttributeIndex::MSA_NonSubstrateAttributes)
 	{
-		return *NSLOCTEXT("SetSubstrateAttributes", "InputName", NON_SUBSTRATE_ATTRIBUTES_TEXT).ToString();
+		return FName(NON_SUBSTRATE_ATTRIBUTES_TEXT);
 	}
 	if (InputIndex == ESubstrateAttributeIndex::MSA_FrontMaterial)
 	{
-		return *NSLOCTEXT("SetSubstrateAttributes", "InputName", FRONT_MATERIAL_ATTRIBUTES_TEXT).ToString();
+		return FName(FRONT_MATERIAL_ATTRIBUTES_TEXT);
 	}
 	return NAME_None;
 }
