Actor
* Added ensure to detect if the ActorComponent::EndPlay did not called due to an override behavior

diff --git a/Engine/Source/Runtime/Engine/Private/Actor.cpp b/Engine/Source/Runtime/Engine/Private/Actor.cpp
--- a/Engine/Source/Runtime/Engine/Private/Actor.cpp
+++ b/Engine/Source/Runtime/Engine/Private/Actor.cpp
@@ -2812,6 +2812,7 @@ void AActor::EndPlay(const EEndPlayReason::Type EndPlayReason)
 			if (Component->HasBegunPlay())
 			{
 				Component->EndPlay(EndPlayReason);
+				ensureMsgf(Component->HasBegunPlay() == false, TEXT("EndPlay on %s failed. Make sure to call Super::EndPlay() in the override function in class %s."), *Component->GetName(), *Component->GetClass()->GetName());
 			}
 		}
 	}
