Fix transient texture override on materials.
Logic only took into account the parameter description, which is similar to the logic for transient override of scalar parameters etc.
But texture override is expected to work on non-parameters. For that case a texture override on one non-parameter texture would override all non-parameter textures.
To fix this both the parameter info and the texture slot are taken into account in the map that holds the texture overrides.
#tests Use texture mesh painting tool and observe if the appearance is correct during paint operations when a painting texture is transiently swapped into the material

diff --git a/Engine/Source/Runtime/Engine/Private/Materials/Material.cpp b/Engine/Source/Runtime/Engine/Private/Materials/Material.cpp
--- a/Engine/Source/Runtime/Engine/Private/Materials/Material.cpp
+++ b/Engine/Source/Runtime/Engine/Private/Materials/Material.cpp
@@ -1316,7 +1316,7 @@ void UMaterial::OverrideTexture(const UTexture* InTextureToOverride, UTexture* O
 					if (Texture != NULL && Texture == InTextureToOverride)
 					{
 						// Override this texture!
-						Resource->TransientOverrides.SetTextureOverride((EMaterialTextureParameterType)TypeIndex, Parameter.ParameterInfo, OverrideTexture);
+						Resource->TransientOverrides.SetTextureOverride((EMaterialTextureParameterType)TypeIndex, Parameter, OverrideTexture);
 						bShouldRecacheMaterialExpressions = true;
 					}
 				}
diff --git a/Engine/Source/Runtime/Engine/Private/Materials/MaterialInstance.cpp b/Engine/Source/Runtime/Engine/Private/Materials/MaterialInstance.cpp
--- a/Engine/Source/Runtime/Engine/Private/Materials/MaterialInstance.cpp
+++ b/Engine/Source/Runtime/Engine/Private/Materials/MaterialInstance.cpp
@@ -1467,7 +1467,7 @@ void UMaterialInstance::OverrideTexture(const UTexture* InTextureToOverride, UTe
 				if (Texture != NULL && Texture == InTextureToOverride)
 				{
 					// Override this texture!
-					SourceMaterialResource->TransientOverrides.SetTextureOverride((EMaterialTextureParameterType)TypeIndex, Parameter.ParameterInfo, OverrideTexture);
+					SourceMaterialResource->TransientOverrides.SetTextureOverride((EMaterialTextureParameterType)TypeIndex, Parameter, OverrideTexture);
 					bShouldRecacheMaterialExpressions = true;
 				}
 			}
diff --git a/Engine/Source/Runtime/Engine/Private/Materials/MaterialUniformExpressions.cpp b/Engine/Source/Runtime/Engine/Private/Materials/MaterialUniformExpressions.cpp
--- a/Engine/Source/Runtime/Engine/Private/Materials/MaterialUniformExpressions.cpp
+++ b/Engine/Source/Runtime/Engine/Private/Materials/MaterialUniformExpressions.cpp
@@ -192,47 +192,51 @@ bool FUniformParameterOverrides::GetNumericOverride(EMaterialParameterType Type,
 	return false;
 }
 
-void FUniformParameterOverrides::SetTextureOverride(EMaterialTextureParameterType Type, const FHashedMaterialParameterInfo& ParameterInfo, UTexture* Texture)
+void FUniformParameterOverrides::SetTextureOverride(EMaterialTextureParameterType Type, const FMaterialTextureParameterInfo& ParameterInfo, UTexture* Texture)
 {
 	check(IsInGameThread());
 	const uint32 TypeIndex = (uint32)Type;
+	const FTextureParameterKey Key{ ParameterInfo.ParameterInfo, ParameterInfo.TextureIndex };
+
 	if (Texture)
 	{
-		GameThreadTextureOverides[TypeIndex].FindOrAdd(ParameterInfo) = Texture;
+		GameThreadTextureOverides[TypeIndex].FindOrAdd(Key) = Texture;
 	}
 	else
 	{
-		GameThreadTextureOverides[TypeIndex].Remove(ParameterInfo);
+		GameThreadTextureOverides[TypeIndex].Remove(Key);
 	}
 
 	FUniformParameterOverrides* Self = this;
 	ENQUEUE_RENDER_COMMAND(SetTextureOverrideCommand)(
-		[Self, TypeIndex, ParameterInfo, Texture](FRHICommandListImmediate& RHICmdList)
+		[Self, TypeIndex, Key, Texture](FRHICommandListImmediate& RHICmdList)
 	{
 		if (Texture)
 		{
-			Self->RenderThreadTextureOverrides[TypeIndex].FindOrAdd(ParameterInfo) = Texture;
+			Self->RenderThreadTextureOverrides[TypeIndex].FindOrAdd(Key) = Texture;
 		}
 		else
 		{
-			Self->RenderThreadTextureOverrides[TypeIndex].Remove(ParameterInfo);
+			Self->RenderThreadTextureOverrides[TypeIndex].Remove(Key);
 		}
 	});
 }
 
-UTexture* FUniformParameterOverrides::GetTextureOverride_GameThread(EMaterialTextureParameterType Type, const FHashedMaterialParameterInfo& ParameterInfo) const
+UTexture* FUniformParameterOverrides::GetTextureOverride_GameThread(EMaterialTextureParameterType Type, const FMaterialTextureParameterInfo& ParameterInfo) const
 {
 	check(IsInGameThread());
 	const uint32 TypeIndex = (uint32)Type;
-	UTexture* const* Result = GameThreadTextureOverides[TypeIndex].Find(ParameterInfo);
+	const FTextureParameterKey Key{ ParameterInfo.ParameterInfo, ParameterInfo.TextureIndex };
+	UTexture* const* Result = GameThreadTextureOverides[TypeIndex].Find(Key);
 	return Result ? *Result : nullptr;
 }
 
-UTexture* FUniformParameterOverrides::GetTextureOverride_RenderThread(EMaterialTextureParameterType Type, const FHashedMaterialParameterInfo& ParameterInfo) const
+UTexture* FUniformParameterOverrides::GetTextureOverride_RenderThread(EMaterialTextureParameterType Type, const FMaterialTextureParameterInfo& ParameterInfo) const
 {
 	check(IsInParallelRenderingThread());
 	const uint32 TypeIndex = (uint32)Type;
-	UTexture* const* Result = RenderThreadTextureOverrides[TypeIndex].Find(ParameterInfo);
+	const FTextureParameterKey Key{ ParameterInfo.ParameterInfo, ParameterInfo.TextureIndex };
+	UTexture* const* Result = RenderThreadTextureOverrides[TypeIndex].Find(Key);
 	return Result ? *Result : nullptr;
 }
 
@@ -725,7 +729,7 @@ void FUniformExpressionSet::GetGameThreadTextureValue(EMaterialTextureParameterT
 #if WITH_EDITOR
 	if (bAllowOverride)
 	{
-		UTexture* OverrideTexture = Material.TransientOverrides.GetTextureOverride_GameThread(Type, Parameter.ParameterInfo);
+		UTexture* OverrideTexture = Material.TransientOverrides.GetTextureOverride_GameThread(Type, Parameter);
 		if (OverrideTexture)
 		{
 			OutValue = OverrideTexture;
@@ -742,7 +746,7 @@ void FUniformExpressionSet::GetTextureValue(EMaterialTextureParameterType Type,
 	const FMaterialTextureParameterInfo& Parameter = GetTextureParameter(Type, Index);
 #if WITH_EDITOR
 	{
-		UTexture* OverrideTexture = Material.TransientOverrides.GetTextureOverride_RenderThread(Type, Parameter.ParameterInfo);
+		UTexture* OverrideTexture = Material.TransientOverrides.GetTextureOverride_RenderThread(Type, Parameter);
 		if (OverrideTexture)
 		{
 			OutValue = OverrideTexture;
diff --git a/Engine/Source/Runtime/Engine/Public/MaterialShared.h b/Engine/Source/Runtime/Engine/Public/MaterialShared.h
--- a/Engine/Source/Runtime/Engine/Public/MaterialShared.h
+++ b/Engine/Source/Runtime/Engine/Public/MaterialShared.h
@@ -581,9 +581,9 @@ public:
 	void SetNumericOverride(EMaterialParameterType Type, const FHashedMaterialParameterInfo& ParameterInfo, const UE::Shader::FValue& Value, bool bOverride);
 	bool GetNumericOverride(EMaterialParameterType Type, const FHashedMaterialParameterInfo& ParameterInfo, UE::Shader::FValue& OutValue) const;
 
-	void SetTextureOverride(EMaterialTextureParameterType Type, const FHashedMaterialParameterInfo& ParameterInfo, UTexture* Texture);
-	UTexture* GetTextureOverride_GameThread(EMaterialTextureParameterType Type, const FHashedMaterialParameterInfo& ParameterInfo) const;
-	UTexture* GetTextureOverride_RenderThread(EMaterialTextureParameterType Type, const FHashedMaterialParameterInfo& ParameterInfo) const;
+	void SetTextureOverride(EMaterialTextureParameterType Type, const FMaterialTextureParameterInfo& ParameterInfo, UTexture* Texture);
+	UTexture* GetTextureOverride_GameThread(EMaterialTextureParameterType Type, const FMaterialTextureParameterInfo& ParameterInfo) const;
+	UTexture* GetTextureOverride_RenderThread(EMaterialTextureParameterType Type, const FMaterialTextureParameterInfo& ParameterInfo) const;
 
 private:
 	struct FNumericParameterKey
@@ -606,8 +606,28 @@ private:
 	};
 
 	TMap<FNumericParameterKey, UE::Shader::FValue> NumericOverrides;
-	TMap<FHashedMaterialParameterInfo, UTexture*> GameThreadTextureOverides[NumMaterialTextureParameterTypes];
-	TMap<FHashedMaterialParameterInfo, UTexture*> RenderThreadTextureOverrides[NumMaterialTextureParameterTypes];
+
+	struct FTextureParameterKey
+	{
+		FHashedMaterialParameterInfo ParameterInfo;
+		int32 TextureIndex;
+
+		friend inline bool operator==(const FTextureParameterKey& Lhs, const FTextureParameterKey& Rhs)
+		{
+			return Lhs.ParameterInfo == Rhs.ParameterInfo && Lhs.TextureIndex == Rhs.TextureIndex;
+		}
+		friend inline bool operator!=(const FTextureParameterKey& Lhs, const FTextureParameterKey& Rhs)
+		{
+			return !operator==(Lhs, Rhs);
+		}
+		friend inline uint32 GetTypeHash(const FTextureParameterKey& Value)
+		{
+			return HashCombine(GetTypeHash(Value.ParameterInfo), GetTypeHash(Value.TextureIndex));
+		}
+	};
+
+	TMap<FTextureParameterKey, UTexture*> GameThreadTextureOverides[NumMaterialTextureParameterTypes];
+	TMap<FTextureParameterKey, UTexture*> RenderThreadTextureOverrides[NumMaterialTextureParameterTypes];
 };
 
 /** Stores all uniform expressions for a material generated from a material translation. */
