Prevent ensure when calling Iris Dependent API when not running Iris

#rnx

diff --git a/Engine/Source/Runtime/Engine/Private/Net/Iris/ReplicationSystem/ReplicationSystemUtil.cpp b/Engine/Source/Runtime/Engine/Private/Net/Iris/ReplicationSystem/ReplicationSystemUtil.cpp
--- a/Engine/Source/Runtime/Engine/Private/Net/Iris/ReplicationSystem/ReplicationSystemUtil.cpp
+++ b/Engine/Source/Runtime/Engine/Private/Net/Iris/ReplicationSystem/ReplicationSystemUtil.cpp
@@ -321,6 +321,14 @@ void FReplicationSystemUtil::EndReplicationForActorComponentSubObject(UActorComp
 
 void FReplicationSystemUtil::AddDependentActor(const AActor* Parent, AActor* Child, EDependentObjectSchedulingHint SchedulingHint)
 {
+	check(Parent);
+
+	const UWorld* World = Parent->GetWorld();
+	if (!World || !World->GetNetDriver() || !World->GetNetDriver()->IsUsingIrisReplication())
+	{
+		return;
+	}
+
 	// Can only add dependent actors on already replicating actors
 	FNetHandle ParentHandle = GetNetHandle(Parent);
 	if (!ensureMsgf(ParentHandle.IsValid(), TEXT("FReplicationSystemUtil::AddDependentActor Parent %s is not replicated. Cannot attach child %s as dependent"), *GetPathNameSafe(Parent), *GetNameSafe(Child)))
@@ -328,32 +336,29 @@ void FReplicationSystemUtil::AddDependentActor(const AActor* Parent, AActor* Chi
 		return;
 	}
 
-	if (const UWorld* World = Parent->GetWorld())
+	ReplicationSystemUtil::ForEachReplicationSystem(GEngine, World, [ParentHandle, Child, SchedulingHint](UReplicationSystem* ReplicationSystem)
 	{
-		ReplicationSystemUtil::ForEachReplicationSystem(GEngine, World, [ParentHandle, Child, SchedulingHint](UReplicationSystem* ReplicationSystem)
+		if (ReplicationSystem->IsServer())
 		{
-			if (ReplicationSystem->IsServer())
+			if (UActorReplicationBridge* Bridge = ReplicationSystem->GetReplicationBridgeAs<UActorReplicationBridge>())
 			{
-				if (UActorReplicationBridge* Bridge = ReplicationSystem->GetReplicationBridgeAs<UActorReplicationBridge>())
+				const FNetRefHandle ParentRefHandle = Bridge->GetReplicatedRefHandle(ParentHandle);
+				if (ParentRefHandle.IsValid())
 				{
-					const FNetRefHandle ParentRefHandle = Bridge->GetReplicatedRefHandle(ParentHandle);
-					if (ParentRefHandle.IsValid())
+					FNetRefHandle ChildRefHandle = Bridge->GetReplicatedRefHandle(Child);
+					if (!ChildRefHandle.IsValid())
 					{
-						FNetRefHandle ChildRefHandle = Bridge->GetReplicatedRefHandle(Child);
-						if (!ChildRefHandle.IsValid())
-						{
-							const FActorBeginReplicationParams BeginReplicationParams;
-							ChildRefHandle = Bridge->BeginReplication(Child, BeginReplicationParams);
-						}
-						if (ensureMsgf(ChildRefHandle.IsValid(), TEXT("FReplicationSystemUtil::AddDependentActor Child %s is not replicated"), *GetPathNameSafe(Child)))
-						{
-							Bridge->AddDependentObject(ParentRefHandle, ChildRefHandle, SchedulingHint);
-						}
+						const FActorBeginReplicationParams BeginReplicationParams;
+						ChildRefHandle = Bridge->BeginReplication(Child, BeginReplicationParams);
+					}
+					if (ensureMsgf(ChildRefHandle.IsValid(), TEXT("FReplicationSystemUtil::AddDependentActor Child %s is not replicated"), *GetPathNameSafe(Child)))
+					{
+						Bridge->AddDependentObject(ParentRefHandle, ChildRefHandle, SchedulingHint);
 					}
 				}
 			}
-		});
-	}
+		}
+	});
 }
 
 void FReplicationSystemUtil::AddDependentActor(const AActor* Parent, AActor* Child)
@@ -363,6 +368,13 @@ void FReplicationSystemUtil::AddDependentActor(const AActor* Parent, AActor* Chi
 
 void FReplicationSystemUtil::RemoveDependentActor(const AActor* Parent, AActor* Child)
 {
+	check(Parent);
+	const UWorld* World = Parent->GetWorld();
+	if (!World || !World->GetNetDriver() || !World->GetNetDriver()->IsUsingIrisReplication())
+	{
+		return;
+	}
+
 	const FNetHandle ParentHandle = GetNetHandle(Parent);
 	const FNetHandle ChildHandle = GetNetHandle(Child);
 	if (!ParentHandle.IsValid() || !ChildHandle.IsValid())
@@ -370,24 +382,21 @@ void FReplicationSystemUtil::RemoveDependentActor(const AActor* Parent, AActor*
 		return;
 	}
 
-	if (const UWorld* World = Parent->GetWorld())
+	ReplicationSystemUtil::ForEachReplicationSystem(GEngine, World, [ParentHandle, ChildHandle](UReplicationSystem* ReplicationSystem)
 	{
-		ReplicationSystemUtil::ForEachReplicationSystem(GEngine, World, [ParentHandle, ChildHandle](UReplicationSystem* ReplicationSystem)
+		if (ReplicationSystem->IsServer())
 		{
-			if (ReplicationSystem->IsServer())
+			if (UObjectReplicationBridge* Bridge = ReplicationSystem->GetReplicationBridgeAs<UObjectReplicationBridge>())
 			{
-				if (UObjectReplicationBridge* Bridge = ReplicationSystem->GetReplicationBridgeAs<UObjectReplicationBridge>())
+				const FNetRefHandle ParentRefHandle = Bridge->GetReplicatedRefHandle(ParentHandle);
+				if (ParentRefHandle.IsValid())
 				{
-					const FNetRefHandle ParentRefHandle = Bridge->GetReplicatedRefHandle(ParentHandle);
-					if (ParentRefHandle.IsValid())
-					{
-						const FNetRefHandle ChildRefHandle = Bridge->GetReplicatedRefHandle(ChildHandle);
-						Bridge->RemoveDependentObject(ParentRefHandle, ChildRefHandle);
-					}
+					const FNetRefHandle ChildRefHandle = Bridge->GetReplicatedRefHandle(ChildHandle);
+					Bridge->RemoveDependentObject(ParentRefHandle, ChildRefHandle);
 				}
 			}
-		});
-	}
+		}
+	});
 }
 
 void FReplicationSystemUtil::SetActorComponentNetCondition(const UActorComponent* ActorComponent, ELifetimeCondition NetCondition)
