[World Partition] Fixed console command wp.editor.dumpactordescs to recursively output actor descriptors from level containers.

diff --git a/Engine/Source/Runtime/Engine/Private/WorldPartition/WorldPartition.cpp b/Engine/Source/Runtime/Engine/Private/WorldPartition/WorldPartition.cpp
--- a/Engine/Source/Runtime/Engine/Private/WorldPartition/WorldPartition.cpp
+++ b/Engine/Source/Runtime/Engine/Private/WorldPartition/WorldPartition.cpp
@@ -173,7 +173,7 @@ FString GetActorDescDumpString(const FWorldPartitionActorDescInstance* ActorDesc
 	check(ActorDescInstance);
 	return FString::Printf(
 		TEXT("%s DataLayerNames:%s") LINE_TERMINATOR, 
-		*ActorDescInstance->ToString(FWorldPartitionActorDesc::EToStringMode::Full),
+		*ActorDescInstance->ToString(FWorldPartitionActorDesc::EToStringMode::Verbose),
 		*GetDataLayerString(ActorDescInstance->GetDataLayerInstanceNames().ToArray())
 	);
 }
@@ -1937,19 +1937,56 @@ void UWorldPartition::DumpActorDescs(const FString& Path)
 {
 	if (FArchive* LogFile = IFileManager::Get().CreateFileWriter(*Path))
 	{
-		TArray<const FWorldPartitionActorDescInstance*> ActorDescInstances;
-		TMap<FName, FString> DataLayersDumpString = GetDataLayersDumpString(this);
-		for (FActorDescContainerInstanceCollection::TConstIterator<> Iterator(this); Iterator; ++Iterator)
+		TArray<TPair<const FWorldPartitionActorDescInstance*, uint32>> ActorDescInstances;
+		TFunction<void(const UActorDescContainerInstance*, uint32)> DumpContainerInstanceActors = [&ActorDescInstances, &DumpContainerInstanceActors](const UActorDescContainerInstance* ContainerInstance, uint32 Depth)
 		{
-			ActorDescInstances.Add(*Iterator);
-		}
-		ActorDescInstances.Sort([](const FWorldPartitionActorDescInstance& A, const FWorldPartitionActorDescInstance& B)
+			TArray<const FWorldPartitionActorDescInstance*> SortedActorDescInstances;
+			for (UActorDescContainerInstance::TConstIterator<> It(ContainerInstance); It; ++It)
+			{
+				SortedActorDescInstances.Add(*It);
+			}
+
+			SortedActorDescInstances.Sort([](const FWorldPartitionActorDescInstance& A, const FWorldPartitionActorDescInstance& B)
+			{
+				return A.GetGuid() < B.GetGuid();
+			});
+
+			for (const FWorldPartitionActorDescInstance* ActorDescInstance : SortedActorDescInstances)
+			{
+				ActorDescInstances.Emplace(ActorDescInstance, Depth);
+
+				if (ActorDescInstance->IsChildContainerInstance())
+				{
+					FWorldPartitionActorDesc::FContainerInstance ContainerInstanceDesc;
+					if (ActorDescInstance->GetChildContainerInstance(ContainerInstanceDesc))
+					{
+						UE_LOG(LogWorldPartition, Log, TEXT("%s%s=%d"), FCString::Tab(Depth), *ContainerInstanceDesc.ContainerInstance->GetContainerPackage().ToString(), ContainerInstanceDesc.ContainerInstance->GetActorsByGuid().Num());
+						DumpContainerInstanceActors(ContainerInstanceDesc.ContainerInstance, Depth + 1);
+					}
+				}
+			}
+		};
+
+		TArray<const UActorDescContainerInstance*> SortedActorDescContainerInstances;
+		ForEachActorDescContainerInstance([&SortedActorDescContainerInstances](const UActorDescContainerInstance* InActorDescContainerInstance)
+		{
+			SortedActorDescContainerInstances.Add(InActorDescContainerInstance);
+		});
+
+		SortedActorDescContainerInstances.Sort([](const UActorDescContainerInstance& A, const UActorDescContainerInstance& B)
 		{
-			return A.GetGuid() < B.GetGuid();
+			return A.GetContentBundleGuid() < B.GetContentBundleGuid();
 		});
-		for (const FWorldPartitionActorDescInstance* Iterator : ActorDescInstances)
+
+		for (const UActorDescContainerInstance* ActorDescContainerInstanceIt : SortedActorDescContainerInstances)
+		{
+			DumpContainerInstanceActors(ActorDescContainerInstanceIt, 0);
+		}
+
+		TMap<FName, FString> DataLayersDumpString = GetDataLayersDumpString(this);
+		for (const TPair<const FWorldPartitionActorDescInstance*, uint32>& Iterator : ActorDescInstances)
 		{
-			FString LineEntry = GetActorDescDumpString(Iterator, DataLayersDumpString);
+			FString LineEntry = FString::Printf(TEXT("%s%s"), FCString::Tab(Iterator.Value), *GetActorDescDumpString(Iterator.Key, DataLayersDumpString));
 			LogFile->Serialize(TCHAR_TO_ANSI(*LineEntry), LineEntry.Len());
 		}
 
diff --git a/Engine/Source/Runtime/Engine/Private/WorldPartition/WorldPartitionActorDesc.cpp b/Engine/Source/Runtime/Engine/Private/WorldPartition/WorldPartitionActorDesc.cpp
--- a/Engine/Source/Runtime/Engine/Private/WorldPartition/WorldPartitionActorDesc.cpp
+++ b/Engine/Source/Runtime/Engine/Private/WorldPartition/WorldPartitionActorDesc.cpp
@@ -468,6 +468,13 @@ FString FWorldPartitionActorDesc::ToString(EToStringMode Mode) const
 
 		AppendToString(TEXT("NativeClass"), NativeClass);
 		AppendFromString(TEXT("Name"), GetActorNameString());
+
+		if (Mode >= EToStringMode::Verbose)
+		{
+			AppendToString(TEXT("ActorPackage"), ActorPackage);
+			AppendToString(TEXT("ActorPath"), ActorPath);
+		}
+
 		AppendToString(TEXT("Label"), GetActorLabel());
 		AppendFromBool(TEXT("SpatiallyLoaded"), bIsSpatiallyLoaded);
 		AppendToString(TEXT("EditorBounds"), EditorBounds);
diff --git a/Engine/Source/Runtime/Engine/Public/WorldPartition/WorldPartitionActorDesc.h b/Engine/Source/Runtime/Engine/Public/WorldPartition/WorldPartitionActorDesc.h
--- a/Engine/Source/Runtime/Engine/Public/WorldPartition/WorldPartitionActorDesc.h
+++ b/Engine/Source/Runtime/Engine/Public/WorldPartition/WorldPartitionActorDesc.h
@@ -355,6 +355,7 @@ public:
 		Guid,
 		Compact,
 		Full,
+		Verbose,
 		ForDiff
 	};
 
