World Partition - Avoid assignation of the Default HLOD Layer to non spatially loaded actor desc views during streaming generation

diff --git a/Engine/Source/Runtime/Engine/Private/WorldPartition/WorldPartitionStreamingGeneration.cpp b/Engine/Source/Runtime/Engine/Private/WorldPartition/WorldPartitionStreamingGeneration.cpp
--- a/Engine/Source/Runtime/Engine/Private/WorldPartition/WorldPartitionStreamingGeneration.cpp
+++ b/Engine/Source/Runtime/Engine/Private/WorldPartition/WorldPartitionStreamingGeneration.cpp
@@ -721,7 +721,7 @@ class FWorldPartitionStreamingGenerator
 	{
 		// Only assign the default layer to actors that don't have a valid HLOD layer set. HLOD actors will have their 
 		// parent HLOD layer set during HLOD generation.
-		if (!ActorDescView.GetHLODLayer().IsValid())
+		if (!ActorDescView.GetHLODLayer().IsValid() && ActorDescView.GetIsSpatiallyLoaded())
 		{
 			ActorDescView.SetRuntimeHLODLayer(DefaultHLODLayer);
 		}
@@ -1576,7 +1576,7 @@ class FWorldPartitionStreamingGenerator
 					NbErrorsDetected++;
 				}
 
-				if (ActorDescView.GetHLODLayer().IsValid() && !IsValidHLODLayer(PerInstanceData.RuntimeGrid, ActorDescView.GetHLODLayer()))
+				if (ActorDescView.GetIsSpatiallyLoaded() && ActorDescView.GetHLODLayer().IsValid() && !IsValidHLODLayer(PerInstanceData.RuntimeGrid, ActorDescView.GetHLODLayer()))
 				{
 					if (PassType == EPassType::ErrorReporting)
 					{
