Fixed typo: "loaded loaded"

diff --git a/Engine/Plugins/Editor/DataValidation/Source/DataValidation/Private/WorldPartitionChangelistValidator.cpp b/Engine/Plugins/Editor/DataValidation/Source/DataValidation/Private/WorldPartitionChangelistValidator.cpp
--- a/Engine/Plugins/Editor/DataValidation/Source/DataValidation/Private/WorldPartitionChangelistValidator.cpp
+++ b/Engine/Plugins/Editor/DataValidation/Source/DataValidation/Private/WorldPartitionChangelistValidator.cpp
@@ -339,7 +339,7 @@ void UWorldPartitionChangelistValidator::OnInvalidReferenceGridPlacement(const I
 		if (!ActorDescView.GetIsSpatiallyLoaded())
 		{
 			const FString SpatiallyLoadedActor(TEXT("Spatially loaded actor"));
-			const FString NonSpatiallyLoadedActor(TEXT("Non-spatially loaded loaded actor"));
+			const FString NonSpatiallyLoadedActor(TEXT("Non-spatially loaded actor"));
 
 			FText CurrentError = FText::Format(LOCTEXT("DataValidation.Changelist.WorldPartition.InvalidReferenceGridPlacement", "{0} {1} is referencing {2} {3}."),
 												FText::FromString(ActorDescView.GetIsSpatiallyLoaded() ? *SpatiallyLoadedActor : *NonSpatiallyLoadedActor),
diff --git a/Engine/Source/Runtime/Engine/Classes/Engine/ObjectLibrary.h b/Engine/Source/Runtime/Engine/Classes/Engine/ObjectLibrary.h
--- a/Engine/Source/Runtime/Engine/Classes/Engine/ObjectLibrary.h
+++ b/Engine/Source/Runtime/Engine/Classes/Engine/ObjectLibrary.h
@@ -153,7 +153,7 @@ public:
 		return LoadAssetDataFromPaths(Paths);
 	}
 
-	/** Load an entire subdirectory of blueprints into this object library. Only loads asset data for blueprints of passed in class. Returns number of asset data loaded loaded */
+	/** Load an entire subdirectory of blueprints into this object library. Only loads asset data for blueprints of passed in class. Returns number of asset data loaded */
 	virtual int32 LoadBlueprintAssetDataFromPaths(const TArray<FString>& Paths, bool bForceSynchronousScan = true);
 
 	virtual int32 LoadBlueprintAssetDataFromPath(const FString& Path)
diff --git a/Engine/Source/Runtime/Engine/Private/WorldPartition/ErrorHandling/WorldPartitionStreamingGenerationLogErrorHandler.cpp b/Engine/Source/Runtime/Engine/Private/WorldPartition/ErrorHandling/WorldPartitionStreamingGenerationLogErrorHandler.cpp
--- a/Engine/Source/Runtime/Engine/Private/WorldPartition/ErrorHandling/WorldPartitionStreamingGenerationLogErrorHandler.cpp
+++ b/Engine/Source/Runtime/Engine/Private/WorldPartition/ErrorHandling/WorldPartitionStreamingGenerationLogErrorHandler.cpp
@@ -22,7 +22,7 @@ void FStreamingGenerationLogErrorHandler::OnInvalidReference(const IWorldPartiti
 void FStreamingGenerationLogErrorHandler::OnInvalidReferenceGridPlacement(const IWorldPartitionActorDescInstanceView& ActorDescView, const IWorldPartitionActorDescInstanceView& ReferenceActorDescView)
 {
 	static const FString SpatiallyLoadedActor(TEXT("Spatially loaded actor"));
-	static const FString NonSpatiallyLoadedActor(TEXT("Non-spatially loaded loaded actor"));
+	static const FString NonSpatiallyLoadedActor(TEXT("Non-spatially loaded actor"));
 
 	UE_ASSET_LOG_ACTORDESCVIEW(LogWorldPartition, Log, ActorDescView, TEXT("%s %s reference %s %s"), ActorDescView.GetIsSpatiallyLoaded() ? *SpatiallyLoadedActor : *NonSpatiallyLoadedActor, *GetActorName(ActorDescView), ReferenceActorDescView.GetIsSpatiallyLoaded() ? *SpatiallyLoadedActor : *NonSpatiallyLoadedActor, *GetActorName(ReferenceActorDescView));
 }
