[World Partition World Asset Streaming] Fix missing external streaming object removal when unregistering a world a asset streaming.

diff --git a/Engine/Source/Runtime/Engine/Private/WorldPartition/WorldPartitionRuntimeSpatialHash.cpp b/Engine/Source/Runtime/Engine/Private/WorldPartition/WorldPartitionRuntimeSpatialHash.cpp
--- a/Engine/Source/Runtime/Engine/Private/WorldPartition/WorldPartitionRuntimeSpatialHash.cpp
+++ b/Engine/Source/Runtime/Engine/Private/WorldPartition/WorldPartitionRuntimeSpatialHash.cpp
@@ -2127,7 +2127,14 @@ bool UWorldPartitionRuntimeSpatialHash::UnregisterWorldAssetStreaming(const FGui
 			TrashObject(InStreamingObject);
 		};
 
-		TrashExternalStreamingData(*StreamingObject);
+		if (IsValid(*StreamingObject))
+		{
+			if (IsValid(GetOuterUWorldPartition()))
+			{
+				GetOuterUWorldPartition()->RemoveExternalStreamingObject(*StreamingObject);
+			}
+			TrashExternalStreamingData(*StreamingObject);
+		}
 		WorldAssetStreamingObjects.Remove(InWorldAssetStreamingGuid);
 		return true;
 	}
