[AssetManager] Added source location tracking for async load requests that do not have debug name specified
#tests local run

diff --git a/Engine/Source/Runtime/Core/Public/Misc/SourceLocation.h b/Engine/Source/Runtime/Core/Public/Misc/SourceLocation.h
new file mode 100644
--- /dev/null
+++ b/Engine/Source/Runtime/Core/Public/Misc/SourceLocation.h
@@ -0,0 +1,139 @@
+// Copyright Epic Games, Inc. All Rights Reserved.
+
+#pragma once
+
+#include "CoreTypes.h"
+
+#include "HAL/Platform.h"
+#include "Templates/UnrealTemplate.h"
+
+#ifndef UE_INCLUDE_SOURCE_LOCATION
+#define UE_INCLUDE_SOURCE_LOCATION !UE_BUILD_SHIPPING
+#endif
+
+#if UE_INCLUDE_SOURCE_LOCATION && defined(__cpp_lib_source_location)
+#include <source_location>
+#endif
+
+namespace UE
+{
+	class FSourceLocation
+	{
+#if UE_INCLUDE_SOURCE_LOCATION
+	#ifdef __cpp_lib_source_location
+		using FSourceLocationImpl = std::source_location;
+	#else // __cpp_lib_source_location
+		struct FSourceLocationImpl
+		{
+			static UE_CONSTEVAL FSourceLocationImpl current(
+				const uint_least32_t InLine = __builtin_LINE(),
+				const uint_least32_t InColumn = __builtin_COLUMN(),
+				const char* const InFile = __builtin_FILE(),
+				const char* const InFunction = __builtin_FUNCTION()) noexcept
+			{
+				FSourceLocationImpl Location;
+				Location.Line = InLine;
+				Location.Column = InColumn;
+				Location.FileName = InFile;
+				Location.FunctionName = InFunction;
+				return Location;
+			};
+
+			constexpr FSourceLocationImpl() noexcept = default;
+
+			// source location field access
+			constexpr uint_least32_t line() const noexcept
+			{
+				return Line;
+			}
+			constexpr uint_least32_t column() const noexcept
+			{
+				return Column;
+			}
+			constexpr const char* file_name() const noexcept
+			{
+				return FileName;
+			}
+			constexpr const char* function_name() const noexcept
+			{
+				return FunctionName;
+			}
+
+		private:
+			uint_least32_t Line;
+			uint_least32_t Column;
+			const char* FileName;
+			const char* FunctionName;
+		};
+	#endif // __cpp_lib_source_location
+#else // UE_INCLUDE_SOURCE_LOCATION
+		struct FSourceLocationImpl{
+			static UE_CONSTEVAL FSourceLocationImpl current() noexcept
+			{
+				return {};
+			}
+			constexpr FSourceLocationImpl() noexcept = default;
+			constexpr uint_least32_t line() const noexcept
+			{
+				return 0;
+			}
+			constexpr uint_least32_t column() const noexcept
+			{
+				return 0;
+			}
+			constexpr const char* file_name() const noexcept
+			{
+				return "";
+			}
+			constexpr const char* function_name() const noexcept
+			{
+				return "";
+			}
+		};
+#endif // UE_INCLUDE_SOURCE_LOCATION
+
+	public:
+		/**
+		 * Saves current source file location into a RAII container that can be used to log/save the information about the caller
+		 * Usage:
+		 * 
+		 * #include "SourceLocation.h"
+		 * #include "SourceLocationUtils.h"
+		 * 
+		 * void MyFunction(int Param1, int Param2, UE::FSourceLocation Location = UE::FSourceLocation::Current()) {
+		 *    UE_LOG(TEXT("My caller is %s"), *UE::SourceLocation::ToFullString(Location));
+		 * }
+		 */
+		static UE_CONSTEVAL FSourceLocation Current(FSourceLocationImpl Impl = FSourceLocationImpl::current()) noexcept
+		{
+			return FSourceLocation(MoveTemp(Impl));
+		}
+
+		// source location field access
+		constexpr uint32 GetLine() const noexcept
+		{
+			return Impl.line();
+		}
+		constexpr uint32 GetColumn() const noexcept
+		{
+			return Impl.column();
+		}
+		constexpr const char* GetFileName() const noexcept
+		{
+			return Impl.file_name();
+		}
+		constexpr const char* GetFunctionName() const noexcept
+		{
+			return Impl.function_name();
+		}
+
+	private:
+		constexpr FSourceLocation(FSourceLocationImpl&& In) noexcept
+			: Impl(MoveTemp(In))
+		{
+		}
+
+		const FSourceLocationImpl Impl;
+	};
+
+}
\ No newline at end of file
diff --git a/Engine/Source/Runtime/Core/Public/Misc/SourceLocationUtils.h b/Engine/Source/Runtime/Core/Public/Misc/SourceLocationUtils.h
new file mode 100644
--- /dev/null
+++ b/Engine/Source/Runtime/Core/Public/Misc/SourceLocationUtils.h
@@ -0,0 +1,42 @@
+// Copyright Epic Games, Inc. All Rights Reserved.
+
+#pragma once
+
+#include "CoreTypes.h"
+
+#include "Containers/UnrealString.h"
+#include "Misc/SourceLocation.h"
+#include "Misc/StringBuilder.h"
+
+namespace UE::SourceLocation
+{
+	/**
+	 * Returns an owning string that contains function name as well as file, line and column
+	 */
+	static inline FString ToFullString(const FSourceLocation& Location)
+	{
+#if UE_INCLUDE_SOURCE_LOCATION
+		return FString(WriteToString<512>(
+			Location.GetFileName(), "(",
+			Location.GetLine(), ":",
+			Location.GetColumn(), ") ",
+			Location.GetFunctionName()).ToView());
+#else
+		return FString();
+#endif
+	}
+
+	/**
+	 * Returns an owning string that contains Source filename and line.
+	 * Equvivalent to UE_SOURCE_LOCATION
+	 */
+	static inline FString ToFileAndLineString(const FSourceLocation& Location)
+	{
+#if UE_INCLUDE_SOURCE_LOCATION
+		return FString(WriteToString<300>(
+			Location.GetFileName(), "(", Location.GetLine(), ")").ToView());
+#else
+		return FString();
+#endif
+	}
+}
\ No newline at end of file
diff --git a/Engine/Source/Runtime/Core/Public/Templates/UnrealTemplate.h b/Engine/Source/Runtime/Core/Public/Templates/UnrealTemplate.h
--- a/Engine/Source/Runtime/Core/Public/Templates/UnrealTemplate.h
+++ b/Engine/Source/Runtime/Core/Public/Templates/UnrealTemplate.h
@@ -546,7 +546,7 @@ template <typename T> struct TRemovePointer<T*> { typedef T Type; };
  * const object, because we would prefer to be informed when MoveTemp will have no effect.
  */
 template <typename T>
-UE_INTRINSIC_CAST FORCEINLINE std::remove_reference_t<T>&& MoveTemp(T&& Obj)
+UE_INTRINSIC_CAST FORCEINLINE constexpr std::remove_reference_t<T>&& MoveTemp(T&& Obj) noexcept
 {
 	using CastType = std::remove_reference_t<T>;
 
@@ -564,7 +564,7 @@ UE_INTRINSIC_CAST FORCEINLINE std::remove_reference_t<T>&& MoveTemp(T&& Obj)
  * where you can but not stop compilation.
  */
 template <typename T>
-UE_INTRINSIC_CAST FORCEINLINE std::remove_reference_t<T>&& MoveTempIfPossible(T&& Obj)
+UE_INTRINSIC_CAST FORCEINLINE constexpr std::remove_reference_t<T>&& MoveTempIfPossible(T&& Obj) noexcept
 {
 	using CastType = std::remove_reference_t<T>;
 	return (CastType&&)Obj;
@@ -600,7 +600,7 @@ FORCEINLINE T CopyTemp(const T& Val)
  * create a PR value without stopping compilation.
  */
 template <typename T>
-FORCEINLINE std::decay_t<T> CopyTempIfNecessary(T&& Val)
+FORCEINLINE constexpr std::decay_t<T> CopyTempIfNecessary(T&& Val)
 {
 	return (T&&)Val;
 }
@@ -610,13 +610,13 @@ FORCEINLINE std::decay_t<T> CopyTempIfNecessary(T&& Val)
  * This is UE's equivalent of std::forward.
  */
 template <typename T>
-UE_INTRINSIC_CAST FORCEINLINE T&& Forward(std::remove_reference_t<T>& Obj)
+UE_INTRINSIC_CAST FORCEINLINE constexpr T&& Forward(std::remove_reference_t<T>& Obj) noexcept
 {
 	return (T&&)Obj;
 }
 
 template <typename T>
-UE_INTRINSIC_CAST FORCEINLINE T&& Forward(std::remove_reference_t<T>&& Obj)
+UE_INTRINSIC_CAST FORCEINLINE constexpr T&& Forward(std::remove_reference_t<T>&& Obj) noexcept
 {
 	return (T&&)Obj;
 }
diff --git a/Engine/Source/Runtime/Engine/Classes/Engine/AssetManager.h b/Engine/Source/Runtime/Engine/Classes/Engine/AssetManager.h
--- a/Engine/Source/Runtime/Engine/Classes/Engine/AssetManager.h
+++ b/Engine/Source/Runtime/Engine/Classes/Engine/AssetManager.h
@@ -9,6 +9,8 @@
 #include "AssetRegistry/AssetRegistryModule.h"
 #include "GenericPlatform/GenericPlatformChunkInstall.h"
 #include "ContentEncryptionConfig.h"
+#include "Misc/SourceLocation.h"
+#include "Misc/SourceLocationUtils.h"
 
 #include "AssetManager.generated.h"
 
@@ -343,9 +345,22 @@ public:
 	 */
 	ENGINE_API virtual TSharedPtr<FStreamableHandle> PreloadPrimaryAssets(const TArray<FPrimaryAssetId>& AssetsToLoad, const TArray<FName>& LoadBundles, bool bLoadRecursive, FStreamableDelegate DelegateToCall = FStreamableDelegate(), TAsyncLoadPriority Priority = FStreamableManager::DefaultAsyncLoadPriority);
 
-	/** Quick wrapper to async load some non primary assets with the primary streamable manager. This will not auto release the handle, release it if needed */
-	ENGINE_API virtual TSharedPtr<FStreamableHandle> LoadAssetList(const TArray<FSoftObjectPath>& AssetList, FStreamableDelegate DelegateToCall = FStreamableDelegate(), TAsyncLoadPriority Priority = FStreamableManager::DefaultAsyncLoadPriority, const FString& DebugName = FStreamableHandle::HandleDebugName_AssetList);
-	ENGINE_API virtual TSharedPtr<FStreamableHandle> LoadAssetList(TArray<FSoftObjectPath>&& AssetList, FStreamableDelegate DelegateToCall = FStreamableDelegate(), TAsyncLoadPriority Priority = FStreamableManager::DefaultAsyncLoadPriority, const FString& DebugName = FStreamableHandle::HandleDebugName_AssetList);
+	/**
+	 * Load non primary assets with the primary streamable manager.
+	 * This will not auto release the handle, release it if needed.
+	 * 
+	 * @param AssetList			List of non primary assets to load
+	 * @param DelegateToCall	[optional] Delegate that will be called on completion, may be called before function returns if assets are already loaded
+	 * @param Priority			[optional] Async loading priority for this request
+	 * @param DebguName			[optional] Name of this handle, either FString or anything that can construct FString, will be reported in debug tools, will report Source Location if DebugName is not specified explicitly
+	 * @return					Streamable Handle that must be stored to keep the preloaded assets from being freed
+	 */
+	template< typename DebugNameType = UE::FSourceLocation >
+	TSharedPtr<FStreamableHandle> LoadAssetList(
+		TArray<FSoftObjectPath> AssetList,
+		FStreamableDelegate DelegateToCall = FStreamableDelegate(),
+		TAsyncLoadPriority Priority = FStreamableManager::DefaultAsyncLoadPriority,
+		DebugNameType&& DebugNameOrLocation = UE::FSourceLocation::Current());
 
 	/** Returns a single AssetBundleInfo, matching Scope and Name */
 	ENGINE_API virtual FAssetBundleEntry GetAssetBundleEntry(const FPrimaryAssetId& BundleScope, FName BundleName) const;
@@ -716,6 +731,9 @@ protected:
 	/** Called when a new chunk has been downloaded */
 	ENGINE_API virtual void OnChunkDownloaded(uint32 ChunkId, bool bSuccess);
 
+	/** Called to load asset list */
+	ENGINE_API virtual TSharedPtr<FStreamableHandle> LoadAssetListInternal(TArray<FSoftObjectPath>&& AssetList, FStreamableDelegate&& DelegateToCall, TAsyncLoadPriority Priority, FString&& DebugName);
+
 #if WITH_EDITOR
 	/** Function used during creating Management references to decide when to recurse and set references */
 	ENGINE_API virtual EAssetSetManagerResult::Type ShouldSetManager(const FAssetIdentifier& Manager, const FAssetIdentifier& Source, const FAssetIdentifier& Target,
@@ -924,3 +942,28 @@ private:
 
 	friend struct FCompiledAssetManagerSearchRules;
 };
+
+template< typename DebugNameType >
+TSharedPtr<FStreamableHandle> UAssetManager::LoadAssetList(
+	TArray<FSoftObjectPath> AssetList,
+	FStreamableDelegate DelegateToCall,
+	TAsyncLoadPriority Priority,
+	DebugNameType&& DebugNameOrLocation)
+{
+	if constexpr (std::is_same_v<std::remove_cv_t<DebugNameType>, UE::FSourceLocation>)
+	{
+		return LoadAssetListInternal(
+			MoveTemp(AssetList),
+			MoveTemp(DelegateToCall),
+			Priority,
+			UE::SourceLocation::ToFileAndLineString(DebugNameOrLocation));
+	}
+	else
+	{
+		return LoadAssetListInternal(
+			MoveTemp(AssetList),
+			MoveTemp(DelegateToCall),
+			Priority,
+			FString{ Forward<DebugNameType>(DebugNameOrLocation) });
+	}
+}
diff --git a/Engine/Source/Runtime/Engine/Classes/Engine/StreamableManager.h b/Engine/Source/Runtime/Engine/Classes/Engine/StreamableManager.h
--- a/Engine/Source/Runtime/Engine/Classes/Engine/StreamableManager.h
+++ b/Engine/Source/Runtime/Engine/Classes/Engine/StreamableManager.h
@@ -12,6 +12,8 @@
 #include "UObject/GCObject.h"
 #include "UObject/ObjectMacros.h"
 #include "UObject/SoftObjectPtr.h"
+#include "Misc/SourceLocation.h"
+#include "Misc/SourceLocationUtils.h"
 
 /** Defines FStreamableDelegate delegate interface */
 DECLARE_DELEGATE(FStreamableDelegate);
@@ -481,31 +483,33 @@ struct FStreamableManager : public FGCObject
 	 * This is the primary streamable operation. Requests streaming of one or more target objects. When complete, a delegate function is called. Returns a Streamable Handle.
 	 *
 	 * @param TargetsToStream		Assets to load off disk
-	 * @param DelegateToCall		Delegate to call when load finishes. Will be called on the next tick if asset is already loaded, or many seconds later
-	 * @param Priority				Priority to pass to the streaming system, higher priority will be loaded first
-	 * @param bManageActiveHandle	If true, the manager will keep the streamable handle active until explicitly released
-	 * @param bStartStalled			If true, the handle will start in a stalled state and will not attempt to actually async load until StartStalledHandle is called on it
-	 * @param DebugName				Name of this handle, will be reported in debug tools
+	 * @param DelegateToCall		[optional] Delegate to call when load finishes. Will be called on the next tick if asset is already loaded, or many seconds later
+	 * @param Priority				[optional] Priority to pass to the streaming system, higher priority will be loaded first
+	 * @param bManageActiveHandle	[optional] If true, the manager will keep the streamable handle active until explicitly released
+	 * @param bStartStalled			[optional] If true, the handle will start in a stalled state and will not attempt to actually async load until StartStalledHandle is called on it
+	 * @param DebugName				[optional] Name of this handle, either FString or anything that can construct FString, will be reported in debug tools, will report Source Location if DebugName is not specified explicitly
 	 */
-	ENGINE_API TSharedPtr<FStreamableHandle> RequestAsyncLoad(const TArray<FSoftObjectPath>& TargetsToStream, FStreamableDelegate DelegateToCall = FStreamableDelegate(), TAsyncLoadPriority Priority = DefaultAsyncLoadPriority, bool bManageActiveHandle = false, bool bStartStalled = false, FString DebugName = TEXT("ArrayDelegate"));
-	ENGINE_API TSharedPtr<FStreamableHandle> RequestAsyncLoad(TArray<FSoftObjectPath>&& TargetsToStream, FStreamableDelegate DelegateToCall = FStreamableDelegate(), TAsyncLoadPriority Priority = DefaultAsyncLoadPriority, bool bManageActiveHandle = false, bool bStartStalled = false, FString DebugName = TEXT("ArrayDelegate"));
-	ENGINE_API TSharedPtr<FStreamableHandle> RequestAsyncLoad(const FSoftObjectPath& TargetToStream, FStreamableDelegate DelegateToCall = FStreamableDelegate(), TAsyncLoadPriority Priority = DefaultAsyncLoadPriority, bool bManageActiveHandle = false, bool bStartStalled = false, FString DebugName = TEXT("SingleDelegate"));
-
-	/** Lambda Wrappers. Be aware that Callback may go off multiple seconds in the future. */
-	ENGINE_API TSharedPtr<FStreamableHandle> RequestAsyncLoad(const TArray<FSoftObjectPath>& TargetsToStream, TFunction<void()>&& Callback, TAsyncLoadPriority Priority = DefaultAsyncLoadPriority, bool bManageActiveHandle = false, bool bStartStalled = false, FString DebugName = TEXT("ArrayLambda"));
-	ENGINE_API TSharedPtr<FStreamableHandle> RequestAsyncLoad(TArray<FSoftObjectPath>&& TargetsToStream, TFunction<void()>&& Callback, TAsyncLoadPriority Priority = DefaultAsyncLoadPriority, bool bManageActiveHandle = false, bool bStartStalled = false, FString DebugName = TEXT("ArrayLambda"));
-	ENGINE_API TSharedPtr<FStreamableHandle> RequestAsyncLoad(const FSoftObjectPath& TargetToStream, TFunction<void()>&& Callback, TAsyncLoadPriority Priority = DefaultAsyncLoadPriority, bool bManageActiveHandle = false, bool bStartStalled = false, FString DebugName = TEXT("SingleLambda"));
+	template< typename PathContainerType = TArray<FSoftObjectPath>, typename FuncType = FStreamableDelegate, typename DebugNameType = UE::FSourceLocation >
+	TSharedPtr<FStreamableHandle> RequestAsyncLoad(
+		PathContainerType&& TargetsToStream,
+		FuncType&& DelegateToCall = FStreamableDelegate(),
+		TAsyncLoadPriority Priority = DefaultAsyncLoadPriority,
+		bool bManageActiveHandle = false,
+		bool bStartStalled = false,
+		DebugNameType&& DebugNameOrLocation = UE::FSourceLocation::Current());
 
 	/** 
 	 * Synchronously load a set of assets, and return a handle. This can be very slow and may stall the game thread for several seconds.
 	 * 
 	 * @param TargetsToStream		Assets to load off disk
-	 * @param bManageActiveHandle	If true, the manager will keep the streamable handle active until explicitly released
-	 * @param DebugName				Name of this handle, will be reported in debug tools
+	 * @param bManageActiveHandle	[optional] If true, the manager will keep the streamable handle active until explicitly released
+	 * @param DebugName				[optional] Name of this handle, either FString or anything that can construct FString, will be reported in debug tools, will report Source Location if DebugName is not specified explicitly
 	 */
-	ENGINE_API TSharedPtr<FStreamableHandle> RequestSyncLoad(const TArray<FSoftObjectPath>& TargetsToStream, bool bManageActiveHandle = false, FString DebugName = TEXT("RequestSyncLoad Array"));
-	ENGINE_API TSharedPtr<FStreamableHandle> RequestSyncLoad(TArray<FSoftObjectPath>&& TargetsToStream, bool bManageActiveHandle = false, FString DebugName = TEXT("RequestSyncLoad Array"));
-	ENGINE_API TSharedPtr<FStreamableHandle> RequestSyncLoad(const FSoftObjectPath& TargetToStream, bool bManageActiveHandle = false, FString DebugName = TEXT("RequestSyncLoad Single"));
+	template< typename PathContainerType = TArray<FSoftObjectPath>, typename DebugNameType = UE::FSourceLocation >
+	TSharedPtr<FStreamableHandle> RequestSyncLoad(
+		PathContainerType&& TargetsToStream,
+		bool bManageActiveHandle = false,
+		DebugNameType&& DebugNameOrLocation = UE::FSourceLocation::Current());
 
 	/** 
 	 * Synchronously load the referred asset and return the loaded object, or nullptr if it can't be found. This can be very slow and may stall the game thread for several seconds.
@@ -514,26 +518,26 @@ struct FStreamableManager : public FGCObject
 	 * @param bManageActiveHandle	If true, the manager will keep the streamable handle active until explicitly released
 	 * @param RequestHandlePointer	If non-null, this will set the handle to the handle used to make this request. This useful for later releasing the handle
 	 */
-	ENGINE_API UObject* LoadSynchronous(const FSoftObjectPath& Target, bool bManageActiveHandle = false, TSharedPtr<FStreamableHandle>* RequestHandlePointer = nullptr);
+	ENGINE_API UObject* LoadSynchronous(const FSoftObjectPath& Target, bool bManageActiveHandle = false, TSharedPtr<FStreamableHandle>* RequestHandlePointer = nullptr, UE::FSourceLocation Location = UE::FSourceLocation::Current());
 
 	/** Typed wrappers */
 	template< typename T >
-	T* LoadSynchronous(const FSoftObjectPath& Target, bool bManageActiveHandle = false, TSharedPtr<FStreamableHandle>* RequestHandlePointer = nullptr)
+	T* LoadSynchronous(const FSoftObjectPath& Target, bool bManageActiveHandle = false, TSharedPtr<FStreamableHandle>* RequestHandlePointer = nullptr, UE::FSourceLocation Location = UE::FSourceLocation::Current())
 	{
-		return Cast<T>(LoadSynchronous(Target, bManageActiveHandle, RequestHandlePointer) );
+		return Cast<T>(LoadSynchronous(Target, bManageActiveHandle, RequestHandlePointer, MoveTemp(Location)) );
 	}
 
 	template< typename T >
-	T* LoadSynchronous(const TSoftObjectPtr<T>& Target, bool bManageActiveHandle = false, TSharedPtr<FStreamableHandle>* RequestHandlePointer = nullptr)
+	T* LoadSynchronous(const TSoftObjectPtr<T>& Target, bool bManageActiveHandle = false, TSharedPtr<FStreamableHandle>* RequestHandlePointer = nullptr, UE::FSourceLocation Location = UE::FSourceLocation::Current())
 	{
-		return Cast<T>(LoadSynchronous(Target.ToSoftObjectPath(), bManageActiveHandle, RequestHandlePointer));
+		return Cast<T>(LoadSynchronous(Target.ToSoftObjectPath(), bManageActiveHandle, RequestHandlePointer, MoveTemp(Location)));
 	}
 
 	template< typename T >
-	TSubclassOf<T> LoadSynchronous(const TSoftClassPtr<T>& Target, bool bManageActiveHandle = false, TSharedPtr<FStreamableHandle>* RequestHandlePointer = nullptr)
+	TSubclassOf<T> LoadSynchronous(const TSoftClassPtr<T>& Target, bool bManageActiveHandle = false, TSharedPtr<FStreamableHandle>* RequestHandlePointer = nullptr, UE::FSourceLocation Location = UE::FSourceLocation::Current())
 	{
 		TSubclassOf<T> ReturnClass;
-		ReturnClass = Cast<UClass>(LoadSynchronous(Target.ToSoftObjectPath(), bManageActiveHandle, RequestHandlePointer));
+		ReturnClass = Cast<UClass>(LoadSynchronous(Target.ToSoftObjectPath(), bManageActiveHandle, RequestHandlePointer, MoveTemp(Location)));
 		return ReturnClass;
 	}
 
@@ -603,6 +607,10 @@ private:
 	ENGINE_API void OnPreGarbageCollect();
 	ENGINE_API void AsyncLoadCallback(FSoftObjectPath Request, UPackage* Package);
 
+	ENGINE_API static bool ShouldStripDebugName();
+	ENGINE_API TSharedPtr<FStreamableHandle> RequestAsyncLoadInternal(TArray<FSoftObjectPath>&& TargetsToStream, FStreamableDelegate&& DelegateToCall, TAsyncLoadPriority Priority, bool bManageActiveHandle, bool bStartStalled, FString&& DebugName);
+	ENGINE_API TSharedPtr<FStreamableHandle> RequestSyncLoadInternal(TArray<FSoftObjectPath>&& TargetsToStream, bool bManageActiveHandle, FString&& DebugName);
+
 	/** Map of paths to streamable objects, this will be the post-redirector name */
 	typedef TMap<FSoftObjectPath, struct FStreamable*> TStreamableMap;
 	TStreamableMap StreamableItems;
@@ -671,3 +679,66 @@ void FStreamableHandle::ForEachLoadedAsset(CallableT Callable) const
 		}
 	}
 }
+
+template< typename PathContainerType, typename FuncType, typename DebugNameType >
+TSharedPtr<FStreamableHandle> FStreamableManager::RequestAsyncLoad(
+	PathContainerType&& TargetsToStream,
+	FuncType&& Callback,
+	TAsyncLoadPriority Priority,
+	bool bManageActiveHandle,
+	bool bStartStalled,
+	DebugNameType&& DebugNameOrLocation)
+{
+	FStreamableDelegate DelegateToCall;
+	if constexpr (std::is_constructible_v<FStreamableDelegate, FuncType>)
+	{
+		DelegateToCall = FStreamableDelegate(Forward<FuncType>(Callback));
+	}
+	else
+	{
+		DelegateToCall = FStreamableDelegate::CreateLambda(Forward<FuncType>(Callback));
+	}
+
+	if constexpr (std::is_same_v<std::remove_cv_t<DebugNameType>, UE::FSourceLocation>)
+	{
+		return RequestAsyncLoadInternal(
+			TArray<FSoftObjectPath>{ Forward<PathContainerType>(TargetsToStream) },
+			MoveTemp(DelegateToCall),
+			Priority,
+			bManageActiveHandle,
+			bStartStalled,
+			ShouldStripDebugName() ? FString() : UE::SourceLocation::ToFileAndLineString(DebugNameOrLocation));
+	}
+	else
+	{
+		return RequestAsyncLoadInternal(
+			TArray<FSoftObjectPath>{ Forward<PathContainerType>(TargetsToStream) },
+			MoveTemp(DelegateToCall),
+			Priority,
+			bManageActiveHandle,
+			bStartStalled,
+			ShouldStripDebugName() ? FString() : FString{ Forward<DebugNameType>(DebugNameOrLocation) });
+	}
+}
+
+template< typename PathContainerType, typename DebugNameType >
+TSharedPtr<FStreamableHandle> FStreamableManager::RequestSyncLoad(
+	PathContainerType&& TargetsToStream,
+	bool bManageActiveHandle,
+	DebugNameType&& DebugNameOrLocation)
+{
+	if constexpr (std::is_same_v<std::remove_cv_t<DebugNameType>, UE::FSourceLocation>)
+	{
+		return RequestSyncLoadInternal(
+			TArray<FSoftObjectPath>{ Forward<PathContainerType>(TargetsToStream) },
+			bManageActiveHandle,
+			ShouldStripDebugName() ? FString() : UE::SourceLocation::ToFileAndLineString(DebugNameOrLocation));
+	}
+	else
+	{
+		return RequestSyncLoadInternal(
+			TArray<FSoftObjectPath>{ Forward<PathContainerType>(TargetsToStream) },
+			bManageActiveHandle,
+			ShouldStripDebugName() ? FString() : FString{ Forward<DebugNameType>(DebugNameOrLocation) });
+	}
+}
diff --git a/Engine/Source/Runtime/Engine/Private/AssetManager.cpp b/Engine/Source/Runtime/Engine/Private/AssetManager.cpp
--- a/Engine/Source/Runtime/Engine/Private/AssetManager.cpp
+++ b/Engine/Source/Runtime/Engine/Private/AssetManager.cpp
@@ -2220,13 +2220,7 @@ int32 UAssetManager::UnloadPrimaryAssetsWithType(FPrimaryAssetType PrimaryAssetT
 	return UnloadPrimaryAssets(Assets);
 }
 
-TSharedPtr<FStreamableHandle> UAssetManager::LoadAssetList(const TArray<FSoftObjectPath>& AssetList, FStreamableDelegate DelegateToCall, TAsyncLoadPriority Priority, const FString& DebugName)
-{
-	// Make an explicit Copy of the input array
-	return LoadAssetList(TArray<FSoftObjectPath>(AssetList), DelegateToCall, Priority, DebugName);
-}
-
-TSharedPtr<FStreamableHandle> UAssetManager::LoadAssetList(TArray<FSoftObjectPath>&& AssetList, FStreamableDelegate DelegateToCall, TAsyncLoadPriority Priority, const FString& DebugName)
+TSharedPtr<FStreamableHandle> UAssetManager::LoadAssetListInternal(TArray<FSoftObjectPath>&& AssetList, FStreamableDelegate&& DelegateToCall, TAsyncLoadPriority Priority, FString&& DebugName)
 {
 	TSharedPtr<FStreamableHandle> NewHandle;
 	TArray<int32> MissingChunks, ErrorChunks;
@@ -2246,12 +2240,12 @@ TSharedPtr<FStreamableHandle> UAssetManager::LoadAssetList(TArray<FSoftObjectPat
 	// SynchronousLoad doesn't make sense if chunks are missing
 	if (bShouldUseSynchronousLoad && MissingChunks.Num() == 0)
 	{
-		NewHandle = StreamableManager.RequestSyncLoad(Forward<TArray<FSoftObjectPath>>(AssetList), false, DebugName);
-		FStreamableHandle::ExecuteDelegate(MoveTemp(DelegateToCall));
+		NewHandle = StreamableManager.RequestSyncLoad(Forward<TArray<FSoftObjectPath>>(AssetList), false, Forward<FString>(DebugName));
+		FStreamableHandle::ExecuteDelegate(Forward<FStreamableDelegate>(DelegateToCall));
 	}
 	else
 	{
-		NewHandle = StreamableManager.RequestAsyncLoad(Forward<TArray<FSoftObjectPath>>(AssetList), MoveTemp(DelegateToCall), Priority, false, MissingChunks.Num() > 0, DebugName);
+		NewHandle = StreamableManager.RequestAsyncLoad(Forward<TArray<FSoftObjectPath>>(AssetList), Forward<FStreamableDelegate>(DelegateToCall), Priority, false, MissingChunks.Num() > 0, Forward<FString>(DebugName));
 
 		if (MissingChunks.Num() > 0 && NewHandle.IsValid())
 		{
diff --git a/Engine/Source/Runtime/Engine/Private/StreamableManager.cpp b/Engine/Source/Runtime/Engine/Private/StreamableManager.cpp
--- a/Engine/Source/Runtime/Engine/Private/StreamableManager.cpp
+++ b/Engine/Source/Runtime/Engine/Private/StreamableManager.cpp
@@ -1390,32 +1390,27 @@ FStreamable* FStreamableManager::StreamInternal(const FSoftObjectPath& InTargetN
 	return Existing;
 }
 
-TSharedPtr<FStreamableHandle> FStreamableManager::RequestAsyncLoad(const TArray<FSoftObjectPath>& TargetsToStream, FStreamableDelegate DelegateToCall, TAsyncLoadPriority Priority, bool bManageActiveHandle, bool bStartStalled, FString DebugName)
+bool FStreamableManager::ShouldStripDebugName()
 {
-	// Make an explicit Copy of the input array
-	return RequestAsyncLoad(TArray<FSoftObjectPath>(TargetsToStream), MoveTemp(DelegateToCall), Priority, bManageActiveHandle, bStartStalled, MoveTemp(DebugName));
+	bool bShouldStripDebugName = false;
+#if (PLATFORM_IOS || PLATFORM_ANDROID)
+	bShouldStripDebugName = true;
+#elif UE_BUILD_SHIPPING
+	bShouldStripDebugName = GStreamableStripDebugNameInShipping;
+#endif
+	return bShouldStripDebugName;
 }
 
-TSharedPtr<FStreamableHandle> FStreamableManager::RequestAsyncLoad(TArray<FSoftObjectPath>&& TargetsToStream, FStreamableDelegate DelegateToCall, TAsyncLoadPriority Priority, bool bManageActiveHandle, bool bStartStalled, FString DebugName)
+TSharedPtr<FStreamableHandle> FStreamableManager::RequestAsyncLoadInternal(TArray<FSoftObjectPath>&& TargetsToStream, FStreamableDelegate&& DelegateToCall, TAsyncLoadPriority Priority, bool bManageActiveHandle, bool bStartStalled, FString&& DebugName)
 {
 	LLM_SCOPE(ELLMTag::StreamingManager);
 
 	// Schedule a new callback, this will get called when all related async loads are completed
 	TSharedRef<FStreamableHandle> NewRequest = MakeShareable(new FStreamableHandle());
-	NewRequest->CompleteDelegate = MoveTemp(DelegateToCall);
+	NewRequest->CompleteDelegate = Forward<FStreamableDelegate>(DelegateToCall);
 	NewRequest->OwningManager = this;
 	NewRequest->RequestedAssets = Forward<TArray<FSoftObjectPath>>(TargetsToStream);
-
-	bool bShouldStripDebugName = false;
-#if (PLATFORM_IOS || PLATFORM_ANDROID)
-	bShouldStripDebugName = true;
-#elif UE_BUILD_SHIPPING
-	bShouldStripDebugName = GStreamableStripDebugNameInShipping;
-#endif
-	if (!bShouldStripDebugName)
-	{
-		NewRequest->DebugName = MoveTemp(DebugName);
-	}
+	NewRequest->DebugName = Forward<FString>(DebugName);
 
 	NewRequest->Priority = Priority;
 #if UE_WITH_PACKAGE_ACCESS_TRACKING
@@ -1520,34 +1515,7 @@ TSharedPtr<FStreamableHandle> FStreamableManager::RequestAsyncLoad(TArray<FSoftO
 	return NewRequest;
 }
 
-TSharedPtr<FStreamableHandle> FStreamableManager::RequestAsyncLoad(const FSoftObjectPath& TargetToStream, FStreamableDelegate DelegateToCall, TAsyncLoadPriority Priority, bool bManageActiveHandle, bool bStartStalled, FString DebugName)
-{
-	return RequestAsyncLoad(TArray<FSoftObjectPath>{TargetToStream}, MoveTemp(DelegateToCall), Priority, bManageActiveHandle, bStartStalled, MoveTemp(DebugName));
-}
-
-TSharedPtr<FStreamableHandle> FStreamableManager::RequestAsyncLoad(const TArray<FSoftObjectPath>& TargetsToStream, TFunction<void()>&& Callback, TAsyncLoadPriority Priority, bool bManageActiveHandle, bool bStartStalled, FString DebugName)
-{
-	// Make an explicit Copy of the input array
-	return RequestAsyncLoad(TArray<FSoftObjectPath>(TargetsToStream), Forward<TFunction<void()>>(Callback), Priority, bManageActiveHandle, bStartStalled, MoveTemp(DebugName));
-}
-
-TSharedPtr<FStreamableHandle> FStreamableManager::RequestAsyncLoad(TArray<FSoftObjectPath>&& TargetsToStream, TFunction<void()>&& Callback, TAsyncLoadPriority Priority, bool bManageActiveHandle, bool bStartStalled, FString DebugName)
-{
-	return RequestAsyncLoad(Forward<TArray<FSoftObjectPath>>(TargetsToStream), FStreamableDelegate::CreateLambda( MoveTemp( Callback ) ), Priority, bManageActiveHandle, bStartStalled, MoveTemp(DebugName));
-}
-
-TSharedPtr<FStreamableHandle> FStreamableManager::RequestAsyncLoad(const FSoftObjectPath& TargetToStream, TFunction<void()>&& Callback, TAsyncLoadPriority Priority, bool bManageActiveHandle, bool bStartStalled, FString DebugName)
-{
-	return RequestAsyncLoad(TargetToStream, FStreamableDelegate::CreateLambda( MoveTemp( Callback ) ), Priority, bManageActiveHandle, bStartStalled, MoveTemp(DebugName));
-}
-
-TSharedPtr<FStreamableHandle> FStreamableManager::RequestSyncLoad(const TArray<FSoftObjectPath>& TargetsToStream, bool bManageActiveHandle, FString DebugName)
-{
-	// Make an explicit Copy of the input array
-	return RequestSyncLoad(TArray<FSoftObjectPath>(TargetsToStream), bManageActiveHandle, MoveTemp(DebugName));
-}
-
-TSharedPtr<FStreamableHandle> FStreamableManager::RequestSyncLoad(TArray<FSoftObjectPath>&& TargetsToStream, bool bManageActiveHandle, FString DebugName)
+TSharedPtr<FStreamableHandle> FStreamableManager::RequestSyncLoadInternal(TArray<FSoftObjectPath>&& TargetsToStream, bool bManageActiveHandle, FString&& DebugName)
 {
 	// If in async loading thread or from callback always do sync as recursive tick is unsafe
 	// If in EDL always do sync as EDL internally avoids flushing
@@ -1555,7 +1523,7 @@ TSharedPtr<FStreamableHandle> FStreamableManager::RequestSyncLoad(TArray<FSoftOb
 	bForceSynchronousLoads = IsInAsyncLoadingThread() || IsEventDrivenLoaderEnabled() || !IsAsyncLoading();
 
 	// Do an async load and wait to complete. In some cases this will do a sync load due to safety issues
-	TSharedPtr<FStreamableHandle> Request = RequestAsyncLoad(Forward<TArray<FSoftObjectPath>>(TargetsToStream), FStreamableDelegate(), AsyncLoadHighPriority, bManageActiveHandle, false, MoveTemp(DebugName));
+	TSharedPtr<FStreamableHandle> Request = RequestAsyncLoadInternal(Forward<TArray<FSoftObjectPath>>(TargetsToStream), FStreamableDelegate(), AsyncLoadHighPriority, bManageActiveHandle, false, Forward<FString>(DebugName));
 
 	bForceSynchronousLoads = false;
 
@@ -1570,11 +1538,6 @@ TSharedPtr<FStreamableHandle> FStreamableManager::RequestSyncLoad(TArray<FSoftOb
 	return Request;
 }
 
-TSharedPtr<FStreamableHandle> FStreamableManager::RequestSyncLoad(const FSoftObjectPath& TargetToStream, bool bManageActiveHandle, FString DebugName)
-{
-	return RequestSyncLoad(TArray<FSoftObjectPath>{TargetToStream}, bManageActiveHandle, MoveTemp(DebugName));
-}
-
 void FStreamableManager::StartHandleRequests(TSharedRef<FStreamableHandle> Handle)
 {
 	TRACE_LOADTIME_REQUEST_GROUP_SCOPE(TEXT("StreamableManager - %s"), *Handle->GetDebugName());
@@ -1638,13 +1601,16 @@ UE_TRACE_EVENT_BEGIN(Cpu, StreamableManager_LoadSynchronous, NoSync)
 	UE_TRACE_EVENT_FIELD(UE::Trace::WideString, AssetPath)
 UE_TRACE_EVENT_END()
 
-UObject* FStreamableManager::LoadSynchronous(const FSoftObjectPath& Target, bool bManageActiveHandle, TSharedPtr<FStreamableHandle>* RequestHandlePointer)
+UObject* FStreamableManager::LoadSynchronous(const FSoftObjectPath& Target, bool bManageActiveHandle, TSharedPtr<FStreamableHandle>* RequestHandlePointer, UE::FSourceLocation Location)
 {
 #if CPUPROFILERTRACE_ENABLED
 	UE_TRACE_LOG_SCOPED_T(Cpu, StreamableManager_LoadSynchronous, CpuChannel)
 		<< StreamableManager_LoadSynchronous.AssetPath(*WriteToWideString<FName::StringBufferSize>(Target));
 #endif // CPUPROFILERTRACE_ENABLED
-	TSharedPtr<FStreamableHandle> Request = RequestSyncLoad(Target, bManageActiveHandle, FString::Printf(TEXT("LoadSynchronous of %s"), *Target.ToString()));
+	TSharedPtr<FStreamableHandle> Request = RequestSyncLoadInternal(
+		TArray<FSoftObjectPath>{Target},
+		bManageActiveHandle,
+		FString::Printf(TEXT("LoadSynchronous from %s (%d)"), Location.GetFileName(), Location.GetLine()));
 
 	if (RequestHandlePointer)
 	{
