Fixed typo - Runtimed

diff --git a/Engine/Source/Runtime/Engine/Private/WorldPartition/WorldPartitionActorDescView.cpp b/Engine/Source/Runtime/Engine/Private/WorldPartition/WorldPartitionActorDescView.cpp
--- a/Engine/Source/Runtime/Engine/Private/WorldPartition/WorldPartitionActorDescView.cpp
+++ b/Engine/Source/Runtime/Engine/Private/WorldPartition/WorldPartitionActorDescView.cpp
@@ -99,9 +99,9 @@ FSoftObjectPath FWorldPartitionActorDescView::GetHLODLayer() const
 		return FSoftObjectPath();
 	}
 
-	if (RuntimedHLODLayer.IsSet())
+	if (RuntimeHLODLayer.IsSet())
 	{
-		return RuntimedHLODLayer.GetValue();
+		return RuntimeHLODLayer.GetValue();
 	}
 
 	return ActorDesc->GetHLODLayer();
@@ -309,7 +309,7 @@ void FWorldPartitionActorDescView::SetForcedNoHLODLayer()
 
 void FWorldPartitionActorDescView::SetRuntimeHLODLayer(const FSoftObjectPath& InHLODLayer)
 {
-	RuntimedHLODLayer = InHLODLayer;
+	RuntimeHLODLayer = InHLODLayer;
 }
 
 AActor* FWorldPartitionActorDescView::GetActor() const
diff --git a/Engine/Source/Runtime/Engine/Private/WorldPartition/WorldPartitionStreamingGeneration.cpp b/Engine/Source/Runtime/Engine/Private/WorldPartition/WorldPartitionStreamingGeneration.cpp
--- a/Engine/Source/Runtime/Engine/Private/WorldPartition/WorldPartitionStreamingGeneration.cpp
+++ b/Engine/Source/Runtime/Engine/Private/WorldPartition/WorldPartitionStreamingGeneration.cpp
@@ -212,9 +212,9 @@ FSoftObjectPath FStreamingGenerationActorDescView::GetHLODLayer() const
 		return FSoftObjectPath();
 	}
 
-	if (RuntimedHLODLayer.IsSet())
+	if (RuntimeHLODLayer.IsSet())
 	{
-		return RuntimedHLODLayer.GetValue();
+		return RuntimeHLODLayer.GetValue();
 	}
 
 	return Super::GetHLODLayer();
@@ -375,7 +375,7 @@ void FStreamingGenerationActorDescView::SetForcedNoHLODLayer()
 
 void FStreamingGenerationActorDescView::SetRuntimeHLODLayer(const FSoftObjectPath& InHLODLayer)
 {
-	RuntimedHLODLayer = InHLODLayer;
+	RuntimeHLODLayer = InHLODLayer;
 }
 
 const FDataLayerInstanceNames& FStreamingGenerationActorDescView::GetRuntimeDataLayerInstanceNames() const
diff --git a/Engine/Source/Runtime/Engine/Public/WorldPartition/WorldPartitionActorDescView.h b/Engine/Source/Runtime/Engine/Public/WorldPartition/WorldPartitionActorDescView.h
--- a/Engine/Source/Runtime/Engine/Public/WorldPartition/WorldPartitionActorDescView.h
+++ b/Engine/Source/Runtime/Engine/Public/WorldPartition/WorldPartitionActorDescView.h
@@ -116,7 +116,7 @@ protected:
 	TOptional<TArray<FName>> ResolvedDataLayerInstanceNames;
 	TOptional<TArray<FName>> RuntimeDataLayerInstanceNames;
 	TOptional<TArray<FGuid>> RuntimeReferences;
-	TOptional<FSoftObjectPath> RuntimedHLODLayer;
+	TOptional<FSoftObjectPath> RuntimeHLODLayer;
 	TArray<FGuid> EditorReferences;
 };
 #endif
\ No newline at end of file
diff --git a/Engine/Source/Runtime/Engine/Public/WorldPartition/WorldPartitionStreamingGeneration.h b/Engine/Source/Runtime/Engine/Public/WorldPartition/WorldPartitionStreamingGeneration.h
--- a/Engine/Source/Runtime/Engine/Public/WorldPartition/WorldPartitionStreamingGeneration.h
+++ b/Engine/Source/Runtime/Engine/Public/WorldPartition/WorldPartitionStreamingGeneration.h
@@ -112,7 +112,7 @@ private:
 	TOptional<FDataLayerInstanceNames> ResolvedDataLayerInstanceNames;
 	TOptional<FDataLayerInstanceNames> RuntimeDataLayerInstanceNames;
 	TOptional<TArray<FGuid>> RuntimeReferences;
-	TOptional<FSoftObjectPath> RuntimedHLODLayer;
+	TOptional<FSoftObjectPath> RuntimeHLODLayer;
 	TArray<FGuid> EditorReferences;
 	TOptional<FTransform> EditorOnlyParentTransform;
 };
