Support BLAS streaming when using r.RayTracing.UseReferenceBasedResidency
- serialize FStaticMeshRayTracingProxyLOD buffers into separate FByteBulkData.
    - for the FStaticMeshRayTracingProxy.bUsingRenderingLODs case, this requires shared ownership of resources. FStaticMeshRayTracingProxyLOD owns the RayTracingGeometry while FStaticMeshLODResources owns the VB/IB/etc as usual.
- FRayTracingGeometryManager is responsible for responsible for issuing bulk data read requests and processing the results.
- decoupled r.RayTracing.UseReferenceBasedResidency from r.RayTracing.EnableOnDemand.
- made r.RayTracing.UseReferenceBasedResidency read-only since it affects serialization.

diff --git a/Engine/Source/Runtime/Core/Private/UObject/FortniteMainBranchObjectVersion.cpp b/Engine/Source/Runtime/Core/Private/UObject/FortniteMainBranchObjectVersion.cpp
--- a/Engine/Source/Runtime/Core/Private/UObject/FortniteMainBranchObjectVersion.cpp
+++ b/Engine/Source/Runtime/Core/Private/UObject/FortniteMainBranchObjectVersion.cpp
@@ -13,7 +13,7 @@ TMap<FGuid, FGuid> FFortniteMainBranchObjectVersion::GetSystemGuids()
 	SystemGuids.Add(DevGuids.NIAGARASHADERMAP_DERIVEDDATA_VER, FGuid("C417A50271B3427D98785388266E3B66"));
 	SystemGuids.Add(DevGuids.Niagara_LatestScriptCompileVersion, FGuid("E486231DA967C64B95D4E73131056FA0"));
 	SystemGuids.Add(DevGuids.SkeletalMeshDerivedDataVersion, FGuid("B9AEAA3EE7AC4FCFB3860E8F3E06DD75"));
-	SystemGuids.Add(DevGuids.STATICMESH_DERIVEDDATA_VER, FGuid("24025D0F040E41EB8CAB8D252E2B6848"));
+	SystemGuids.Add(DevGuids.STATICMESH_DERIVEDDATA_VER, FGuid("B385A7D5BC3642E19DD6313E3D86B75C"));
 	SystemGuids.Add(DevGuids.MaterialTranslationDDCVersion, FGuid("F43DE2FD43D353A9ABE260AD4CE9F3D4"));
 	return SystemGuids;
 }
diff --git a/Engine/Source/Runtime/Engine/Private/Rendering/RayTracingGeometryManager.cpp b/Engine/Source/Runtime/Engine/Private/Rendering/RayTracingGeometryManager.cpp
--- a/Engine/Source/Runtime/Engine/Private/Rendering/RayTracingGeometryManager.cpp
+++ b/Engine/Source/Runtime/Engine/Private/Rendering/RayTracingGeometryManager.cpp
@@ -12,6 +12,8 @@
 #include "RayTracingGeometry.h"
 #include "RenderUtils.h"
 
+#include "Serialization/MemoryReader.h"
+
 #if RHI_RAYTRACING
 
 static bool bHasRayTracingEnableChanged = false;
@@ -47,7 +49,15 @@ static TAutoConsoleVariable<bool> CVarRayTracingUseReferenceBasedResidency(
 				}
 			);
 		}),
-	ECVF_RenderThreadSafe
+	ECVF_ReadOnly
+);
+
+static int32 GRayTracingStreamingMaxPendingRequests = 128;
+static FAutoConsoleVariableRef CVarNaniteStreamingMaxPendingRequests(
+	TEXT("r.RayTracing.Streaming.MaxPendingRequests"),
+	GRayTracingStreamingMaxPendingRequests,
+	TEXT("Maximum number of requests that can be pending streaming."),
+	ECVF_ReadOnly
 );
 
 static int32 GRayTracingMaxBuiltPrimitivesPerFrame = -1;
@@ -76,6 +86,12 @@ DECLARE_MEMORY_STAT(TEXT("Resident Memory"), STAT_RayTracingGeometryResidentMemo
 DECLARE_DWORD_ACCUMULATOR_STAT(TEXT("Pending Builds"), STAT_RayTracingPendingBuilds, STATGROUP_RayTracingGeometry);
 DECLARE_DWORD_ACCUMULATOR_STAT(TEXT("Pending Build Primitives"), STAT_RayTracingPendingBuildPrimitives, STATGROUP_RayTracingGeometry);
 
+
+FRayTracingGeometryManager::FRayTracingGeometryManager()
+{
+	StreamingRequests.SetNum(GRayTracingStreamingMaxPendingRequests);
+}
+
 FRayTracingGeometryManager::~FRayTracingGeometryManager()
 {
 	check(GeometryBuildRequests.IsEmpty());
@@ -135,6 +151,7 @@ RayTracing::GeometryGroupHandle FRayTracingGeometryManager::RegisterRayTracingGe
 	FRayTracingGeometryGroup Group;
 	Group.Geometries.AddDefaulted(NumLODs);
 	Group.NumReferences = 1;
+	Group.CurrentFirstLODIdx = NumLODs;
 
 	RayTracing::GeometryGroupHandle Handle = RegisteredGroups.Add(MoveTemp(Group));
 
@@ -178,11 +195,6 @@ FRayTracingGeometryManager::RayTracingGeometryHandle FRayTracingGeometryManager:
 {
 	check(InGeometry);
 
-	if (GetRayTracingMode() != ERayTracingMode::Dynamic)
-	{
-		return INDEX_NONE;
-	}
-
 	FScopeLock ScopeLock(&MainCS);
 
 	RayTracingGeometryHandle Handle = RegisteredGeometries.Add({});
@@ -212,12 +224,6 @@ FRayTracingGeometryManager::RayTracingGeometryHandle FRayTracingGeometryManager:
 
 void FRayTracingGeometryManager::ReleaseRayTracingGeometryHandle(RayTracingGeometryHandle Handle)
 {
-	if (GetRayTracingMode() != ERayTracingMode::Dynamic)
-	{
-		checkf(Handle == INDEX_NONE, TEXT("When dynamic ray tracing is disabled, ray tracing geometries shouldn't have a valid handle."));
-		return;
-	}
-
 	check(Handle != INDEX_NONE);
 
 	FScopeLock ScopeLock(&MainCS);
@@ -253,6 +259,38 @@ void FRayTracingGeometryManager::ReleaseRayTracingGeometryHandle(RayTracingGeome
 	DEC_DWORD_STAT(STAT_RayTracingGeometryCount);
 }
 
+void FRayTracingGeometryManager::SetRayTracingGeometryStreamingData(const FRayTracingGeometry* Geometry, FByteBulkData& BulkData, uint32 Offset, uint32 Size)
+{
+	FScopeLock ScopeLock(&MainCS);
+
+	checkf(RegisteredGeometries.IsValidIndex(Geometry->RayTracingGeometryHandle), TEXT("SetRayTracingGeometryStreamingData(...) can only be used with FRayTracingGeometry that has been registered with FRayTracingGeometryManager."));
+
+	FRegisteredGeometry& RegisteredGeometry = RegisteredGeometries[Geometry->RayTracingGeometryHandle];
+	RegisteredGeometry.StreamableData = &BulkData;
+	RegisteredGeometry.StreamableDataOffset = Offset;
+	RegisteredGeometry.StreamableDataSize = Size;
+}
+
+
+void FRayTracingGeometryManager::SetRayTracingGeometryGroupCurrentFirstLODIndex(FRHICommandListBase& RHICmdList, RayTracing::GeometryGroupHandle Handle, uint8 NewCurrentFirstLODIdx)
+{
+	FScopeLock ScopeLock(&MainCS);
+
+	FRayTracingGeometryGroup& Group = RegisteredGroups[Handle];
+
+	// immediately release streamed out LODs
+	if(NewCurrentFirstLODIdx > Group.CurrentFirstLODIdx)
+	{
+		FRHIResourceReplaceBatcher Batcher(RHICmdList, NewCurrentFirstLODIdx - Group.CurrentFirstLODIdx);
+		for (int32 LODIdx = Group.CurrentFirstLODIdx; LODIdx < NewCurrentFirstLODIdx; ++LODIdx)
+		{
+			Group.Geometries[LODIdx]->ReleaseRHIForStreaming(Batcher);
+		}
+	}
+
+	Group.CurrentFirstLODIdx = NewCurrentFirstLODIdx;
+}
+
 void FRayTracingGeometryManager::RefreshRegisteredGeometry(RayTracingGeometryHandle Handle)
 {
 	FScopeLock ScopeLock(&MainCS);
@@ -285,6 +323,11 @@ void FRayTracingGeometryManager::RefreshRegisteredGeometry(RayTracingGeometryHan
 				TotalResidentSize -= RegisteredGeometry.Size;
 			}
 		}
+
+		if (RegisteredGeometry.Geometry->Initializer.Type == ERayTracingGeometryInitializerType::StreamingDestination)
+		{
+			RegisteredGeometry.Status = FRegisteredGeometry::FStatus::StreamedOut;
+		}
 	}
 }
 
@@ -297,11 +340,6 @@ void FRayTracingGeometryManager::Tick(FRHICommandList& RHICmdList)
 {
 	check(IsInRenderingThread());
 
-	if (GetRayTracingMode() != ERayTracingMode::Dynamic)
-	{
-		return;
-	}
-
 	TRACE_CPUPROFILER_EVENT_SCOPE(FRayTracingGeometryManager::Tick);
 	QUICK_SCOPE_CYCLE_COUNTER(STAT_FRayTracingGeometryManager_Tick);
 
@@ -391,9 +429,12 @@ void FRayTracingGeometryManager::Tick(FRHICommandList& RHICmdList)
 
 			const FRayTracingGeometryGroup& Group = RegisteredGroups[GroupHandle];
 
-			for (FRayTracingGeometry* Geometry : Group.Geometries)
+			for (uint8 LODIndex = Group.CurrentFirstLODIdx; LODIndex < Group.Geometries.Num(); ++LODIndex)
 			{
-				if (Geometry != nullptr) // some LODs might be stripped during cook
+				FRayTracingGeometry* Geometry = Group.Geometries[LODIndex];
+
+				// some LODs might be stripped during cook
+				if (Geometry != nullptr)
 				{
 					ReferencedGeometries.Add(Geometry);
 					NotReferencedResidentGeometries.Remove(Geometry);
@@ -409,12 +450,104 @@ void FRayTracingGeometryManager::Tick(FRHICommandList& RHICmdList)
 		// - make referenced geometries resident
 		for (FRayTracingGeometry* Geometry : ReferencedGeometries)
 		{
+			NotReferencedResidentGeometries.Remove(Geometry);
+
 			if (Geometry->IsEvicted())
 			{
 				Geometry->MakeResident(RHICmdList);
 			}
 
-			NotReferencedResidentGeometries.Remove(Geometry);
+			FRegisteredGeometry& RegisteredGeometry = RegisteredGeometries[Geometry->RayTracingGeometryHandle];
+
+			if (Geometry->Initializer.Type == ERayTracingGeometryInitializerType::StreamingDestination && RegisteredGeometry.StreamableData != nullptr)
+			{
+				if (RegisteredGeometry.Status == FRegisteredGeometry::FStatus::Streaming)
+				{
+					continue;
+				}
+
+				FByteBulkData* StreamableData = RegisteredGeometry.StreamableData;
+				check(StreamableData);
+
+				TResourceArray<uint8> RawData;
+				RawData.SetAllowCPUAccess(true);
+
+				FResourceArrayInterface* OfflineData = nullptr;
+
+				if (RegisteredGeometry.StreamableDataSize == 0)
+				{
+					// no offline data -> build from VB/IB at runtime
+
+					RegisteredGeometry.Status = FRegisteredGeometry::FStatus::StreamedIn;
+				}
+				else if (StreamableData->IsBulkDataLoaded())
+				{
+					{
+						const uint8* Ptr = (const uint8*)StreamableData->LockReadOnly();
+
+						FMemoryView MemView(Ptr + RegisteredGeometry.StreamableDataOffset, RegisteredGeometry.StreamableDataSize);
+
+						FMemoryReaderView MemReader(MemView, true);
+						RawData.BulkSerialize(MemReader);
+						StreamableData->Unlock();
+					}
+
+					if (!RawData.IsEmpty())
+					{
+						OfflineData = &RawData;
+					}
+
+					RegisteredGeometry.Status = FRegisteredGeometry::FStatus::StreamedIn;
+				}
+				else
+				{
+					checkf(StreamableData->CanLoadFromDisk(), TEXT("Bulk data is not loaded and cannot be loaded from disk!"));
+					check(!StreamableData->IsStoredCompressedOnDisk()); // We do not support compressed Bulkdata for this system. Limitation of the streaming request/bulk data
+
+					check(StreamableData->IsUsingIODispatcher());
+
+					if (NumStreamingRequests >= GRayTracingStreamingMaxPendingRequests)
+					{
+						continue;
+					}
+
+					FStreamingRequest& StreamingRequest = StreamingRequests[NextStreamingRequestIndex];
+					NextStreamingRequestIndex = (NextStreamingRequestIndex + 1) % GRayTracingStreamingMaxPendingRequests;
+					++NumStreamingRequests;
+
+					StreamingRequest.GeometryHandle = Geometry->RayTracingGeometryHandle;
+					StreamingRequest.RequestBuffer = FIoBuffer(RegisteredGeometry.StreamableDataSize); // TODO: Use FIoBuffer::Wrap with preallocated memory
+
+					// TODO: We're currently using a single batch per request so we can individually cancel and wait on requests.
+					// This isn't ideal and should be revisited in the future.
+					FBulkDataBatchRequest::FScatterGatherBuilder Batch = FBulkDataBatchRequest::ScatterGather(1);
+					Batch.Read(*StreamableData, RegisteredGeometry.StreamableDataOffset, RegisteredGeometry.StreamableDataSize);
+					Batch.Issue(StreamingRequest.RequestBuffer, AIOP_Low, [](FBulkDataRequest::EStatus) {}, StreamingRequest.Request);
+
+					RegisteredGeometry.Status = FRegisteredGeometry::FStatus::Streaming;
+				}
+
+				if(RegisteredGeometry.Status == FRegisteredGeometry::FStatus::StreamedIn)
+				{
+					{
+						FRHIResourceReplaceBatcher Batcher(RHICmdList, 1);
+						FRayTracingGeometryInitializer IntermediateInitializer = Geometry->Initializer;
+						IntermediateInitializer.Type = ERayTracingGeometryInitializerType::StreamingSource;
+						IntermediateInitializer.OfflineData = OfflineData;
+
+						FRayTracingGeometryRHIRef IntermediateRayTracingGeometry = RHICmdList.CreateRayTracingGeometry(IntermediateInitializer);
+
+						Geometry->SetRequiresBuild(IntermediateInitializer.OfflineData == nullptr || IntermediateRayTracingGeometry->IsCompressed());
+
+						Geometry->InitRHIForStreaming(IntermediateRayTracingGeometry, Batcher);
+
+						// When Batcher goes out of scope it will add commands to copy the BLAS buffers on RHI thread.
+						// We need to do it before we build the current geometry (also on RHI thread).
+					}
+
+					Geometry->RequestBuildIfNeeded(RHICmdList, ERTAccelerationStructureBuildPriority::Normal);
+				}
+			}
 		}
 
 		// 4th step
@@ -427,6 +560,8 @@ void FRayTracingGeometryManager::Tick(FRHICommandList& RHICmdList)
 				Geometry->Evict();
 			}
 		}
+
+		ProcessCompletedStreamingRequests(RHICmdList);
 	}
 	else
 	{
@@ -463,6 +598,86 @@ void FRayTracingGeometryManager::Tick(FRHICommandList& RHICmdList)
 	SET_MEMORY_STAT(STAT_RayTracingGeometryResidentMemory, TotalResidentSize);
 }
 
+void FRayTracingGeometryManager::ProcessCompletedStreamingRequests(FRHICommandList& RHICmdList)
+{
+	const int32 StartPendingRequestIndex = (NextStreamingRequestIndex + GRayTracingStreamingMaxPendingRequests - NumStreamingRequests) % GRayTracingStreamingMaxPendingRequests;
+
+	int32 NumCompletedRequests = 0;
+
+	for (int32 Index = 0; Index < NumStreamingRequests; ++Index)
+	{
+		const int32 PendingRequestIndex = (StartPendingRequestIndex + Index) % GRayTracingStreamingMaxPendingRequests;
+		FStreamingRequest& PendingRequest = StreamingRequests[PendingRequestIndex];
+
+		if (PendingRequest.Request.IsCompleted())
+		{
+			++NumCompletedRequests;
+
+			FRegisteredGeometry& RegisteredGeometry = RegisteredGeometries[PendingRequest.GeometryHandle];
+
+			const FRayTracingGeometryGroup& Group = RegisteredGroups[RegisteredGeometry.Geometry->GroupHandle];
+
+			if (RegisteredGeometry.Geometry->IsEvicted() || RegisteredGeometry.Geometry->LODIndex < Group.CurrentFirstLODIdx)
+			{
+				// skip if geometry was evicted while streaming request was being processed
+				continue;
+			}
+
+			if (!PendingRequest.Request.IsOk())
+			{
+				// Retry if IO request failed for some reason
+
+				FByteBulkData* StreamableData = RegisteredGeometry.StreamableData;
+
+				FBulkDataBatchRequest::FScatterGatherBuilder Batch = FBulkDataBatchRequest::ScatterGather(1);
+				Batch.Read(*StreamableData, RegisteredGeometry.StreamableDataOffset, RegisteredGeometry.StreamableDataSize);
+				Batch.Issue(PendingRequest.RequestBuffer, AIOP_Low, [](FBulkDataRequest::EStatus) {}, PendingRequest.Request);
+						
+				// TODO: Could other requests already be completed?
+				break;
+			}
+			else
+			{
+				{
+					FMemoryReaderView Ar(PendingRequest.RequestBuffer.GetView(), /*bIsPersistent=*/ true);
+					RegisteredGeometry.Geometry->RawData.BulkSerialize(Ar);
+				}
+
+				{
+					FRHIResourceReplaceBatcher Batcher(RHICmdList, 1);
+					FRayTracingGeometryInitializer IntermediateInitializer = RegisteredGeometry.Geometry->Initializer;
+					IntermediateInitializer.Type = ERayTracingGeometryInitializerType::StreamingSource;
+
+					if (!RegisteredGeometry.Geometry->RawData.IsEmpty())
+					{
+						IntermediateInitializer.OfflineData = &RegisteredGeometry.Geometry->RawData;
+					}
+
+					FRayTracingGeometryRHIRef IntermediateRayTracingGeometry = RHICmdList.CreateRayTracingGeometry(IntermediateInitializer);
+
+					RegisteredGeometry.Geometry->SetRequiresBuild(IntermediateInitializer.OfflineData == nullptr || IntermediateRayTracingGeometry->IsCompressed());
+
+					RegisteredGeometry.Geometry->InitRHIForStreaming(IntermediateRayTracingGeometry, Batcher);
+
+					// When Batcher goes out of scope it will add commands to copy the BLAS buffers on RHI thread.
+					// We need to do it before we build the current geometry (also on RHI thread).
+				}
+
+				RegisteredGeometry.Status = FRegisteredGeometry::FStatus::StreamedIn;
+
+				RegisteredGeometry.Geometry->RequestBuildIfNeeded(RHICmdList, ERTAccelerationStructureBuildPriority::Normal);
+			}
+		}
+		else
+		{
+			// TODO: Could other requests already be completed?
+			break;
+		}
+	}
+
+	NumStreamingRequests -= NumCompletedRequests;
+}
+
 void FRayTracingGeometryManager::BoostPriority(BuildRequestIndex InRequestIndex, float InBoostValue)
 {
 	FScopeLock ScopeLock(&RequestCS);
diff --git a/Engine/Source/Runtime/Engine/Private/StaticMesh.cpp b/Engine/Source/Runtime/Engine/Private/StaticMesh.cpp
--- a/Engine/Source/Runtime/Engine/Private/StaticMesh.cpp
+++ b/Engine/Source/Runtime/Engine/Private/StaticMesh.cpp
@@ -1648,6 +1648,22 @@ void FStaticMeshRayTracingProxyLOD::ReleaseResources()
 	}
 }
 
+void FStaticMeshRayTracingProxyLOD::DiscardCPUData()
+{
+	if (bOwnsRayTracingGeometry)
+	{
+		RayTracingGeometry->RawData.Discard();
+	}
+
+	if (bOwnsBuffers)
+	{
+		VertexBuffers->StaticMeshVertexBuffer.CleanUp();
+		VertexBuffers->PositionVertexBuffer.CleanUp();
+		VertexBuffers->ColorVertexBuffer.CleanUp();
+		IndexBuffer->Discard();
+	}
+}
+
 void FStaticMeshRayTracingProxyLOD::Serialize(FArchive& Ar, UObject* Owner, int32 Index)
 {
 	UStaticMesh* OwnerStaticMesh = Cast<UStaticMesh>(Owner);
@@ -1689,11 +1705,49 @@ void FStaticMeshRayTracingProxyLOD::Serialize(FArchive& Ar, UObject* Owner, int3
 		RayTracingGeometry = new FRayTracingGeometry();
 	}
 
-	SerializeBuffers(Ar, OwnerStaticMesh, 0);
+	if (bBuffersInlined)
+	{
+		SerializeBuffers(Ar, OwnerStaticMesh, 0);
+	}
+	else
+	{
+#if WITH_EDITOR
+		// Serialize buffers to FByteBulkData when saving to DDC
+		// (don't need to do it again when cooking)
+		if (Ar.IsSaving() && !Ar.IsCooking())
+		{
+			TArray<uint8> TmpBuff;
+			uint32 BuffersSize = 0;
+
+			{
+				FMemoryWriter MemWriter(TmpBuff, /*bIsPersistent=*/ true);
+				MemWriter.SetCookData(Ar.GetCookData());
+				MemWriter.SetByteSwapping(Ar.IsByteSwapping());
+
+				BuffersSize = SerializeBuffers(MemWriter, OwnerStaticMesh, 0);
+			}
+
+			if(BuffersSize > 0)
+			{
+				StreamableData.Lock(LOCK_READ_WRITE);
+
+				void* Ptr = StreamableData.Realloc(TmpBuff.Num());
+				FMemory::Memcpy(Ptr, TmpBuff.GetData(), TmpBuff.Num());
+
+				StreamableData.Unlock();
+				StreamableData.SetBulkDataFlags(BULKDATA_Force_NOT_InlinePayload); // TODO: Support inline LOD
+			}
+		}
+#endif
+		
+		StreamableData.Serialize(Ar, Owner, Index);
+	}
 }
 
-void FStaticMeshRayTracingProxyLOD::SerializeBuffers(FArchive& Ar, UStaticMesh* OwnerStaticMesh, uint8 InStripFlags)
+uint32 FStaticMeshRayTracingProxyLOD::SerializeBuffers(FArchive& Ar, UStaticMesh* OwnerStaticMesh, uint8 InStripFlags)
 {
+	uint32 OutSize = 0;
+
 	if (bOwnsBuffers)
 	{
 		VertexBuffers->PositionVertexBuffer.Serialize(Ar, false);
@@ -1701,12 +1755,22 @@ void FStaticMeshRayTracingProxyLOD::SerializeBuffers(FArchive& Ar, UStaticMesh*
 		VertexBuffers->ColorVertexBuffer.Serialize(Ar, false);
 
 		IndexBuffer->Serialize(Ar, false);
+
+		OutSize += VertexBuffers->PositionVertexBuffer.GetNumVertices() * VertexBuffers->PositionVertexBuffer.GetStride();
+		OutSize += VertexBuffers->StaticMeshVertexBuffer.GetResourceSize();
+		OutSize += VertexBuffers->ColorVertexBuffer.GetNumVertices() * VertexBuffers->ColorVertexBuffer.GetStride();
+
+		OutSize += IndexBuffer->GetIndexDataSize();
 	}
 
 	if (bOwnsRayTracingGeometry)
 	{
 		RayTracingGeometry->RawData.BulkSerialize(Ar);
+
+		OutSize += RayTracingGeometry->RawData.Num();
 	}
+
+	return OutSize;
 }
 
 /*------------------------------------------------------------------------------
@@ -1766,13 +1830,18 @@ void FStaticMeshRayTracingProxy::Serialize(FArchive& Ar, UObject* Owner, FStatic
 
 				for (int32 Index = 0; Index < LODs.Num(); ++Index)
 				{
-					LODs[Index].Sections = &LODResources[Index].Sections;
-					LODs[Index].VertexBuffers = &LODResources[Index].VertexBuffers;
-					LODs[Index].IndexBuffer = &LODResources[Index].IndexBuffer;
-					LODs[Index].bOwnsBuffers = false;
+					FStaticMeshRayTracingProxyLOD& RayTracingLOD = LODs[Index];
+
+					checkf(RayTracingLOD.bOwnsBuffers == false, TEXT("FStaticMeshRayTracingProxyLOD shouldn't own VB/IB/etc when bUsingRenderingLODs since they're owned by FStaticMeshLODResources"));
 
-					LODs[Index].RayTracingGeometry = LODResources[Index].RayTracingGeometry;
-					LODs[Index].bOwnsRayTracingGeometry = false;
+					RayTracingLOD.Sections = &LODResources[Index].Sections;
+					RayTracingLOD.VertexBuffers = &LODResources[Index].VertexBuffers;
+					RayTracingLOD.IndexBuffer = &LODResources[Index].IndexBuffer;
+
+					if (!RayTracingLOD.bOwnsRayTracingGeometry)
+					{
+						RayTracingLOD.RayTracingGeometry = LODResources[Index].RayTracingGeometry;
+					}
 				}
 
 				LODVertexFactories = &RenderData->LODVertexFactories;
@@ -2227,6 +2296,11 @@ void FStaticMeshRenderData::InitResources(ERHIFeatureLevel::Type InFeatureLevel,
 			{
 				RayTracingGeometryGroupHandle = GRayTracingGeometryManager->RegisterRayTracingGeometryGroup(LODResources.Num());
 
+				if (IsRayTracingUsingReferenceBasedResidency())
+				{
+					((FRayTracingGeometryManager*)GRayTracingGeometryManager)->SetRayTracingGeometryGroupCurrentFirstLODIndex(RHICmdList, RayTracingGeometryGroupHandle, CurrentFirstLODIdx);
+				}
+
 				for (int32 LODIndex = 0; LODIndex < LODResources.Num(); ++LODIndex)
 				{
 					// Skip LODs that have their render data stripped
@@ -2243,8 +2317,10 @@ void FStaticMeshRenderData::InitResources(ERHIFeatureLevel::Type InFeatureLevel,
 						LODResources[LODIndex].RayTracingGeometry->InitResource(RHICmdList);
 					}
 				}
+				
+				FStaticMeshRayTracingProxyLODArray& RayTracingLODs = RayTracingProxy->LODs;
 
-				TIndirectArray<FStaticMeshRayTracingProxyLOD>& RayTracingLODs = RayTracingProxy->LODs;
+				check(!RayTracingProxy->bUsingRenderingLODs || RayTracingLODs.Num() == LODResources.Num());
 
 				for (int32 LODIndex = 0; LODIndex < RayTracingLODs.Num(); ++LODIndex)
 				{
@@ -2255,13 +2331,22 @@ void FStaticMeshRenderData::InitResources(ERHIFeatureLevel::Type InFeatureLevel,
 					{
 						RayTracingLOD.RayTracingGeometry->GroupHandle = RayTracingGeometryGroupHandle;
 
-						if (LODIndex < CurrentFirstLODIdx)
+						if (LODIndex < CurrentFirstLODIdx || IsRayTracingUsingReferenceBasedResidency())
 						{
 							RayTracingLOD.RayTracingGeometry->Initializer.Type = ERayTracingGeometryInitializerType::StreamingDestination;
 						}
 
 						RayTracingLOD.RayTracingGeometry->LODIndex = LODIndex;
 						RayTracingLOD.RayTracingGeometry->InitResource(RHICmdList);
+
+						if (IsRayTracingUsingReferenceBasedResidency())
+						{
+							((FRayTracingGeometryManager*)GRayTracingGeometryManager)->SetRayTracingGeometryStreamingData(
+								RayTracingLOD.RayTracingGeometry,
+								RayTracingLOD.StreamableData,
+								0,
+								RayTracingLOD.StreamableData.GetBulkDataSize());
+						}
 					}
 				}
 			}
@@ -2332,15 +2417,18 @@ void FStaticMeshRenderData::InitializeRayTracingRepresentationFromRenderingLODs(
 {
 	const int32 NumLODs = LODResources.Num();
 
-	for (int32 LODIndex = 0; LODIndex < NumLODs; ++LODIndex)
+	if (!IsRayTracingUsingReferenceBasedResidency())
 	{
-		FStaticMeshLODResources& LODModel = LODResources[LODIndex];
+		for (int32 LODIndex = 0; LODIndex < NumLODs; ++LODIndex)
+		{
+			FStaticMeshLODResources& LODModel = LODResources[LODIndex];
 
-		check(LODModel.RayTracingGeometry == nullptr);
-		LODModel.RayTracingGeometry = new FRayTracingGeometry();
+			checkf(LODModel.RayTracingGeometry == nullptr, TEXT("LODModel.RayTracingGeometry expected to be null. Was the static mesh ray tracing representation already initialized?"));
+			LODModel.RayTracingGeometry = new FRayTracingGeometry();
+		}
 	}
 
-	check(RayTracingProxy == nullptr);
+	checkf(RayTracingProxy == nullptr, TEXT("RayTracingProxy expected to be null. Was the static mesh ray tracing representation already initialized?"));
 	RayTracingProxy = new FStaticMeshRayTracingProxy();
 	RayTracingProxy->bUsingRenderingLODs = true;
 
@@ -2354,8 +2442,8 @@ void FStaticMeshRenderData::InitializeRayTracingRepresentationFromRenderingLODs(
 
 		FStaticMeshRayTracingProxyLOD* RayTracingLOD = new FStaticMeshRayTracingProxyLOD;
 
-		RayTracingLOD->RayTracingGeometry = LODModel.RayTracingGeometry;
-		RayTracingLOD->bOwnsRayTracingGeometry = false;
+		RayTracingLOD->bOwnsRayTracingGeometry = IsRayTracingUsingReferenceBasedResidency();
+		RayTracingLOD->RayTracingGeometry = RayTracingLOD->bOwnsRayTracingGeometry ? new FRayTracingGeometry() : LODModel.RayTracingGeometry;
 
 		RayTracingLOD->Sections = &LODModel.Sections;
 		RayTracingLOD->VertexBuffers = &LODModel.VertexBuffers;
@@ -2402,8 +2490,15 @@ void UStaticMesh::RequestUpdateCachedRenderState() const
 #if RHI_RAYTRACING
 	if (IsRayTracingAllowed() && bSupportRayTracing)
 	{
-		// TODO: this should only be necessary when a BLAS build was not requested (ie: non-compressed offline BLAS)
-		GRayTracingGeometryManager->RequestUpdateCachedRenderState(GetRenderData()->RayTracingGeometryGroupHandle);
+		if (IsRayTracingUsingReferenceBasedResidency())
+		{
+			((FRayTracingGeometryManager*)GRayTracingGeometryManager)->SetRayTracingGeometryGroupCurrentFirstLODIndex(FRHICommandListImmediate::Get(), GetRenderData()->RayTracingGeometryGroupHandle, GetRenderData()->CurrentFirstLODIdx);
+		}
+		else
+		{
+			// TODO: this should only be necessary when a BLAS build was not requested (ie: non-compressed offline BLAS)
+			GRayTracingGeometryManager->RequestUpdateCachedRenderState(GetRenderData()->RayTracingGeometryGroupHandle);
+		}
 	}
 #endif
 
@@ -3172,7 +3267,19 @@ static FString BuildStaticMeshDerivedDataKeySuffix(const ITargetPlatform* Target
 		KeySuffix += TEXT("zzzzzzzz");
 	}
 
-	KeySuffix += Mesh->bSupportRayTracing && TargetPlatform->UsesRayTracing() ? TEXT("RT1") : TEXT("RT0");
+	if (Mesh->bSupportRayTracing)
+	{
+		KeySuffix += TEXT("RT1");
+
+		if (IsRayTracingUsingReferenceBasedResidency())
+		{
+			KeySuffix += TEXT("_RRBR");
+		}
+	}
+	else
+	{
+		KeySuffix += TEXT("RT0");
+	}
 
 	KeySuffix.AppendChar(Mesh->bSupportUniformlyDistributedSampling ? TEXT('1') : TEXT('0'));
 
diff --git a/Engine/Source/Runtime/Engine/Public/Rendering/RayTracingGeometryManager.h b/Engine/Source/Runtime/Engine/Public/Rendering/RayTracingGeometryManager.h
--- a/Engine/Source/Runtime/Engine/Public/Rendering/RayTracingGeometryManager.h
+++ b/Engine/Source/Runtime/Engine/Public/Rendering/RayTracingGeometryManager.h
@@ -5,6 +5,8 @@
 
 #include "Containers/SparseArray.h"
 #include "Containers/Map.h"
+#include "IO/IoBuffer.h"
+#include "Serialization/BulkData.h"
 
 #if RHI_RAYTRACING
 
@@ -15,6 +17,7 @@ class FRayTracingGeometryManager : public IRayTracingGeometryManager
 {
 public:
 
+	ENGINE_API FRayTracingGeometryManager();
 	ENGINE_API virtual ~FRayTracingGeometryManager();
 
 	ENGINE_API virtual BuildRequestIndex RequestBuildAccelerationStructure(FRayTracingGeometry* InGeometry, ERTAccelerationStructureBuildPriority InPriority, EAccelerationStructureBuildMode InBuildMode) override;
@@ -32,6 +35,9 @@ public:
 
 	ENGINE_API virtual void RefreshRegisteredGeometry(RayTracingGeometryHandle Handle) override;
 
+	void SetRayTracingGeometryStreamingData(const FRayTracingGeometry* Geometry, FByteBulkData& BulkData, uint32 Offset, uint32 Size);
+	void SetRayTracingGeometryGroupCurrentFirstLODIndex(FRHICommandListBase& RHICmdList, RayTracing::GeometryGroupHandle Handle, uint8 CurrentFirstLODIdx);
+
 	ENGINE_API virtual void PreRender() override;
 	ENGINE_API virtual void Tick(FRHICommandList& RHICmdList) override;
 
@@ -65,6 +71,8 @@ private:
 
 	void ReleaseRayTracingGeometryGroupReference(RayTracing::GeometryGroupHandle Handle);
 
+	void ProcessCompletedStreamingRequests(FRHICommandList& RHICmdList);
+
 	FCriticalSection RequestCS;
 
 	TSparseArray<FBuildRequest> GeometryBuildRequests;
@@ -83,6 +91,8 @@ private:
 
 		TSet<FPrimitiveSceneProxy*> ProxiesWithCachedRayTracingState;
 
+		uint8 CurrentFirstLODIdx = INDEX_NONE;
+
 		// Due to the way we batch release FRenderResource and SceneProxies, 
 		// ReleaseRayTracingGeometryGroup(...) can end up being called before all FRayTracingGeometry and SceneProxies are actually released.
 		// To deal with this, we keep track of whether the group is still referenced and only release the group handle once all references are released.
@@ -95,6 +105,27 @@ private:
 	{
 		FRayTracingGeometry* Geometry = nullptr;
 		uint32 Size = 0;
+
+		FByteBulkData* StreamableData = nullptr;
+		uint32 StreamableDataOffset = 0;
+		uint32 StreamableDataSize = 0;
+
+		enum class FStatus : uint8
+		{
+			StreamedOut,
+			Streaming,
+			StreamedIn,
+		};
+
+		FStatus Status = FStatus::StreamedOut;
+	};
+
+	struct FStreamingRequest
+	{
+		FIoBuffer RequestBuffer;
+		FBulkDataBatchRequest Request;
+
+		RayTracingGeometryHandle GeometryHandle;
 	};
 
 	TSparseArray<FRayTracingGeometryGroup> RegisteredGroups;
@@ -108,6 +139,10 @@ private:
 	TSet<RayTracingGeometryHandle> ReferencedGeometryHandles;
 	TSet<RayTracing::GeometryGroupHandle> ReferencedGeometryGroups;
 
+	TArray<FStreamingRequest> StreamingRequests;
+	int32 NumStreamingRequests = 0;
+	int32 NextStreamingRequestIndex = 0;
+
 	bool bRenderedFrame = false;
 };
 
diff --git a/Engine/Source/Runtime/Engine/Public/StaticMeshResources.h b/Engine/Source/Runtime/Engine/Public/StaticMeshResources.h
--- a/Engine/Source/Runtime/Engine/Public/StaticMeshResources.h
+++ b/Engine/Source/Runtime/Engine/Public/StaticMeshResources.h
@@ -687,8 +687,11 @@ struct FStaticMeshRayTracingProxyLOD
 	FStaticMeshVertexBuffers* VertexBuffers = nullptr;
 	FRawStaticIndexBuffer* IndexBuffer = nullptr;
 
+	FByteBulkData StreamableData;
+
 	bool bOwnsRayTracingGeometry : 1 = true;
 	bool bOwnsBuffers : 1 = true;
+	bool bBuffersInlined : 1 = false;
 
 	ENGINE_API ~FStaticMeshRayTracingProxyLOD();
 
@@ -697,7 +700,10 @@ struct FStaticMeshRayTracingProxyLOD
 
 	void Serialize(FArchive& Ar, UObject* Owner, int32 Index);
 
-	void SerializeBuffers(FArchive& Ar, UStaticMesh* OwnerStaticMesh, uint8 InStripFlags);
+	uint32 SerializeBuffers(FArchive& Ar, UStaticMesh* OwnerStaticMesh, uint8 InStripFlags);
+
+	/** Discard loaded data. Used when a streaming request is cancelled */
+	void DiscardCPUData();
 
 #if RHI_RAYTRACING
 	void SetupRayTracingGeometryInitializer(FRayTracingGeometryInitializer& Initializer, const FDebugName& DebugName, const FName& OwnerName) const;
diff --git a/Engine/Source/Runtime/RenderCore/Private/RayTracingGeometry.cpp b/Engine/Source/Runtime/RenderCore/Private/RayTracingGeometry.cpp
--- a/Engine/Source/Runtime/RenderCore/Private/RayTracingGeometry.cpp
+++ b/Engine/Source/Runtime/RenderCore/Private/RayTracingGeometry.cpp
@@ -96,6 +96,10 @@ void FRayTracingGeometry::MakeResident(FRHICommandList& RHICmdList)
 	// In that case we will recreate the geometry as-if it was streamed in.
 	if (EnumHasAnyFlags(GeometryState, FRayTracingGeometry::EGeometryStateFlags::StreamedIn))
 	{
+		FResourceArrayUploadInterface* OfflineData = Initializer.OfflineData;
+
+		Initializer.OfflineData = nullptr;
+
 		// When a mesh is streamed in (FStaticMeshStreamIn::DoFinishUpdate) we update the geometry initializer using just streamed in VB/IB.
 		// That initializer sets a Rendering type but RHI object was created as StreamingDestination and we have a mismatch between geometry initializer and RHI initializer.
 		// It's not an issue unless we try to initialize the geometry again using the geometry's initializer.
@@ -112,8 +116,12 @@ void FRayTracingGeometry::MakeResident(FRHICommandList& RHICmdList)
 			FRHIResourceReplaceBatcher Batcher(RHICmdList, 1);
 			FRayTracingGeometryInitializer IntermediateInitializer = Initializer;
 			IntermediateInitializer.Type = ERayTracingGeometryInitializerType::StreamingSource;
+			IntermediateInitializer.OfflineData = OfflineData;
 
 			FRayTracingGeometryRHIRef IntermediateRayTracingGeometry = RHICmdList.CreateRayTracingGeometry(IntermediateInitializer);
+
+			SetRequiresBuild(IntermediateInitializer.OfflineData == nullptr || IntermediateRayTracingGeometry->IsCompressed());
+
 			InitRHIForStreaming(IntermediateRayTracingGeometry, Batcher);
 
 			// When Batcher goes out of scope it will add commands to copy the BLAS buffers on RHI thread.
@@ -134,6 +142,12 @@ void FRayTracingGeometry::Evict()
 	RemoveBuildRequest();
 	RayTracingGeometryRHI.SafeRelease();
 	EnumAddFlags(GeometryState, EGeometryStateFlags::Evicted);
+
+	if (IsRayTracingUsingReferenceBasedResidency() && EnumHasAllFlags(GeometryState, EGeometryStateFlags::StreamedIn))
+	{
+		FRHIResourceReplaceBatcher Batcher(FRHICommandListImmediate::Get(), 1);
+		ReleaseRHIForStreaming(Batcher);
+	}
 	
 	GRayTracingGeometryManager->RefreshRegisteredGeometry(RayTracingGeometryHandle);
 	
@@ -199,7 +213,7 @@ void FRayTracingGeometry::CreateRayTracingGeometry(FRHICommandListBase& RHICmdLi
 				}
 				SetRequiresBuild(false);
 			}
-			else
+			else if (bWithNativeResource)
 			{
 				SetRequiresBuild(true);
 			}
diff --git a/Engine/Source/Runtime/RenderCore/Private/RenderUtils.cpp b/Engine/Source/Runtime/RenderCore/Private/RenderUtils.cpp
--- a/Engine/Source/Runtime/RenderCore/Private/RenderUtils.cpp
+++ b/Engine/Source/Runtime/RenderCore/Private/RenderUtils.cpp
@@ -1666,7 +1666,7 @@ bool IsRayTracingUsingReferenceBasedResidency()
 {
 	static const auto CVar = IConsoleManager::Get().FindTConsoleVariableDataBool(TEXT("r.RayTracing.UseReferenceBasedResidency"));
 
-	return GetRayTracingMode() == ERayTracingMode::Dynamic && IsRayTracingEnabled() && CVar && CVar->GetValueOnAnyThread();
+	return CVar && CVar->GetValueOnAnyThread();
 }
 
 bool UseSplineMeshSceneResources(const FStaticShaderPlatform Platform)
