Make sure that the buffer visualization materials are ready when taking a screenshot.

Without this, the first time HighResShot was asked to dump buffers in the editor, it wrote incorrect images, because the shaders were still compiling.

#rnx

diff --git a/Engine/Source/Runtime/Engine/Private/SceneView.cpp b/Engine/Source/Runtime/Engine/Private/SceneView.cpp
--- a/Engine/Source/Runtime/Engine/Private/SceneView.cpp
+++ b/Engine/Source/Runtime/Engine/Private/SceneView.cpp
@@ -2360,7 +2360,12 @@ void FSceneView::ConfigureBufferVisualizationSettings()
 		// Copy current material list into settings material list
 		for (TArray<UMaterialInterface*>::TConstIterator It = BufferVisualizationData.GetOverviewMaterials().CreateConstIterator(); It; ++It)
 		{
-			FinalPostProcessSettings.BufferVisualizationOverviewMaterials.Add(*It);
+			UMaterialInterface* VisMat = *It;
+			if(VisMat && bBufferDumpingRequired)
+			{
+				VisMat->EnsureIsComplete();
+			}
+			FinalPostProcessSettings.BufferVisualizationOverviewMaterials.Add(VisMat);
 		}
 	}
 }
