Improve error reporting for ODSC when the global shader map is not compatible between the editor and the client. It usually happens when the client/existing shadermap are up-to-date, but the ODSC server isn't
Also, add more statistics about what shader types / how material requests are currently being processed by the ODSC server

#rnx

diff --git a/Engine/Source/Runtime/Engine/Private/ODSC/ODSCManager.cpp b/Engine/Source/Runtime/Engine/Private/ODSC/ODSCManager.cpp
--- a/Engine/Source/Runtime/Engine/Private/ODSC/ODSCManager.cpp
+++ b/Engine/Source/Runtime/Engine/Private/ODSC/ODSCManager.cpp
@@ -50,9 +50,39 @@ FODSCManager::FODSCManager()
 		Thread->StartThread();
 		OnScreenMessagesHandle = FCoreDelegates::OnGetOnScreenMessages.AddLambda([this](TMultiMap<FCoreDelegates::EOnScreenMessageSeverity, FText >& OutMessages)
 			{
-				if (Thread && Thread->HasPendingRequests())
+				if (Thread)
 				{
-					OutMessages.Add(FCoreDelegates::EOnScreenMessageSeverity::Info, FText::FromString(FString::Printf(TEXT("Recompiling shaders"))));
+					FString LocalErrorMessage;
+					RetrieveErrorMessage(LocalErrorMessage);
+					if (!LocalErrorMessage.IsEmpty())
+					{
+						OutMessages.Add(FCoreDelegates::EOnScreenMessageSeverity::Error, FText::FromString(LocalErrorMessage));
+					}
+
+					bool bHasPendingGlobalShaders;
+					uint32 NumPendingMaterialsRecompile;
+					uint32 NumPendingMaterialsShaders;
+					if (Thread->GetPendingShaderData(bHasPendingGlobalShaders, NumPendingMaterialsRecompile, NumPendingMaterialsShaders))
+					{
+						FString Message = TEXT("Recompiling shaders (");
+						if (bHasPendingGlobalShaders)
+						{
+							Message += "global";
+						}
+
+						if (NumPendingMaterialsRecompile > 0)
+						{
+							Message += FString::Printf(TEXT(" %d materials"), NumPendingMaterialsRecompile);
+						}
+
+						if (NumPendingMaterialsShaders > 0)
+						{
+							Message += FString::Printf(TEXT("%d pipelines"), NumPendingMaterialsShaders);
+						}
+
+						Message += ")";
+						OutMessages.Add(FCoreDelegates::EOnScreenMessageSeverity::Info, FText::FromString(Message));
+					}
 				}
 			}
 		);
@@ -128,6 +158,8 @@ void FODSCManager::AddThreadedRequest(
 {
 	if (IsHandlingRequests())
 	{
+		ClearErrorMessage();
+
 		if ((RecompileCommandType == ODSCRecompileCommand::Material || RecompileCommandType == ODSCRecompileCommand::Changed)
 			&& (CVarODSCRecompileMode.GetValueOnAnyThread() > 0))
 		{
@@ -264,4 +296,26 @@ void FODSCManager::TryLoadGlobalShaders(EShaderPlatform ShaderPlatform)
 		ProcessCookOnTheFlyShaders(false, CompletedRequest->GetMeshMaterialMaps(), CompletedRequest->GetMaterialsToLoad(), CompletedRequest->GetGlobalShaderMap());
 		delete CompletedRequest;
 	}
-}
\ No newline at end of file
+}
+
+void FODSCManager::ReportODSCError(const FString& InErrorMessage)
+{
+	if (GODSCManager && !InErrorMessage.IsEmpty())
+	{
+		FScopeLock Lock(&GODSCManager->ErrorMessageCS);
+		GODSCManager->ErrorMessage += InErrorMessage;
+		GODSCManager->ErrorMessage += FString(TEXT("\n"));
+	}
+}
+
+void FODSCManager::RetrieveErrorMessage(FString& OutErrorMessage)
+{
+	FScopeLock Lock(&ErrorMessageCS);
+	OutErrorMessage = ErrorMessage;
+}
+
+void FODSCManager::ClearErrorMessage()
+{
+	FScopeLock Lock(&ErrorMessageCS);
+	ErrorMessage.Empty();
+}
diff --git a/Engine/Source/Runtime/Engine/Private/ODSC/ODSCThread.cpp b/Engine/Source/Runtime/Engine/Private/ODSC/ODSCThread.cpp
--- a/Engine/Source/Runtime/Engine/Private/ODSC/ODSCThread.cpp
+++ b/Engine/Source/Runtime/Engine/Private/ODSC/ODSCThread.cpp
@@ -324,12 +324,24 @@ void FODSCThread::Process()
 	// cache material requests.
 	FODSCMessageHandler* Request = nullptr;
 	TArray<FODSCMessageHandler*> RequestsToStart;
+	bool bHasGlobalShaders = false;
+	uint32 NumMaterials = 0;
 	while (PendingMaterialThreadedRequests.Dequeue(Request))
 	{
+		if (Request->GetRecompileCommandType() != ODSCRecompileCommand::Material)
+		{
+			bHasGlobalShaders = true;
+		}
+		else
+		{
+			NumMaterials += Request->GetMaterialsToLoad().Num();
+		}
+
 		RequestsToStart.Add(Request);
 	}
 
-	bHasPendingRequests = PayloadsToAggregate.Num() > 0 || RequestsToStart.Num() > 0;
+	bHasPendingGlobalShaders.store(bHasGlobalShaders, std::memory_order_release);
+	NumPendingMaterialsRecompile.store(NumMaterials, std::memory_order_release);
 
 	// process any material or recompile change shader requests or global shader compile requests.
 	for (FODSCMessageHandler* NextRequest : RequestsToStart)
@@ -340,6 +352,10 @@ void FODSCThread::Process()
 		CompletedThreadedRequests.Enqueue(NextRequest);
 	}
 
+	bHasPendingGlobalShaders.store(false, std::memory_order_release);
+	NumPendingMaterialsRecompile.store(0, std::memory_order_release);
+	NumPendingMaterialsShaders.store(PayloadsToAggregate.Num(), std::memory_order_release);
+
 	// process any specific mesh material shader requests.
 	if (PayloadsToAggregate.Num())
 	{
@@ -355,8 +371,9 @@ void FODSCThread::Process()
 		CompletedThreadedRequests.Enqueue(RequestHandler);
 	}
 
+	NumPendingMaterialsShaders.store(0, std::memory_order_release);
+
 	// SendMessageToServer is synchronous, so when we're here, we know we've processed all the requests
-	bHasPendingRequests = false;
 	WakeupEvent->Reset();
 	AllRequestsDoneEvent->Trigger();
 }
@@ -406,3 +423,11 @@ void FODSCThread::SendMessageToServer(IPlatformFile::IFileServerMessageHandler*
 		CookOnTheFlyServerConnection.Reset();
 	}
 }
+
+bool FODSCThread::GetPendingShaderData(bool& bOutHasPendingGlobalShaders, uint32& OutNumPendingMaterialsRecompile, uint32& OutNumPendingMaterialsShaders) const
+{
+	bOutHasPendingGlobalShaders = bHasPendingGlobalShaders.load(std::memory_order_acquire);
+	OutNumPendingMaterialsRecompile = NumPendingMaterialsRecompile.load(std::memory_order_acquire);
+	OutNumPendingMaterialsShaders = NumPendingMaterialsShaders.load(std::memory_order_acquire);
+	return bOutHasPendingGlobalShaders || OutNumPendingMaterialsRecompile > 0 || OutNumPendingMaterialsShaders > 0;
+}
diff --git a/Engine/Source/Runtime/Engine/Private/ODSC/ODSCThread.h b/Engine/Source/Runtime/Engine/Private/ODSC/ODSCThread.h
--- a/Engine/Source/Runtime/Engine/Private/ODSC/ODSCThread.h
+++ b/Engine/Source/Runtime/Engine/Private/ODSC/ODSCThread.h
@@ -40,6 +40,7 @@ public:
 	const TArray<uint8>& GetMeshMaterialMaps() const;
 	const TArray<uint8>& GetGlobalShaderMap() const;
 	bool ReloadGlobalShaders() const;
+	ODSCRecompileCommand GetRecompileCommandType() const { return RecompileCommandType; };
 
 private:
 	/** The time when this command was issued.  This isn't serialized to the cooking server. */
@@ -155,7 +156,7 @@ public:
 	*/
 	void WaitUntilAllRequestsDone();
 
-	bool HasPendingRequests() const { return bHasPendingRequests; }
+	bool GetPendingShaderData(bool& bOutHasPendingGlobalShaders, uint32& OutNumPendingMaterialsRecompile, uint32& OutNumPendingMaterialsShaders) const;
 
 	void RegisterMaterialShaderMap(const FMaterialShaderMap& MaterialShaderMap);
 
@@ -222,6 +223,8 @@ private:
 
 	FString ODSCHostIP;
 
-	bool bHasPendingRequests = false;
+	std::atomic<bool> bHasPendingGlobalShaders = false;
+	std::atomic<uint32> NumPendingMaterialsRecompile = 0;
+	std::atomic<uint32> NumPendingMaterialsShaders = 0;
 	bool bHasDefaultConnection = false;
 };
diff --git a/Engine/Source/Runtime/Engine/Private/ShaderCompiler/ShaderCompiler.cpp b/Engine/Source/Runtime/Engine/Private/ShaderCompiler/ShaderCompiler.cpp
--- a/Engine/Source/Runtime/Engine/Private/ShaderCompiler/ShaderCompiler.cpp
+++ b/Engine/Source/Runtime/Engine/Private/ShaderCompiler/ShaderCompiler.cpp
@@ -9607,7 +9607,7 @@ static inline bool ShouldCacheGlobalShaderTypeName(const FGlobalShaderType* Glob
 };
 
 
-bool IsGlobalShaderMapComplete(const TCHAR* TypeNameSubstring, FGlobalShaderMap* GlobalShaderMap, EShaderPlatform Platform)
+bool IsGlobalShaderMapComplete(const TCHAR* TypeNameSubstring, FGlobalShaderMap* GlobalShaderMap, EShaderPlatform Platform, FString* FailureReason = nullptr)
 {
 	// look at any shadermap in the GlobalShaderMap for the permutation flags, as they will all be the same
 	if (GlobalShaderMap)
@@ -9631,6 +9631,12 @@ bool IsGlobalShaderMapComplete(const TCHAR* TypeNameSubstring, FGlobalShaderMap*
 				{
 					if (!GlobalShaderMap->HasShader(GlobalShaderType, PermutationId))
 					{
+						if (FailureReason)
+						{
+							FString GlobalShaderTypeName = GlobalShaderType ? GlobalShaderType->GetFName().ToString() : FString(TEXT("Unknown shader type"));
+							*FailureReason = FString::Printf(TEXT("Failed to find global shader \"%s\", permutation %d"), *GlobalShaderTypeName, PermutationId);
+						}
+
 						return false;
 					}
 				}
@@ -9662,6 +9668,11 @@ bool IsGlobalShaderMapComplete(const TCHAR* TypeNameSubstring, FGlobalShaderMap*
 				{
 					if (!GlobalShaderMap->HasShaderPipeline(Pipeline))
 					{
+						if (FailureReason)
+						{
+							*FailureReason = FString::Printf(TEXT("Failed to find global pipeline \"%s\""), *Pipeline->GetFName().ToString());
+						}
+
 						return false;
 					}
 				}
@@ -10649,7 +10660,8 @@ void LoadGlobalShadersForRemoteRecompile(FArchive& Ar, EShaderPlatform ShaderPla
 		{
 			NewGlobalShaderMap->LoadFromGlobalArchive(Ar);
 
-			bool bIsNewGlobalShaderMapComplete = IsGlobalShaderMapComplete(nullptr, NewGlobalShaderMap, ShaderPlatform);
+			FString FailureReason;
+			bool bIsNewGlobalShaderMapComplete = IsGlobalShaderMapComplete(nullptr, NewGlobalShaderMap, ShaderPlatform, &FailureReason);
 
 			if (bIsNewGlobalShaderMapComplete)
 			{
@@ -10673,7 +10685,14 @@ void LoadGlobalShadersForRemoteRecompile(FArchive& Ar, EShaderPlatform ShaderPla
 			}
 			else
 			{
-				UE_LOG(LogShaderCompilers, Error, TEXT("New shader map is incomplete. Look at the ODSC server log to see shader errors"));
+				FString ErrorMessage = FString::Printf(TEXT("New global shader map is incomplete and will not be used. Reason:\n%s\n"
+													        "Please check the ODSC server log & that client/editor are compiled"), *FailureReason);
+
+				UE_LOG(LogShaderCompilers, Error, TEXT("%s"), *ErrorMessage);
+#if WITH_ODSC
+				FODSCManager::ReportODSCError(ErrorMessage);
+#endif
+				
 				delete NewGlobalShaderMap;
 			}
 		}
diff --git a/Engine/Source/Runtime/Engine/Public/ODSC/ODSCManager.h b/Engine/Source/Runtime/Engine/Public/ODSC/ODSCManager.h
--- a/Engine/Source/Runtime/Engine/Public/ODSC/ODSCManager.h
+++ b/Engine/Source/Runtime/Engine/Public/ODSC/ODSCManager.h
@@ -107,6 +107,8 @@ public:
 	static void ResumeODSCForceRecompile();
 	void TryLoadGlobalShaders(EShaderPlatform ShaderPlatform);
 
+	static void ReportODSCError(const FString& InErrorMessage);
+
 private:
 
 	ENGINE_API void OnEnginePreExit();
@@ -115,12 +117,18 @@ private:
 	bool HasAsyncLoadingInstances();
 	bool ShouldForceRecompileInternal(const FMaterialShaderMap* MaterialShaderMap, const FMaterial* Material);
 
+	void RetrieveErrorMessage(FString& OutErrorMessage);
+	void ClearErrorMessage();
+
 	/** Handles communicating directly with the cook on the fly server. */
 	FODSCThread* Thread = nullptr;
 
 	FDelegateHandle OnScreenMessagesHandle;
 	FCriticalSection MaterialInstancesCachedUniformExpressionsCS;
 	TMap<const void*, TWeakObjectPtr<const UMaterialInstance> > MaterialInstancesCachedUniformExpressions;
+
+	FCriticalSection ErrorMessageCS;
+	FString ErrorMessage;
 };
 
 struct FODSCSuspendForceRecompileScope
