[World Partition Optimization] Delay refreshing world partition editor and always loaded actors when saving actors to the next tick.
Saving 10K actors save time from ~20 minutes down to ~30 seconds.

diff --git a/Engine/Source/Runtime/Engine/Private/WorldPartition/WorldPartition.cpp b/Engine/Source/Runtime/Engine/Private/WorldPartition/WorldPartition.cpp
--- a/Engine/Source/Runtime/Engine/Private/WorldPartition/WorldPartition.cpp
+++ b/Engine/Source/Runtime/Engine/Private/WorldPartition/WorldPartition.cpp
@@ -331,6 +331,8 @@ UWorldPartition::UWorldPartition(const FObjectInitializer& ObjectInitializer)
 	, bShouldCheckEnableStreamingWarning(false)
 	, bForceGarbageCollection(false)
 	, bForceGarbageCollectionPurge(false)
+	, bForceRefreshAlwaysLoaded(false)
+	, bForceRefreshEditor(false)
 	, bEnablingStreamingJustified(false)
 	, bIsPIE(false)
 	, NumUserCreatedLoadedRegions(0)
@@ -1321,15 +1323,8 @@ void UWorldPartition::OnActorDescInstanceAdded(FWorldPartitionActorDescInstance*
 		ForceLoadedActors->AddActors({ NewActorDescInstance->GetGuid() });
 	}
 
-	if (AlwaysLoadedActors && !NewActorDescInstance->GetIsSpatiallyLoaded())
-	{
-		AlwaysLoadedActors->RefreshLoadedState();
-	}
-
-	if (WorldPartitionEditor)
-	{
-		WorldPartitionEditor->Refresh();
-	}
+	bForceRefreshAlwaysLoaded = !NewActorDescInstance->GetIsSpatiallyLoaded();
+	bForceRefreshEditor = true;
 }
 
 void UWorldPartition::OnActorDescInstanceRemoved(FWorldPartitionActorDescInstance* ActorDescInstance)
@@ -1346,15 +1341,8 @@ void UWorldPartition::OnActorDescInstanceRemoved(FWorldPartitionActorDescInstanc
 		ForceLoadedActors->RemoveActors({ ActorDescInstance->GetGuid() });
 	}
 
-	if (AlwaysLoadedActors && !ActorDescInstance->GetIsSpatiallyLoaded())
-	{
-		AlwaysLoadedActors->RefreshLoadedState();
-	}
-
-	if (WorldPartitionEditor)
-	{
-		WorldPartitionEditor->Refresh();
-	}
+	bForceRefreshAlwaysLoaded = !ActorDescInstance->GetIsSpatiallyLoaded();
+	bForceRefreshEditor = true;
 }
 
 void UWorldPartition::OnActorDescInstanceUpdating(FWorldPartitionActorDescInstance* ActorDescInstance)
@@ -1371,10 +1359,7 @@ void UWorldPartition::OnActorDescInstanceUpdated(FWorldPartitionActorDescInstanc
 
 	HashActorDescInstance(ActorDescInstance);
 
-	if (WorldPartitionEditor)
-	{
-		WorldPartitionEditor->Refresh();
-	}
+	bForceRefreshEditor = true;
 }
 
 bool UWorldPartition::ShouldHashUnhashActorDescInstances() const
@@ -1659,6 +1644,26 @@ void UWorldPartition::Tick(float DeltaSeconds)
 		ExternalDirtyActorsTracker->Tick(DeltaSeconds);
 	}
 
+	if (bForceRefreshAlwaysLoaded)
+	{
+		if (AlwaysLoadedActors)
+		{
+			AlwaysLoadedActors->RefreshLoadedState();
+		}
+
+		bForceRefreshAlwaysLoaded = false;
+	}
+
+	if (bForceRefreshEditor)
+	{
+		if (WorldPartitionEditor)
+		{
+			WorldPartitionEditor->Refresh();
+		}
+
+		bForceRefreshEditor = false;
+	}
+
 	if (bForceGarbageCollection)
 	{
 		GEngine->ForceGarbageCollection(bForceGarbageCollectionPurge);
diff --git a/Engine/Source/Runtime/Engine/Public/WorldPartition/WorldPartition.h b/Engine/Source/Runtime/Engine/Public/WorldPartition/WorldPartition.h
--- a/Engine/Source/Runtime/Engine/Public/WorldPartition/WorldPartition.h
+++ b/Engine/Source/Runtime/Engine/Public/WorldPartition/WorldPartition.h
@@ -608,6 +608,8 @@ private:
 #if WITH_EDITOR
 	bool bForceGarbageCollection;
 	bool bForceGarbageCollectionPurge;
+	bool bForceRefreshAlwaysLoaded;
+	bool bForceRefreshEditor;
 	bool bEnablingStreamingJustified;
 	bool bIsPIE;
 	int32 NumUserCreatedLoadedRegions;
