Fixed inconsistencies between UStaticMesh's ContainsPhysicsTriMeshData, GetPhysicsTriMeshData, and GetTriMeshSizeEstimates in cooked builds.

We now make sure to check bEnableCollision on each FStaticMeshSection in cooked builds instead of assuming lod 0 has collision enabled.

Fixes logspam "LogStaticMesh: Warning: UStaticMesh::GetPhysicsTriMeshData: CPU data not available on selected LOD (UseLODIndex=0, LODForCollision=0)"
Fixes logspam "LogStaticMesh: Warning: UStaticMesh::GetPhysicsTriMeshData: Triangle data from '<AssetName>' cannot be accessed at runtime on a mesh that isn't flagged as Allow CPU Access. This asset needs to be flagged as such (in the Advanced section)"

This issue seems to have originated from CL 10285605. bInUseAllTriData is intended to be used at a higher level, eg from bMeshCollideAll (originated from CL 1375311).
Also looks like this solves the mystery of some cook crashes last year (CL 24680981 speculative fix).

#tests win64 test client, PIE

diff --git a/Engine/Source/Runtime/Engine/Classes/Engine/StaticMesh.h b/Engine/Source/Runtime/Engine/Classes/Engine/StaticMesh.h
--- a/Engine/Source/Runtime/Engine/Classes/Engine/StaticMesh.h
+++ b/Engine/Source/Runtime/Engine/Classes/Engine/StaticMesh.h
@@ -39,6 +39,7 @@ class UPackage;
 struct FMeshDescription;
 struct FStaticMeshLODResources;
 class IStaticMeshComponent;
+struct FStaticMeshSection;
 
 /*-----------------------------------------------------------------------------
 	Async Static Mesh Compilation
@@ -1870,8 +1871,10 @@ public:
 	ENGINE_API virtual bool GetTriMeshSizeEstimates(struct FTriMeshCollisionDataEstimates& OutTriMeshEstimates, bool bInUseAllTriData) const override;
 
 private:
-		bool GetPhysicsTriMeshDataCheckComplex(struct FTriMeshCollisionData* CollisionData, bool bInUseAllTriData, bool bInCheckComplexCollisionMesh);
-		bool ContainsPhysicsTriMeshDataCheckComplex(bool InUseAllTriData, bool bInCheckComplexCollisionMesh) const;
+	bool GetPhysicsTriMeshDataCheckComplex(struct FTriMeshCollisionData* CollisionData, bool bInUseAllTriData, bool bInCheckComplexCollisionMesh);
+	bool ContainsPhysicsTriMeshDataCheckComplex(bool InUseAllTriData, bool bInCheckComplexCollisionMesh) const;
+	bool SectionHasCollisionEnabled(const FStaticMeshSection& Section, int32 LODIndex, int32 SectionIndex) const;
+
 public:
 
 	virtual bool WantsNegXTriMesh() override
diff --git a/Engine/Source/Runtime/Engine/Private/StaticMesh.cpp b/Engine/Source/Runtime/Engine/Private/StaticMesh.cpp
--- a/Engine/Source/Runtime/Engine/Private/StaticMesh.cpp
+++ b/Engine/Source/Runtime/Engine/Private/StaticMesh.cpp
@@ -7812,9 +7812,6 @@ bool UStaticMesh::GetPhysicsTriMeshDataCheckComplex(struct FTriMeshCollisionData
 		UE_LOG(LogStaticMesh, Warning, TEXT("UStaticMesh::GetPhysicsTriMeshData: Triangle data from '%s' cannot be accessed at runtime on a mesh that isn't flagged as Allow CPU Access. This asset needs to be flagged as such (in the Advanced section)."), *GetFullName());
 		return false;
 	}
-
-	// without editor data, we can't selectively generate a physics mesh for a given LOD index (we're missing access to GetSectionInfoMap()) so force bInUseAllTriData in order to use LOD index 0
-	bInUseAllTriData = true;
 #endif // #if !WITH_EDITORONLY_DATA
 
 #if WITH_EDITORONLY_DATA
@@ -7836,7 +7833,6 @@ bool UStaticMesh::GetPhysicsTriMeshDataCheckComplex(struct FTriMeshCollisionData
 	FStaticMeshLODResources& LOD = GetRenderData()->LODResources[UseLODIndex];
 
 	// Make sure the LOD we selected actually has CPU data
-	// NOTE: for non-editor builds we forced LOD0
 	if (!LOD.IndexBuffer.GetAllowCPUAccess())
 	{
 		UE_LOG(LogStaticMesh, Warning, TEXT("UStaticMesh::GetPhysicsTriMeshData: CPU data not available on selected LOD (UseLODIndex=%d, LODForCollision=%d) on '%s'."), UseLODIndex, LODForCollision, *GetFullName());
@@ -7859,12 +7855,7 @@ bool UStaticMesh::GetPhysicsTriMeshDataCheckComplex(struct FTriMeshCollisionData
 	{
 		const FStaticMeshSection& Section = LOD.Sections[SectionIndex];
 
-#if WITH_EDITORONLY_DATA
-		// we can only use GetSectionInfoMap() in WITH_EDITORONLY_DATA mode, otherwise, assume bInUseAllTriData :
-		if (bInUseAllTriData || GetSectionInfoMap().Get(UseLODIndex, SectionIndex).bEnableCollision)
-#else // #if WITH_EDITORONLY_DATA
-		check(bInUseAllTriData && bAllowCPUAccess);
-#endif // #if !WITH_EDITORONLY_DATA
+		if (bInUseAllTriData || SectionHasCollisionEnabled(Section, UseLODIndex, SectionIndex))
 		{
 			const uint32 OnePastLastIndex  = Section.FirstIndex + Section.NumTriangles*3;
 
@@ -7909,10 +7900,7 @@ bool UStaticMesh::GetTriMeshSizeEstimates(struct FTriMeshCollisionDataEstimates&
 		ComplexCollisionMesh->ConditionalPostLoad();
 		return ComplexCollisionMesh->GetTriMeshSizeEstimates(OutTriMeshEstimates, bInUseAllTriData);
 	}
-#else // #if WITH_EDITORONLY_DATA
-	// without editor data, we can't selectively generate a physics mesh for a given LOD index (we're missing access to GetSectionInfoMap()) so force bInUseAllTriData in order to use LOD index 0
-	bInUseAllTriData = true;
-#endif // #if !WITH_EDITORONLY_DATA
+#endif // #if WITH_EDITORONLY_DATA
 
 	if (GetRenderData() == nullptr || GetRenderData()->LODResources.Num() == 0)
 	{
@@ -7930,7 +7918,11 @@ bool UStaticMesh::GetTriMeshSizeEstimates(struct FTriMeshCollisionDataEstimates&
 	for (int32 SectionIndex = 0; SectionIndex < LOD.Sections.Num(); ++SectionIndex)
 	{
 		const FStaticMeshSection& Section = LOD.Sections[SectionIndex];
-		OutTriMeshEstimates.VerticeCount += Section.NumTriangles * 3;
+
+		if (bInUseAllTriData || SectionHasCollisionEnabled(Section, UseLODIndex, SectionIndex))
+		{
+			OutTriMeshEstimates.VerticeCount += Section.NumTriangles * 3;
+		}
 	}
 
 	return true;
@@ -7944,10 +7936,7 @@ bool UStaticMesh::ContainsPhysicsTriMeshDataCheckComplex(bool bInUseAllTriData,
 		ComplexCollisionMesh->ConditionalPostLoad();
 		return ComplexCollisionMesh->ContainsPhysicsTriMeshDataCheckComplex(bInUseAllTriData, false); // One level of recursion
 	}
-#else // #if WITH_EDITORONLY_DATA
-	// without editor data, we can't selectively generate a physics mesh for a given LOD index (we're missing access to GetSectionInfoMap()) so force bInUseAllTriData in order to use LOD index 0
-	bInUseAllTriData = true;
-#endif // #if !WITH_EDITORONLY_DATA
+#endif // #if WITH_EDITORONLY_DATA
 	
 #if WITH_EDITORONLY_DATA
 	// if we're a cooked cooker, just use the canned data
@@ -7966,30 +7955,36 @@ bool UStaticMesh::ContainsPhysicsTriMeshDataCheckComplex(bool bInUseAllTriData,
 	// Always use 0 if asking for 'all tri data'
 	const int32 UseLODIndex = bInUseAllTriData ? 0 : FMath::Clamp(LODForCollision, 0, GetRenderData()->LODResources.Num() - 1);
 
-	if (GetRenderData()->LODResources[UseLODIndex].VertexBuffers.PositionVertexBuffer.GetNumVertices() > 0)
+	const FStaticMeshLODResources& LOD = GetRenderData()->LODResources[UseLODIndex];
+	if (LOD.VertexBuffers.PositionVertexBuffer.GetNumVertices() > 0)
 	{
-		// Get the LOD level to use for collision
-		const FStaticMeshLODResources& LOD = GetRenderData()->LODResources[UseLODIndex];
-#if WITH_EDITORONLY_DATA
 		for (int32 SectionIndex = 0; SectionIndex < LOD.Sections.Num(); ++SectionIndex)
 		{
 			const FStaticMeshSection& Section = LOD.Sections[SectionIndex];
-			// we can only use GetSectionInfoMap() in WITH_EDITORONLY_DATA mode, otherwise, assume bInUseAllTriData :
-			if ((bInUseAllTriData || GetSectionInfoMap().Get(UseLODIndex, SectionIndex).bEnableCollision) && Section.NumTriangles > 0)
+			if (Section.NumTriangles == 0)
+			{
+				continue;
+			}
+
+			if (bInUseAllTriData || SectionHasCollisionEnabled(Section, UseLODIndex, SectionIndex))
 			{
 				return true;
 			}
 		}
-#else // #if WITH_EDITORONLY_DATA
-		if (LOD.Sections.Num() > 0)
-		{
-			return true;
-		}
-#endif // #if WITH_EDITORONLY_DATA
 	}
 	return false; 
 }
 
+bool UStaticMesh::SectionHasCollisionEnabled(const FStaticMeshSection& Section, int32 LODIndex, int32 SectionIndex) const
+{
+#if WITH_EDITORONLY_DATA
+	// we can only use GetSectionInfoMap() in WITH_EDITORONLY_DATA mode
+	return GetSectionInfoMap().Get(LODIndex, SectionIndex).bEnableCollision;
+#else // #if WITH_EDITORONLY_DATA
+	return Section.bEnableCollision;
+#endif // #if !WITH_EDITORONLY_DATA
+}
+
 bool UStaticMesh::PollAsyncPhysicsTriMeshData(bool InUseAllTriData) const
 {
 #if WITH_EDITORONLY_DATA
