[UE] AssetManager review suggestions: overloads and moving instead of forwarding
#rnx

diff --git a/Engine/Source/Runtime/Engine/Classes/Engine/AssetManager.h b/Engine/Source/Runtime/Engine/Classes/Engine/AssetManager.h
--- a/Engine/Source/Runtime/Engine/Classes/Engine/AssetManager.h
+++ b/Engine/Source/Runtime/Engine/Classes/Engine/AssetManager.h
@@ -357,7 +357,15 @@ public:
 	 */
 	template< typename DebugNameType = UE::FSourceLocation >
 	TSharedPtr<FStreamableHandle> LoadAssetList(
-		TArray<FSoftObjectPath> AssetList,
+		const TArray<FSoftObjectPath>& AssetList,
+		FStreamableDelegate DelegateToCall = FStreamableDelegate(),
+		TAsyncLoadPriority Priority = FStreamableManager::DefaultAsyncLoadPriority,
+		DebugNameType&& DebugNameOrLocation = UE::FSourceLocation::Current());
+
+	/** rvalue reference overload for Asset List */
+	template< typename DebugNameType = UE::FSourceLocation >
+	TSharedPtr<FStreamableHandle> LoadAssetList(
+		TArray<FSoftObjectPath>&& AssetList,
 		FStreamableDelegate DelegateToCall = FStreamableDelegate(),
 		TAsyncLoadPriority Priority = FStreamableManager::DefaultAsyncLoadPriority,
 		DebugNameType&& DebugNameOrLocation = UE::FSourceLocation::Current());
@@ -945,12 +953,12 @@ private:
 
 template< typename DebugNameType >
 TSharedPtr<FStreamableHandle> UAssetManager::LoadAssetList(
-	TArray<FSoftObjectPath> AssetList,
+	TArray<FSoftObjectPath>&& AssetList,
 	FStreamableDelegate DelegateToCall,
 	TAsyncLoadPriority Priority,
 	DebugNameType&& DebugNameOrLocation)
 {
-	if constexpr (std::is_same_v<std::remove_cv_t<DebugNameType>, UE::FSourceLocation>)
+	if constexpr (std::is_same_v<std::decay_t<DebugNameType>, UE::FSourceLocation>)
 	{
 		return LoadAssetListInternal(
 			MoveTemp(AssetList),
@@ -967,3 +975,19 @@ TSharedPtr<FStreamableHandle> UAssetManager::LoadAssetList(
 			FString{ Forward<DebugNameType>(DebugNameOrLocation) });
 	}
 }
+
+
+template< typename DebugNameType >
+TSharedPtr<FStreamableHandle> UAssetManager::LoadAssetList(
+	const TArray<FSoftObjectPath>& AssetList,
+	FStreamableDelegate DelegateToCall,
+	TAsyncLoadPriority Priority,
+	DebugNameType&& DebugNameOrLocation)
+{
+	// explicit copy
+	return LoadAssetList(
+		TArray<FSoftObjectPath>{AssetList},
+		MoveTemp(DelegateToCall),
+		Priority,
+		Forward<DebugNameType>(DebugNameOrLocation));
+}
diff --git a/Engine/Source/Runtime/Engine/Classes/Engine/StreamableManager.h b/Engine/Source/Runtime/Engine/Classes/Engine/StreamableManager.h
--- a/Engine/Source/Runtime/Engine/Classes/Engine/StreamableManager.h
+++ b/Engine/Source/Runtime/Engine/Classes/Engine/StreamableManager.h
@@ -699,7 +699,7 @@ TSharedPtr<FStreamableHandle> FStreamableManager::RequestAsyncLoad(
 		DelegateToCall = FStreamableDelegate::CreateLambda(Forward<FuncType>(Callback));
 	}
 
-	if constexpr (std::is_same_v<std::remove_cv_t<DebugNameType>, UE::FSourceLocation>)
+	if constexpr (std::is_same_v<std::decay_t<DebugNameType>, UE::FSourceLocation>)
 	{
 		return RequestAsyncLoadInternal(
 			TArray<FSoftObjectPath>{ Forward<PathContainerType>(TargetsToStream) },
@@ -727,7 +727,7 @@ TSharedPtr<FStreamableHandle> FStreamableManager::RequestSyncLoad(
 	bool bManageActiveHandle,
 	DebugNameType&& DebugNameOrLocation)
 {
-	if constexpr (std::is_same_v<std::remove_cv_t<DebugNameType>, UE::FSourceLocation>)
+	if constexpr (std::is_same_v<std::decay_t<DebugNameType>, UE::FSourceLocation>)
 	{
 		return RequestSyncLoadInternal(
 			TArray<FSoftObjectPath>{ Forward<PathContainerType>(TargetsToStream) },
diff --git a/Engine/Source/Runtime/Engine/Private/AssetManager.cpp b/Engine/Source/Runtime/Engine/Private/AssetManager.cpp
--- a/Engine/Source/Runtime/Engine/Private/AssetManager.cpp
+++ b/Engine/Source/Runtime/Engine/Private/AssetManager.cpp
@@ -2240,12 +2240,12 @@ TSharedPtr<FStreamableHandle> UAssetManager::LoadAssetListInternal(TArray<FSoftO
 	// SynchronousLoad doesn't make sense if chunks are missing
 	if (bShouldUseSynchronousLoad && MissingChunks.Num() == 0)
 	{
-		NewHandle = StreamableManager.RequestSyncLoad(Forward<TArray<FSoftObjectPath>>(AssetList), false, Forward<FString>(DebugName));
-		FStreamableHandle::ExecuteDelegate(Forward<FStreamableDelegate>(DelegateToCall));
+		NewHandle = StreamableManager.RequestSyncLoad(MoveTemp(AssetList), false, MoveTemp(DebugName));
+		FStreamableHandle::ExecuteDelegate(MoveTemp(DelegateToCall));
 	}
 	else
 	{
-		NewHandle = StreamableManager.RequestAsyncLoad(Forward<TArray<FSoftObjectPath>>(AssetList), Forward<FStreamableDelegate>(DelegateToCall), Priority, false, MissingChunks.Num() > 0, Forward<FString>(DebugName));
+		NewHandle = StreamableManager.RequestAsyncLoad(MoveTemp(AssetList), MoveTemp(DelegateToCall), Priority, false, MissingChunks.Num() > 0, MoveTemp(DebugName));
 
 		if (MissingChunks.Num() > 0 && NewHandle.IsValid())
 		{
diff --git a/Engine/Source/Runtime/Engine/Private/StreamableManager.cpp b/Engine/Source/Runtime/Engine/Private/StreamableManager.cpp
--- a/Engine/Source/Runtime/Engine/Private/StreamableManager.cpp
+++ b/Engine/Source/Runtime/Engine/Private/StreamableManager.cpp
@@ -1407,10 +1407,10 @@ TSharedPtr<FStreamableHandle> FStreamableManager::RequestAsyncLoadInternal(TArra
 
 	// Schedule a new callback, this will get called when all related async loads are completed
 	TSharedRef<FStreamableHandle> NewRequest = MakeShareable(new FStreamableHandle());
-	NewRequest->CompleteDelegate = Forward<FStreamableDelegate>(DelegateToCall);
+	NewRequest->CompleteDelegate = MoveTemp(DelegateToCall);
 	NewRequest->OwningManager = this;
-	NewRequest->RequestedAssets = Forward<TArray<FSoftObjectPath>>(TargetsToStream);
-	NewRequest->DebugName = Forward<FString>(DebugName);
+	NewRequest->RequestedAssets = MoveTemp(TargetsToStream);
+	NewRequest->DebugName = MoveTemp(DebugName);
 
 	NewRequest->Priority = Priority;
 #if UE_WITH_PACKAGE_ACCESS_TRACKING
@@ -1523,7 +1523,7 @@ TSharedPtr<FStreamableHandle> FStreamableManager::RequestSyncLoadInternal(TArray
 	bForceSynchronousLoads = IsInAsyncLoadingThread() || IsEventDrivenLoaderEnabled() || !IsAsyncLoading();
 
 	// Do an async load and wait to complete. In some cases this will do a sync load due to safety issues
-	TSharedPtr<FStreamableHandle> Request = RequestAsyncLoadInternal(Forward<TArray<FSoftObjectPath>>(TargetsToStream), FStreamableDelegate(), AsyncLoadHighPriority, bManageActiveHandle, false, Forward<FString>(DebugName));
+	TSharedPtr<FStreamableHandle> Request = RequestAsyncLoadInternal(MoveTemp(TargetsToStream), FStreamableDelegate(), AsyncLoadHighPriority, bManageActiveHandle, false, MoveTemp(DebugName));
 
 	bForceSynchronousLoads = false;
 
