Fixed problematic namespace

diff --git a/Engine/Source/Runtime/Engine/Classes/Components/SplineMeshComponent.h b/Engine/Source/Runtime/Engine/Classes/Components/SplineMeshComponent.h
--- a/Engine/Source/Runtime/Engine/Classes/Components/SplineMeshComponent.h
+++ b/Engine/Source/Runtime/Engine/Classes/Components/SplineMeshComponent.h
@@ -239,7 +239,7 @@ public:
 public:
 	ENGINE_API virtual class FStaticMeshStaticLightingMesh* AllocateStaticLightingMesh(int32 LODIndex, const TArray<ULightComponent*>& InRelevantLights) override;
 protected:
-	ENGINE_API virtual FPrimitiveSceneProxy* CreateStaticMeshSceneProxy(Nanite::FMaterialAudit& NaniteMaterials, bool bCreateNanite) override;
+	ENGINE_API virtual FPrimitiveSceneProxy* CreateStaticMeshSceneProxy(::Nanite::FMaterialAudit& NaniteMaterials, bool bCreateNanite) override;
 	//End UStaticMeshComponent Interface
 
 	//~ Begin Interface_CollisionDataProvider Interface
