Skip FRayTracingGeometryManager::Tick(...) when running commandlet.

diff --git a/Engine/Source/Runtime/Engine/Private/Rendering/RayTracingGeometryManager.cpp b/Engine/Source/Runtime/Engine/Private/Rendering/RayTracingGeometryManager.cpp
--- a/Engine/Source/Runtime/Engine/Private/Rendering/RayTracingGeometryManager.cpp
+++ b/Engine/Source/Runtime/Engine/Private/Rendering/RayTracingGeometryManager.cpp
@@ -338,6 +338,11 @@ void FRayTracingGeometryManager::PreRender()
 
 void FRayTracingGeometryManager::Tick(FRHICommandList& RHICmdList)
 {
+	if (IsRunningCommandlet())
+	{
+		return;
+	}
+
 	check(IsInRenderingThread());
 
 	TRACE_CPUPROFILER_EVENT_SCOPE(FRayTracingGeometryManager::Tick);
