Added redirector support to FExternalPackageHelper::LoadObjectsFromExternalPackages

External packages that have objects with classes that were previously redirected will now be correctly loaded
We skip the redirector handing in the case of 'any UObject', as that will be caught by the regular bRecursiveClasses case.

diff --git a/Engine/Source/Runtime/Engine/Public/ExternalPackageHelper.h b/Engine/Source/Runtime/Engine/Public/ExternalPackageHelper.h
--- a/Engine/Source/Runtime/Engine/Public/ExternalPackageHelper.h
+++ b/Engine/Source/Runtime/Engine/Public/ExternalPackageHelper.h
@@ -19,6 +19,7 @@
 #include "WorldPartition/DataLayer/ExternalDataLayerAsset.h"
 #include "WorldPartition/DataLayer/ExternalDataLayerHelper.h"
 #include "Engine/Level.h"
+#include "UObject/CoreRedirects.h"
 
 class FExternalPackageHelper
 {
@@ -154,10 +155,32 @@ void FExternalPackageHelper::LoadObjectsFromExternalPackages(UObject* InOuter, T
 	IAssetRegistry& AssetRegistry = FModuleManager::LoadModuleChecked<FAssetRegistryModule>(TEXT("AssetRegistry")).Get();
 	AssetRegistry.ScanSynchronous({ ExternalObjectsPath }, TArray<FString>());
 
+	// Find any redirects for this class and any of its derived classes so we can load older external objects (asset registry class does not redirect)
+	TArray<FTopLevelAssetPath> ClassPaths;
+	ClassPaths.Add(T::StaticClass()->GetClassPathName());
+
+	if(T::StaticClass() != UObject::StaticClass())
+	{
+		TArray<UClass*> DerivedClasses;
+		DerivedClasses.Add(T::StaticClass());
+		GetDerivedClasses(T::StaticClass(), DerivedClasses);
+
+		for(UClass* Class : DerivedClasses)
+		{
+			TArray<FCoreRedirectObjectName> PreviousRedirectedNames;
+			FCoreRedirects::FindPreviousNames(ECoreRedirectFlags::Type_Class, FCoreRedirectObjectName(Class->GetClassPathName()), PreviousRedirectedNames);
+
+			for(const FCoreRedirectObjectName& PreviousRedirectedName : PreviousRedirectedNames)
+			{
+				ClassPaths.Add(FTopLevelAssetPath(PreviousRedirectedName.PackageName, PreviousRedirectedName.ObjectName));
+			}
+		}
+	}
+
 	FARFilter Filter;
 	Filter.bRecursivePaths = true;
 	Filter.bIncludeOnlyOnDiskAssets = true;
-	Filter.ClassPaths.Add(T::StaticClass()->GetClassPathName());
+	Filter.ClassPaths = MoveTemp(ClassPaths);
 	Filter.bRecursiveClasses = true;
 	Filter.PackagePaths.Add(*ExternalObjectsPath);
 	TArray<FAssetData> Assets;
