[Materials] Avoid unnecessary copy when passing material shader code to virtual include map.

Generated material HLSL shader code is stored in usually large FString objects.
Using MoveTemp() here when passing this string to the virtual include map avoids an unnecessary copy of a potentially large string.

#rnx

diff --git a/Engine/Source/Runtime/Engine/Private/Materials/MaterialShared.cpp b/Engine/Source/Runtime/Engine/Private/Materials/MaterialShared.cpp
--- a/Engine/Source/Runtime/Engine/Private/Materials/MaterialShared.cpp
+++ b/Engine/Source/Runtime/Engine/Private/Materials/MaterialShared.cpp
@@ -3380,9 +3380,10 @@ bool FMaterial::Translate_Legacy(const FMaterialShaderMapId& ShaderMapId,
 	OutMaterialEnvironment = new FSharedShaderCompilerEnvironment();
 	OutMaterialEnvironment->TargetPlatform = InTargetPlatform;
 	MaterialTranslator.GetMaterialEnvironment(InPlatform, *OutMaterialEnvironment);
-	const FString MaterialShaderCode = MaterialTranslator.GetMaterialShaderCode();
 
-	OutMaterialEnvironment->IncludeVirtualPathToContentsMap.Add(TEXT("/Engine/Generated/Material.ush"), MaterialShaderCode);
+	// Add generated HLSL shader code to virtual include map to be included by the respective base shader (e.g. BasePassPixelShader.usf)
+	FString MaterialShaderCode = MaterialTranslator.GetMaterialShaderCode();
+	OutMaterialEnvironment->IncludeVirtualPathToContentsMap.Add(TEXT("/Engine/Generated/Material.ush"), MoveTemp(MaterialShaderCode));
 	
 	return true;
 }
