Adding metadata to texture compiler to try and help diagnose all these async ddc key modifications.
#rnx

diff --git a/Engine/Source/Runtime/Engine/Private/TextureCompiler.cpp b/Engine/Source/Runtime/Engine/Private/TextureCompiler.cpp
--- a/Engine/Source/Runtime/Engine/Private/TextureCompiler.cpp
+++ b/Engine/Source/Runtime/Engine/Private/TextureCompiler.cpp
@@ -205,6 +205,20 @@ void FTextureCompilingManager::PostCompilation(UTexture* Texture)
 	UE_LOG(LogTexture, Verbose, TEXT("Refreshing texture %s because it is ready"), *Texture->GetName());
 
 	Texture->FinishCachePlatformData();
+
+	// Track the DDC key suffix of the texture we are done with so that if we re-enter we can
+	// log info and hopefully be able to do some post-mortem on it.
+	CurrentPostCompilationTexture = Texture;
+	CurrentPostCompilationDDCKey.Empty();
+	if (FTexturePlatformData** RunningPlatformData = Texture->GetRunningPlatformData(); RunningPlatformData && *RunningPlatformData)
+	{
+		// only works for ddc1 right now... 
+		if (FString* DDCKey = RunningPlatformData[0]->DerivedDataKey.TryGet<FString>(); DDCKey)
+		{
+			CurrentPostCompilationDDCKey = *DDCKey;
+		}
+	}
+
 	Texture->UpdateResource();
 
 	// Generate an empty property changed event, to force the asset registry tag
@@ -251,10 +265,9 @@ void FTextureCompilingManager::AddTextures(TArrayView<UTexture* const> InTexture
 	// and then immediately tried again - and then tried to launch another build because the ddc keys changed. This means that
 	// during the async build, a property or otherwise that is an input to the ddc key changed. This shouldn't happen because
 	// PreEditChange completes the async build before allowing the change.
-	// Debugging this can be a huge pain. If you have a repro, IMO the best way is to hack GetTextureDerivedDataKeySuffix
-	// to strcmp on the name of the repro texture and just log the full key suffix. Then you should immediately see the changed
-	// keys right before the crash and you can backsolve what value changed. Once you have that, you can set a data breakpoint on
-	// the property and see who is poking it.
+	// Debugging this can be a huge pain. NEW AND IMPROVED: We should now be printing the relevant DDC keys below (if DDC1). This
+	// should facilitate at least finding out what property is getting changed, as well as what texture. If you can't divine what's
+	// causing the change from that, you'll need to put a data breakpoint on it and see who is doing it.
 	//
 	// **
 	//
@@ -266,6 +279,33 @@ void FTextureCompilingManager::AddTextures(TArrayView<UTexture* const> InTexture
 	// in question - only the compilation manager shoulid be calling that for editor resources.
 	if (bIsRoutingPostCompilation)
 	{
+		UE_LOG(LogTexture, Error, TEXT("PostCompilation Texture: %s"), CurrentPostCompilationTexture ? *CurrentPostCompilationTexture->GetPathName() : TEXT("<nullptr>"));
+
+		// Empty keys most likely means we are on ddc2.
+		UE_LOG(LogTexture, Error, TEXT("PostCompilation DDCKey: %s"), *CurrentPostCompilationDDCKey);
+		UE_LOG(LogTexture, Error, TEXT("AddTextures Count: %d"), InTextures.Num());
+
+		for (const UTexture* ConstTexture : InTextures)
+		{
+			// We're about to crash anyway.
+			UTexture* Texture = const_cast<UTexture*>(ConstTexture);
+
+			UE_LOG(LogTexture, Error, TEXT("%s:"), *Texture->GetPathName());
+
+			FTexturePlatformData** RunningPlatformData = Texture->GetRunningPlatformData();
+			if (!RunningPlatformData || !RunningPlatformData[0])
+			{
+				UE_LOG(LogTexture, Error, TEXT("   -> No RunningPlatformData!"));
+			}
+			else
+			{
+				FString* Key = RunningPlatformData[0]->FetchFirstDerivedDataKey.TryGet<FString>();
+				UE_LOG(LogTexture, Error, TEXT("    FetchFirstKey: %s"), Key ? **Key : TEXT("<empty, likely new texture build flow?>"));
+				Key = RunningPlatformData[0]->FetchOrBuildDerivedDataKey.TryGet<FString>();
+				UE_LOG(LogTexture, Error, TEXT("    FetchOrBuildKey: %s"), Key ? **Key : TEXT("<empty, likely new texture build flow?>"));
+			}
+		}
+
 		// This has been updated to Fatal because it potentially modifies RegisteredTextureBuckets below which is iterated upon
 		// during PostCompilation routing. That modification can put us in an unstable state and crash in unexpected and rather
 		// undebuggable ways.
diff --git a/Engine/Source/Runtime/Engine/Public/TextureCompiler.h b/Engine/Source/Runtime/Engine/Public/TextureCompiler.h
--- a/Engine/Source/Runtime/Engine/Public/TextureCompiler.h
+++ b/Engine/Source/Runtime/Engine/Public/TextureCompiler.h
@@ -108,6 +108,9 @@ private:
 
 	void ProcessDeferredRequests();
 
+
+	FString CurrentPostCompilationDDCKey;
+	UTexture* CurrentPostCompilationTexture = nullptr;
 	double LastReschedule = 0.0f;
 	bool bHasShutdown = false;
 	bool bIsRoutingPostCompilation = false;
