SceneViewExtensions: Standardize view argument name.

Minor naming adjustment for consistency.

#rnx

diff --git a/Engine/Source/Runtime/Engine/Public/SceneViewExtension.h b/Engine/Source/Runtime/Engine/Public/SceneViewExtension.h
--- a/Engine/Source/Runtime/Engine/Public/SceneViewExtension.h
+++ b/Engine/Source/Runtime/Engine/Public/SceneViewExtension.h
@@ -178,17 +178,14 @@ public:
 	/**
 	 * Called right before Post Processing rendering begins
 	 */
-	virtual void PrePostProcessPass_RenderThread(FRDGBuilder& GraphBuilder, const FSceneView& View, const FPostProcessingInputs& Inputs) {};
+	virtual void PrePostProcessPass_RenderThread(FRDGBuilder& GraphBuilder, const FSceneView& InView, const FPostProcessingInputs& Inputs) {};
 
 	/**
 	* This will be called at the beginning of post processing to make sure that each view extension gets a chance to subscribe to an after pass event.
 	*  - The pass MUST write to the override output texture if it is active (this occurs when the pass is the last in the post processing chain writing to the back buffer).
 	*    For performance reasons it is recommended to only subscribe to a pass when the pass will produce a GPU resource. Calling 
 	*/
-	virtual void SubscribeToPostProcessingPass(EPostProcessingPass Pass, const FSceneView& View, FAfterPassCallbackDelegateArray& InOutPassCallbacks, bool bIsPassEnabled) {}
-
-	UE_DEPRECATED(5.5, "SubscribeToPostProcessingPass now takes a SceneView")
-	virtual void SubscribeToPostProcessingPass(EPostProcessingPass Pass, FAfterPassCallbackDelegateArray& InOutPassCallbacks, bool bIsPassEnabled) {}
+	virtual void SubscribeToPostProcessingPass(EPostProcessingPass Pass, const FSceneView& InView, FAfterPassCallbackDelegateArray& InOutPassCallbacks, bool bIsPassEnabled) {}
 
 	/**
 	 * Allows to render content after the 3D content scene, useful for debugging
@@ -228,6 +225,9 @@ public:
 	UE_DEPRECATED(5.1, "PostRenderBasePass now has a mobile and deferred variant.")
 	virtual void PostRenderBasePass_RenderThread(FRHICommandListImmediate& RHICmdList, FSceneView& InView) {};
 
+	UE_DEPRECATED(5.5, "SubscribeToPostProcessingPass now takes a SceneView")
+	virtual void SubscribeToPostProcessingPass(EPostProcessingPass Pass, FAfterPassCallbackDelegateArray& InOutPassCallbacks, bool bIsPassEnabled) {}
+
 	///////////////////////////////////////////////////////////////////////////////////////
 
 protected:
