Handle case where TesselateRecursive reaches maximum recursion.
If spline still isn't represented by line segments within threshold by MaxRecursion levels of recursion, return the current points instead of nothing.
This prevents a crash where FlattenSplineSegements (ZoneShapeUtilities.cpp) assumes at least one point is produced by UE::CubicBezier::Tessellate.

diff --git a/Engine/Source/Runtime/Engine/Private/Curves/BezierUtilities.cpp b/Engine/Source/Runtime/Engine/Private/Curves/BezierUtilities.cpp
--- a/Engine/Source/Runtime/Engine/Private/Curves/BezierUtilities.cpp
+++ b/Engine/Source/Runtime/Engine/Private/Curves/BezierUtilities.cpp
@@ -164,6 +164,12 @@ float ArcLengthApproximate(const FVector& P0, const FVector& P1, const FVector&
 
 void TessellateRecursive(TArray<FVector>& Output, const FVector& P0, const FVector& P1, const FVector& P2, const FVector& P3, const float ToleranceSqr, int Level, const int MaxLevel)
 {
+	if (Level >= MaxLevel)
+	{
+		Output.Add(P3);
+		return;
+	}
+
 	// Handle degenerate segment.
 	FVector Dir = P3 - P0;
 	if (Dir.IsNearlyZero())
@@ -186,19 +192,16 @@ void TessellateRecursive(TArray<FVector>& Output, const FVector& P0, const FVect
 		return;
 	}
 
-	if (Level < MaxLevel)
-	{
-		// Split the curve in half and recurse.
-		const FVector P01 = FMath::Lerp(P0, P1, 0.5f);
-		const FVector P12 = FMath::Lerp(P1, P2, 0.5f);
-		const FVector P23 = FMath::Lerp(P2, P3, 0.5f);
-		const FVector P012 = FMath::Lerp(P01, P12, 0.5f);
-		const FVector P123 = FMath::Lerp(P12, P23, 0.5f);
-		const FVector P0123 = FMath::Lerp(P012, P123, 0.5f);
-
-		TessellateRecursive(Output, P0, P01, P012, P0123, ToleranceSqr, Level + 1, MaxLevel);
-		TessellateRecursive(Output, P0123, P123, P23, P3, ToleranceSqr, Level + 1, MaxLevel);
-	}
+	// Split the curve in half and recurse.
+	const FVector P01 = FMath::Lerp(P0, P1, 0.5f);
+	const FVector P12 = FMath::Lerp(P1, P2, 0.5f);
+	const FVector P23 = FMath::Lerp(P2, P3, 0.5f);
+	const FVector P012 = FMath::Lerp(P01, P12, 0.5f);
+	const FVector P123 = FMath::Lerp(P12, P23, 0.5f);
+	const FVector P0123 = FMath::Lerp(P012, P123, 0.5f);
+
+	TessellateRecursive(Output, P0, P01, P012, P0123, ToleranceSqr, Level + 1, MaxLevel);
+	TessellateRecursive(Output, P0123, P123, P23, P3, ToleranceSqr, Level + 1, MaxLevel);
 }
 
 void Tessellate(TArray<FVector>& Output, const FVector& P0, const FVector& P1, const FVector& P2, const FVector& P3, const float Tolerance, const int MaxLevel)
