Fix cached ray tracing state not being updated when LODs are streamed out if using r.RayTracing.UseReferenceBasedResidency=True

diff --git a/Engine/Source/Runtime/Engine/Private/StaticMesh.cpp b/Engine/Source/Runtime/Engine/Private/StaticMesh.cpp
--- a/Engine/Source/Runtime/Engine/Private/StaticMesh.cpp
+++ b/Engine/Source/Runtime/Engine/Private/StaticMesh.cpp
@@ -2494,11 +2494,11 @@ void UStaticMesh::RequestUpdateCachedRenderState() const
 		{
 			((FRayTracingGeometryManager*)GRayTracingGeometryManager)->SetRayTracingGeometryGroupCurrentFirstLODIndex(FRHICommandListImmediate::Get(), GetRenderData()->RayTracingGeometryGroupHandle, GetRenderData()->CurrentFirstLODIdx);
 		}
-		else
-		{
-			// TODO: this should only be necessary when a BLAS build was not requested (ie: non-compressed offline BLAS)
-			GRayTracingGeometryManager->RequestUpdateCachedRenderState(GetRenderData()->RayTracingGeometryGroupHandle);
-		}
+
+		// TODO: this could be skipped in some cases
+		// - compressed offline BLAS (since build request handle it)
+		// - FStaticMeshLODResources doesn't own the ray tracing geometry
+		GRayTracingGeometryManager->RequestUpdateCachedRenderState(GetRenderData()->RayTracingGeometryGroupHandle);
 	}
 #endif
 
