[Material Editor] Fixed OutputsTranslucentVelocity flag to use Instance override.

Before this change the override of OutputsTranslucentVelocity in Material Instances was ignored and only the override set in the parent Material was considered.

diff --git a/Engine/Source/Runtime/Engine/Private/Materials/MaterialInterface.cpp b/Engine/Source/Runtime/Engine/Private/Materials/MaterialInterface.cpp
--- a/Engine/Source/Runtime/Engine/Private/Materials/MaterialInterface.cpp
+++ b/Engine/Source/Runtime/Engine/Private/Materials/MaterialInterface.cpp
@@ -430,108 +430,106 @@ bool UMaterialInterface::GetStaticComponentMaskParameterValue(const FHashedMater
 
 FMaterialRelevance UMaterialInterface::GetRelevance_Internal(const UMaterial* Material, ERHIFeatureLevel::Type InFeatureLevel) const
 {
-	if(Material)
+	if (!Material)
 	{
-		const FMaterialResource* MaterialResource = GetMaterialResource(InFeatureLevel);
+		return FMaterialRelevance();
+	}
 
-		// If material is invalid e.g. unparented instance, fallback to the passed in material
-		bool bIsValidMaterialResource = (MaterialResource != nullptr) && (MaterialResource->GetMaterial() != nullptr);
-		if (!bIsValidMaterialResource && (Material != nullptr))
-		{
-			MaterialResource = Material->GetMaterialResource(InFeatureLevel);
-		}
+	const FMaterialResource* MaterialResource = GetMaterialResource(InFeatureLevel);
 
-		if (MaterialResource == nullptr)
-		{
-			return FMaterialRelevance();
-		}
+	// If material is invalid e.g. unparented instance, fallback to the passed in material
+	bool bIsValidMaterialResource = (MaterialResource != nullptr) && (MaterialResource->GetMaterial() != nullptr);
+	if (!bIsValidMaterialResource && (Material != nullptr))
+	{
+		MaterialResource = Material->GetMaterialResource(InFeatureLevel);
+	}
 
-		const bool bIsMobile = InFeatureLevel <= ERHIFeatureLevel::ES3_1;
-		const bool bUsesSingleLayerWaterMaterial = MaterialResource->GetShadingModels().HasShadingModel(MSM_SingleLayerWater);
-		const bool bIsSinglePassWaterTranslucent = bIsMobile && bUsesSingleLayerWaterMaterial;
-		const bool bIsMobilePixelProjectedTranslucent = MaterialResource->IsUsingPlanarForwardReflections() 
-														&& IsUsingMobilePixelProjectedReflection(GetFeatureLevelShaderPlatform(InFeatureLevel));
+	if (MaterialResource == nullptr)
+	{
+		return FMaterialRelevance();
+	}
 
-		// Note that even though XX_GameThread() api is called, this function can be called on non game thread via 
-		// GetRelevance_Concurrent()
-		bool bUsesAnisotropy = MaterialResource->GetShadingModels().HasAnyShadingModel({ MSM_DefaultLit, MSM_ClearCoat }) && 
-			MaterialResource->MaterialUsesAnisotropy_GameThread();
+	const bool bIsMobile = InFeatureLevel <= ERHIFeatureLevel::ES3_1;
+	const bool bUsesSingleLayerWaterMaterial = MaterialResource->GetShadingModels().HasShadingModel(MSM_SingleLayerWater);
+	const bool bIsSinglePassWaterTranslucent = bIsMobile && bUsesSingleLayerWaterMaterial;
+	const bool bIsMobilePixelProjectedTranslucent = MaterialResource->IsUsingPlanarForwardReflections() 
+													&& IsUsingMobilePixelProjectedReflection(GetFeatureLevelShaderPlatform(InFeatureLevel));
 
-		const EBlendMode BlendMode = (EBlendMode)GetBlendMode();
-		const bool bIsTranslucent = IsTranslucentBlendMode(BlendMode) || bIsSinglePassWaterTranslucent || bIsMobilePixelProjectedTranslucent; // We want meshes with water materials to be scheduled for translucent pass on mobile. And we also have to render the meshes used for mobile pixel projection reflection in translucent pass.
+	// Note that even though XX_GameThread() api is called, this function can be called on non game thread via 
+	// GetRelevance_Concurrent()
+	bool bUsesAnisotropy = MaterialResource->GetShadingModels().HasAnyShadingModel({ MSM_DefaultLit, MSM_ClearCoat }) && 
+		MaterialResource->MaterialUsesAnisotropy_GameThread();
 
-		EMaterialDomain Domain = (EMaterialDomain)MaterialResource->GetMaterialDomain();
-		bool bDecal = (Domain == MD_DeferredDecal);
+	const EBlendMode BlendMode = (EBlendMode)GetBlendMode();
+	const bool bIsTranslucent = IsTranslucentBlendMode(BlendMode) || bIsSinglePassWaterTranslucent || bIsMobilePixelProjectedTranslucent; // We want meshes with water materials to be scheduled for translucent pass on mobile. And we also have to render the meshes used for mobile pixel projection reflection in translucent pass.
 
-		// Determine the material's view relevance.
-		FMaterialRelevance MaterialRelevance;
+	EMaterialDomain Domain = (EMaterialDomain)MaterialResource->GetMaterialDomain();
+	bool bDecal = (Domain == MD_DeferredDecal);
 
-		MaterialRelevance.ShadingModelMask = GetShadingModels().GetShadingModelField();
-		MaterialRelevance.CustomDepthStencilUsageMask = MaterialResource->GetCustomDepthStencilUsageMask_GameThread();
-		MaterialRelevance.bDecal = bDecal;
+	// Determine the material's view relevance.
+	FMaterialRelevance MaterialRelevance;
 
-		// Check whether the material can be drawn in the separate translucency pass as per FMaterialResource::IsTranslucencyAfterDOFEnabled and IsMobileSeparateTranslucencyEnabled
-		EMaterialTranslucencyPass TranslucencyPass = MTP_BeforeDOF;
-		const bool bSupportsSeparateTranslucency = Material->MaterialDomain != MD_UI && Material->MaterialDomain != MD_DeferredDecal;
-		if (bIsTranslucent && bSupportsSeparateTranslucency)
+	MaterialRelevance.ShadingModelMask = GetShadingModels().GetShadingModelField();
+	MaterialRelevance.CustomDepthStencilUsageMask = MaterialResource->GetCustomDepthStencilUsageMask_GameThread();
+	MaterialRelevance.bDecal = bDecal;
+
+	// Check whether the material can be drawn in the separate translucency pass as per FMaterialResource::IsTranslucencyAfterDOFEnabled and IsMobileSeparateTranslucencyEnabled
+	EMaterialTranslucencyPass TranslucencyPass = MTP_BeforeDOF;
+	const bool bSupportsSeparateTranslucency = Material->MaterialDomain != MD_UI && Material->MaterialDomain != MD_DeferredDecal;
+	if (bIsTranslucent && bSupportsSeparateTranslucency)
+	{
+		if (bIsMobile)
 		{
-			if (bIsMobile)
+			if (Material->bEnableMobileSeparateTranslucency)
 			{
-				if (Material->bEnableMobileSeparateTranslucency)
-				{
-					TranslucencyPass = MTP_AfterDOF;
-				}
-			}
-			else
-			{
-				TranslucencyPass = Material->TranslucencyPass;
+				TranslucencyPass = MTP_AfterDOF;
 			}
-		}			
-
-		// If dual blending is supported, and we are rendering post-DOF translucency, then we also need to render a second pass to the modulation buffer.
-		// The modulation buffer can also be used for regular modulation shaders after DoF.
-		const bool bMaterialSeparateModulation = MaterialResource->IsDualBlendingEnabled(GShaderPlatformForFeatureLevel[InFeatureLevel]) || IsModulateBlendMode(BlendMode);
-
-		// Encode Substrate BSDF into a mask where each bit correspond to a number of BSDF (1-8)
-		const uint8 SubstrateBSDFCount = FMath::Max(MaterialResource->MaterialGetSubstrateClosureCount_GameThread(), uint8(1u));
-		const uint8 SubstrateBSDFCountMask = 1u << uint8(FMath::Min(SubstrateBSDFCount - 1, 8));
-		const uint8 SubstrateUintPerPixel = FMath::Max(MaterialResource->MaterialGetSubstrateUintPerPixel_GameThread(), uint8(1u));
-
-		MaterialRelevance.bOpaque = !bIsTranslucent;
-		MaterialRelevance.bMasked = IsMasked();
-		MaterialRelevance.bDistortion = MaterialResource->IsDistorted();
-		MaterialRelevance.bHairStrands = IsCompatibleWithHairStrands(MaterialResource, InFeatureLevel);
-		MaterialRelevance.bTwoSided = MaterialResource->IsTwoSided();
-		MaterialRelevance.bSeparateTranslucency = bIsTranslucent && (TranslucencyPass == MTP_AfterDOF);
-		MaterialRelevance.bTranslucencyModulate = bMaterialSeparateModulation;
-		MaterialRelevance.bPostMotionBlurTranslucency = (TranslucencyPass == MTP_AfterMotionBlur);
-		MaterialRelevance.bNormalTranslucency = bIsTranslucent && (TranslucencyPass == MTP_BeforeDOF);
-		MaterialRelevance.bDisableDepthTest = bIsTranslucent && Material->bDisableDepthTest;		
-		MaterialRelevance.bUsesSceneColorCopy = bIsTranslucent && MaterialResource->RequiresSceneColorCopy_GameThread();
-		MaterialRelevance.bOutputsTranslucentVelocity = Material->IsTranslucencyWritingVelocity();
-		MaterialRelevance.bUsesGlobalDistanceField = MaterialResource->UsesGlobalDistanceField_GameThread();
-		MaterialRelevance.bUsesWorldPositionOffset = MaterialResource->MaterialUsesWorldPositionOffset_GameThread();
-		MaterialRelevance.bUsesDisplacement = MaterialResource->MaterialUsesDisplacement_GameThread();
-		MaterialRelevance.bUsesPixelDepthOffset = MaterialResource->MaterialUsesPixelDepthOffset_GameThread();
-		ETranslucencyLightingMode TranslucencyLightingMode = MaterialResource->GetTranslucencyLightingMode();
-		MaterialRelevance.bTranslucentSurfaceLighting = bIsTranslucent && (TranslucencyLightingMode == TLM_SurfacePerPixelLighting || TranslucencyLightingMode == TLM_Surface);
-		MaterialRelevance.bUsesSceneDepth = MaterialResource->MaterialUsesSceneDepthLookup_GameThread();
-		MaterialRelevance.bHasVolumeMaterialDomain = MaterialResource->IsVolumetricPrimitive();
-		MaterialRelevance.bUsesDistanceCullFade = MaterialResource->MaterialUsesDistanceCullFade_GameThread();
-		MaterialRelevance.bUsesSkyMaterial = Material->bIsSky;
-		MaterialRelevance.bUsesSingleLayerWaterMaterial = bUsesSingleLayerWaterMaterial;
-		MaterialRelevance.bUsesAnisotropy = bUsesAnisotropy;
-		MaterialRelevance.SubstrateClosureCountMask = SubstrateBSDFCountMask;
-		MaterialRelevance.SubstrateUintPerPixel = SubstrateUintPerPixel;
-		MaterialRelevance.bUsesComplexSpecialRenderPath = MaterialResource->MaterialGetSubstrateUsesComplexSpecialRenderPath_GameThread();
-		MaterialRelevance.bIsLightFunctionAtlasCompatible = MaterialResource->MaterialIsLightFunctionAtlasCompatible_GameThread();
-
-		return MaterialRelevance;
-	}
-	else
-	{
-		return FMaterialRelevance();
-	}
+		}
+		else
+		{
+			TranslucencyPass = Material->TranslucencyPass;
+		}
+	}			
+
+	// If dual blending is supported, and we are rendering post-DOF translucency, then we also need to render a second pass to the modulation buffer.
+	// The modulation buffer can also be used for regular modulation shaders after DoF.
+	const bool bMaterialSeparateModulation = MaterialResource->IsDualBlendingEnabled(GShaderPlatformForFeatureLevel[InFeatureLevel]) || IsModulateBlendMode(BlendMode);
+
+	// Encode Substrate BSDF into a mask where each bit correspond to a number of BSDF (1-8)
+	const uint8 SubstrateBSDFCount = FMath::Max(MaterialResource->MaterialGetSubstrateClosureCount_GameThread(), uint8(1u));
+	const uint8 SubstrateBSDFCountMask = 1u << uint8(FMath::Min(SubstrateBSDFCount - 1, 8));
+	const uint8 SubstrateUintPerPixel = FMath::Max(MaterialResource->MaterialGetSubstrateUintPerPixel_GameThread(), uint8(1u));
+
+	MaterialRelevance.bOpaque = !bIsTranslucent;
+	MaterialRelevance.bMasked = IsMasked();
+	MaterialRelevance.bDistortion = MaterialResource->IsDistorted();
+	MaterialRelevance.bHairStrands = IsCompatibleWithHairStrands(MaterialResource, InFeatureLevel);
+	MaterialRelevance.bTwoSided = MaterialResource->IsTwoSided();
+	MaterialRelevance.bSeparateTranslucency = bIsTranslucent && (TranslucencyPass == MTP_AfterDOF);
+	MaterialRelevance.bTranslucencyModulate = bMaterialSeparateModulation;
+	MaterialRelevance.bPostMotionBlurTranslucency = (TranslucencyPass == MTP_AfterMotionBlur);
+	MaterialRelevance.bNormalTranslucency = bIsTranslucent && (TranslucencyPass == MTP_BeforeDOF);
+	MaterialRelevance.bDisableDepthTest = bIsTranslucent && Material->bDisableDepthTest;		
+	MaterialRelevance.bUsesSceneColorCopy = bIsTranslucent && MaterialResource->RequiresSceneColorCopy_GameThread();
+	MaterialRelevance.bOutputsTranslucentVelocity = MaterialResource->IsTranslucencyWritingVelocity();
+	MaterialRelevance.bUsesGlobalDistanceField = MaterialResource->UsesGlobalDistanceField_GameThread();
+	MaterialRelevance.bUsesWorldPositionOffset = MaterialResource->MaterialUsesWorldPositionOffset_GameThread();
+	MaterialRelevance.bUsesDisplacement = MaterialResource->MaterialUsesDisplacement_GameThread();
+	MaterialRelevance.bUsesPixelDepthOffset = MaterialResource->MaterialUsesPixelDepthOffset_GameThread();
+	ETranslucencyLightingMode TranslucencyLightingMode = MaterialResource->GetTranslucencyLightingMode();
+	MaterialRelevance.bTranslucentSurfaceLighting = bIsTranslucent && (TranslucencyLightingMode == TLM_SurfacePerPixelLighting || TranslucencyLightingMode == TLM_Surface);
+	MaterialRelevance.bUsesSceneDepth = MaterialResource->MaterialUsesSceneDepthLookup_GameThread();
+	MaterialRelevance.bHasVolumeMaterialDomain = MaterialResource->IsVolumetricPrimitive();
+	MaterialRelevance.bUsesDistanceCullFade = MaterialResource->MaterialUsesDistanceCullFade_GameThread();
+	MaterialRelevance.bUsesSkyMaterial = Material->bIsSky;
+	MaterialRelevance.bUsesSingleLayerWaterMaterial = bUsesSingleLayerWaterMaterial;
+	MaterialRelevance.bUsesAnisotropy = bUsesAnisotropy;
+	MaterialRelevance.SubstrateClosureCountMask = SubstrateBSDFCountMask;
+	MaterialRelevance.SubstrateUintPerPixel = SubstrateUintPerPixel;
+	MaterialRelevance.bUsesComplexSpecialRenderPath = MaterialResource->MaterialGetSubstrateUsesComplexSpecialRenderPath_GameThread();
+	MaterialRelevance.bIsLightFunctionAtlasCompatible = MaterialResource->MaterialIsLightFunctionAtlasCompatible_GameThread();
+
+	return MaterialRelevance;
 }
 
 FMaterialParameterInfo UMaterialInterface::GetParameterInfo(EMaterialParameterAssociation Association, FName ParameterName, UMaterialFunctionInterface* LayerFunction) const
