[ScreenSaverInhibitor]
 - Fix boolean possibly not being read for every loop because it was neither atomic nor volatile

diff --git a/Engine/Source/Runtime/Engine/Private/UnrealEngine.cpp b/Engine/Source/Runtime/Engine/Private/UnrealEngine.cpp
--- a/Engine/Source/Runtime/Engine/Private/UnrealEngine.cpp
+++ b/Engine/Source/Runtime/Engine/Private/UnrealEngine.cpp
@@ -1442,7 +1442,6 @@ class FScreenSaverInhibitor : public FRunnable
 public:
 	/** Default constructor. */
 	FScreenSaverInhibitor()
-		: bEnabled(true)
 	{}
 
 protected:
@@ -1454,8 +1453,7 @@ protected:
 
 	void Stop() override
 	{
-		bEnabled = false;
-		FPlatformMisc::MemoryBarrier();
+		bEnabled.store(false, std::memory_order_relaxed);
 	}
 
 	/**
@@ -1463,10 +1461,10 @@ protected:
 	*/
 	uint32 Run() override
 	{
-		while( bEnabled )
+		while( bEnabled.load(std::memory_order_relaxed) )
 		{
 			const int32 NUM_SECONDS_TO_SLEEP = 50;
-			for( int32 Sec = 0; Sec < NUM_SECONDS_TO_SLEEP && bEnabled; ++Sec )
+			for( int32 Sec = 0; Sec < NUM_SECONDS_TO_SLEEP && bEnabled.load(std::memory_order_relaxed); ++Sec )
 			{
 				FPlatformProcess::Sleep( 1 );
 			}
@@ -1475,7 +1473,7 @@ protected:
 		return 0;
 	}
 
-	bool bEnabled;
+	std::atomic<bool> bEnabled{ true };
 };
 
 /*-----------------------------------------------------------------------------
