- UObject::Rename will always remove a mismatching linker unless explicitly told not to via a new rename flag REN_AllowPackageLinkerMismatch. This is meant to be distinct from the deprecated REN_ForceNoResetLoaders as the intent is inverted: There are few reasons to not want to reset loaders when there is a mismatch between package and linker, and when there is no mismatch, we do not want Rename to implicitly resetloaders (we no longer do as of 33136565). The only use case where mismatched linkers is desired is when objects are _temporarily_ renamed to belong to a staging package and desire to have the object continue loading when the object is moved back to the matching linker package (commonly done during blueprint reinstancing). In such a case, the new flag acknowledging the intentional mismatch is meant to be used.

- Any code causing a linker mismatch needs to explicitly clear loading flags for their object to avoid asserting. This is to make the intention clear that rename will make loading correctly impossible when moving the object to a different linker unintentionally. So either clear the loading flags of the object (i.e. finish loading via load calls such as ConditionalPostLoad) or forcibly allow mismatching via the new flag.

- When marking an export as invalid in FLinkerLoad::InvalidateExport() we now also clear loading flags since it is wasteful to load invalid objects.

- Exports marked as invalid are no longer reset to be valid when the linker is reset for the object instance. This is to prevent reloading invalid objects when re-using a linker already resident in memory if a subsequent package load is requested.

- When UStruct upgrades UFields to FFields, the UStruct objects are marked as invalid to prevent reloading the objects unnecessarily

- Fixes an issue where since Rename can clear the linker for mismatched types, any calls to FLinkerLoad::InvalidateExport needs to be moved before the rename operation to be effective.

- CollectUnreachableObjects will no longer mark unreachable objects as invalid in the linkerloader. That code would mark objects as invalid until the linker is destroyed which incorrectly assumed would happen soon. If however a load for the same package occurs before the linker is destroyed, the in-memory linker would be re-used keeping the marked exports as invalid so those exports would not be loaded (even though they should have been).

[RN] UObject::Rename will always remove the renamed object's linker if the rename moves the object into a different package. As a result, any renames occuring while an object is loading will now assert instead of leading to hard to diagnose crashes long after the Rename call has completed. If renaming to a different package is intentional and the linker should not be cleared, the new flag REN_AllowPackageLinkerMismatch can be used to prevent clearing linkers on the object even though the linker is for a different package than the package has been renamed to be part of.

diff --git a/Engine/Plugins/Animation/DeformerGraph/Source/OptimusCore/Private/OptimusHelpers.cpp b/Engine/Plugins/Animation/DeformerGraph/Source/OptimusCore/Private/OptimusHelpers.cpp
--- a/Engine/Plugins/Animation/DeformerGraph/Source/OptimusCore/Private/OptimusHelpers.cpp
+++ b/Engine/Plugins/Animation/DeformerGraph/Source/OptimusCore/Private/OptimusHelpers.cpp
@@ -73,7 +73,7 @@ bool Optimus::RenameObject(UObject* InObjectToRename, const TCHAR* InNewName, UO
 
 void Optimus::RemoveObject(UObject* InObjectToRemove)
 {
-	InObjectToRemove->Rename(nullptr, GetTransientPackage(), REN_ForceNoResetLoaders | REN_DoNotDirty | REN_DontCreateRedirectors | REN_NonTransactional);
+	InObjectToRemove->Rename(nullptr, GetTransientPackage(), REN_AllowPackageLinkerMismatch | REN_DoNotDirty | REN_DontCreateRedirectors | REN_NonTransactional);
 	InObjectToRemove->MarkAsGarbage();
 }
 
diff --git a/Engine/Plugins/FX/Niagara/Source/Niagara/Private/NiagaraScratchPadContainer.cpp b/Engine/Plugins/FX/Niagara/Source/Niagara/Private/NiagaraScratchPadContainer.cpp
--- a/Engine/Plugins/FX/Niagara/Source/Niagara/Private/NiagaraScratchPadContainer.cpp
+++ b/Engine/Plugins/FX/Niagara/Source/Niagara/Private/NiagaraScratchPadContainer.cpp
@@ -20,6 +20,8 @@ void UNiagaraScratchPadContainer::CheckConsistency()
 	{
 		if (Script && Script->GetOuter() != this)
 		{
+			// Ensure the object has finished loading before renaming since we are changing the outer which will invalidate the linker
+			Script->ConditionalPostLoad();
 			Script->Rename(*Script->GetName(), this, REN_ForceNoResetLoaders | REN_NonTransactional);
 		}
 	}
diff --git a/Engine/Plugins/PCG/Source/PCG/Private/PCGComponent.cpp b/Engine/Plugins/PCG/Source/PCG/Private/PCGComponent.cpp
--- a/Engine/Plugins/PCG/Source/PCG/Private/PCGComponent.cpp
+++ b/Engine/Plugins/PCG/Source/PCG/Private/PCGComponent.cpp
@@ -3047,7 +3047,10 @@ void UPCGComponent::RefreshSchedulingPolicy()
 	if (IsValid(SchedulingPolicy) && SchedulingPolicy->GetOuter() == this)
 	{
 #if WITH_EDITOR
-		SchedulingPolicy->Rename(nullptr, GetTransientPackage(), REN_DontCreateRedirectors | REN_ForceNoResetLoaders);
+        // We are renaming to a new outer on an object that may still be loading. Since we are destroying this object
+        // pass REN_AllowPackageLinkerMismatch to allow the linker to remain on the object so we don't have to force a load 
+        // to complete before the rename 
+		SchedulingPolicy->Rename(nullptr, GetTransientPackage(), REN_DontCreateRedirectors | REN_AllowPackageLinkerMismatch);
 #endif
 		SchedulingPolicy->MarkAsGarbage();
 	}
diff --git a/Engine/Plugins/Runtime/ModelViewViewModel/Source/ModelViewViewModelBlueprint/Private/Bindings/MVVMConversionFunctionHelper.cpp b/Engine/Plugins/Runtime/ModelViewViewModel/Source/ModelViewViewModelBlueprint/Private/Bindings/MVVMConversionFunctionHelper.cpp
--- a/Engine/Plugins/Runtime/ModelViewViewModel/Source/ModelViewViewModelBlueprint/Private/Bindings/MVVMConversionFunctionHelper.cpp
+++ b/Engine/Plugins/Runtime/ModelViewViewModel/Source/ModelViewViewModelBlueprint/Private/Bindings/MVVMConversionFunctionHelper.cpp
@@ -81,10 +81,12 @@ namespace Private
 	{
 		const ERenameFlags RenFlags = REN_DoNotDirty | REN_ForceNoResetLoaders | REN_DontCreateRedirectors;
 
-		ObjectToRename->Rename(nullptr, GetTransientPackage(), RenFlags);
 		ObjectToRename->SetFlags(RF_Transient);
 		ObjectToRename->ClearFlags(RF_Public | RF_Standalone | RF_ArchetypeObject);
+
+        // Rename will remove the renamed object's linker when moving to a new package so invalidate the export beforehand
 		FLinkerLoad::InvalidateExport(ObjectToRename);
+		ObjectToRename->Rename(nullptr, GetTransientPackage(), RenFlags);
 	}
 
 	struct FCreateGraphResult
diff --git a/Engine/Plugins/Runtime/ModelViewViewModel/Source/ModelViewViewModelBlueprint/Private/MVVMViewBlueprintCompiler.cpp b/Engine/Plugins/Runtime/ModelViewViewModel/Source/ModelViewViewModelBlueprint/Private/MVVMViewBlueprintCompiler.cpp
--- a/Engine/Plugins/Runtime/ModelViewViewModel/Source/ModelViewViewModelBlueprint/Private/MVVMViewBlueprintCompiler.cpp
+++ b/Engine/Plugins/Runtime/ModelViewViewModel/Source/ModelViewViewModelBlueprint/Private/MVVMViewBlueprintCompiler.cpp
@@ -83,10 +83,12 @@ void RenameObjectToTransientPackage(UObject* ObjectToRename)
 {
 	const ERenameFlags RenFlags = REN_DoNotDirty | REN_ForceNoResetLoaders | REN_DontCreateRedirectors;
 
-	ObjectToRename->Rename(nullptr, GetTransientPackage(), RenFlags);
 	ObjectToRename->SetFlags(RF_Transient);
 	ObjectToRename->ClearFlags(RF_Public | RF_Standalone | RF_ArchetypeObject);
+
+    // Rename will remove the renamed object's linker when moving to a new package so invalidate the export beforehand
 	FLinkerLoad::InvalidateExport(ObjectToRename);
+	ObjectToRename->Rename(nullptr, GetTransientPackage(), RenFlags);
 }
 
 FString PropertyPathToString(const UClass* InSelfContext, const UMVVMBlueprintView* BlueprintView, const FMVVMBlueprintPropertyPath& PropertyPath)
diff --git a/Engine/Plugins/Runtime/RigVM/Source/RigVMDeveloper/Private/RigVMBlueprint.cpp b/Engine/Plugins/Runtime/RigVM/Source/RigVMDeveloper/Private/RigVMBlueprint.cpp
--- a/Engine/Plugins/Runtime/RigVM/Source/RigVMDeveloper/Private/RigVMBlueprint.cpp
+++ b/Engine/Plugins/Runtime/RigVM/Source/RigVMDeveloper/Private/RigVMBlueprint.cpp
@@ -980,8 +980,11 @@ void URigVMBlueprint::PostLoad()
 			}
 			else
 			{
+                // We are renaming an object to a new outer while we may still be loading. Since we
+                // are destroying the object, pass REN_AllowPackageLinkerMismatch to avoid forcing
+                // the load to complete since that is wasteful.
 				Graph->MarkAsGarbage();
-				Graph->Rename(nullptr, GetTransientPackage(), REN_ForceNoResetLoaders);
+				Graph->Rename(nullptr, GetTransientPackage(), REN_AllowPackageLinkerMismatch);
 			}
 		}
 		UbergraphPages = NewUberGraphPages;
diff --git a/Engine/Plugins/Runtime/RigVM/Source/RigVMDeveloper/Private/RigVMModel/RigVMController.cpp b/Engine/Plugins/Runtime/RigVM/Source/RigVMDeveloper/Private/RigVMModel/RigVMController.cpp
--- a/Engine/Plugins/Runtime/RigVM/Source/RigVMDeveloper/Private/RigVMModel/RigVMController.cpp
+++ b/Engine/Plugins/Runtime/RigVM/Source/RigVMDeveloper/Private/RigVMModel/RigVMController.cpp
@@ -18765,9 +18765,9 @@ void URigVMController::RewireLinks(URigVMPin* InOldPin, URigVMPin* InNewPin, boo
 
 #endif
 
-bool URigVMController::RenameObject(UObject* InObjectToRename, const TCHAR* InNewName, UObject* InNewOuter) const
+bool URigVMController::RenameObject(UObject* InObjectToRename, const TCHAR* InNewName, UObject* InNewOuter, ERenameFlags InFlags) const
 {
-	const bool bSuccess = InObjectToRename->Rename(InNewName, InNewOuter, REN_ForceNoResetLoaders | REN_DoNotDirty | REN_DontCreateRedirectors | REN_NonTransactional);
+	const bool bSuccess = InObjectToRename->Rename(InNewName, InNewOuter, InFlags | REN_DoNotDirty | REN_DontCreateRedirectors | REN_NonTransactional);
 	if(bSuccess)
 	{
 		if(const URigVMCollapseNode* CollapseNode = Cast<URigVMCollapseNode>(InObjectToRename))
@@ -18796,8 +18796,11 @@ void URigVMController::DestroyObject(UObject* InObjectToDestroy) const
 			}
 		}
 	}
-
-	RenameObject(InObjectToDestroy, nullptr, GetTransientPackage());
+	
+	// We are renaming an object that may be loading, but ultimately want to destroy.
+	// Pass REN_AllowPackageLinkerMismatch to allow the linker to remain on the object
+	// so we may rename without forcing the load to complete.
+	RenameObject(InObjectToDestroy, nullptr, GetTransientPackage(), REN_AllowPackageLinkerMismatch);
 	InObjectToDestroy->RemoveFromRoot();
 	InObjectToDestroy->MarkAsGarbage();
 }
diff --git a/Engine/Plugins/Runtime/RigVM/Source/RigVMDeveloper/Public/RigVMModel/RigVMController.h b/Engine/Plugins/Runtime/RigVM/Source/RigVMDeveloper/Public/RigVMModel/RigVMController.h
--- a/Engine/Plugins/Runtime/RigVM/Source/RigVMDeveloper/Public/RigVMModel/RigVMController.h
+++ b/Engine/Plugins/Runtime/RigVM/Source/RigVMDeveloper/Public/RigVMModel/RigVMController.h
@@ -1276,8 +1276,7 @@ private:
 	void RewireLinks(URigVMPin* OldPin, URigVMPin* NewPin, bool bAsInput, bool bSetupUndoRedo, TArray<URigVMLink*> InLinks = TArray<URigVMLink*>());
 #endif
 
-
-	bool RenameObject(UObject* InObjectToRename, const TCHAR* InNewName, UObject* InNewOuter = nullptr) const;
+	bool RenameObject(UObject* InObjectToRename, const TCHAR* InNewName, UObject* InNewOuter = nullptr, ERenameFlags InFlags = REN_None) const;
 	void DestroyObject(UObject* InObjectToDestroy) const ;
 	static URigVMPin* MakeExecutePin(URigVMNode* InNode, const FName& InName);
 	static bool MakeExecutePin(URigVMPin* InOutPin);
diff --git a/Engine/Plugins/Runtime/nDisplay/Source/DisplayClusterConfigurator/Private/DisplayClusterConfiguratorCompiler.cpp b/Engine/Plugins/Runtime/nDisplay/Source/DisplayClusterConfigurator/Private/DisplayClusterConfiguratorCompiler.cpp
--- a/Engine/Plugins/Runtime/nDisplay/Source/DisplayClusterConfigurator/Private/DisplayClusterConfiguratorCompiler.cpp
+++ b/Engine/Plugins/Runtime/nDisplay/Source/DisplayClusterConfigurator/Private/DisplayClusterConfiguratorCompiler.cpp
@@ -163,10 +163,12 @@ void FDisplayClusterConfiguratorKismetCompilerContext::CopyTermDefaultsToDefault
 
 		// Rename our new config data (along with all new sub-objects) to the transient package.
 		{
-			NewConfigData->Rename(nullptr, GetTransientPackage(), RenFlags);
 			NewConfigData->SetFlags(RF_Transient);
 			NewConfigData->ClearFlags(RF_Transactional);
+
+            // Rename will remove the renamed object's linker when moving to a new package so invalidate the export beforehand
 			FLinkerLoad::InvalidateExport(NewConfigData);
+			NewConfigData->Rename(nullptr, GetTransientPackage(), RenFlags);
 		}
 
 		// Rename our old config data (along with all sub-objects) to our new CDO.
diff --git a/Engine/Source/Editor/Kismet/Private/BlueprintCompilationManager.cpp b/Engine/Source/Editor/Kismet/Private/BlueprintCompilationManager.cpp
--- a/Engine/Source/Editor/Kismet/Private/BlueprintCompilationManager.cpp
+++ b/Engine/Source/Editor/Kismet/Private/BlueprintCompilationManager.cpp
@@ -2819,7 +2819,7 @@ void FBlueprintCompilationManagerImpl::ReinstanceBatch(TArray<FReinstancingJob>&
 					// out of the way so we can reuse its name:
 					Destination, 
 					// Rename options:
-					REN_DoNotDirty | REN_DontCreateRedirectors | REN_ForceNoResetLoaders );
+					REN_DoNotDirty | REN_DontCreateRedirectors | REN_AllowPackageLinkerMismatch);
 
 				// reconstruct
 				FMakeClassSpawnableOnScope TemporarilySpawnable(NewClass);
diff --git a/Engine/Source/Editor/KismetCompiler/Private/KismetCompiler.cpp b/Engine/Source/Editor/KismetCompiler/Private/KismetCompiler.cpp
--- a/Engine/Source/Editor/KismetCompiler/Private/KismetCompiler.cpp
+++ b/Engine/Source/Editor/KismetCompiler/Private/KismetCompiler.cpp
@@ -660,9 +660,10 @@ void FKismetCompilerContext::CleanAndSanitizeClass(UBlueprintGeneratedClass* Cla
 	{
 		FString TransientCDOString = FString::Printf(TEXT("TRASH_%s"), *InOldCDO->GetName());
 		FName TransientCDOName = MakeUniqueObjectName(GetTransientPackage(), TransientClass, FName(*TransientCDOString));
-		InOldCDO->Rename(*TransientCDOName.ToString(), GetTransientPackage(), RenFlags);
+
+        // Rename will remove the renamed object's linker when moving to a new package so invalidate the export beforehand
 		FLinkerLoad::InvalidateExport(InOldCDO);
-		InOldCDO->ClearFlags(RF_NeedLoad | RF_NeedPostLoad | RF_NeedPostLoadSubobjects);
+		InOldCDO->Rename(*TransientCDOName.ToString(), GetTransientPackage(), RenFlags);
 	}
 
 	// Purge all subobjects (properties, functions, params) of the class, as they will be regenerated
@@ -724,9 +725,10 @@ void FKismetCompilerContext::CleanAndSanitizeClass(UBlueprintGeneratedClass* Cla
 		}
 
 		FName NewSubobjectName = MakeUniqueObjectName(TransientClass, CurrSubObj->GetClass(), CurrSubObj->GetFName());
-		CurrSubObj->Rename(*NewSubobjectName.ToString(), TransientClass, RenFlags | REN_ForceNoResetLoaders);
+
+        // Rename will remove the renamed object's linker when moving to a new package so invalidate the export beforehand
 		FLinkerLoad::InvalidateExport(CurrSubObj);
-		CurrSubObj->ClearFlags(RF_NeedLoad | RF_NeedPostLoad | RF_NeedPostLoadSubobjects);
+		CurrSubObj->Rename(*NewSubobjectName.ToString(), TransientClass, RenFlags | REN_ForceNoResetLoaders);
 	}
 
 	// Purge the class to get it back to a "base" state
diff --git a/Engine/Source/Editor/KismetCompiler/Private/KismetCompilerModule.cpp b/Engine/Source/Editor/KismetCompiler/Private/KismetCompilerModule.cpp
--- a/Engine/Source/Editor/KismetCompiler/Private/KismetCompilerModule.cpp
+++ b/Engine/Source/Editor/KismetCompiler/Private/KismetCompilerModule.cpp
@@ -106,7 +106,7 @@ void FKismet2CompilerModule::RefreshVariables(UBlueprint* Blueprint)
 		// move old cdo aside:
 		if(OldCDO)
 		{
-			OldCDO->Rename(NULL, GetTransientPackage(), REN_DontCreateRedirectors | REN_NonTransactional | REN_ForceNoResetLoaders);
+			OldCDO->Rename(NULL, GetTransientPackage(), REN_DontCreateRedirectors | REN_NonTransactional | REN_ForceNoResetLoaders | REN_AllowPackageLinkerMismatch);
 		}
 		Blueprint->GeneratedClass->ClassDefaultObject = nullptr;
 
diff --git a/Engine/Source/Editor/UMGEditor/Private/WidgetBlueprintCompiler.cpp b/Engine/Source/Editor/UMGEditor/Private/WidgetBlueprintCompiler.cpp
--- a/Engine/Source/Editor/UMGEditor/Private/WidgetBlueprintCompiler.cpp
+++ b/Engine/Source/Editor/UMGEditor/Private/WidgetBlueprintCompiler.cpp
@@ -291,6 +291,16 @@ void FWidgetBlueprintCompilerContext::CleanAndSanitizeClass(UBlueprintGeneratedC
 	const bool bRecompilingOnLoad = Blueprint->bIsRegeneratingOnLoad;
 	auto RenameObjectToTransientPackage = [bRecompilingOnLoad](UObject* ObjectToRename, const FName BaseName,  bool bClearFlags)
 	{
+		ObjectToRename->SetFlags(RF_Transient);
+
+		if (bClearFlags)
+		{
+			ObjectToRename->ClearFlags(RF_Public | RF_Standalone | RF_ArchetypeObject);
+		}
+
+        // Rename will remove the renamed object's linker when moving to a new package so invalidate the export beforehand
+		FLinkerLoad::InvalidateExport(ObjectToRename);
+
 		const ERenameFlags RenFlags = REN_DontCreateRedirectors | (bRecompilingOnLoad ? REN_ForceNoResetLoaders : 0) | REN_NonTransactional | REN_DoNotDirty;
 
 		if (BaseName.IsNone())
@@ -302,14 +312,6 @@ void FWidgetBlueprintCompilerContext::CleanAndSanitizeClass(UBlueprintGeneratedC
 			FName TransientArchetypeName = MakeUniqueObjectName(GetTransientPackage(), ObjectToRename->GetClass(), BaseName);
 			ObjectToRename->Rename(*TransientArchetypeName.ToString(), GetTransientPackage(), RenFlags);
 		}
-
-		ObjectToRename->SetFlags(RF_Transient);
-
-		if (bClearFlags)
-		{
-			ObjectToRename->ClearFlags(RF_Public | RF_Standalone | RF_ArchetypeObject);
-		}
-		FLinkerLoad::InvalidateExport(ObjectToRename);
 	};
 
 	if ( !Blueprint->bIsRegeneratingOnLoad && bIsFullCompile )
diff --git a/Engine/Source/Editor/UnrealEd/Private/Kismet2/BlueprintEditorUtils.cpp b/Engine/Source/Editor/UnrealEd/Private/Kismet2/BlueprintEditorUtils.cpp
--- a/Engine/Source/Editor/UnrealEd/Private/Kismet2/BlueprintEditorUtils.cpp
+++ b/Engine/Source/Editor/UnrealEd/Private/Kismet2/BlueprintEditorUtils.cpp
@@ -1243,12 +1243,14 @@ void FBlueprintEditorUtils::RemoveStaleFunctions(UBlueprintGeneratedClass* Class
 		{
 			UFunction* Function = *Fn;
 			Class->RemoveFunctionFromFunctionMap(Function);
-			Function->Rename(nullptr, OrphanedClass, RenFlags);
 
 			// invalidate this package's reference to this function, so 
 			// subsequent packages that import it will treat it as if it didn't 
 			// exist (because data-only blueprints shouldn't have functions)
+            // Note, Rename will remove the renamed object's linker when moving 
+            // to a new package so invalidate the export beforehand
 			FLinkerLoad::InvalidateExport(Function); 
+			Function->Rename(nullptr, OrphanedClass, RenFlags);
 			++Fn;
 		}
 	}
diff --git a/Engine/Source/Editor/UnrealEd/Private/Kismet2/KismetReinstanceUtilities.cpp b/Engine/Source/Editor/UnrealEd/Private/Kismet2/KismetReinstanceUtilities.cpp
--- a/Engine/Source/Editor/UnrealEd/Private/Kismet2/KismetReinstanceUtilities.cpp
+++ b/Engine/Source/Editor/UnrealEd/Private/Kismet2/KismetReinstanceUtilities.cpp
@@ -514,7 +514,7 @@ FBlueprintCompileReinstancer::FBlueprintCompileReinstancer(UClass* InClassToRein
 		if(!bAvoidCDODuplication)
 		{
 			ensure( ClassToReinstance->ClassDefaultObject->GetClass() == DuplicatedClass );
-			ClassToReinstance->ClassDefaultObject->Rename(nullptr, GetTransientPackage(), REN_DoNotDirty | REN_DontCreateRedirectors | REN_ForceNoResetLoaders);
+			ClassToReinstance->ClassDefaultObject->Rename(nullptr, GetTransientPackage(), REN_DoNotDirty | REN_DontCreateRedirectors | REN_AllowPackageLinkerMismatch);
 		}
 		
 		// Note that we can't clear ClassToReinstance->ClassDefaultObject even though
@@ -1280,7 +1280,7 @@ void FBlueprintCompileReinstancer::TakeOwnershipOfSparseClassData(UClass* ForCla
 	OriginalSCD = const_cast<void*>(ForClass->GetSparseClassData(EGetSparseClassDataMethod::ReturnIfNull));
 	if (OriginalSCDStruct->GetOuter() == ForClass)
 	{
-		OriginalSCDStruct->Rename(nullptr, DuplicatedClass, REN_DoNotDirty | REN_DontCreateRedirectors | REN_ForceNoResetLoaders | REN_NonTransactional);
+		OriginalSCDStruct->Rename(nullptr, DuplicatedClass, REN_DoNotDirty | REN_DontCreateRedirectors | REN_AllowPackageLinkerMismatch | REN_NonTransactional);
 	}
 	// We own these now, remove ForClass's knowledge of the sparse class data - they
 	// will be freed when reinstancing is complete:
@@ -2068,7 +2068,7 @@ UClass* FBlueprintCompileReinstancer::MoveCDOToNewClass(UClass* OwnerClass, cons
 		OriginalNames.Add(OwnedObject->GetFName());
 		if(OwnedObject->HasAnyFlags(RF_ArchetypeObject))
 		{
-			OwnedObject->Rename(nullptr, GetTransientPackage(), REN_DoNotDirty | REN_DontCreateRedirectors | REN_ForceNoResetLoaders | REN_NonTransactional);
+			OwnedObject->Rename(nullptr, GetTransientPackage(), REN_DoNotDirty | REN_DontCreateRedirectors | REN_AllowPackageLinkerMismatch | REN_NonTransactional);
 		}
 	}
 
@@ -2114,7 +2114,7 @@ UClass* FBlueprintCompileReinstancer::MoveCDOToNewClass(UClass* OwnerClass, cons
 		UObject* OwnedArchetype = OwnedObjects[I];
 		if(OwnedArchetype->HasAnyFlags(RF_ArchetypeObject))
 		{
-			OwnedArchetype->Rename(*OriginalNames[I].ToString(), OwnerClass, REN_DoNotDirty | REN_DontCreateRedirectors | REN_ForceNoResetLoaders | REN_NonTransactional);
+			OwnedArchetype->Rename(*OriginalNames[I].ToString(), OwnerClass, REN_DoNotDirty | REN_DontCreateRedirectors | REN_AllowPackageLinkerMismatch | REN_NonTransactional);
 		}
 	}
 
@@ -2137,7 +2137,7 @@ UClass* FBlueprintCompileReinstancer::MoveCDOToNewClass(UClass* OwnerClass, cons
 		if(bAvoidCDODuplication)
 		{
 			OwnerClass->ClassDefaultObject = nullptr;
-			OldCDO->Rename(nullptr, CopyOfOwnerClass->GetOuter(), REN_DoNotDirty | REN_DontCreateRedirectors | REN_ForceNoResetLoaders | REN_NonTransactional);
+			OldCDO->Rename(nullptr, CopyOfOwnerClass->GetOuter(), REN_DoNotDirty | REN_DontCreateRedirectors | REN_AllowPackageLinkerMismatch | REN_NonTransactional);
 			CopyOfOwnerClass->ClassDefaultObject = OldCDO;
 		}
 		OldCDO->SetClass(CopyOfOwnerClass);
@@ -2230,7 +2230,7 @@ static void ReplaceObjectHelper(UObject*& OldObject, UClass* OldClass, UObject*&
 		}
 		else
 		{
-			OldObject->Rename(nullptr, OldObject->GetOuter(), REN_DoNotDirty | REN_DontCreateRedirectors | REN_ForceNoResetLoaders);
+			OldObject->Rename(nullptr, OldObject->GetOuter(), REN_DoNotDirty | REN_DontCreateRedirectors | REN_ForceNoResetLoaders | REN_AllowPackageLinkerMismatch);
 		}
 	}
 						
diff --git a/Engine/Source/Runtime/CoreUObject/Private/Serialization/AsyncLoading2.cpp b/Engine/Source/Runtime/CoreUObject/Private/Serialization/AsyncLoading2.cpp
--- a/Engine/Source/Runtime/CoreUObject/Private/Serialization/AsyncLoading2.cpp
+++ b/Engine/Source/Runtime/CoreUObject/Private/Serialization/AsyncLoading2.cpp
@@ -5594,18 +5594,10 @@ EEventLoadNodeExecutionResult FAsyncPackage2::ProcessLinkerLoadPackageSummary(FA
 			FObjectExport& LinkerExport = LinkerLoadState->Linker->ExportMap[LinkerLoadState->MetaDataIndex];
 			FExportObject& ExportObject = Data.Exports[LinkerLoadState->MetaDataIndex];
 			ExportObject.Object = LinkerExport.Object;
-			if (ExportObject.Object)
-			{
-				ExportObject.bWasFoundInMemory = true; // Make sure that the async flags are cleared in ClearConstructedObjects
-			}
-			else
-			{
-				ExportObject.bExportLoadFailed = LinkerExport.bExportLoadFailed;
-				if (!ExportObject.bExportLoadFailed)
-				{
-					ExportObject.bFiltered = true;
-				}
-			}
+
+			ExportObject.bWasFoundInMemory = !!LinkerExport.Object; // Make sure that the async flags are cleared in ClearConstructedObjects
+			ExportObject.bExportLoadFailed = LinkerExport.bExportLoadFailed;
+			ExportObject.bFiltered = LinkerExport.bWasFiltered;
 		}
 #endif
 	}
@@ -5918,6 +5910,7 @@ EEventLoadNodeExecutionResult FAsyncPackage2::ExecuteDeferredPostLoadLinkerLoadP
 				{
 					UE_ASYNC_PACKAGE_LOG(Verbose, Desc, TEXT("ExecuteDeferredPostLoadLinkerLoadPackageExports"), TEXT("Patching export %d: %s -> %s"), ObjectIndex, *GetPathNameSafe(ExportObject.Object), *GetPathNameSafe(LinkerExport.Object));
 					ExportObject.Object = LinkerExport.Object;
+					ExportObject.bWasFoundInMemory = !!LinkerExport.Object;
 				}
 			}
 		}
@@ -6481,12 +6474,12 @@ UObject* FAsyncPackage2::ConditionalSerializeExport(const FAsyncPackageHeaderDat
 {
 	FExportObject& Export = Header.ExportsView[LocalExportIndex];
 	
-	if (!Export.Object && !(Export.bFiltered | Export.bExportLoadFailed))
+	if (!Export.Object && !(Export.bFiltered || Export.bExportLoadFailed))
 	{
 		ConditionalCreateExport(Header, LocalExportIndex);
 	}
 	
-	if (!Export.Object || (Export.bFiltered | Export.bExportLoadFailed))
+	if (!Export.Object || (Export.bFiltered || Export.bExportLoadFailed))
 	{
 		return nullptr;
 	}
@@ -6598,7 +6591,7 @@ void FAsyncPackage2::EventDrivenCreateExport(const FAsyncPackageHeaderData& Head
 	};
 
 	ExportObject.bFiltered = FilterExport(Export.FilterFlags);
-	if (ExportObject.bFiltered | ExportObject.bExportLoadFailed)
+	if (ExportObject.bFiltered || ExportObject.bExportLoadFailed)
 	{
 		if (ExportObject.bExportLoadFailed)
 		{
@@ -6804,11 +6797,11 @@ bool FAsyncPackage2::EventDrivenSerializeExport(const FAsyncPackageHeaderData& H
 	const FExportMapEntry& Export = Header.ExportMap[LocalExportIndex];
 	FExportObject& ExportObject = Header.ExportsView[LocalExportIndex];
 	UObject* Object = ExportObject.Object;
-	check(Object || (ExportObject.bFiltered | ExportObject.bExportLoadFailed));
+	check(Object || (ExportObject.bFiltered || ExportObject.bExportLoadFailed));
 
 	TRACE_LOADTIME_SERIALIZE_EXPORT_SCOPE(Object, Export.CookedSerialSize);
 
-	if ((ExportObject.bFiltered | ExportObject.bExportLoadFailed) || !(Object && Object->HasAnyFlags(RF_NeedLoad)))
+	if ((ExportObject.bFiltered || ExportObject.bExportLoadFailed) || !(Object && Object->HasAnyFlags(RF_NeedLoad)))
 	{
 		if (ExportObject.bExportLoadFailed)
 		{
@@ -7367,7 +7360,7 @@ EEventLoadNodeExecutionResult FAsyncPackage2::Event_PostLoadExportBundle(FAsyncL
 				do
 				{
 					FExportObject& Export = HeaderData->ExportsView[BundleEntry.LocalExportIndex];
-					if (Export.bFiltered | Export.bExportLoadFailed)
+					if (Export.bFiltered || Export.bExportLoadFailed)
 					{
 						break;
 					}
@@ -7481,7 +7474,7 @@ EEventLoadNodeExecutionResult FAsyncPackage2::Event_DeferredPostLoadExportBundle
 				do
 				{
 					FExportObject& Export = HeaderData->ExportsView[BundleEntry.LocalExportIndex];
-					if (Export.bFiltered | Export.bExportLoadFailed)
+					if (Export.bFiltered || Export.bExportLoadFailed)
 					{
 						break;
 					}
@@ -7796,7 +7789,7 @@ EAsyncPackageState::Type FAsyncLoadingThread2::ProcessLoadedPackagesFromGameThre
 				// Clear async loading flags (we still want RF_Async, but EInternalObjectFlags::AsyncLoading can be cleared)
 				for (const FExportObject& Export : Package->Data.Exports)
 				{
-					if (Export.bFiltered | Export.bExportLoadFailed)
+					if (Export.bFiltered || Export.bExportLoadFailed)
 					{
 						continue;
 					}
@@ -8543,12 +8536,6 @@ void FAsyncLoadingThread2::CollectUnreachableObjects(
 			{
 				const int32 CachedLinkerIndex = Object->GetLinkerIndex();
 				Object->SetLinker(nullptr, INDEX_NONE);
-				// As we are garbaging the object, mark it as invalid in the linker
-				// Either it is now truly invalid to access this object 
-				// (i.e the asset ran upgraded, migrated away from this object and doesn't hold any references to it anymore.)
-				// or we are gc'ing the entire asset in which case the linker will eventually get purged and those entry won't be marked invalid anymore on recreation
-				FObjectExport& ObjExport = ObjectLinker->ExportMap[CachedLinkerIndex];
-				ObjExport.bExportLoadFailed = true;
 			}
 #endif // ALT2_ENABLE_LINKERLOAD_SUPPORT
 		});
@@ -9044,7 +9031,7 @@ EAsyncPackageState::Type FAsyncPackage2::PostLoadInstances(FAsyncLoadingThreadSt
 	{
 		const FExportObject& Export = Data.Exports[PostLoadInstanceIndex++];
 
-		if (!(Export.bFiltered || Export.bExportLoadFailed))
+		if (Export.Object && !(Export.bFiltered || Export.bExportLoadFailed))
 		{
 			UClass* ObjClass = Export.Object->GetClass();
 			ObjClass->PostLoadInstance(Export.Object);
@@ -9060,7 +9047,7 @@ EAsyncPackageState::Type FAsyncPackage2::CreateClusters(FAsyncLoadingThreadState
 	{
 		const FExportObject& Export = Data.Exports[DeferredClusterIndex++];
 
-		if (!(Export.bFiltered | Export.bExportLoadFailed) && Export.Object->CanBeClusterRoot())
+		if (!(Export.bFiltered || Export.bExportLoadFailed) && Export.Object->CanBeClusterRoot())
 		{
 			Export.Object->CreateCluster();
 			if (DeferredClusterIndex < ExportCount && ThreadState.IsTimeLimitExceeded(TEXT("CreateClusters")))
diff --git a/Engine/Source/Runtime/CoreUObject/Private/UObject/Class.cpp b/Engine/Source/Runtime/CoreUObject/Private/UObject/Class.cpp
--- a/Engine/Source/Runtime/CoreUObject/Private/UObject/Class.cpp
+++ b/Engine/Source/Runtime/CoreUObject/Private/UObject/Class.cpp
@@ -2091,6 +2091,9 @@ void UStruct::ConvertUFieldsToFFields()
 			{
 				Children = OldField->Next;
 			}
+
+            // Rename will remove the renamed object's linker when moving to a new package so invalidate the export beforehand
+			FLinkerLoad::InvalidateExport(OldField);
 			// Move the old UProperty to the transient package and rename it to something unique
 			OldField->Rename(*MakeUniqueObjectName(GetTransientPackage(), OldField->GetClass()).ToString(), GetTransientPackage(), REN_ForceNoResetLoaders | REN_DoNotDirty | REN_DontCreateRedirectors | REN_NonTransactional);
 			OldField->RemoveFromRoot();
diff --git a/Engine/Source/Runtime/CoreUObject/Private/UObject/LinkerLoad.cpp b/Engine/Source/Runtime/CoreUObject/Private/UObject/LinkerLoad.cpp
--- a/Engine/Source/Runtime/CoreUObject/Private/UObject/LinkerLoad.cpp
+++ b/Engine/Source/Runtime/CoreUObject/Private/UObject/LinkerLoad.cpp
@@ -686,6 +686,9 @@ void FLinkerLoad::InvalidateExport(UObject* OldObject)
 
 	if (OldObjectLinker && OldObjectLinker->ExportMap.IsValidIndex(CachedLinkerIndex))
 	{
+		// Prevent any further loading as this export is now invalid
+		OldObject->ClearFlags(RF_NeedLoad | RF_NeedPostLoad | RF_NeedPostLoadSubobjects);
+
 		FObjectExport& ObjExport = OldObjectLinker->ExportMap[CachedLinkerIndex];
 		ObjExport.bExportLoadFailed = true;
 	}
diff --git a/Engine/Source/Runtime/CoreUObject/Private/UObject/Obj.cpp b/Engine/Source/Runtime/CoreUObject/Private/UObject/Obj.cpp
--- a/Engine/Source/Runtime/CoreUObject/Private/UObject/Obj.cpp
+++ b/Engine/Source/Runtime/CoreUObject/Private/UObject/Obj.cpp
@@ -214,7 +214,7 @@ UObject* UObject::GetDefaultSubobjectByName(FName ToFind)
 	return Object;
 }
 
-bool UObject::Rename( const TCHAR* InName, UObject* NewOuter, ERenameFlags Flags )
+bool UObject::Rename(const TCHAR* InName, UObject* NewOuter, ERenameFlags Flags)
 {
 	TRACE_CPUPROFILER_EVENT_SCOPE(UObject::Rename);
 
@@ -224,7 +224,7 @@ bool UObject::Rename( const TCHAR* InName, UObject* NewOuter, ERenameFlags Flags
 	FMetaDataUtilities::FMoveMetadataHelperContext MoveMetaData(this, true);
 #endif //WITH_EDITOR
 
-	if(NewOuter)
+	if (NewOuter)
 	{
 		// Renaming the CDO of a Blueprint is a special case so we do not validate what would otherwise be incorrect use of Rename.
 		// Moving objects to the transient package is commonly used halfway through destroying them so that is also fine, otherwise
@@ -244,11 +244,12 @@ bool UObject::Rename( const TCHAR* InName, UObject* NewOuter, ERenameFlags Flags
 					*NewOuter->GetFullName(),
 					*GetClass()->ClassWithin->GetName());
 			}
-			// If moving the object to a new package, remove it's linker .
-			else if (GetLinker() && GetPackage() != NewOuter->GetPackage())
-			{
-				SetLinker(nullptr, INDEX_NONE);
-			}
+		}
+
+		// If moving the object to a new package, remove its linker and detach the object
+		if (((Flags & REN_AllowPackageLinkerMismatch) == 0) && GetLinker() != NewOuter->GetLinker())
+		{
+			SetLinker(nullptr, INDEX_NONE);
 		}
 	}
 
@@ -996,29 +997,8 @@ void UObject::BeginDestroy()
 			);
 	}
 
-#if WITH_EDITORONLY_DATA
-	// Make sure the linker entry stays as 'bExportLoadFailed' if the entry was marked as such, 
-	// doing this prevents the object from being reloaded by subsequent load calls:
-	FLinkerLoad* Linker = GetLinker();
-	const int32 CachedLinkerIndex = GetLinkerIndex();
-	bool bLinkerEntryWasInvalid = false;
-	if(Linker && Linker->ExportMap.IsValidIndex(CachedLinkerIndex))
-	{
-		FObjectExport& ObjExport = Linker->ExportMap[CachedLinkerIndex];
-		bLinkerEntryWasInvalid = ObjExport.bExportLoadFailed;
-	}
-#endif // WITH_EDITORONLY_DATA
-
 	// Remove from linker's export table.
 	SetLinker( NULL, INDEX_NONE );
-	
-#if WITH_EDITORONLY_DATA
-	if(bLinkerEntryWasInvalid)
-	{
-		FObjectExport& ObjExport = Linker->ExportMap[CachedLinkerIndex];
-		ObjExport.bExportLoadFailed = true;
-	}
-#endif // WITH_EDITORONLY_DATA
 
 	LowLevelRename(NAME_None);
 	// Remove any associated external package, at this point
diff --git a/Engine/Source/Runtime/CoreUObject/Private/UObject/UObjectLinker.cpp b/Engine/Source/Runtime/CoreUObject/Private/UObject/UObjectLinker.cpp
--- a/Engine/Source/Runtime/CoreUObject/Private/UObject/UObjectLinker.cpp
+++ b/Engine/Source/Runtime/CoreUObject/Private/UObject/UObjectLinker.cpp
@@ -125,7 +125,12 @@ void UObject::SetLinker( FLinkerLoad* LinkerLoad, int32 LinkerIndex, bool bShoul
 		FObjectExport& ExportObject = Existing.Linker->ExportMap[Existing.LinkerIndex];
 		checkf(ExportObject.Object != nullptr, TEXT("Expected ExportMap[%d].Object to not be null for this ('%s')"), Existing.LinkerIndex, *GetFName().ToString());
 		checkf(ExportObject.Object == this, TEXT("Expected ExportMap[%d].Object ('%s') to equal this ('%s')"), Existing.LinkerIndex, *ExportObject.Object->GetFName().ToString(), *GetFName().ToString());
+
+		// Detach the object but keep the object marked as invalid if it was previously. Now that it's detached
+		// if we reload the package we don't want to reload invalid objects.
+		bool bIsInvalid = ExportObject.bExportLoadFailed;
 		ExportObject.ResetObject();
+		ExportObject.bExportLoadFailed = bIsInvalid;
 	}
 
 	if (Existing.Linker == LinkerLoad)
diff --git a/Engine/Source/Runtime/CoreUObject/Public/UObject/ObjectMacros.h b/Engine/Source/Runtime/CoreUObject/Public/UObject/ObjectMacros.h
--- a/Engine/Source/Runtime/CoreUObject/Public/UObject/ObjectMacros.h
+++ b/Engine/Source/Runtime/CoreUObject/Public/UObject/ObjectMacros.h
@@ -2133,23 +2133,25 @@ private:
 typedef uint32 ERenameFlags;
 
 /** Default rename behavior */
-#define REN_None					(0x0000)
+#define REN_None					    (0x0000)
 /** UE_DEPRECATED(5.5, "Rename will no longer call ResetLoaders making this flag no longer needed.") -  Rename won't call ResetLoaders or flush async loading. You should pass this if you are renaming a deep subobject and do not need to reset loading for the outer package */
-#define REN_ForceNoResetLoaders		(0x0001) 
+#define REN_ForceNoResetLoaders		    (0x0001) 
 /** Just test to make sure that the rename is guaranteed to succeed if an non test rename immediately follows */
-#define REN_Test					(0x0002) 
+#define REN_Test					    (0x0002) 
 /** Indicates that the object (and new outer) should not be dirtied */
-#define REN_DoNotDirty				(0x0004) 
+#define REN_DoNotDirty				    (0x0004) 
 /** Don't create an object redirector, even if the class is marked RF_Public */
-#define REN_DontCreateRedirectors	(0x0010) 
+#define REN_DontCreateRedirectors	    (0x0010) 
 /** Don't call Modify() on the objects, so they won't be stored in the transaction buffer */
-#define REN_NonTransactional		(0x0020) 
+#define REN_NonTransactional		    (0x0020) 
 /** Force unique names across all packages not just within the scope of the new outer */
-#define REN_ForceGlobalUnique		(0x0040) 
+#define REN_ForceGlobalUnique		    (0x0040) 
 /** Prevent renaming of any child generated classes and CDO's in blueprints */
-#define REN_SkipGeneratedClasses	(0x0080) 
+#define REN_SkipGeneratedClasses	    (0x0080) 
 /** Prevents renaming from unregistering/registering all components */
-#define REN_SkipComponentRegWork	(0x0100) 
+#define REN_SkipComponentRegWork	    (0x0100) 
+/** Prevents rename from clearing an object's linker when being renamed to a new package */
+#define REN_AllowPackageLinkerMismatch	(0x0200) 
 
 /*-----------------------------------------------------------------------------
 	Misc.
diff --git a/Engine/Source/Runtime/Engine/Private/EdGraph/EdGraphNode.cpp b/Engine/Source/Runtime/Engine/Private/EdGraph/EdGraphNode.cpp
--- a/Engine/Source/Runtime/Engine/Private/EdGraph/EdGraphNode.cpp
+++ b/Engine/Source/Runtime/Engine/Private/EdGraph/EdGraphNode.cpp
@@ -682,7 +682,7 @@ void UEdGraphNode::PostLoad()
 	{
 		for (UEdGraphPin_Deprecated* LegacyPin : DeprecatedPins)
 		{
-			LegacyPin->Rename(nullptr, GetTransientPackage(), REN_ForceNoResetLoaders|REN_NonTransactional);
+			LegacyPin->Rename(nullptr, GetTransientPackage(), REN_AllowPackageLinkerMismatch | REN_NonTransactional);
 			LegacyPin->SetFlags(RF_Transient);
 			LegacyPin->MarkAsGarbage();
 		}
diff --git a/Engine/Source/Runtime/Engine/Private/InheritableComponentHandler.cpp b/Engine/Source/Runtime/Engine/Private/InheritableComponentHandler.cpp
--- a/Engine/Source/Runtime/Engine/Private/InheritableComponentHandler.cpp
+++ b/Engine/Source/Runtime/Engine/Private/InheritableComponentHandler.cpp
@@ -254,18 +254,15 @@ void UInheritableComponentHandler::ValidateTemplates()
 					{
 						Record.ComponentTemplate->SetFlags(RF_Transient);
 #if WITH_EDITOR
+                        // Mark the export as invalid to ensure the loader won't reload the object, 
+                        // should the package owning this object be reloaded.
+                        // Note, Rename will remove the renamed object's linker when moving to a new package so invalidate the export beforehand
+						FLinkerLoad::InvalidateExport(Record.ComponentTemplate);
 						// in editor, move the component template aside so its name is free:
 						Record.ComponentTemplate->Rename(nullptr, GetTransientPackage(), REN_DoNotDirty | REN_DontCreateRedirectors | REN_ForceNoResetLoaders | REN_NonTransactional);
 						Record.ComponentTemplate->ClearFlags(RF_Standalone);
 						Record.ComponentTemplate->RemoveFromRoot();
 						Record.ComponentTemplate->MarkAsGarbage();
-						// Rename won't invalidate the linker's export, and linker lifetime extends long beyond an actual loadpackage 
-						// invocation. Consequently, if the template object is garbage collected (as we hope it will be) it 
-						// could tragically be recreated by FLinkerLoad unless we invalidate the export. Zen loader has some 
-						// logic to avoid recreating the object, but it is buggy and we want to avoid object recreation when 
-						// not using Zen, anyway. We would not need to invalidate the export if: 1. Rename invalidated the 
-						// export or 2. FLinkerLoad's lifetime were reigned in.
-						FLinkerLoad::InvalidateExport(Record.ComponentTemplate);
 #endif // WITH_EDITOR
 						UnnecessaryComponents.AddUnique(Record.ComponentTemplate);
 					}
diff --git a/Engine/Source/Runtime/Engine/Private/Tests/Loading/AsyncLoadingTests.cpp b/Engine/Source/Runtime/Engine/Private/Tests/Loading/AsyncLoadingTests.cpp
--- a/Engine/Source/Runtime/Engine/Private/Tests/Loading/AsyncLoadingTests.cpp
+++ b/Engine/Source/Runtime/Engine/Private/Tests/Loading/AsyncLoadingTests.cpp
@@ -6,6 +6,9 @@
 #include "AssetRegistry/AssetRegistryState.h"
 #include "AssetRegistry/IAssetRegistry.h"
 #include "Async/ParallelFor.h"
+#include "AsyncLoadingTests_Shared.h"
+#include "Engine/UserDefinedStruct.h"
+#include "UObject/UObjectHash.h"
 
 #if WITH_DEV_AUTOMATION_TESTS
 
@@ -62,7 +65,7 @@ bool FThreadSafeAsyncLoadingTest::RunTest(const FString& Parameters)
 IMPLEMENT_SIMPLE_AUTOMATION_TEST(FCircularDependencyLoadingTest, TEXT(TEST_NAME_ROOT ".LoadBlueprintWithCircularDependencyTest"), EAutomationTestFlags::EditorContext | EAutomationTestFlags::EngineFilter)
 bool FCircularDependencyLoadingTest::RunTest(const FString& Parameters)
 {
-	static constexpr const TCHAR* ActorWithCircularReferences = TEXT("/Game/Tests/Core/AssetLoading/BlueprintActorWithCircularReferences");
+	static constexpr const TCHAR* ActorWithCircularReferences = TEXT("/Game/Tests/Core/AssetLoading/RecursiveLoads/BlueprintActorWithCircularReferences");
 
 	UPackage* Package = LoadPackage(nullptr, ActorWithCircularReferences, LOAD_None);
 	TestTrue(TEXT("The object should have been properly loaded recursively"), Package != nullptr);
@@ -75,5 +78,167 @@ bool FCircularDependencyLoadingTest::RunTest(const FString& Parameters)
 	return true;
 }
 
+/**
+ * PackagePath1 leads to a user defined structure containing UFields which will be forcibly converted to FField when loading. The UFields should be marked as invalid and the loader should not consider them for loading.
+ */
+IMPLEMENT_SIMPLE_AUTOMATION_TEST(FLoadingTests_ImportPackageConvertedOnLoad, TEXT(TEST_NAME_ROOT ".InvalidExports.ImportPackageConvertedOnLoad"), EAutomationTestFlags::EditorContext | EAutomationTestFlags::EngineFilter)
+bool FLoadingTests_ImportPackageConvertedOnLoad::RunTest(const FString& Parameters)
+{
+	static constexpr const TCHAR* PackagePath1 = TEXT("/Game/Tests/Core/AssetLoading/InvalidatingExports/OldProcMeshData");
+	static constexpr const TCHAR* PackagePath2 = TEXT("/Game/Tests/Core/AssetLoading/InvalidatingExports/BP_UsingUStructWithUnconvertedFields");
+	static constexpr const TCHAR* ObjectName = TEXT("ProcMeshData");
+	typedef TSet<TPair<FName, FName>> FObjPathAndClassNameSet;
+	FObjPathAndClassNameSet ExpectedObjects;
+	FObjPathAndClassNameSet ActualObjects;
+
+	auto VerifyLoadedObjects = [this](const FObjPathAndClassNameSet& ExpectedObjects, const FObjPathAndClassNameSet& ActualObjects)
+		{
+			TestTrue(FString::Printf(TEXT("Expected the same number of objects in package %s after reloading. Expected: %d != Actual: %d"), PackagePath1, ExpectedObjects.Num(), ActualObjects.Num()), ExpectedObjects.Num() == ActualObjects.Num());
+			for (const TPair<FName, FName>& Pair : ExpectedObjects)
+			{
+				TestTrue(FString::Printf(TEXT("Missing obj %s (type: %s) after reloading %s"), *Pair.Key.ToString(), *Pair.Value.ToString(), PackagePath1), ActualObjects.Find(Pair) != nullptr);
+			}
+		};
+
+	// Load our package and populate the ExpectedObjects set
+	{
+		UPackage* Package = LoadPackage(nullptr, PackagePath1, LOAD_None);
+		TestTrue(FString::Printf(TEXT("Failed to load package at %s"), PackagePath1), Package != nullptr);
+
+		ForEachObjectWithPackage(Package,
+			[&ExpectedObjects](UObject* Object)
+			{
+				ExpectedObjects.Add({ Object->GetFName(), Object->GetClass()->GetFName() });
+				return true;
+			}
+		);
+
+		FLoadingTestsScope::GarbageCollect({ {PackagePath1} }, *this);
+	}
+
+	// Load a package that imports the objects from the package at PackagePath1 and ensure we still load only the objects we expect
+	{
+		UPackage* Package = LoadPackage(nullptr, PackagePath2, LOAD_None);
+		TestTrue(FString::Printf(TEXT("Failed to load package at %s"), PackagePath2), Package != nullptr);
+
+		// Path2 imports PackagePath1 so we should be able to still find our Obj and our PackagePath1 package
+		UUserDefinedStruct* Obj = FindFirstObject<UUserDefinedStruct>(ObjectName);
+		TestTrue(FString::Printf(TEXT("Failed to find expected object %s in package %s"), ObjectName, PackagePath1), Obj != nullptr);
+		if (Obj == nullptr)
+		{
+			return false;
+		}
+
+		UPackage* PackagePath1Package = Obj->GetPackage();
+		ForEachObjectWithPackage(PackagePath1Package,
+			[&ActualObjects](UObject* Object)
+			{
+				ActualObjects.Add({ Object->GetFName(), Object->GetClass()->GetFName() });
+				return true;
+			}
+		);
+
+		VerifyLoadedObjects(ExpectedObjects, ActualObjects);
+		ActualObjects.Empty();
+
+		// Note, we are only GC'ing the importing package.
+		Package->ClearFlags(RF_Standalone);
+		FLoadingTestsScope::GarbageCollect({ { PackagePath2 } }, *this);
+	}
+
+	// Load the old package from memory and ensure we don't have any new objects created from importing
+	{
+		UPackage* Package = LoadPackage(nullptr, PackagePath1, LOAD_None);
+		TestTrue(FString::Printf(TEXT("Failed to load package at %s"), PackagePath1), Package != nullptr);
+
+		ForEachObjectWithPackage(Package,
+			[&ActualObjects](UObject* Object)
+			{
+				ActualObjects.Add({ Object->GetFName(), Object->GetClass()->GetFName() });
+				return true;
+			}
+		);
+
+		VerifyLoadedObjects(ExpectedObjects, ActualObjects);
+		ActualObjects.Empty();
+
+		FLoadingTestsScope::GarbageCollect({ { PackagePath2 } }, *this);
+	}
+
+	return true;
+}
+
+
+/**
+ * PackagePath1 leads to a user defined structure containing UFields which will be forcibly converted to FField when loading.
+ * When a sub-object is deleted, if the package is requested again we should reload the missing object and not create objects marked as invalid
+ */
+IMPLEMENT_SIMPLE_AUTOMATION_TEST(FLoadingTests_ReloadDestroyedExport, TEXT(TEST_NAME_ROOT ".InvalidExports.ReloadDestroyedExport"), EAutomationTestFlags::EditorContext | EAutomationTestFlags::EngineFilter)
+bool FLoadingTests_ReloadDestroyedExport::RunTest(const FString& Parameters)
+{
+	static constexpr const TCHAR* PackagePath1 = TEXT("/Game/Tests/Core/AssetLoading/InvalidatingExports/OldProcMeshData");
+	static constexpr const TCHAR* PackagePath2 = TEXT("/Game/Tests/Core/AssetLoading/InvalidatingExports/BP_UsingUStructWithUnconvertedFields");
+	static constexpr const TCHAR* ObjectName = TEXT("ProcMeshData");
+	typedef TSet<TPair<FName, FName>> FObjPathAndClassNameSet;
+	FObjPathAndClassNameSet ExpectedObjects;
+	FObjPathAndClassNameSet ActualObjects;
+
+	auto VerifyLoadedObjects = [this](const FObjPathAndClassNameSet& ExpectedObjects, const FObjPathAndClassNameSet& ActualObjects)
+		{
+			TestTrue(FString::Printf(TEXT("Expected the same number of objects in package %s after reloading. Expected: %d != Actual: %d"), PackagePath1, ExpectedObjects.Num(), ActualObjects.Num()), ExpectedObjects.Num() == ActualObjects.Num());
+			for (const TPair<FName, FName>& Pair : ExpectedObjects)
+			{
+				TestTrue(FString::Printf(TEXT("Missing obj %s (type: %s) after reloading %s"), *Pair.Key.ToString(), *Pair.Value.ToString(), PackagePath1), ActualObjects.Find(Pair) != nullptr);
+			}
+		};
+
+	// Load our package and populate the ExpectedObjects set
+	{
+		UPackage* Package = LoadPackage(nullptr, PackagePath1, LOAD_None);
+		TestTrue(FString::Printf(TEXT("Failed to load package at %s"), PackagePath1), Package != nullptr);
+
+		ForEachObjectWithPackage(Package,
+			[&ExpectedObjects](UObject* Object)
+			{
+				ExpectedObjects.Add({ Object->GetFName(), Object->GetClass()->GetFName() });
+				return true;
+			}
+		);
+
+		// Find a sub-object from the package and mark it for garbage collection, but do not destroy the package entirely
+		UUserDefinedStruct* Obj = FindObject<UUserDefinedStruct>(Package, ObjectName);
+		TestTrue(FString::Printf(TEXT("Failed to find expected object %s in package %s"), ObjectName, PackagePath1), Obj != nullptr);
+		Obj->ClearFlags(RF_Standalone);
+
+		// Note we are not using FLoadingTestsScope::GarbageCollect as that will mark all sub-objects and the package for collection
+		CollectGarbage(RF_Standalone);
+	}
+
+	// Reloading the package should re-populate any deleted sub-objects as long as the objects weren't marked as invalid to load
+	{
+		UPackage* Package = LoadPackage(nullptr, PackagePath1, LOAD_None);
+		TestTrue(FString::Printf(TEXT("Failed to load package at %s"), PackagePath1), Package != nullptr);
+
+		ForEachObjectWithPackage(Package,
+			[&ActualObjects](UObject* Object)
+			{
+				ActualObjects.Add({ Object->GetFName(), Object->GetClass()->GetFName() });
+				return true;
+			}
+		);
+
+		// Our deleted object should have been restored
+		UUserDefinedStruct* Obj = FindObject<UUserDefinedStruct>(Package, ObjectName);
+		TestTrue(FString::Printf(TEXT("Failed to find expected object %s in package %s"), ObjectName, PackagePath1), Obj != nullptr);
+
+		VerifyLoadedObjects(ExpectedObjects, ActualObjects);
+		ActualObjects.Empty();
+
+		FLoadingTestsScope::GarbageCollect({ {PackagePath1} }, *this);
+	}
+
+	return true;
+}
+
 #undef TEST_NAME_ROOT
 #endif // WITH_DEV_AUTOMATION_TESTS
\ No newline at end of file
diff --git a/Engine/Source/Runtime/Engine/Private/Tests/Loading/AsyncLoadingTests_Shared.cpp b/Engine/Source/Runtime/Engine/Private/Tests/Loading/AsyncLoadingTests_Shared.cpp
--- a/Engine/Source/Runtime/Engine/Private/Tests/Loading/AsyncLoadingTests_Shared.cpp
+++ b/Engine/Source/Runtime/Engine/Private/Tests/Loading/AsyncLoadingTests_Shared.cpp
@@ -50,6 +50,11 @@ void FLoadingTestsScope::SavePackages()
 }
 
 void FLoadingTestsScope::GarbageCollect()
+{
+	GarbageCollect(PackageNames, AutomationTest);
+}
+
+void FLoadingTestsScope::GarbageCollect(const TArray<FString>& PackageNames, FAutomationTestBase& AutomationTest)
 {
 	TArray<FString> ObjectPaths;
 
@@ -75,16 +80,21 @@ void FLoadingTestsScope::GarbageCollect()
 	// Make sure everything we gathered can be properly found
 	for (const FString& ObjectPath : ObjectPaths)
 	{
-		checkf(FindObject<UObject>(nullptr, *ObjectPath) != nullptr, TEXT("%s should be present in memory"), *ObjectPath);
+		AutomationTest.TestTrue(FString::Printf(TEXT("%s should be present in memory"), *ObjectPath), FindObject<UObject>(nullptr, *ObjectPath) != nullptr);
 	}
 
-	// GC and make sure everything gets cleaned up before loading
-	CollectGarbage(GARBAGE_COLLECTION_KEEPFLAGS);
+	{
+		// Make sure the GIsInitialLoad flag is false.  Otherwise GC does nothing
+		TGuardValue<bool> GuardIsInitialLoad(GIsInitialLoad, false);
+
+		// GC and make sure everything gets cleaned up before loading
+		CollectGarbage(GARBAGE_COLLECTION_KEEPFLAGS);
+	}
 
 	// Now make sure everything is gone
 	for (const FString& ObjectPath : ObjectPaths)
 	{
-		checkf(FindObject<UObject>(nullptr, *ObjectPath) == nullptr, TEXT("%s should have been garbage collected"), *ObjectPath);
+		AutomationTest.TestTrue(FString::Printf(TEXT("%s should have been garbage collected"), *ObjectPath), FindObject<UObject>(nullptr, *ObjectPath) == nullptr);
 	}
 }
 
diff --git a/Engine/Source/Runtime/Engine/Private/Tests/Loading/AsyncLoadingTests_Shared.h b/Engine/Source/Runtime/Engine/Private/Tests/Loading/AsyncLoadingTests_Shared.h
--- a/Engine/Source/Runtime/Engine/Private/Tests/Loading/AsyncLoadingTests_Shared.h
+++ b/Engine/Source/Runtime/Engine/Private/Tests/Loading/AsyncLoadingTests_Shared.h
@@ -105,6 +105,7 @@ public:
 	void LoadObjects();
 	void CleanupObjects();
 	void GarbageCollect();
+	static void GarbageCollect(const TArray<FString>& PackageNames, FAutomationTestBase& AutomationTest);
 
 	FLoadingTestsScope(FAutomationTestBase* InAutomationTest, TFunction<void (FLoadingTestsScope&)> InMutateObjects = nullptr)
 		: AutomationTest(*InAutomationTest)
diff --git a/Engine/Source/Runtime/UMG/Private/WidgetBlueprintGeneratedClass.cpp b/Engine/Source/Runtime/UMG/Private/WidgetBlueprintGeneratedClass.cpp
--- a/Engine/Source/Runtime/UMG/Private/WidgetBlueprintGeneratedClass.cpp
+++ b/Engine/Source/Runtime/UMG/Private/WidgetBlueprintGeneratedClass.cpp
@@ -382,16 +382,18 @@ void UWidgetBlueprintGeneratedClass::PurgeClass(bool bRecompilingOnLoad)
 	// Remove the old widdget tree.
 	if ( WidgetTree )
 	{
-		WidgetTree->Rename(nullptr, GetTransientPackage(), RenFlags);
+        // Rename will remove the renamed object's linker when moving to a new package so invalidate the export beforehand
 		FLinkerLoad::InvalidateExport(WidgetTree);
+		WidgetTree->Rename(nullptr, GetTransientPackage(), RenFlags);
 		WidgetTree = nullptr;
 	}
 
 	// Remove all animations.
 	for ( UWidgetAnimation* Animation : Animations )
 	{
-		Animation->Rename(nullptr, GetTransientPackage(), RenFlags);
+        // Rename will remove the renamed object's linker when moving to a new package so invalidate the export beforehand
 		FLinkerLoad::InvalidateExport(Animation);
+		Animation->Rename(nullptr, GetTransientPackage(), RenFlags);
 	}
 
 	Animations.Empty();
