Remove 5.0 deprecated online functions
Modify UGameInstance::FindLocalPlayerFromUniqueNetId to prefer the repl implementation

diff --git a/Engine/Source/Runtime/CoreOnline/Public/Online/CoreOnline.h b/Engine/Source/Runtime/CoreOnline/Public/Online/CoreOnline.h
--- a/Engine/Source/Runtime/CoreOnline/Public/Online/CoreOnline.h
+++ b/Engine/Source/Runtime/CoreOnline/Public/Online/CoreOnline.h
@@ -38,14 +38,6 @@ class FLazySingleton;
 	#endif
 #endif //MAX_LOCAL_PLAYERS
 
-UE_DEPRECATED(5.0, "Use NAME_GameSession.")
-inline constexpr EName GameSessionName = NAME_GameSession;
-UE_DEPRECATED(5.0, "Use NAME_PartySession.")
-inline constexpr EName PartySessionName = NAME_PartySession;
-UE_DEPRECATED(5.0, "Use NAME_GamePort.")
-inline constexpr EName GamePort = NAME_GamePort;
-UE_DEPRECATED(5.0, "Use NAME_BeaconPort.")
-inline constexpr EName BeaconPort = NAME_BeaconPort;
 
 USTRUCT(noexport)
 struct FJoinabilitySettings
diff --git a/Engine/Source/Runtime/Engine/Classes/Engine/Engine.h b/Engine/Source/Runtime/Engine/Classes/Engine/Engine.h
--- a/Engine/Source/Runtime/Engine/Classes/Engine/Engine.h
+++ b/Engine/Source/Runtime/Engine/Classes/Engine/Engine.h
@@ -2750,10 +2750,6 @@ public:
 	/** @return the currently active audio device */
 	ENGINE_API FAudioDeviceHandle GetActiveAudioDevice();
 
-	/** @return whether we currently have more than one local player */
-	UE_DEPRECATED(5.0, "IsSplitScreen was only ever checking if there are more than one local player. Use HasMultipleLocalPlayers instead.")
-	ENGINE_API virtual bool IsSplitScreen(UWorld *InWorld);
-
 	/** @returns whether there are currently multiple local players in the given world */
 	ENGINE_API virtual bool HasMultipleLocalPlayers(UWorld* InWorld);
 
diff --git a/Engine/Source/Runtime/Engine/Classes/Engine/GameInstance.h b/Engine/Source/Runtime/Engine/Classes/Engine/GameInstance.h
--- a/Engine/Source/Runtime/Engine/Classes/Engine/GameInstance.h
+++ b/Engine/Source/Runtime/Engine/Classes/Engine/GameInstance.h
@@ -404,14 +404,6 @@ public:
 	 */
 	ENGINE_API APlayerController* GetPrimaryPlayerController(bool bRequiresValidUniqueId = true) const;
 
-	/**
-	 * Get the unique id for the primary player on this machine (others are splitscreen children)
-	 *
-	 * @return the unique id of the primary player on this machine
-	 */
-	UE_DEPRECATED(5.0, "Use GetPrimaryPlayerUniqueIdRepl.")
-	ENGINE_API FUniqueNetIdPtr GetPrimaryPlayerUniqueId() const;
-
 	/**
 	 * Get the unique id for the primary player on this machine (others are splitscreen children)
 	 *
diff --git a/Engine/Source/Runtime/Engine/Classes/Engine/LocalPlayer.h b/Engine/Source/Runtime/Engine/Classes/Engine/LocalPlayer.h
--- a/Engine/Source/Runtime/Engine/Classes/Engine/LocalPlayer.h
+++ b/Engine/Source/Runtime/Engine/Classes/Engine/LocalPlayer.h
@@ -534,12 +534,6 @@ public:
 	 */
 	ENGINE_API FUniqueNetIdRepl GetCachedUniqueNetId() const;
 
-	/** Sets the players current cached unique net id */
-	UE_DEPRECATED(5.0, "Use SetCachedUniqueNetId with FUniqueNetIdRepl")
-	ENGINE_API void SetCachedUniqueNetId(FUniqueNetIdPtr NewUniqueNetId);
-	/** Sets the players current cached unique net id */
-	UE_DEPRECATED(5.0, "Use SetCachedUniqueNetId with FUniqueNetIdRepl")
-	ENGINE_API void SetCachedUniqueNetId(TYPE_OF_NULLPTR);
 	/** Sets the players current cached unique net id */
 	ENGINE_API void SetCachedUniqueNetId(const FUniqueNetIdRepl& NewUniqueNetId);
 
diff --git a/Engine/Source/Runtime/Engine/Classes/Engine/World.h b/Engine/Source/Runtime/Engine/Classes/Engine/World.h
--- a/Engine/Source/Runtime/Engine/Classes/Engine/World.h
+++ b/Engine/Source/Runtime/Engine/Classes/Engine/World.h
@@ -3515,8 +3515,6 @@ public:
 	 * @param InNetPlayerIndex (optional) - the NetPlayerIndex to set on the PlayerController
 	 * @return the PlayerController that was spawned (may fail and return NULL)
 	 */
-	UE_DEPRECATED(5.0, "Use SpawnPlayActor with FUniqueNetIdRepl")
-	APlayerController* SpawnPlayActor(class UPlayer* Player, ENetRole RemoteRole, const FURL& InURL, const FUniqueNetIdPtr& UniqueId, FString& Error, uint8 InNetPlayerIndex = 0);
 	APlayerController* SpawnPlayActor(class UPlayer* Player, ENetRole RemoteRole, const FURL& InURL, const FUniqueNetIdRepl& UniqueId, FString& Error, uint8 InNetPlayerIndex = 0);
 	
 	/**
diff --git a/Engine/Source/Runtime/Engine/Classes/GameFramework/GameSession.h b/Engine/Source/Runtime/Engine/Classes/GameFramework/GameSession.h
--- a/Engine/Source/Runtime/Engine/Classes/GameFramework/GameSession.h
+++ b/Engine/Source/Runtime/Engine/Classes/GameFramework/GameSession.h
@@ -77,15 +77,6 @@ class AGameSession : public AInfo
 	 */
 	ENGINE_API virtual FString ApproveLogin(const FString& Options);
 
-	/**
-	 * Register a player with the online service session
-	 * @param NewPlayer player to register
-	 * @param UniqueId uniqueId they sent over on Login
-	 * @param bWasFromInvite was this from an invite
-	 */
-	UE_DEPRECATED(5.0, "Use RegisterPlayer with FUniqueNetIdRepl")
-	ENGINE_API virtual void RegisterPlayer(APlayerController* NewPlayer, const FUniqueNetIdPtr& UniqueId, bool bWasFromInvite);
-
 	/**
 	 * Register a player with the online service session
 	 * @param NewPlayer player to register
@@ -130,14 +121,6 @@ class AGameSession : public AInfo
 	 */
 	ENGINE_API virtual void UnregisterPlayer(FName InSessionName, const FUniqueNetIdRepl& UniqueId);
 
-	/**
-	 * Unregister players from the online service session
-	 *
-	 * @param SessionName name of session to unregister from
-	 * @param Players ids of the players to unregister
-	 */
-	UE_DEPRECATED(5.0, "Use UnregisterPlayers with FUniqueNetIdRepl")
-	ENGINE_API virtual void UnregisterPlayers(FName InSessionName, const TArray< FUniqueNetIdRef >& Players);
 	ENGINE_API virtual void UnregisterPlayers(FName InSessionName, const TArray< FUniqueNetIdRepl >& Players);
 	
 	/**
@@ -285,14 +268,6 @@ private:
 	TOptional<int32> MaxSpectatorsOptionOverride;
 };
 
-/** 
- * Returns the player controller associated with this net id
- * @param PlayerNetId the id to search for
- * @return the player controller if found, otherwise NULL
- */
-UE_DEPRECATED(5.0, "Use GetPlayerControllerFromNetId with FUniqueNetIdRepl")
-ENGINE_API APlayerController* GetPlayerControllerFromNetId(UWorld* World, const FUniqueNetId& PlayerNetId);
-
 /**
  * Returns the player controller associated with this net id
  * @param PlayerNetId the id to search for
diff --git a/Engine/Source/Runtime/Engine/Private/GameInstance.cpp b/Engine/Source/Runtime/Engine/Private/GameInstance.cpp
--- a/Engine/Source/Runtime/Engine/Private/GameInstance.cpp
+++ b/Engine/Source/Runtime/Engine/Private/GameInstance.cpp
@@ -1134,12 +1134,6 @@ APlayerController* UGameInstance::GetPrimaryPlayerController(bool bRequiresValid
 	return PrimaryController;
 }
 
-FUniqueNetIdPtr UGameInstance::GetPrimaryPlayerUniqueId() const
-{
-	FUniqueNetIdRepl UniqueIdRepl = GetPrimaryPlayerUniqueIdRepl();
-	return UniqueIdRepl.GetV1();
-}
-
 FUniqueNetIdRepl UGameInstance::GetPrimaryPlayerUniqueIdRepl() const
 {
 	ULocalPlayer* PrimaryLP = nullptr;
@@ -1190,6 +1184,21 @@ ULocalPlayer* UGameInstance::FindLocalPlayerFromPlatformUserId(const FPlatformUs
 }
 
 ULocalPlayer* UGameInstance::FindLocalPlayerFromUniqueNetId(const FUniqueNetId& UniqueNetId) const
+{
+	return FindLocalPlayerFromUniqueNetId(FUniqueNetIdRepl(UniqueNetId.AsShared()));
+}
+
+ULocalPlayer* UGameInstance::FindLocalPlayerFromUniqueNetId(FUniqueNetIdPtr UniqueNetId) const
+{
+	if (!UniqueNetId.IsValid())
+	{
+		return nullptr;
+	}
+
+	return FindLocalPlayerFromUniqueNetId(*UniqueNetId);
+}
+
+ULocalPlayer* UGameInstance::FindLocalPlayerFromUniqueNetId(const FUniqueNetIdRepl& UniqueNetId) const
 {
 	for (ULocalPlayer* Player : LocalPlayers)
 	{
@@ -1211,26 +1220,6 @@ ULocalPlayer* UGameInstance::FindLocalPlayerFromUniqueNetId(const FUniqueNetId&
 	return nullptr;
 }
 
-ULocalPlayer* UGameInstance::FindLocalPlayerFromUniqueNetId(FUniqueNetIdPtr UniqueNetId) const
-{
-	if (!UniqueNetId.IsValid())
-	{
-		return nullptr;
-	}
-
-	return FindLocalPlayerFromUniqueNetId(*UniqueNetId);
-}
-
-ULocalPlayer* UGameInstance::FindLocalPlayerFromUniqueNetId(const FUniqueNetIdRepl& UniqueNetId) const
-{
-	if (!UniqueNetId.IsValid())
-	{
-		return nullptr;
-	}
-
-	return FindLocalPlayerFromUniqueNetId(*UniqueNetId);
-}
-
 ULocalPlayer* UGameInstance::GetFirstGamePlayer() const
 {
 	return (LocalPlayers.Num() > 0) ? LocalPlayers[0] : nullptr;
diff --git a/Engine/Source/Runtime/Engine/Private/GameSession.cpp b/Engine/Source/Runtime/Engine/Private/GameSession.cpp
--- a/Engine/Source/Runtime/Engine/Private/GameSession.cpp
+++ b/Engine/Source/Runtime/Engine/Private/GameSession.cpp
@@ -267,11 +267,6 @@ int32 AGameSession::GetNextPlayerID()
 	return NextPlayerID++;
 }
 
-void AGameSession::RegisterPlayer(APlayerController* NewPlayer, const FUniqueNetIdPtr& UniqueId, bool bWasFromInvite)
-{
-	RegisterPlayer(NewPlayer, FUniqueNetIdRepl(UniqueId), bWasFromInvite);
-}
-
 void AGameSession::RegisterPlayer(APlayerController* NewPlayer, const FUniqueNetIdRepl& UniqueId, bool bWasFromInvite)
 {
 	if (NewPlayer != NULL)
@@ -296,23 +291,6 @@ void AGameSession::UnregisterPlayer(FName InSessionName, const FUniqueNetIdRepl&
 	}
 }
 
-void AGameSession::UnregisterPlayers(FName InSessionName, const TArray<FUniqueNetIdRef>& Players)
-{
-	UWorld* World = GetWorld();
-	if (GetNetMode() != NM_Standalone &&
-		Players.Num() > 0 &&
-		UOnlineEngineInterface::Get()->DoesSessionExist(World, InSessionName))
-	{
-		// Remove the player from the session
-		TArray<FUniqueNetIdWrapper> PlayerIdsAsWrappers;
-		for (const FUniqueNetIdRef& PlayerId : Players)
-		{
-			PlayerIdsAsWrappers.Emplace(PlayerId);
-		}
-		UOnlineEngineInterface::Get()->UnregisterPlayers(World, InSessionName, PlayerIdsAsWrappers);
-	}
-}
-
 void AGameSession::UnregisterPlayers(FName InSessionName, const TArray<FUniqueNetIdRepl>& Players)
 {
 	UWorld * World = GetWorld();
diff --git a/Engine/Source/Runtime/Engine/Private/LevelActor.cpp b/Engine/Source/Runtime/Engine/Private/LevelActor.cpp
--- a/Engine/Source/Runtime/Engine/Private/LevelActor.cpp
+++ b/Engine/Source/Runtime/Engine/Private/LevelActor.cpp
@@ -1019,12 +1019,6 @@ bool UWorld::DestroyActor( AActor* ThisActor, bool bNetForce, bool bShouldModify
 	Player spawning.
 -----------------------------------------------------------------------------*/
 
-APlayerController* UWorld::SpawnPlayActor(UPlayer* NewPlayer, ENetRole RemoteRole, const FURL& InURL, const FUniqueNetIdPtr& UniqueId, FString& Error, uint8 InNetPlayerIndex)
-{
-	FUniqueNetIdRepl UniqueIdRepl(UniqueId);
-	return SpawnPlayActor(NewPlayer, RemoteRole, InURL, UniqueIdRepl, Error, InNetPlayerIndex);
-}
-
 APlayerController* UWorld::SpawnPlayActor(UPlayer* NewPlayer, ENetRole RemoteRole, const FURL& InURL, const FUniqueNetIdRepl& UniqueId, FString& Error, uint8 InNetPlayerIndex)
 {
 	Error = TEXT("");
diff --git a/Engine/Source/Runtime/Engine/Private/LocalPlayer.cpp b/Engine/Source/Runtime/Engine/Private/LocalPlayer.cpp
--- a/Engine/Source/Runtime/Engine/Private/LocalPlayer.cpp
+++ b/Engine/Source/Runtime/Engine/Private/LocalPlayer.cpp
@@ -1748,16 +1748,6 @@ FUniqueNetIdRepl ULocalPlayer::GetCachedUniqueNetId() const
 	return CachedUniqueNetId;
 }
 
-void ULocalPlayer::SetCachedUniqueNetId(FUniqueNetIdPtr NewUniqueNetId)
-{
-	CachedUniqueNetId = NewUniqueNetId;
-}
-
-void ULocalPlayer::SetCachedUniqueNetId(TYPE_OF_NULLPTR)
-{
-	CachedUniqueNetId = FUniqueNetIdRepl(nullptr);
-}
-
 void ULocalPlayer::SetCachedUniqueNetId(const FUniqueNetIdRepl& NewUniqueNetId)
 {
 	CachedUniqueNetId = NewUniqueNetId;
diff --git a/Engine/Source/Runtime/Engine/Private/UnrealEngine.cpp b/Engine/Source/Runtime/Engine/Private/UnrealEngine.cpp
--- a/Engine/Source/Runtime/Engine/Private/UnrealEngine.cpp
+++ b/Engine/Source/Runtime/Engine/Private/UnrealEngine.cpp
@@ -4070,12 +4070,6 @@ void UEngine::RecordHMDAnalytics()
 	}
 }
 
-/** @return whether we currently have more than one local player */
-bool UEngine::IsSplitScreen(UWorld *InWorld)
-{
-	return HasMultipleLocalPlayers(InWorld);
-}
-
 bool UEngine::HasMultipleLocalPlayers(UWorld* InWorld)
 {
 	if (InWorld == NULL)
diff --git a/Engine/Source/Runtime/Engine/Public/Net/OnlineEngineInterface.h b/Engine/Source/Runtime/Engine/Public/Net/OnlineEngineInterface.h
--- a/Engine/Source/Runtime/Engine/Public/Net/OnlineEngineInterface.h
+++ b/Engine/Source/Runtime/Engine/Public/Net/OnlineEngineInterface.h
@@ -64,9 +64,6 @@ public:
 	/** Returns the name of the default online subsystem, generally the DefaultPlatformService config value. */
 	virtual FName GetDefaultOnlineSubsystemName() const { return NAME_None; }
 	/** Returns whether the unique net id is compatible */
-	UE_DEPRECATED(5.0, "Use IsCompatibleUniqueNetId with FUniqueNetIdWrapper")
-	virtual bool IsCompatibleUniqueNetId(const FUniqueNetId& InUniqueNetId) const { return IsCompatibleUniqueNetId(FUniqueNetIdWrapper(InUniqueNetId)); }
-	/** Returns whether the unique net id is compatible */
 	virtual bool IsCompatibleUniqueNetId(const FUniqueNetIdWrapper& InUniqueNetId) const { return (InUniqueNetId.IsV1() && InUniqueNetId.GetType() == GetDefaultOnlineSubsystemName()); }
 
 	/**
@@ -95,20 +92,11 @@ public:
 	 * Identity
 	 */
 
-	/** @return a proper FUniqueNetId pointer generated by the default OnlineSubsystem */
-	UE_DEPRECATED(5.0, "Use CreateUniquePlayerIdWrapper")
-	virtual FUniqueNetIdPtr CreateUniquePlayerId(const FString& Str, FName Type = NAME_None) { return CreateUniquePlayerIdWrapper(Str, Type).GetV1(); }
 	/** @return a proper FUniqueNetId pointer generated by the default OnlineSubsystem */
 	virtual FUniqueNetIdWrapper CreateUniquePlayerIdWrapper(const FString& Str, FName Type = NAME_None) { return FUniqueNetIdWrapper(); }
 	/** @return the unique id for a given local user on the given OnlineSubsystem */
-	UE_DEPRECATED(5.0, "Use GetUniquePlayerIdWrapper")
-	virtual FUniqueNetIdPtr GetUniquePlayerId(UWorld* World, int32 LocalUserNum, FName Type = NAME_None) { return GetUniquePlayerIdWrapper(World, LocalUserNum, Type).GetV1(); }
-	/** @return the unique id for a given local user on the given OnlineSubsystem */
 	virtual FUniqueNetIdWrapper GetUniquePlayerIdWrapper(UWorld* World, int32 LocalUserNum, FName Type = NAME_None) { return FUniqueNetIdWrapper(); }
 
-	/** @return the stored nickname on the OnlineSubsystem for a player with a given unique id */
-	UE_DEPRECATED(5.0, "Use GetPlayerNickname with FUniqueNetIdWrapper")
-	virtual FString GetPlayerNickname(UWorld* World, const FUniqueNetId& UniqueId) { return GetPlayerNickname(World, FUniqueNetIdWrapper(UniqueId)); }
 	/** @return the stored nickname on the OnlineSubsystem for a player with a given unique id */
 	virtual FString GetPlayerNickname(UWorld* World, const FUniqueNetIdWrapper& UniqueId) { return TEXT("INVALID"); }
 	/** @return the stored nickname on the secondary/console OnlineSubsystem for a player with a given unique id */
@@ -133,28 +121,11 @@ public:
 	virtual bool GetSessionJoinability(UWorld* World, FName SessionName, FJoinabilitySettings& OutSettings) { return false; }
 	virtual void UpdateSessionJoinability(UWorld* World, FName SessionName, bool bPublicSearchable, bool bAllowInvites, bool bJoinViaPresence, bool bJoinViaPresenceFriendsOnly) {}
 
-	/** Register a given player with the online session */
-	UE_DEPRECATED(5.0, "Use RegisterPlayer with FUniqueNetIdWrapper")
-	virtual void RegisterPlayer(UWorld* World, FName SessionName, const FUniqueNetId& UniqueId, bool bWasInvited) { RegisterPlayer(World, SessionName, FUniqueNetIdWrapper(UniqueId), bWasInvited); }
 	/** Register a given player with the online session */
 	virtual void RegisterPlayer(UWorld* World, FName SessionName, const FUniqueNetIdWrapper& UniqueId, bool bWasInvited) {}
 	/** Unregister a given player from the online session */
-	UE_DEPRECATED(5.0, "Use UnregisterPlayer with FUniqueNetIdWrapper")
-	virtual void UnregisterPlayer(UWorld* World, FName SessionName, const FUniqueNetId& UniqueId) { UnregisterPlayer(World, SessionName, FUniqueNetIdWrapper(UniqueId)); }
-	/** Unregister a given player from the online session */
 	virtual void UnregisterPlayer(UWorld* World, FName SessionName, const FUniqueNetIdWrapper& UniqueId) {}
 	/** Unregister a given player from the online session */
-	UE_DEPRECATED(5.0, "Use UnregisterPlayers with FUniqueNetIdWrapper")
-	virtual void UnregisterPlayers(UWorld* World, FName SessionName, const TArray<FUniqueNetIdRef>& Players)
-	{
-		TArray<FUniqueNetIdWrapper> PlayerIdsAsWrappers;
-		for (const FUniqueNetIdRef& PlayerId : Players)
-		{
-			PlayerIdsAsWrappers.Emplace(PlayerId);
-		}
-		UnregisterPlayers(World, SessionName, PlayerIdsAsWrappers);
-	}
-	/** Unregister a given player from the online session */
 	virtual void UnregisterPlayers(UWorld* World, FName SessionName, const TArray<FUniqueNetIdWrapper>& Players) {}
 	/** @return true if there is a valid URL that can been used to connect to a given session, false otherwise */
 	virtual bool GetResolvedConnectString(UWorld* World, FName SessionName, FString& URL) { return false; }
@@ -175,15 +146,9 @@ public:
 	/** Clears all voice packets currently queued for send */
 	virtual void ClearVoicePackets(UWorld* World) {}
 
-	/** Mute a remote user for a given local player */
-	UE_DEPRECATED(5.0, "Use MuteRemoteTalker with FUniqueNetIdWrapper")
-	virtual bool MuteRemoteTalker(UWorld* World, uint8 LocalUserNum, const FUniqueNetId& PlayerId, bool bIsSystemWide) { return MuteRemoteTalker(World, LocalUserNum, FUniqueNetIdWrapper(PlayerId), bIsSystemWide); }
 	/** Mute a remote user for a given local player */
 	virtual bool MuteRemoteTalker(UWorld* World, uint8 LocalUserNum, const FUniqueNetIdWrapper& PlayerId, bool bIsSystemWide) { return false; }
 	/** UnMute a remote user for a given local player */
-	UE_DEPRECATED(5.0, "Use UnmuteRemoteTalker with FUniqueNetIdWrapper")
-	virtual bool UnmuteRemoteTalker(UWorld* World, uint8 LocalUserNum, const FUniqueNetId& PlayerId, bool bIsSystemWide) { return UnmuteRemoteTalker(World, LocalUserNum, FUniqueNetIdWrapper(PlayerId), bIsSystemWide); }
-	/** UnMute a remote user for a given local player */
 	virtual bool UnmuteRemoteTalker(UWorld* World, uint8 LocalUserNum, const FUniqueNetIdWrapper& PlayerId, bool bIsSystemWide) { return false; }
 	/** @return number of available/registered local talkers */
 	virtual int32 GetNumLocalTalkers(UWorld* World) { return 0; }
