Prevent a null pointer crash when attempting to replicate an actor that doesn't belong to a world (e.g. was streamed out).

diff --git a/Engine/Source/Runtime/Engine/Private/DataChannel.cpp b/Engine/Source/Runtime/Engine/Private/DataChannel.cpp
--- a/Engine/Source/Runtime/Engine/Private/DataChannel.cpp
+++ b/Engine/Source/Runtime/Engine/Private/DataChannel.cpp
@@ -3407,7 +3407,11 @@ int64 UActorChannel::ReplicateActor()
 	check(nullptr != Cast<UPackageMapClient>(Connection->PackageMap));
 
 	const UWorld* const ActorWorld = Actor->GetWorld();
-	checkf(ActorWorld, TEXT("ActorWorld for Actor [%s] is Null"), *GetPathNameSafe(Actor));
+	ensureMsgf(ActorWorld, TEXT("ActorWorld for Actor [%s] is Null"), *GetPathNameSafe(Actor));
+	if (ActorWorld == nullptr)
+	{
+		return 0;
+	}
 
 #if STATS || ENABLE_STATNAMEDEVENTS
 	UClass* ParentNativeClass = GetParentNativeClass(Actor->GetClass());
