Change client pawn controller replication logic to properly call OnRep_Controller and NotifyControllerChanged when the controller is changed by code like ClientRestart that bypasses normal replication.
This change can be disabled by setting Controller.AlwaysNotifyClientOnControllerChange to 0.
Moved some existing cvars into UE::Gameplay::CVars to match coding standard.
This is based on pull requests #9467 and #9008

diff --git a/Engine/Source/Runtime/Engine/Private/Controller.cpp b/Engine/Source/Runtime/Engine/Private/Controller.cpp
--- a/Engine/Source/Runtime/Engine/Private/Controller.cpp
+++ b/Engine/Source/Runtime/Engine/Private/Controller.cpp
@@ -28,17 +28,22 @@ DEFINE_LOG_CATEGORY(LogPath);
 
 #define LOCTEXT_NAMESPACE "Controller"
 
-namespace ControllerStatics
+namespace UE::Gameplay::CVars
 {
-	static float InvalidControlRotationMagnitude = 8388608.f; // 2^23, largest float when fractions are lost, and where FMod loses meaningful precision.
+	bool bAlwaysNotifyClientOnControllerChange = true;
+	static FAutoConsoleVariableRef CVarAlwaysNotifyClientOnControllerChange(
+		TEXT("Controller.AlwaysNotifyClientOnControllerChange"), bAlwaysNotifyClientOnControllerChange,
+		TEXT("If true, OnRep_Controller and NotifyControllerChanged are always called on the client when a pawn's controller changes.\n")
+		TEXT("If false, use backward compatible controller notification logic that does not handle networking properly."),
+		ECVF_Default);
+
+	float InvalidControlRotationMagnitude = 8388608.f; // 2^23, largest float when fractions are lost, and where FMod loses meaningful precision.
 	static FAutoConsoleVariableRef CVarInvalidControlRotationMagnitude(
 		TEXT("Controller.InvalidControlRotationMagnitude"), InvalidControlRotationMagnitude,
 		TEXT("If any component of an FRotator passed to SetControlRotation is larger than this magnitude, ignore the value. Huge values are usually from uninitialized variables and can cause NaN/Inf to propagate later."),
 		ECVF_Default);
 }
 
-
-
 AController::AController(const FObjectInitializer& ObjectInitializer)
 	: Super(ObjectInitializer)
 {
@@ -145,9 +150,9 @@ bool AController::IsValidControlRotation(FRotator CheckRotation) const
 
 	// Really large values can be technically valid but are usually the result of uninitialized values, and those can cause
 	// conversion to FQuat or Vector to fail and generate NaN or Inf.
-	if (FMath::Abs(CheckRotation.Pitch) >= ControllerStatics::InvalidControlRotationMagnitude ||
-		FMath::Abs(CheckRotation.Yaw  ) >= ControllerStatics::InvalidControlRotationMagnitude ||
-		FMath::Abs(CheckRotation.Roll ) >= ControllerStatics::InvalidControlRotationMagnitude)
+	if (FMath::Abs(CheckRotation.Pitch) >= UE::Gameplay::CVars::InvalidControlRotationMagnitude ||
+		FMath::Abs(CheckRotation.Yaw  ) >= UE::Gameplay::CVars::InvalidControlRotationMagnitude ||
+		FMath::Abs(CheckRotation.Roll ) >= UE::Gameplay::CVars::InvalidControlRotationMagnitude)
 	{
 		return false;
 	}
@@ -530,6 +535,11 @@ void AController::OnRep_Pawn()
 	{
 		// Set the old controller to NULL, since we are no longer the owner, and can't rely on it replicating to us anymore
 		StrongOldPawn->Controller = nullptr;
+		if (UE::Gameplay::CVars::bAlwaysNotifyClientOnControllerChange)
+		{
+			// This will notify other systems like the game instance
+			StrongOldPawn->OnRep_Controller();
+		}
 	}
 
 	OldPawn = Pawn;
diff --git a/Engine/Source/Runtime/Engine/Private/Pawn.cpp b/Engine/Source/Runtime/Engine/Private/Pawn.cpp
--- a/Engine/Source/Runtime/Engine/Private/Pawn.cpp
+++ b/Engine/Source/Runtime/Engine/Private/Pawn.cpp
@@ -37,6 +37,11 @@ DEFINE_LOG_CATEGORY_STATIC(LogPawn, Warning, All);
 
 FOnPawnBeginPlay APawn::OnPawnBeginPlay;
 
+namespace UE::Gameplay::CVars
+{
+	extern bool bAlwaysNotifyClientOnControllerChange;
+}
+
 APawn::APawn(const FObjectInitializer& ObjectInitializer)
 	: Super(ObjectInitializer)
 {
@@ -582,8 +587,10 @@ FRotator APawn::GetControlRotation() const
 
 void APawn::OnRep_Controller()
 {
-	bool bNotifyControllerChange = (Controller == nullptr);
-
+	bool bNotifyControllerChange = UE::Gameplay::CVars::bAlwaysNotifyClientOnControllerChange ?
+		(Controller != PreviousController) :	// By default, notify whenever the PreviousController is out of date for any reason
+		(Controller == nullptr);				// In backward compatibility, only notify when changing from null or the edge case below
+	
 	if ( (Controller != nullptr) && (Controller->GetPawn() == nullptr) )
 	{
 		// This ensures that AController::OnRep_Pawn is called. Since we cant ensure replication order of APawn::Controller and AController::Pawn,
diff --git a/Engine/Source/Runtime/Engine/Private/PlayerController.cpp b/Engine/Source/Runtime/Engine/Private/PlayerController.cpp
--- a/Engine/Source/Runtime/Engine/Private/PlayerController.cpp
+++ b/Engine/Source/Runtime/Engine/Private/PlayerController.cpp
@@ -86,23 +86,25 @@ DECLARE_CYCLE_STAT(TEXT("PC Build Input Stack"), STAT_PC_BuildInputStack, STATGR
 DECLARE_CYCLE_STAT(TEXT("PC Process Input Stack"), STAT_PC_ProcessInputStack, STATGROUP_PlayerController);
 
 // CVars
-namespace PlayerControllerCVars
+namespace UE::Gameplay::CVars
 {
 	// Resync timestamps on pawn ack
-	static int32 NetResetServerPredictionDataOnPawnAck = 1;
-	FAutoConsoleVariableRef CVarNetResetServerPredictionDataOnPawnAck(
+	int32 NetResetServerPredictionDataOnPawnAck = 1;
+	static FAutoConsoleVariableRef CVarNetResetServerPredictionDataOnPawnAck(
 		TEXT("PlayerController.NetResetServerPredictionDataOnPawnAck"),
 		NetResetServerPredictionDataOnPawnAck,
 		TEXT("Whether to reset server prediction data for the possessed Pawn when the pawn ack handshake completes.\n")
 		TEXT("0: Disable, 1: Enable"),
 		ECVF_Default);
 
-	static int32 ForceUsingCameraAsStreamingSource = 0;
-	FAutoConsoleVariableRef CVarForceUsingCameraAsStreamingSource(
+	int32 ForceUsingCameraAsStreamingSource = 0;
+	static FAutoConsoleVariableRef CVarForceUsingCameraAsStreamingSource(
 		TEXT("wp.Runtime.PlayerController.ForceUsingCameraAsStreamingSource"),
 		ForceUsingCameraAsStreamingSource,
 		TEXT("Whether to force the use of the camera as the streaming source for World Partition. By default the player pawn is used.\n")
 		TEXT("0: Use pawn as streaming source, 1: Use camera as streaming source"));
+
+	extern bool bAlwaysNotifyClientOnControllerChange;
 }
 
 namespace NetworkPhysicsCvars
@@ -806,7 +808,15 @@ void APlayerController::ClientRestart_Implementation(APawn* NewPawn)
 	if (OldController != this)
 	{
 		// In case this is received before APawn::OnRep_Controller is called
-		GetPawn()->NotifyControllerChanged();
+		if (UE::Gameplay::CVars::bAlwaysNotifyClientOnControllerChange)
+		{
+			// When not in backward compatibility mode, OnRep_Controller will properly call NotifyControllerChanged
+			GetPawn()->OnRep_Controller();
+		}
+		else
+		{
+			GetPawn()->NotifyControllerChanged();
+		}
 	}
 	GetPawn()->DispatchRestart(true);
 	
@@ -1278,7 +1288,7 @@ void APlayerController::ServerAcknowledgePossession_Implementation(APawn* P)
 	UE_LOG(LogPlayerController, Verbose, TEXT("ServerAcknowledgePossession_Implementation %s"), *GetNameSafe(P));
 	AcknowledgedPawn = P;
 
-	if (PlayerControllerCVars::NetResetServerPredictionDataOnPawnAck != 0)
+	if (UE::Gameplay::CVars::NetResetServerPredictionDataOnPawnAck != 0)
 	{
 		if (AcknowledgedPawn && AcknowledgedPawn == GetPawn())
 		{
@@ -3674,7 +3684,7 @@ void APlayerController::OnRemovedFromPlayerControllerList()
 
 void APlayerController::GetStreamingSourceLocationAndRotation(FVector& OutLocation, FRotator& OutRotation) const
 {
-	if (!PlayerControllerCVars::ForceUsingCameraAsStreamingSource)
+	if (!UE::Gameplay::CVars::ForceUsingCameraAsStreamingSource)
 	{
 		if (const AActor* ViewTarget = GetViewTarget())
 		{
@@ -5512,11 +5522,23 @@ void APlayerController::EndSpectatingState()
 
 void APlayerController::BeginInactiveState()
 {
-	if ( (GetPawn() != NULL) && (GetPawn()->Controller == this) )
+	if ( (GetPawn() != nullptr) && (GetPawn()->Controller == this) )
 	{
-		GetPawn()->Controller = NULL;
+		GetPawn()->Controller = nullptr;
+		if (UE::Gameplay::CVars::bAlwaysNotifyClientOnControllerChange)
+		{
+			if (HasAuthority())
+			{
+				// OnRep is not called on the server so call notify directly
+				GetPawn()->NotifyControllerChanged();
+			}
+			else
+			{
+				GetPawn()->OnRep_Controller();
+			}
+		}
 	}
-	SetPawn(NULL);
+	SetPawn(nullptr);
 
 	GetWorldTimerManager().SetTimer(TimerHandle_UnFreeze, this, &APlayerController::UnFreeze, GetMinRespawnDelay());
 }
