Fix some materials not being considered for ODSC reloading when ODSC.recompilemode=1 : we need to push these requests so that the material gets its shadermap updated in FMaterialShaderMap::LoadForRemoteRecompile
The shaders requests will be deduplicated on the server side in FMaterial::BeginCacheShaders

#rnx

diff --git a/Engine/Source/Runtime/Engine/Private/ODSC/ODSCThread.cpp b/Engine/Source/Runtime/Engine/Private/ODSC/ODSCThread.cpp
--- a/Engine/Source/Runtime/Engine/Private/ODSC/ODSCThread.cpp
+++ b/Engine/Source/Runtime/Engine/Private/ODSC/ODSCThread.cpp
@@ -225,9 +225,11 @@ void FODSCThread::AddShaderPipelineRequest(
 
 	for (const FShaderId& ShaderId : RequestShaderIds)
 	{
-		if (!RequestHashes.Contains(ShaderId))
+		FMaterialRequestsHashes& MaterialRequestHashes = RequestHashes.FindOrAdd(ShaderId);
+		bool bAlreadyInSet = false;
+		MaterialRequestHashes.RequestStrings.Add(RequestString, &bAlreadyInSet);
+		if (!bAlreadyInSet)
 		{
-			RequestHashes.Add(ShaderId);
 			bShouldAddRequest = true;
 		}
 	}
diff --git a/Engine/Source/Runtime/Engine/Private/ODSC/ODSCThread.h b/Engine/Source/Runtime/Engine/Private/ODSC/ODSCThread.h
--- a/Engine/Source/Runtime/Engine/Private/ODSC/ODSCThread.h
+++ b/Engine/Source/Runtime/Engine/Private/ODSC/ODSCThread.h
@@ -204,8 +204,14 @@ private:
 	/** Lock to access the RequestHashes TMap */
 	FCriticalSection RequestHashCriticalSection;
 
+
+	struct FMaterialRequestsHashes
+	{
+		TSet<FString> RequestStrings;
+	};
+
 	/** Hashes for all Pending or Completed requests.  This is so we avoid making the same request multiple times. */
-	TSet<FShaderId> RequestHashes;
+	TMap<FShaderId, FMaterialRequestsHashes> RequestHashes;
 
 	/** Pointer to Runnable Thread */
 	FRunnableThread* Thread = nullptr;
