Moved some logic in FRayTracingGeometryManager::Tick() to RequestRayTracingGeometryStreamIn(...)

diff --git a/Engine/Source/Runtime/Engine/Private/Rendering/RayTracingGeometryManager.cpp b/Engine/Source/Runtime/Engine/Private/Rendering/RayTracingGeometryManager.cpp
--- a/Engine/Source/Runtime/Engine/Private/Rendering/RayTracingGeometryManager.cpp
+++ b/Engine/Source/Runtime/Engine/Private/Rendering/RayTracingGeometryManager.cpp
@@ -462,97 +462,7 @@ void FRayTracingGeometryManager::Tick(FRHICommandList& RHICmdList)
 				Geometry->MakeResident(RHICmdList);
 			}
 
-			FRegisteredGeometry& RegisteredGeometry = RegisteredGeometries[Geometry->RayTracingGeometryHandle];
-
-			if (Geometry->Initializer.Type == ERayTracingGeometryInitializerType::StreamingDestination && RegisteredGeometry.StreamableData != nullptr)
-			{
-				if (RegisteredGeometry.Status == FRegisteredGeometry::FStatus::Streaming)
-				{
-					continue;
-				}
-
-				FByteBulkData* StreamableData = RegisteredGeometry.StreamableData;
-				check(StreamableData);
-
-				TResourceArray<uint8> RawData;
-				RawData.SetAllowCPUAccess(true);
-
-				FResourceArrayInterface* OfflineData = nullptr;
-
-				if (RegisteredGeometry.StreamableDataSize == 0)
-				{
-					// no offline data -> build from VB/IB at runtime
-
-					RegisteredGeometry.Status = FRegisteredGeometry::FStatus::StreamedIn;
-				}
-				else if (StreamableData->IsBulkDataLoaded())
-				{
-					{
-						const uint8* Ptr = (const uint8*)StreamableData->LockReadOnly();
-
-						FMemoryView MemView(Ptr + RegisteredGeometry.StreamableDataOffset, RegisteredGeometry.StreamableDataSize);
-
-						FMemoryReaderView MemReader(MemView, true);
-						RawData.BulkSerialize(MemReader);
-						StreamableData->Unlock();
-					}
-
-					if (!RawData.IsEmpty())
-					{
-						OfflineData = &RawData;
-					}
-
-					RegisteredGeometry.Status = FRegisteredGeometry::FStatus::StreamedIn;
-				}
-				else
-				{
-					checkf(StreamableData->CanLoadFromDisk(), TEXT("Bulk data is not loaded and cannot be loaded from disk!"));
-					check(!StreamableData->IsStoredCompressedOnDisk()); // We do not support compressed Bulkdata for this system. Limitation of the streaming request/bulk data
-
-					check(StreamableData->IsUsingIODispatcher());
-
-					if (NumStreamingRequests >= GRayTracingStreamingMaxPendingRequests)
-					{
-						continue;
-					}
-
-					FStreamingRequest& StreamingRequest = StreamingRequests[NextStreamingRequestIndex];
-					NextStreamingRequestIndex = (NextStreamingRequestIndex + 1) % GRayTracingStreamingMaxPendingRequests;
-					++NumStreamingRequests;
-
-					StreamingRequest.GeometryHandle = Geometry->RayTracingGeometryHandle;
-					StreamingRequest.RequestBuffer = FIoBuffer(RegisteredGeometry.StreamableDataSize); // TODO: Use FIoBuffer::Wrap with preallocated memory
-
-					// TODO: We're currently using a single batch per request so we can individually cancel and wait on requests.
-					// This isn't ideal and should be revisited in the future.
-					FBulkDataBatchRequest::FScatterGatherBuilder Batch = FBulkDataBatchRequest::ScatterGather(1);
-					Batch.Read(*StreamableData, RegisteredGeometry.StreamableDataOffset, RegisteredGeometry.StreamableDataSize);
-					Batch.Issue(StreamingRequest.RequestBuffer, AIOP_Low, [](FBulkDataRequest::EStatus) {}, StreamingRequest.Request);
-
-					RegisteredGeometry.Status = FRegisteredGeometry::FStatus::Streaming;
-				}
-
-				if(RegisteredGeometry.Status == FRegisteredGeometry::FStatus::StreamedIn)
-				{
-					{
-						FRHIResourceReplaceBatcher Batcher(RHICmdList, 1);
-						FRayTracingGeometryInitializer IntermediateInitializer = Geometry->Initializer;
-						IntermediateInitializer.Type = ERayTracingGeometryInitializerType::StreamingSource;
-						IntermediateInitializer.OfflineData = OfflineData;
-
-						FRayTracingGeometryRHIRef IntermediateRayTracingGeometry = RHICmdList.CreateRayTracingGeometry(IntermediateInitializer);
-
-						Geometry->SetRequiresBuild(IntermediateInitializer.OfflineData == nullptr || IntermediateRayTracingGeometry->IsCompressed());
-
-						Geometry->InitRHIForStreaming(IntermediateRayTracingGeometry, Batcher);
-
-						// When Batcher goes out of scope it will add commands to copy the BLAS buffers on RHI thread.
-						// We need to do it before we build the current geometry (also on RHI thread).
-					}
-
-					Geometry->RequestBuildIfNeeded(RHICmdList, ERTAccelerationStructureBuildPriority::Normal);
-				}
-			}
+			RequestRayTracingGeometryStreamIn(RHICmdList, Geometry);
 		}
 
 		// 4th step
@@ -603,6 +513,100 @@ void FRayTracingGeometryManager::Tick(FRHICommandList& RHICmdList)
 	SET_MEMORY_STAT(STAT_RayTracingGeometryResidentMemory, TotalResidentSize);
 }
 
+void FRayTracingGeometryManager::RequestRayTracingGeometryStreamIn(FRHICommandList& RHICmdList, FRayTracingGeometry* Geometry)
+{
+	FRegisteredGeometry& RegisteredGeometry = RegisteredGeometries[Geometry->RayTracingGeometryHandle];
+
+	if (RegisteredGeometry.Geometry->Initializer.Type == ERayTracingGeometryInitializerType::StreamingDestination)
+	{
+		if (RegisteredGeometry.Status == FRegisteredGeometry::FStatus::Streaming)
+		{
+			return;
+		}
+
+		FByteBulkData* StreamableData = RegisteredGeometry.StreamableData;
+
+		TResourceArray<uint8> RawData;
+		RawData.SetAllowCPUAccess(true);
+
+		FResourceArrayInterface* OfflineData = nullptr;
+
+		if (RegisteredGeometry.StreamableDataSize == 0)
+		{
+			// no offline data -> build from VB/IB at runtime
+
+			RegisteredGeometry.Status = FRegisteredGeometry::FStatus::StreamedIn;
+		}
+		else if (StreamableData->IsBulkDataLoaded())
+		{
+			{
+				const uint8* Ptr = (const uint8*)StreamableData->LockReadOnly();
+
+				FMemoryView MemView(Ptr + RegisteredGeometry.StreamableDataOffset, RegisteredGeometry.StreamableDataSize);
+
+				FMemoryReaderView MemReader(MemView, true);
+				RawData.BulkSerialize(MemReader);
+				StreamableData->Unlock();
+			}
+
+			if (!RawData.IsEmpty())
+			{
+				OfflineData = &RawData;
+			}
+
+			RegisteredGeometry.Status = FRegisteredGeometry::FStatus::StreamedIn;
+		}
+		else
+		{
+			checkf(StreamableData->CanLoadFromDisk(), TEXT("Bulk data is not loaded and cannot be loaded from disk!"));
+			check(!StreamableData->IsStoredCompressedOnDisk()); // We do not support compressed Bulkdata for this system. Limitation of the streaming request/bulk data
+
+			check(StreamableData->IsUsingIODispatcher());
+
+			if (NumStreamingRequests >= GRayTracingStreamingMaxPendingRequests)
+			{
+				return;
+			}
+
+			FStreamingRequest& StreamingRequest = StreamingRequests[NextStreamingRequestIndex];
+			NextStreamingRequestIndex = (NextStreamingRequestIndex + 1) % GRayTracingStreamingMaxPendingRequests;
+			++NumStreamingRequests;
+
+			StreamingRequest.GeometryHandle = Geometry->RayTracingGeometryHandle;
+			StreamingRequest.RequestBuffer = FIoBuffer(RegisteredGeometry.StreamableDataSize); // TODO: Use FIoBuffer::Wrap with preallocated memory
+
+			// TODO: We're currently using a single batch per request so we can individually cancel and wait on requests.
+			// This isn't ideal and should be revisited in the future.
+			FBulkDataBatchRequest::FScatterGatherBuilder Batch = FBulkDataBatchRequest::ScatterGather(1);
+			Batch.Read(*StreamableData, RegisteredGeometry.StreamableDataOffset, RegisteredGeometry.StreamableDataSize);
+			Batch.Issue(StreamingRequest.RequestBuffer, AIOP_Low, [](FBulkDataRequest::EStatus) {}, StreamingRequest.Request);
+
+			RegisteredGeometry.Status = FRegisteredGeometry::FStatus::Streaming;
+		}
+
+		if (RegisteredGeometry.Status == FRegisteredGeometry::FStatus::StreamedIn)
+		{
+			{
+				FRHIResourceReplaceBatcher Batcher(RHICmdList, 1);
+				FRayTracingGeometryInitializer IntermediateInitializer = Geometry->Initializer;
+				IntermediateInitializer.Type = ERayTracingGeometryInitializerType::StreamingSource;
+				IntermediateInitializer.OfflineData = OfflineData;
+
+				FRayTracingGeometryRHIRef IntermediateRayTracingGeometry = RHICmdList.CreateRayTracingGeometry(IntermediateInitializer);
+
+				Geometry->SetRequiresBuild(IntermediateInitializer.OfflineData == nullptr || IntermediateRayTracingGeometry->IsCompressed());
+
+				Geometry->InitRHIForStreaming(IntermediateRayTracingGeometry, Batcher);
+
+				// When Batcher goes out of scope it will add commands to copy the BLAS buffers on RHI thread.
+				// We need to do it before we build the current geometry (also on RHI thread).
+			}
+
+			Geometry->RequestBuildIfNeeded(RHICmdList, ERTAccelerationStructureBuildPriority::Normal);
+		}
+	}
+}
+
 void FRayTracingGeometryManager::ProcessCompletedStreamingRequests(FRHICommandList& RHICmdList)
 {
 	const int32 StartPendingRequestIndex = (NextStreamingRequestIndex + GRayTracingStreamingMaxPendingRequests - NumStreamingRequests) % GRayTracingStreamingMaxPendingRequests;
diff --git a/Engine/Source/Runtime/Engine/Public/Rendering/RayTracingGeometryManager.h b/Engine/Source/Runtime/Engine/Public/Rendering/RayTracingGeometryManager.h
--- a/Engine/Source/Runtime/Engine/Public/Rendering/RayTracingGeometryManager.h
+++ b/Engine/Source/Runtime/Engine/Public/Rendering/RayTracingGeometryManager.h
@@ -71,6 +71,7 @@ private:
 
 	void ReleaseRayTracingGeometryGroupReference(RayTracing::GeometryGroupHandle Handle);
 
+	void RequestRayTracingGeometryStreamIn(FRHICommandList& RHICmdList, FRayTracingGeometry* Geometry);
 	void ProcessCompletedStreamingRequests(FRHICommandList& RHICmdList);
 
 	FCriticalSection RequestCS;
