Early out and change indentation in FRayTracingGeometryManager::RequestRayTracingGeometryStreamIn(...)

diff --git a/Engine/Source/Runtime/Engine/Private/Rendering/RayTracingGeometryManager.cpp b/Engine/Source/Runtime/Engine/Private/Rendering/RayTracingGeometryManager.cpp
--- a/Engine/Source/Runtime/Engine/Private/Rendering/RayTracingGeometryManager.cpp
+++ b/Engine/Source/Runtime/Engine/Private/Rendering/RayTracingGeometryManager.cpp
@@ -517,93 +517,92 @@ void FRayTracingGeometryManager::RequestRayTracingGeometryStreamIn(FRHICommandLi
 {
 	FRegisteredGeometry& RegisteredGeometry = RegisteredGeometries[Geometry->RayTracingGeometryHandle];
 
-	if (RegisteredGeometry.Geometry->Initializer.Type == ERayTracingGeometryInitializerType::StreamingDestination)
+	if (RegisteredGeometry.Geometry->Initializer.Type != ERayTracingGeometryInitializerType::StreamingDestination
+		|| RegisteredGeometry.Status == FRegisteredGeometry::FStatus::Streaming)
 	{
-		if (RegisteredGeometry.Status == FRegisteredGeometry::FStatus::Streaming)
-		{
-			return;
-		}
+		// no streaming required or streaming request already in-flight
+		return;
+	}
 
-		FByteBulkData* StreamableData = RegisteredGeometry.StreamableData;
+	FByteBulkData* StreamableData = RegisteredGeometry.StreamableData;
 
-		TResourceArray<uint8> RawData;
-		RawData.SetAllowCPUAccess(true);
+	TResourceArray<uint8> RawData;
+	RawData.SetAllowCPUAccess(true);
 
-		FResourceArrayInterface* OfflineData = nullptr;
+	FResourceArrayInterface* OfflineData = nullptr;
 
-		if (RegisteredGeometry.StreamableDataSize == 0)
-		{
-			// no offline data -> build from VB/IB at runtime
+	if (RegisteredGeometry.StreamableDataSize == 0)
+	{
+		// no offline data -> build from VB/IB at runtime
 
-			RegisteredGeometry.Status = FRegisteredGeometry::FStatus::StreamedIn;
-		}
-		else if (StreamableData->IsBulkDataLoaded())
+		RegisteredGeometry.Status = FRegisteredGeometry::FStatus::StreamedIn;
+	}
+	else if (StreamableData->IsBulkDataLoaded())
+	{
 		{
-			{
-				const uint8* Ptr = (const uint8*)StreamableData->LockReadOnly();
+			const uint8* Ptr = (const uint8*)StreamableData->LockReadOnly();
 
-				FMemoryView MemView(Ptr + RegisteredGeometry.StreamableDataOffset, RegisteredGeometry.StreamableDataSize);
+			FMemoryView MemView(Ptr + RegisteredGeometry.StreamableDataOffset, RegisteredGeometry.StreamableDataSize);
 
-				FMemoryReaderView MemReader(MemView, true);
-				RawData.BulkSerialize(MemReader);
-				StreamableData->Unlock();
-			}
-
-			if (!RawData.IsEmpty())
-			{
-				OfflineData = &RawData;
-			}
-
-			RegisteredGeometry.Status = FRegisteredGeometry::FStatus::StreamedIn;
+			FMemoryReaderView MemReader(MemView, true);
+			RawData.BulkSerialize(MemReader);
+			StreamableData->Unlock();
 		}
-		else
+
+		if (!RawData.IsEmpty())
 		{
-			checkf(StreamableData->CanLoadFromDisk(), TEXT("Bulk data is not loaded and cannot be loaded from disk!"));
-			check(!StreamableData->IsStoredCompressedOnDisk()); // We do not support compressed Bulkdata for this system. Limitation of the streaming request/bulk data
+			OfflineData = &RawData;
+		}
 
-			check(StreamableData->IsUsingIODispatcher());
+		RegisteredGeometry.Status = FRegisteredGeometry::FStatus::StreamedIn;
+	}
+	else
+	{
+		checkf(StreamableData->CanLoadFromDisk(), TEXT("Bulk data is not loaded and cannot be loaded from disk!"));
+		check(!StreamableData->IsStoredCompressedOnDisk()); // We do not support compressed Bulkdata for this system. Limitation of the streaming request/bulk data
 
-			if (NumStreamingRequests >= GRayTracingStreamingMaxPendingRequests)
-			{
-				return;
-			}
+		check(StreamableData->IsUsingIODispatcher());
 
-			FStreamingRequest& StreamingRequest = StreamingRequests[NextStreamingRequestIndex];
-			NextStreamingRequestIndex = (NextStreamingRequestIndex + 1) % GRayTracingStreamingMaxPendingRequests;
-			++NumStreamingRequests;
+		if (NumStreamingRequests >= GRayTracingStreamingMaxPendingRequests)
+		{
+			return;
+		}
 
-			StreamingRequest.GeometryHandle = Geometry->RayTracingGeometryHandle;
-			StreamingRequest.RequestBuffer = FIoBuffer(RegisteredGeometry.StreamableDataSize); // TODO: Use FIoBuffer::Wrap with preallocated memory
+		FStreamingRequest& StreamingRequest = StreamingRequests[NextStreamingRequestIndex];
+		NextStreamingRequestIndex = (NextStreamingRequestIndex + 1) % GRayTracingStreamingMaxPendingRequests;
+		++NumStreamingRequests;
 
-			// TODO: We're currently using a single batch per request so we can individually cancel and wait on requests.
-			// This isn't ideal and should be revisited in the future.
-			FBulkDataBatchRequest::FScatterGatherBuilder Batch = FBulkDataBatchRequest::ScatterGather(1);
-			Batch.Read(*StreamableData, RegisteredGeometry.StreamableDataOffset, RegisteredGeometry.StreamableDataSize);
-			Batch.Issue(StreamingRequest.RequestBuffer, AIOP_Low, [](FBulkDataRequest::EStatus) {}, StreamingRequest.Request);
+		StreamingRequest.GeometryHandle = Geometry->RayTracingGeometryHandle;
+		StreamingRequest.RequestBuffer = FIoBuffer(RegisteredGeometry.StreamableDataSize); // TODO: Use FIoBuffer::Wrap with preallocated memory
 
-			RegisteredGeometry.Status = FRegisteredGeometry::FStatus::Streaming;
-		}
+		// TODO: We're currently using a single batch per request so we can individually cancel and wait on requests.
+		// This isn't ideal and should be revisited in the future.
+		FBulkDataBatchRequest::FScatterGatherBuilder Batch = FBulkDataBatchRequest::ScatterGather(1);
+		Batch.Read(*StreamableData, RegisteredGeometry.StreamableDataOffset, RegisteredGeometry.StreamableDataSize);
+		Batch.Issue(StreamingRequest.RequestBuffer, AIOP_Low, [](FBulkDataRequest::EStatus) {}, StreamingRequest.Request);
 
-		if (RegisteredGeometry.Status == FRegisteredGeometry::FStatus::StreamedIn)
-		{
-			{
-				FRHIResourceReplaceBatcher Batcher(RHICmdList, 1);
-				FRayTracingGeometryInitializer IntermediateInitializer = Geometry->Initializer;
-				IntermediateInitializer.Type = ERayTracingGeometryInitializerType::StreamingSource;
-				IntermediateInitializer.OfflineData = OfflineData;
+		RegisteredGeometry.Status = FRegisteredGeometry::FStatus::Streaming;
+	}
 
-				FRayTracingGeometryRHIRef IntermediateRayTracingGeometry = RHICmdList.CreateRayTracingGeometry(IntermediateInitializer);
+	if (RegisteredGeometry.Status == FRegisteredGeometry::FStatus::StreamedIn)
+	{
+		{
+			FRHIResourceReplaceBatcher Batcher(RHICmdList, 1);
+			FRayTracingGeometryInitializer IntermediateInitializer = Geometry->Initializer;
+			IntermediateInitializer.Type = ERayTracingGeometryInitializerType::StreamingSource;
+			IntermediateInitializer.OfflineData = OfflineData;
 
-				Geometry->SetRequiresBuild(IntermediateInitializer.OfflineData == nullptr || IntermediateRayTracingGeometry->IsCompressed());
+			FRayTracingGeometryRHIRef IntermediateRayTracingGeometry = RHICmdList.CreateRayTracingGeometry(IntermediateInitializer);
 
-				Geometry->InitRHIForStreaming(IntermediateRayTracingGeometry, Batcher);
+			Geometry->SetRequiresBuild(IntermediateInitializer.OfflineData == nullptr || IntermediateRayTracingGeometry->IsCompressed());
 
-				// When Batcher goes out of scope it will add commands to copy the BLAS buffers on RHI thread.
-				// We need to do it before we build the current geometry (also on RHI thread).
-			}
+			Geometry->InitRHIForStreaming(IntermediateRayTracingGeometry, Batcher);
 
-			Geometry->RequestBuildIfNeeded(RHICmdList, ERTAccelerationStructureBuildPriority::Normal);
+			// When Batcher goes out of scope it will add commands to copy the BLAS buffers on RHI thread.
+			// We need to do it before we build the current geometry (also on RHI thread).
 		}
+
+		Geometry->RequestBuildIfNeeded(RHICmdList, ERTAccelerationStructureBuildPriority::Normal);
 	}
 }
 
