Improvements to RayTracingGeometry management in FStaticMeshRenderData.
- unify code paths in InitStaticMeshRayTracingGeometry
- always initialize RayTracingGeometryGroup CurrentFirstLODIndex and update it when bUsingRenderingLODs.

diff --git a/Engine/Source/Runtime/Engine/Private/StaticMesh.cpp b/Engine/Source/Runtime/Engine/Private/StaticMesh.cpp
--- a/Engine/Source/Runtime/Engine/Private/StaticMesh.cpp
+++ b/Engine/Source/Runtime/Engine/Private/StaticMesh.cpp
@@ -2292,31 +2292,11 @@ void FStaticMeshRenderData::InitResources(ERHIFeatureLevel::Type InFeatureLevel,
 		checkf(Owner->bSupportRayTracing, TEXT("Unexpected RayTracingProxy on '%s' (support ray tracing is disabled on this UStaticMesh)."), *GetPathNameSafe(Owner));
 
 		ENQUEUE_RENDER_COMMAND(InitStaticMeshRayTracingGeometry)(
-			[this](FRHICommandListImmediate& RHICmdList)
+			[this, Owner](FRHICommandListImmediate& RHICmdList)
 			{
 				RayTracingGeometryGroupHandle = GRayTracingGeometryManager->RegisterRayTracingGeometryGroup(LODResources.Num());
 
-				if (IsRayTracingUsingReferenceBasedResidency())
-				{
-					((FRayTracingGeometryManager*)GRayTracingGeometryManager)->SetRayTracingGeometryGroupCurrentFirstLODIndex(RHICmdList, RayTracingGeometryGroupHandle, CurrentFirstLODIdx);
-				}
-
-				for (int32 LODIndex = 0; LODIndex < LODResources.Num(); ++LODIndex)
-				{
-					// Skip LODs that have their render data stripped
-					if (LODResources[LODIndex].RayTracingGeometry != nullptr && LODResources[LODIndex].VertexBuffers.StaticMeshVertexBuffer.GetNumVertices() > 0)
-					{
-						LODResources[LODIndex].RayTracingGeometry->GroupHandle = RayTracingGeometryGroupHandle;
-
-						if (LODIndex < CurrentFirstLODIdx)
-						{
-							LODResources[LODIndex].RayTracingGeometry->Initializer.Type = ERayTracingGeometryInitializerType::StreamingDestination;
-						}
-
-						LODResources[LODIndex].RayTracingGeometry->LODIndex = LODIndex;
-						LODResources[LODIndex].RayTracingGeometry->InitResource(RHICmdList);
-					}
-				}
+				((FRayTracingGeometryManager*)GRayTracingGeometryManager)->SetRayTracingGeometryGroupCurrentFirstLODIndex(RHICmdList, RayTracingGeometryGroupHandle, RayTracingProxy->bUsingRenderingLODs ? CurrentFirstLODIdx : 0);
 				
 				FStaticMeshRayTracingProxyLODArray& RayTracingLODs = RayTracingProxy->LODs;
 
@@ -2327,11 +2307,16 @@ void FStaticMeshRenderData::InitResources(ERHIFeatureLevel::Type InFeatureLevel,
 					FStaticMeshRayTracingProxyLOD& RayTracingLOD = RayTracingLODs[LODIndex];
 
 					// Skip LODs that have their render data stripped
-					if (RayTracingLOD.bOwnsRayTracingGeometry && RayTracingLOD.VertexBuffers->StaticMeshVertexBuffer.GetNumVertices() > 0)
+					if (RayTracingLOD.VertexBuffers->StaticMeshVertexBuffer.GetNumVertices() > 0)
 					{
 						RayTracingLOD.RayTracingGeometry->GroupHandle = RayTracingGeometryGroupHandle;
 
-						if (LODIndex < CurrentFirstLODIdx || IsRayTracingUsingReferenceBasedResidency())
+						const bool bHasStreamableData = RayTracingLOD.StreamableData.GetBulkDataSize() > 0;
+
+						checkf(RayTracingLOD.bOwnsRayTracingGeometry || !bHasStreamableData,
+							TEXT("Unexpected StreamableData on '%s' (StreamableData is only expected when ray tracing proxy owns the geometry)."), *GetPathNameSafe(Owner));
+
+						if (bHasStreamableData || LODIndex < CurrentFirstLODIdx)
 						{
 							RayTracingLOD.RayTracingGeometry->Initializer.Type = ERayTracingGeometryInitializerType::StreamingDestination;
 						}
@@ -2339,7 +2324,7 @@ void FStaticMeshRenderData::InitResources(ERHIFeatureLevel::Type InFeatureLevel,
 						RayTracingLOD.RayTracingGeometry->LODIndex = LODIndex;
 						RayTracingLOD.RayTracingGeometry->InitResource(RHICmdList);
 
-						if (IsRayTracingUsingReferenceBasedResidency())
+						if (bHasStreamableData)
 						{
 							((FRayTracingGeometryManager*)GRayTracingGeometryManager)->SetRayTracingGeometryStreamingData(
 								RayTracingLOD.RayTracingGeometry,
@@ -2490,7 +2475,7 @@ void UStaticMesh::RequestUpdateCachedRenderState() const
 #if RHI_RAYTRACING
 	if (IsRayTracingAllowed() && bSupportRayTracing)
 	{
-		if (IsRayTracingUsingReferenceBasedResidency())
+		if (GetRenderData()->RayTracingProxy->bUsingRenderingLODs)
 		{
 			((FRayTracingGeometryManager*)GRayTracingGeometryManager)->SetRayTracingGeometryGroupCurrentFirstLODIndex(FRHICommandListImmediate::Get(), GetRenderData()->RayTracingGeometryGroupHandle, GetRenderData()->CurrentFirstLODIdx);
 		}
