Replaced EAutomationTestFlags::ApplicationContextMask with EAutomationTestFlags_ApplicationContextMask.

diff --git a/Engine/Source/Runtime/AssetRegistry/Private/AssetRegistryArchive.cpp b/Engine/Source/Runtime/AssetRegistry/Private/AssetRegistryArchive.cpp
--- a/Engine/Source/Runtime/AssetRegistry/Private/AssetRegistryArchive.cpp
+++ b/Engine/Source/Runtime/AssetRegistry/Private/AssetRegistryArchive.cpp
@@ -344,7 +344,7 @@ void FAssetRegistryReader::SerializeTagsAndBundlesOldVersion(FAssetData& Out, FA
 #include "Serialization/MemoryReader.h"
 #include "Serialization/MemoryWriter.h"
 
-IMPLEMENT_SIMPLE_AUTOMATION_TEST(FAssetRegistryTagSerializationTest, "System.AssetRegistry.SerializeTagMap", EAutomationTestFlags::ApplicationContextMask | EAutomationTestFlags::SmokeFilter)
+IMPLEMENT_SIMPLE_AUTOMATION_TEST(FAssetRegistryTagSerializationTest, "System.AssetRegistry.SerializeTagMap", EAutomationTestFlags_ApplicationContextMask | EAutomationTestFlags::SmokeFilter)
 
 FAssetDataTagMapSharedView MakeLooseMap(std::initializer_list<TPairInitializer<const char*, FString>> Pairs)
 {
diff --git a/Engine/Source/Runtime/AssetRegistry/Private/AssetRegistryState.cpp b/Engine/Source/Runtime/AssetRegistry/Private/AssetRegistryState.cpp
--- a/Engine/Source/Runtime/AssetRegistry/Private/AssetRegistryState.cpp
+++ b/Engine/Source/Runtime/AssetRegistry/Private/AssetRegistryState.cpp
@@ -3145,7 +3145,7 @@ void FAssetRegistryState::Dump(const TArray<FString>& Arguments, TArray<FString>
 
 #include "Misc/AutomationTest.h"
 
-IMPLEMENT_SIMPLE_AUTOMATION_TEST(FAssetRegistryAssetPathStringsTest, "System.AssetRegistry.AssetPathStrings", EAutomationTestFlags::ApplicationContextMask | EAutomationTestFlags::EngineFilter);
+IMPLEMENT_SIMPLE_AUTOMATION_TEST(FAssetRegistryAssetPathStringsTest, "System.AssetRegistry.AssetPathStrings", EAutomationTestFlags_ApplicationContextMask | EAutomationTestFlags::EngineFilter);
 
 // Tests that we can produce correct paths for objects represented by FCachedAssetKey and FAssetData
 // E.g	PackageName.AssetName
diff --git a/Engine/Source/Runtime/AssetRegistry/Private/DependsNodeTest.cpp b/Engine/Source/Runtime/AssetRegistry/Private/DependsNodeTest.cpp
--- a/Engine/Source/Runtime/AssetRegistry/Private/DependsNodeTest.cpp
+++ b/Engine/Source/Runtime/AssetRegistry/Private/DependsNodeTest.cpp
@@ -9,7 +9,7 @@
 #include "Serialization/MemoryReader.h"
 #include "Serialization/MemoryWriter.h"
 
-IMPLEMENT_SIMPLE_AUTOMATION_TEST(FDependsNodeTest, "System.AssetRegistry.DependsNode", EAutomationTestFlags::ApplicationContextMask | EAutomationTestFlags::EngineFilter);
+IMPLEMENT_SIMPLE_AUTOMATION_TEST(FDependsNodeTest, "System.AssetRegistry.DependsNode", EAutomationTestFlags_ApplicationContextMask | EAutomationTestFlags::EngineFilter);
 
 namespace
 {
diff --git a/Engine/Source/Runtime/AssetRegistry/Private/PropertyCombinationSetTest.cpp b/Engine/Source/Runtime/AssetRegistry/Private/PropertyCombinationSetTest.cpp
--- a/Engine/Source/Runtime/AssetRegistry/Private/PropertyCombinationSetTest.cpp
+++ b/Engine/Source/Runtime/AssetRegistry/Private/PropertyCombinationSetTest.cpp
@@ -432,7 +432,7 @@ protected:
 	}
 };
 
-IMPLEMENT_CUSTOM_SIMPLE_AUTOMATION_TEST(FPropertyCombinationSetTestSubClass, FPropertyCombinationSetTest, "System.AssetRegistry.PropertyCombinationSet", EAutomationTestFlags::ApplicationContextMask | EAutomationTestFlags::EngineFilter);
+IMPLEMENT_CUSTOM_SIMPLE_AUTOMATION_TEST(FPropertyCombinationSetTestSubClass, FPropertyCombinationSetTest, "System.AssetRegistry.PropertyCombinationSet", EAutomationTestFlags_ApplicationContextMask | EAutomationTestFlags::EngineFilter);
 
 bool FPropertyCombinationSetTestSubClass::RunTest(const FString& Parameters)
 {
diff --git a/Engine/Source/Runtime/Core/Private/Async/Fundamental/UnitTests.cpp b/Engine/Source/Runtime/Core/Private/Async/Fundamental/UnitTests.cpp
--- a/Engine/Source/Runtime/Core/Private/Async/Fundamental/UnitTests.cpp
+++ b/Engine/Source/Runtime/Core/Private/Async/Fundamental/UnitTests.cpp
@@ -467,7 +467,7 @@ namespace Tasks2Tests
 		SGarbage++;
 	}
 
-	IMPLEMENT_SIMPLE_AUTOMATION_TEST(FTasksPerfTests2, "System.Core.LowLevelTasks.PerfTests", EAutomationTestFlags::ApplicationContextMask | EAutomationTestFlags::EngineFilter | EAutomationTestFlags::Disabled);
+	IMPLEMENT_SIMPLE_AUTOMATION_TEST(FTasksPerfTests2, "System.Core.LowLevelTasks.PerfTests", EAutomationTestFlags_ApplicationContextMask | EAutomationTestFlags::EngineFilter | EAutomationTestFlags::Disabled);
 
 	bool FTasksPerfTests2::RunTest(const FString& Parameters)
 	{
@@ -480,7 +480,7 @@ namespace Tasks2Tests
 		return true;
 	}
 
-	IMPLEMENT_SIMPLE_AUTOMATION_TEST(FTasksLocalGlobalPriorities, "System.Core.LowLevelTasks.LocalGlobalPriorities", EAutomationTestFlags::ApplicationContextMask | EAutomationTestFlags::EngineFilter);
+	IMPLEMENT_SIMPLE_AUTOMATION_TEST(FTasksLocalGlobalPriorities, "System.Core.LowLevelTasks.LocalGlobalPriorities", EAutomationTestFlags_ApplicationContextMask | EAutomationTestFlags::EngineFilter);
 	bool FTasksLocalGlobalPriorities::RunTest(const FString& Parameters)
 	{
 		using namespace LowLevelTasks;
@@ -533,7 +533,7 @@ namespace Tasks2Tests
 		return true;
 	};
 
-	IMPLEMENT_SIMPLE_AUTOMATION_TEST(FTasksUnitTests2, "System.Core.LowLevelTasks.UnitTests", EAutomationTestFlags::ApplicationContextMask | EAutomationTestFlags::EngineFilter);
+	IMPLEMENT_SIMPLE_AUTOMATION_TEST(FTasksUnitTests2, "System.Core.LowLevelTasks.UnitTests", EAutomationTestFlags_ApplicationContextMask | EAutomationTestFlags::EngineFilter);
 	bool FTasksUnitTests2::RunTest(const FString& Parameters)
 	{
 		using namespace LowLevelTasks;
@@ -667,7 +667,7 @@ namespace Tasks2Tests
 		return true;
 	}
 
-	IMPLEMENT_SIMPLE_AUTOMATION_TEST(FTasksCoroutineTests, "System.Core.Coroutine.UnitTests", EAutomationTestFlags::ApplicationContextMask | EAutomationTestFlags::EngineFilter);
+	IMPLEMENT_SIMPLE_AUTOMATION_TEST(FTasksCoroutineTests, "System.Core.Coroutine.UnitTests", EAutomationTestFlags_ApplicationContextMask | EAutomationTestFlags::EngineFilter);
 	bool FTasksCoroutineTests::RunTest(const FString& Parameters)
 	{
 #if WITH_CPP_COROUTINES
diff --git a/Engine/Source/Runtime/Core/Private/Containers/Algo/AlgosTest.cpp b/Engine/Source/Runtime/Core/Private/Containers/Algo/AlgosTest.cpp
--- a/Engine/Source/Runtime/Core/Private/Containers/Algo/AlgosTest.cpp
+++ b/Engine/Source/Runtime/Core/Private/Containers/Algo/AlgosTest.cpp
@@ -811,7 +811,7 @@ private:
 	TArray<int> TestData2;
 };
 
-IMPLEMENT_CUSTOM_SIMPLE_AUTOMATION_TEST(FAlgosTest, FAlgosTestBase, "System.Core.Misc.Algos", EAutomationTestFlags::ApplicationContextMask | EAutomationTestFlags::SmokeFilter)
+IMPLEMENT_CUSTOM_SIMPLE_AUTOMATION_TEST(FAlgosTest, FAlgosTestBase, "System.Core.Misc.Algos", EAutomationTestFlags_ApplicationContextMask | EAutomationTestFlags::SmokeFilter)
 
 bool FAlgosTest::RunTest(const FString& Parameters)
 {
diff --git a/Engine/Source/Runtime/Core/Private/Containers/ContainersTest.cpp b/Engine/Source/Runtime/Core/Private/Containers/ContainersTest.cpp
--- a/Engine/Source/Runtime/Core/Private/Containers/ContainersTest.cpp
+++ b/Engine/Source/Runtime/Core/Private/Containers/ContainersTest.cpp
@@ -661,7 +661,7 @@ namespace
 	}
 }
 
-IMPLEMENT_SIMPLE_AUTOMATION_TEST(FContainersSmokeTest, "System.Core.Containers.Smoke", EAutomationTestFlags::ApplicationContextMask | EAutomationTestFlags::SmokeFilter)
+IMPLEMENT_SIMPLE_AUTOMATION_TEST(FContainersSmokeTest, "System.Core.Containers.Smoke", EAutomationTestFlags_ApplicationContextMask | EAutomationTestFlags::SmokeFilter)
 bool FContainersSmokeTest::RunTest( const FString& Parameters )
 {
 	RunContainerTests<TMap<int32, FContainerTestValueType>, int32>();
@@ -671,7 +671,7 @@ bool FContainersSmokeTest::RunTest( const FString& Parameters )
 	return true;
 }
 
-IMPLEMENT_SIMPLE_AUTOMATION_TEST(FContainersFullTest, "System.Core.Containers.Full", EAutomationTestFlags::ApplicationContextMask | EAutomationTestFlags::EngineFilter)
+IMPLEMENT_SIMPLE_AUTOMATION_TEST(FContainersFullTest, "System.Core.Containers.Full", EAutomationTestFlags_ApplicationContextMask | EAutomationTestFlags::EngineFilter)
 bool FContainersFullTest::RunTest(const FString& Parameters)
 {
 	RunContainerTests<TMap<int32, FContainerTestValueType>, int32>();
@@ -712,7 +712,7 @@ bool FContainersFullTest::RunTest(const FString& Parameters)
 	return true;
 }
 
-IMPLEMENT_SIMPLE_AUTOMATION_TEST(FContainerPerformanceTest, "System.Core.Containers.Performance", EAutomationTestFlags::ApplicationContextMask | EAutomationTestFlags::EngineFilter)
+IMPLEMENT_SIMPLE_AUTOMATION_TEST(FContainerPerformanceTest, "System.Core.Containers.Performance", EAutomationTestFlags_ApplicationContextMask | EAutomationTestFlags::EngineFilter)
 bool FContainerPerformanceTest::RunTest(const FString& Parameters)
 {
 	RunPerformanceTest<TMap<int32, FString>, int32>(TEXT("TMap int32"), 1, 1000000);
@@ -810,7 +810,7 @@ namespace
 	}
 
 }
-IMPLEMENT_SIMPLE_AUTOMATION_TEST(FContainersTSetTest, "System.Core.Containers.TSet", EAutomationTestFlags::ApplicationContextMask | EAutomationTestFlags::EngineFilter)
+IMPLEMENT_SIMPLE_AUTOMATION_TEST(FContainersTSetTest, "System.Core.Containers.TSet", EAutomationTestFlags_ApplicationContextMask | EAutomationTestFlags::EngineFilter)
 bool FContainersTSetTest::RunTest(const FString& Parameters)
 {
 	enum class EArgType
diff --git a/Engine/Source/Runtime/Core/Private/Containers/StringTest.cpp b/Engine/Source/Runtime/Core/Private/Containers/StringTest.cpp
--- a/Engine/Source/Runtime/Core/Private/Containers/StringTest.cpp
+++ b/Engine/Source/Runtime/Core/Private/Containers/StringTest.cpp
@@ -11,7 +11,7 @@
 
 #if WITH_DEV_AUTOMATION_TESTS
 
-IMPLEMENT_SIMPLE_AUTOMATION_TEST(FStringSanitizeFloatTest, "System.Core.String.SanitizeFloat", EAutomationTestFlags::ApplicationContextMask | EAutomationTestFlags::SmokeFilter)
+IMPLEMENT_SIMPLE_AUTOMATION_TEST(FStringSanitizeFloatTest, "System.Core.String.SanitizeFloat", EAutomationTestFlags_ApplicationContextMask | EAutomationTestFlags::SmokeFilter)
 bool FStringSanitizeFloatTest::RunTest(const FString& Parameters)
 {
 	auto DoTest = [this](const double InVal, const int32 InMinFractionalDigits, const FString& InExpected)
@@ -50,7 +50,7 @@ bool FStringSanitizeFloatTest::RunTest(const FString& Parameters)
 	return true;
 }
 
-IMPLEMENT_SIMPLE_AUTOMATION_TEST(FStringAppendIntTest, "System.Core.String.AppendInt", EAutomationTestFlags::ApplicationContextMask | EAutomationTestFlags::SmokeFilter)
+IMPLEMENT_SIMPLE_AUTOMATION_TEST(FStringAppendIntTest, "System.Core.String.AppendInt", EAutomationTestFlags_ApplicationContextMask | EAutomationTestFlags::SmokeFilter)
 bool FStringAppendIntTest::RunTest(const FString& Parameters)
 {
 	auto DoTest = [this](const TCHAR* Call, const FString& Result, const TCHAR* InExpected)
@@ -90,7 +90,7 @@ bool FStringAppendIntTest::RunTest(const FString& Parameters)
 	return true;
 }
 
-IMPLEMENT_SIMPLE_AUTOMATION_TEST(FStringUnicodeTest, "System.Core.String.Unicode", EAutomationTestFlags::ApplicationContextMask | EAutomationTestFlags::SmokeFilter)
+IMPLEMENT_SIMPLE_AUTOMATION_TEST(FStringUnicodeTest, "System.Core.String.Unicode", EAutomationTestFlags_ApplicationContextMask | EAutomationTestFlags::SmokeFilter)
 bool FStringUnicodeTest::RunTest(const FString& Parameters)
 {
 	auto DoTest = [this](const TCHAR* Call, const FString& Result, const TCHAR* InExpected)
@@ -137,7 +137,7 @@ bool FStringUnicodeTest::RunTest(const FString& Parameters)
 	return true;
 }
 
-IMPLEMENT_SIMPLE_AUTOMATION_TEST(FLexTryParseStringTest, "System.Core.Misc.LexTryParseString", EAutomationTestFlags::ApplicationContextMask | EAutomationTestFlags::SmokeFilter)
+IMPLEMENT_SIMPLE_AUTOMATION_TEST(FLexTryParseStringTest, "System.Core.Misc.LexTryParseString", EAutomationTestFlags_ApplicationContextMask | EAutomationTestFlags::SmokeFilter)
 
 bool FLexTryParseStringTest::RunTest(const FString& Parameters)
 {
@@ -240,7 +240,7 @@ bool FLexTryParseStringTest::RunTest(const FString& Parameters)
 	return true;
 }
 
-IMPLEMENT_SIMPLE_AUTOMATION_TEST(FStringParseIntoArrayTest, "System.Core.String.ParseIntoArray", EAutomationTestFlags::ApplicationContextMask | EAutomationTestFlags::SmokeFilter)
+IMPLEMENT_SIMPLE_AUTOMATION_TEST(FStringParseIntoArrayTest, "System.Core.String.ParseIntoArray", EAutomationTestFlags_ApplicationContextMask | EAutomationTestFlags::SmokeFilter)
 
 bool FStringParseIntoArrayTest::RunTest(const FString& Parameters)
 {
@@ -283,7 +283,7 @@ bool FStringParseIntoArrayTest::RunTest(const FString& Parameters)
 	return true;
 }
 
-IMPLEMENT_SIMPLE_AUTOMATION_TEST(FStringSubstringTest, "System.Core.String.Substring", EAutomationTestFlags::ApplicationContextMask | EAutomationTestFlags::SmokeFilter)
+IMPLEMENT_SIMPLE_AUTOMATION_TEST(FStringSubstringTest, "System.Core.String.Substring", EAutomationTestFlags_ApplicationContextMask | EAutomationTestFlags::SmokeFilter)
 bool FStringSubstringTest::RunTest(const FString& Parameters)
 {
 	auto DoTest = [this](const TCHAR* Call, const FString& Result, const TCHAR* InExpected)
@@ -349,7 +349,7 @@ bool FStringSubstringTest::RunTest(const FString& Parameters)
 	return true;
 }
 
-IMPLEMENT_SIMPLE_AUTOMATION_TEST(FStringFromStringViewTest, "System.Core.String.FromStringView", EAutomationTestFlags::ApplicationContextMask | EAutomationTestFlags::SmokeFilter)
+IMPLEMENT_SIMPLE_AUTOMATION_TEST(FStringFromStringViewTest, "System.Core.String.FromStringView", EAutomationTestFlags_ApplicationContextMask | EAutomationTestFlags::SmokeFilter)
 bool FStringFromStringViewTest::RunTest(const FString& Parameters)
 {
 	// Verify basic construction and assignment from a string view.
@@ -424,7 +424,7 @@ bool FStringFromStringViewTest::RunTest(const FString& Parameters)
 	return true;
 }
 
-IMPLEMENT_SIMPLE_AUTOMATION_TEST(FStringConstructorWithSlackTest, "System.Core.String.ConstructorWithSlack", EAutomationTestFlags::ApplicationContextMask | EAutomationTestFlags::SmokeFilter)
+IMPLEMENT_SIMPLE_AUTOMATION_TEST(FStringConstructorWithSlackTest, "System.Core.String.ConstructorWithSlack", EAutomationTestFlags_ApplicationContextMask | EAutomationTestFlags::SmokeFilter)
 bool FStringConstructorWithSlackTest::RunTest(const FString& Parameters)
 {
 	// Note that the total capacity of a string might be greater than the string length + slack + a null terminator due to
@@ -440,10 +440,10 @@ bool FStringConstructorWithSlackTest::RunTest(const FString& Parameters)
 
 		const SIZE_T ExpectedCapacity = FMemory::QuantizeSize(NumElements * sizeof(TCHAR));
 
-		FString StringFromTChar(TestString, ExtraSlack);
+		FString StringFromTChar = FString::ConstructWithSlack(TestString, ExtraSlack);
 		TestEqual(TEXT("(TCHAR: Valid string with valid slack) resulting capacity"), StringFromTChar.GetAllocatedSize(), ExpectedCapacity);
 
-		FString StringFromAscii(TestAsciiString, ExtraSlack);
+		FString StringFromAscii = FString::ConstructWithSlack(TestAsciiString, ExtraSlack);
 		TestEqual(TEXT("(ASCII: Valid string with valid slack) resulting capacity"), StringFromAscii.GetAllocatedSize(), ExpectedCapacity);
 
 		FString StringFromFStringView(FStringView(TestString), ExtraSlack);
@@ -463,10 +463,10 @@ bool FStringConstructorWithSlackTest::RunTest(const FString& Parameters)
 
 		const SIZE_T ExpectedCapacity = FMemory::QuantizeSize(NumElements * sizeof(TCHAR));
 
-		FString StringFromTChar(TestString, ExtraSlack);
+		FString StringFromTChar = FString::ConstructWithSlack(TestString, ExtraSlack);
 		TestEqual(TEXT("(TCHAR: Valid string with zero slack) resulting capacity"), StringFromTChar.GetAllocatedSize(), ExpectedCapacity);
 
-		FString StringFromAscii(TestAsciiString, ExtraSlack);
+		FString StringFromAscii = FString::ConstructWithSlack(TestAsciiString, ExtraSlack);
 		TestEqual(TEXT("(ASCII: Valid string with zero slack) resulting capacity"), StringFromAscii.GetAllocatedSize(), ExpectedCapacity);
 
 		FString StringFromFStringView(FStringView(TestString), ExtraSlack);
@@ -486,10 +486,10 @@ bool FStringConstructorWithSlackTest::RunTest(const FString& Parameters)
 
 		const SIZE_T ExpectedCapacity = FMemory::QuantizeSize(NumElements * sizeof(TCHAR));
 
-		FString StringFromTChar(TestString, ExtraSlack);
+		FString StringFromTChar = FString::ConstructWithSlack(TestString, ExtraSlack);
 		TestEqual(TEXT("(TCHAR: Empty string with slack) resulting capacity"), StringFromTChar.GetAllocatedSize(), ExpectedCapacity);
 
-		FString StringFromAscii(TestAsciiString, ExtraSlack);
+		FString StringFromAscii = FString::ConstructWithSlack(TestAsciiString, ExtraSlack);
 		TestEqual(TEXT("(ASCII: Empty string with slack) resulting capacity"), StringFromAscii.GetAllocatedSize(), ExpectedCapacity);
 
 		FString StringFromFStringView(FStringView(TestString), ExtraSlack);
@@ -507,10 +507,10 @@ bool FStringConstructorWithSlackTest::RunTest(const FString& Parameters)
 
 		const SIZE_T ExpectedCapacity = 0u;
 
-		FString StringFromTChar(TestString, ExtraSlack);
+		FString StringFromTChar = FString::ConstructWithSlack(TestString, ExtraSlack);
 		TestEqual(TEXT("(TCHAR: Empty string with zero slack) resulting capacity"), StringFromTChar.GetAllocatedSize(), ExpectedCapacity);
 
-		FString StringFromAscii(TestAsciiString, ExtraSlack);
+		FString StringFromAscii = FString::ConstructWithSlack(TestAsciiString, ExtraSlack);
 		TestEqual(TEXT("(ASCII: Empty string with zero slack) resulting capacity"), StringFromAscii.GetAllocatedSize(), ExpectedCapacity);
 
 		FString StringFromFStringView(FStringView(TestString), ExtraSlack);
@@ -523,7 +523,7 @@ bool FStringConstructorWithSlackTest::RunTest(const FString& Parameters)
 	return true;
 }
 
-IMPLEMENT_SIMPLE_AUTOMATION_TEST(FStringEqualityTest, "System.Core.String.Equality", EAutomationTestFlags::ApplicationContextMask | EAutomationTestFlags::SmokeFilter)
+IMPLEMENT_SIMPLE_AUTOMATION_TEST(FStringEqualityTest, "System.Core.String.Equality", EAutomationTestFlags_ApplicationContextMask | EAutomationTestFlags::SmokeFilter)
 bool FStringEqualityTest::RunTest(const FString& Parameters)
 {
 	auto TestSelfEquality = [this](const TCHAR* A)
@@ -562,7 +562,7 @@ bool FStringEqualityTest::RunTest(const FString& Parameters)
 	return true;	
 }
 
-IMPLEMENT_SIMPLE_AUTOMATION_TEST(FStringPathConcatCompoundOperatorTest, "System.Core.String.PathConcatCompoundOperator", EAutomationTestFlags::ApplicationContextMask | EAutomationTestFlags::SmokeFilter)
+IMPLEMENT_SIMPLE_AUTOMATION_TEST(FStringPathConcatCompoundOperatorTest, "System.Core.String.PathConcatCompoundOperator", EAutomationTestFlags_ApplicationContextMask | EAutomationTestFlags::SmokeFilter)
 bool FStringPathConcatCompoundOperatorTest::RunTest(const FString& Parameters)
 {
 	// No need to test a nullptr TCHAR* as an input parameter as this is expected to cause a crash
@@ -653,7 +653,7 @@ bool FStringPathConcatCompoundOperatorTest::RunTest(const FString& Parameters)
 	return true;
 }
 
-IMPLEMENT_SIMPLE_AUTOMATION_TEST(FStringFindAndContainsTest, "System.Core.String.FindAndContains", EAutomationTestFlags::ApplicationContextMask | EAutomationTestFlags::EngineFilter)
+IMPLEMENT_SIMPLE_AUTOMATION_TEST(FStringFindAndContainsTest, "System.Core.String.FindAndContains", EAutomationTestFlags_ApplicationContextMask | EAutomationTestFlags::EngineFilter)
 bool FStringFindAndContainsTest::RunTest(const FString& Parameters)
 {
 	auto RunTest = [this](FStringView Search, FStringView Find, ESearchCase::Type SearchCase,
@@ -833,12 +833,12 @@ bool FStringFindAndContainsTest::RunTest(const FString& Parameters)
 	return true;
 }
 
-IMPLEMENT_SIMPLE_AUTOMATION_TEST(FStringConstructorWithLengthTest, "System.Core.String.ConstructorWithLength", EAutomationTestFlags::ApplicationContextMask | EAutomationTestFlags::SmokeFilter)
+IMPLEMENT_SIMPLE_AUTOMATION_TEST(FStringConstructorWithLengthTest, "System.Core.String.ConstructorWithLength", EAutomationTestFlags_ApplicationContextMask | EAutomationTestFlags::SmokeFilter)
 bool FStringConstructorWithLengthTest::RunTest(const FString& Parameters)
 {
 	auto DoTest = [this](const TCHAR* Ptr, int32 Size, const TArray<TCHAR>& Expected)
 	{
-		FString Str(Size, Ptr);
+		FString Str = FString::ConstructFromPtrSize(Ptr, Size);
 
 		const TArray<TCHAR>& StrArr = Str.GetCharArray();
 		if (StrArr != Expected)
@@ -865,7 +865,7 @@ bool FStringConstructorWithLengthTest::RunTest(const FString& Parameters)
 	return true;
 }
 
-IMPLEMENT_SIMPLE_AUTOMATION_TEST(FStringConstructWithSlackTest, "System.Core.String.ConstructWithSlack", EAutomationTestFlags::ApplicationContextMask | EAutomationTestFlags::SmokeFilter)
+IMPLEMENT_SIMPLE_AUTOMATION_TEST(FStringConstructWithSlackTest, "System.Core.String.ConstructWithSlack", EAutomationTestFlags_ApplicationContextMask | EAutomationTestFlags::SmokeFilter)
 bool FStringConstructWithSlackTest::RunTest(const FString& Parameters)
 {
 	auto DoTest = [this](const TCHAR* Ptr, int32 ExtraSlack, const TArray<TCHAR>& Expected)
@@ -900,7 +900,7 @@ bool FStringConstructWithSlackTest::RunTest(const FString& Parameters)
 	return true;
 }
 
-IMPLEMENT_SIMPLE_AUTOMATION_TEST(FStringConstructFromPtrSizeTest, "System.Core.String.ConstructFromPtrSize", EAutomationTestFlags::ApplicationContextMask | EAutomationTestFlags::SmokeFilter)
+IMPLEMENT_SIMPLE_AUTOMATION_TEST(FStringConstructFromPtrSizeTest, "System.Core.String.ConstructFromPtrSize", EAutomationTestFlags_ApplicationContextMask | EAutomationTestFlags::SmokeFilter)
 bool FStringConstructFromPtrSizeTest::RunTest(const FString& Parameters)
 {
 	auto DoTest = [this](const TCHAR* Ptr, int32 Size, const TArray<TCHAR>& Expected)
@@ -931,7 +931,7 @@ bool FStringConstructFromPtrSizeTest::RunTest(const FString& Parameters)
 	return true;
 }
 
-IMPLEMENT_SIMPLE_AUTOMATION_TEST(FStringConstructFromPtrSizeWithSlackTest, "System.Core.String.ConstructFromPtrSizeWithSlack", EAutomationTestFlags::ApplicationContextMask | EAutomationTestFlags::SmokeFilter)
+IMPLEMENT_SIMPLE_AUTOMATION_TEST(FStringConstructFromPtrSizeWithSlackTest, "System.Core.String.ConstructFromPtrSizeWithSlack", EAutomationTestFlags_ApplicationContextMask | EAutomationTestFlags::SmokeFilter)
 bool FStringConstructFromPtrSizeWithSlackTest::RunTest(const FString& Parameters)
 {
 	auto DoTest = [this](const TCHAR* Ptr, int32 Size, int32 ExtraSlack, const TArray<TCHAR>& Expected)
diff --git a/Engine/Source/Runtime/Core/Private/Experimental/Graph/GraphConvert.cpp b/Engine/Source/Runtime/Core/Private/Experimental/Graph/GraphConvert.cpp
--- a/Engine/Source/Runtime/Core/Private/Experimental/Graph/GraphConvert.cpp
+++ b/Engine/Source/Runtime/Core/Private/Experimental/Graph/GraphConvert.cpp
@@ -395,7 +395,7 @@ FGraph ConstructPartialTransposeGraph(
 
 #if WITH_DEV_AUTOMATION_TESTS
 
-IMPLEMENT_SIMPLE_AUTOMATION_TEST(FConvertToGraphTest, "System.Core.Graph.ConvertToGraph", EAutomationTestFlags::ApplicationContextMask | EAutomationTestFlags::EngineFilter);
+IMPLEMENT_SIMPLE_AUTOMATION_TEST(FConvertToGraphTest, "System.Core.Graph.ConvertToGraph", EAutomationTestFlags_ApplicationContextMask | EAutomationTestFlags::EngineFilter);
 bool FConvertToGraphTest::RunTest(const FString& Parameters)
 {
 	using namespace UE::Graph;
@@ -468,7 +468,7 @@ bool FConvertToGraphTest::RunTest(const FString& Parameters)
 	return TestGraph(DestGraph);
 }
 
-IMPLEMENT_SIMPLE_AUTOMATION_TEST(FConvertToSingleBufferGraphTest, "System.Core.Graph.ConvertToSingleBufferGraph", EAutomationTestFlags::ApplicationContextMask | EAutomationTestFlags::EngineFilter);
+IMPLEMENT_SIMPLE_AUTOMATION_TEST(FConvertToSingleBufferGraphTest, "System.Core.Graph.ConvertToSingleBufferGraph", EAutomationTestFlags_ApplicationContextMask | EAutomationTestFlags::EngineFilter);
 bool FConvertToSingleBufferGraphTest::RunTest(const FString& Parameters)
 {
 	using namespace UE::Graph;
@@ -512,7 +512,7 @@ bool FConvertToSingleBufferGraphTest::RunTest(const FString& Parameters)
 	return true;
 }
 
-IMPLEMENT_SIMPLE_AUTOMATION_TEST(FCondensationGraphTest_Core, "System.Core.Graph.CondensationGraph", EAutomationTestFlags::ApplicationContextMask | EAutomationTestFlags::EngineFilter);
+IMPLEMENT_SIMPLE_AUTOMATION_TEST(FCondensationGraphTest_Core, "System.Core.Graph.CondensationGraph", EAutomationTestFlags_ApplicationContextMask | EAutomationTestFlags::EngineFilter);
 
 bool FCondensationGraphTest_Core::RunTest(const FString& Parameters)
 {
diff --git a/Engine/Source/Runtime/Core/Private/GenericPlatform/GenericPlatformStricmp.cpp b/Engine/Source/Runtime/Core/Private/GenericPlatform/GenericPlatformStricmp.cpp
--- a/Engine/Source/Runtime/Core/Private/GenericPlatform/GenericPlatformStricmp.cpp
+++ b/Engine/Source/Runtime/Core/Private/GenericPlatform/GenericPlatformStricmp.cpp
@@ -124,7 +124,7 @@ int32 FGenericPlatformStricmp::Strnicmp(const UTF32CHAR* Str1, const ANSICHAR*
 
 #include "Misc/AutomationTest.h"
 
-IMPLEMENT_SIMPLE_AUTOMATION_TEST(FGenericPlatformStricmpTest, "System.Core.GenericPlatform.Stricmp", EAutomationTestFlags::ApplicationContextMask | EAutomationTestFlags::SmokeFilter)
+IMPLEMENT_SIMPLE_AUTOMATION_TEST(FGenericPlatformStricmpTest, "System.Core.GenericPlatform.Stricmp", EAutomationTestFlags_ApplicationContextMask | EAutomationTestFlags::SmokeFilter)
 
 // Simpler reference implementation based on ToLower() instead of a lookup table.
 // Used to verify correctness on non-Windows platforms
diff --git a/Engine/Source/Runtime/Core/Private/HAL/LowLevelMemTrackerTest.cpp b/Engine/Source/Runtime/Core/Private/HAL/LowLevelMemTrackerTest.cpp
--- a/Engine/Source/Runtime/Core/Private/HAL/LowLevelMemTrackerTest.cpp
+++ b/Engine/Source/Runtime/Core/Private/HAL/LowLevelMemTrackerTest.cpp
@@ -15,7 +15,7 @@ LLM_DEFINE_TAG(LLMTestParentTag_LLMTestChildTag3); // A test tag to verify that
 
 #if WITH_DEV_AUTOMATION_TESTS
 
-IMPLEMENT_SIMPLE_AUTOMATION_TEST(FLLMScopeTest, "System.Core.LLM.Scope", EAutomationTestFlags::ApplicationContextMask | EAutomationTestFlags::EngineFilter);
+IMPLEMENT_SIMPLE_AUTOMATION_TEST(FLLMScopeTest, "System.Core.LLM.Scope", EAutomationTestFlags_ApplicationContextMask | EAutomationTestFlags::EngineFilter);
 
 bool FLLMScopeTest::RunTest(const FString& Parameters)
 {
@@ -69,7 +69,7 @@ bool FLLMScopeTest::RunTest(const FString& Parameters)
 	return true;
 }
 
-IMPLEMENT_SIMPLE_AUTOMATION_TEST(FLLMMiscTest, "System.Core.LLM.Misc", EAutomationTestFlags::ApplicationContextMask | EAutomationTestFlags::EngineFilter);
+IMPLEMENT_SIMPLE_AUTOMATION_TEST(FLLMMiscTest, "System.Core.LLM.Misc", EAutomationTestFlags_ApplicationContextMask | EAutomationTestFlags::EngineFilter);
 
 bool FLLMMiscTest::RunTest(const FString& Parameters)
 {
@@ -92,7 +92,7 @@ bool FLLMMiscTest::RunTest(const FString& Parameters)
 	return true;
 }
 
-IMPLEMENT_SIMPLE_AUTOMATION_TEST(FTopologicalSortTest, "System.Core.LLM.TopologicalSort", EAutomationTestFlags::ApplicationContextMask | EAutomationTestFlags::EngineFilter);
+IMPLEMENT_SIMPLE_AUTOMATION_TEST(FTopologicalSortTest, "System.Core.LLM.TopologicalSort", EAutomationTestFlags_ApplicationContextMask | EAutomationTestFlags::EngineFilter);
 
 bool FTopologicalSortTest::RunTest(const FString& Parameters)
 {
diff --git a/Engine/Source/Runtime/Core/Private/Math/BasicMathExpressionEvaluator.cpp b/Engine/Source/Runtime/Core/Private/Math/BasicMathExpressionEvaluator.cpp
--- a/Engine/Source/Runtime/Core/Private/Math/BasicMathExpressionEvaluator.cpp
+++ b/Engine/Source/Runtime/Core/Private/Math/BasicMathExpressionEvaluator.cpp
@@ -304,7 +304,7 @@ bool TestInvalidExpression(FAutomationTestBase* Test, const TCHAR* Expression)
 }
 
 
-IMPLEMENT_SIMPLE_AUTOMATION_TEST(FBasicMathExpressionEvaluatorTest, "System.Core.Math.Evaluate.Valid Expressions", EAutomationTestFlags::ApplicationContextMask | EAutomationTestFlags::SmokeFilter)
+IMPLEMENT_SIMPLE_AUTOMATION_TEST(FBasicMathExpressionEvaluatorTest, "System.Core.Math.Evaluate.Valid Expressions", EAutomationTestFlags_ApplicationContextMask | EAutomationTestFlags::SmokeFilter)
 // Evaluates valid math expressions.
 bool FBasicMathExpressionEvaluatorTest::RunTest(const FString& Parameters)
 {
@@ -332,7 +332,7 @@ bool FBasicMathExpressionEvaluatorTest::RunTest(const FString& Parameters)
 	return true;
 }
 
-IMPLEMENT_SIMPLE_AUTOMATION_TEST(FBasicMathExpressionEvaluatorWhitespaceExpressionsTest, "System.Core.Math.Evaluate.Valid Expressions With Whitespaces", EAutomationTestFlags::ApplicationContextMask | EAutomationTestFlags::SmokeFilter)
+IMPLEMENT_SIMPLE_AUTOMATION_TEST(FBasicMathExpressionEvaluatorWhitespaceExpressionsTest, "System.Core.Math.Evaluate.Valid Expressions With Whitespaces", EAutomationTestFlags_ApplicationContextMask | EAutomationTestFlags::SmokeFilter)
 // Evaluates a valid math expression with leading and trailing white spaces.
 bool FBasicMathExpressionEvaluatorWhitespaceExpressionsTest::RunTest(const FString& Parameters)
 {
@@ -341,7 +341,7 @@ bool FBasicMathExpressionEvaluatorWhitespaceExpressionsTest::RunTest(const FStri
 	return true;
 }
 
-IMPLEMENT_SIMPLE_AUTOMATION_TEST(FBasicMathExpressionEvaluatorGroupedExpressionsTest, "System.Core.Math.Evaluate.Valid Grouped Expressions", EAutomationTestFlags::ApplicationContextMask | EAutomationTestFlags::SmokeFilter)
+IMPLEMENT_SIMPLE_AUTOMATION_TEST(FBasicMathExpressionEvaluatorGroupedExpressionsTest, "System.Core.Math.Evaluate.Valid Grouped Expressions", EAutomationTestFlags_ApplicationContextMask | EAutomationTestFlags::SmokeFilter)
 // Evaluates valid math expressions that are grouped
 bool FBasicMathExpressionEvaluatorGroupedExpressionsTest::RunTest(const FString& Parameters)
 {
@@ -353,7 +353,7 @@ bool FBasicMathExpressionEvaluatorGroupedExpressionsTest::RunTest(const FString&
 	return true;
 }
 
-IMPLEMENT_SIMPLE_AUTOMATION_TEST(FBasicMathExpressionEvaluatorInvalidExpressionTest, "System.Core.Math.Evaluate.Invalid Expressions", EAutomationTestFlags::ApplicationContextMask | EAutomationTestFlags::SmokeFilter)
+IMPLEMENT_SIMPLE_AUTOMATION_TEST(FBasicMathExpressionEvaluatorInvalidExpressionTest, "System.Core.Math.Evaluate.Invalid Expressions", EAutomationTestFlags_ApplicationContextMask | EAutomationTestFlags::SmokeFilter)
 
 // Evaluates invalid expressions.
 // Invalid expressions will report errors and not crash.
diff --git a/Engine/Source/Runtime/Core/Private/Misc/CStringTest.cpp b/Engine/Source/Runtime/Core/Private/Misc/CStringTest.cpp
--- a/Engine/Source/Runtime/Core/Private/Misc/CStringTest.cpp
+++ b/Engine/Source/Runtime/Core/Private/Misc/CStringTest.cpp
@@ -44,7 +44,7 @@ protected:
 	}
 };
 
-IMPLEMENT_CUSTOM_SIMPLE_AUTOMATION_TEST(FCStringGetVarArgsTest, FCStringGetVarArgsTestBase, "System.Core.Misc.CString.GetVarArgs", EAutomationTestFlags::ApplicationContextMask | EAutomationTestFlags::EngineFilter)
+IMPLEMENT_CUSTOM_SIMPLE_AUTOMATION_TEST(FCStringGetVarArgsTest, FCStringGetVarArgsTestBase, "System.Core.Misc.CString.GetVarArgs", EAutomationTestFlags_ApplicationContextMask | EAutomationTestFlags::EngineFilter)
 bool FCStringGetVarArgsTest::RunTest(const FString& Parameters)
 {
 #if PLATFORM_64BITS
@@ -97,7 +97,7 @@ bool FCStringGetVarArgsTest::RunTest(const FString& Parameters)
 	return true;
 }
 
-IMPLEMENT_SIMPLE_AUTOMATION_TEST(FCStringStrstrTest, "System.Core.Misc.CString.Strstr", EAutomationTestFlags::ApplicationContextMask | EAutomationTestFlags::EngineFilter)
+IMPLEMENT_SIMPLE_AUTOMATION_TEST(FCStringStrstrTest, "System.Core.Misc.CString.Strstr", EAutomationTestFlags_ApplicationContextMask | EAutomationTestFlags::EngineFilter)
 bool FCStringStrstrTest::RunTest(const FString& Parameters)
 {
 	auto RunTest = [this](const TCHAR* Search, const TCHAR* Find, int32 ExpectedSensitiveIndex, int32 ExpectedInsensitiveIndex)
@@ -142,7 +142,7 @@ bool FCStringStrstrTest::RunTest(const FString& Parameters)
 	return true;
 }
 
-IMPLEMENT_SIMPLE_AUTOMATION_TEST(FCStringStrnstrTest, "System.Core.Misc.CString.Strnstr", EAutomationTestFlags::ApplicationContextMask | EAutomationTestFlags::EngineFilter)
+IMPLEMENT_SIMPLE_AUTOMATION_TEST(FCStringStrnstrTest, "System.Core.Misc.CString.Strnstr", EAutomationTestFlags_ApplicationContextMask | EAutomationTestFlags::EngineFilter)
 bool FCStringStrnstrTest::RunTest(const FString& Parameters)
 {
 	auto RunTest = [this](FStringView Search, FStringView Find, int32 ExpectedSensitiveIndex, int32 ExpectedInsensitiveIndex)
diff --git a/Engine/Source/Runtime/Core/Private/Misc/ExpressionParser.cpp b/Engine/Source/Runtime/Core/Private/Misc/ExpressionParser.cpp
--- a/Engine/Source/Runtime/Core/Private/Misc/ExpressionParser.cpp
+++ b/Engine/Source/Runtime/Core/Private/Misc/ExpressionParser.cpp
@@ -1063,7 +1063,7 @@ private:
 
 };
 
-IMPLEMENT_SIMPLE_AUTOMATION_TEST(FShortCircuitParserTest, "System.Core.Expression Parser.Short Circuit", EAutomationTestFlags::ApplicationContextMask | EAutomationTestFlags::SmokeFilter | EAutomationTestFlags::HighPriority)
+IMPLEMENT_SIMPLE_AUTOMATION_TEST(FShortCircuitParserTest, "System.Core.Expression Parser.Short Circuit", EAutomationTestFlags_ApplicationContextMask | EAutomationTestFlags::SmokeFilter | EAutomationTestFlags::HighPriority)
 bool FShortCircuitParserTest::RunTest(const FString& Parameters)
 {
 	struct FExpectedResult
diff --git a/Engine/Source/Runtime/Core/Private/Misc/MTAccessDetector.cpp b/Engine/Source/Runtime/Core/Private/Misc/MTAccessDetector.cpp
--- a/Engine/Source/Runtime/Core/Private/Misc/MTAccessDetector.cpp
+++ b/Engine/Source/Runtime/Core/Private/Misc/MTAccessDetector.cpp
@@ -29,7 +29,7 @@ FMRSWRecursiveAccessDetector::FDestructionSentinelStackTls& FMRSWRecursiveAccess
 //----------------------------------------------------------------------//
 // FRWAccessDetector tests
 //----------------------------------------------------------------------//
-IMPLEMENT_SIMPLE_AUTOMATION_TEST(FRWAccessDetector_ConcurentReadTest, "System.Core.Misc.MTAccessDetector.ConcurrentReadAccess", EAutomationTestFlags::ApplicationContextMask | EAutomationTestFlags::EngineFilter)
+IMPLEMENT_SIMPLE_AUTOMATION_TEST(FRWAccessDetector_ConcurentReadTest, "System.Core.Misc.MTAccessDetector.ConcurrentReadAccess", EAutomationTestFlags_ApplicationContextMask | EAutomationTestFlags::EngineFilter)
 bool FRWAccessDetector_ConcurentReadTest::RunTest(const FString& Parameters)
 {
 	bool bReading1 = false;
@@ -54,7 +54,7 @@ bool FRWAccessDetector_ConcurentReadTest::RunTest(const FString& Parameters)
 	return Success;
 }
 
-IMPLEMENT_SIMPLE_AUTOMATION_TEST(FRWAccessDetector_ConcurentWriteTest, "System.Core.Misc.MTAccessDetector.ConcurrentWriteAccess", EAutomationTestFlags::ApplicationContextMask | EAutomationTestFlags::EngineFilter)
+IMPLEMENT_SIMPLE_AUTOMATION_TEST(FRWAccessDetector_ConcurentWriteTest, "System.Core.Misc.MTAccessDetector.ConcurrentWriteAccess", EAutomationTestFlags_ApplicationContextMask | EAutomationTestFlags::EngineFilter)
 bool FRWAccessDetector_ConcurentWriteTest::RunTest(const FString& Parameters)
 {
 	bool bWriting1 = false;
@@ -79,7 +79,7 @@ bool FRWAccessDetector_ConcurentWriteTest::RunTest(const FString& Parameters)
 	return Success;
 }
 
-IMPLEMENT_SIMPLE_AUTOMATION_TEST(FRWAccessDetector_ConcurentReadWriteTest, "System.Core.Misc.MTAccessDetector.ConcurrentReadWriteAccess", EAutomationTestFlags::ApplicationContextMask | EAutomationTestFlags::EngineFilter)
+IMPLEMENT_SIMPLE_AUTOMATION_TEST(FRWAccessDetector_ConcurentReadWriteTest, "System.Core.Misc.MTAccessDetector.ConcurrentReadWriteAccess", EAutomationTestFlags_ApplicationContextMask | EAutomationTestFlags::EngineFilter)
 bool FRWAccessDetector_ConcurentReadWriteTest::RunTest(const FString& Parameters)
 {
 	bool bReading = false;
@@ -104,7 +104,7 @@ bool FRWAccessDetector_ConcurentReadWriteTest::RunTest(const FString& Parameters
 	return Success;
 }
 
-IMPLEMENT_SIMPLE_AUTOMATION_TEST(FRWAccessDetector_WriteReentrance, "System.Core.Misc.MTAccessDetector.WriteReentrance", EAutomationTestFlags::ApplicationContextMask | EAutomationTestFlags::EngineFilter)
+IMPLEMENT_SIMPLE_AUTOMATION_TEST(FRWAccessDetector_WriteReentrance, "System.Core.Misc.MTAccessDetector.WriteReentrance", EAutomationTestFlags_ApplicationContextMask | EAutomationTestFlags::EngineFilter)
 bool FRWAccessDetector_WriteReentrance::RunTest(const FString& Parameters)
 {
 	UE_MT_DECLARE_RW_ACCESS_DETECTOR(MTAccessDetector);
@@ -118,7 +118,7 @@ bool FRWAccessDetector_WriteReentrance::RunTest(const FString& Parameters)
 	return Success;
 }
 
-IMPLEMENT_SIMPLE_AUTOMATION_TEST(FRWAccessDetector_ReadReentrance, "System.Core.Misc.MTAccessDetector.ReadReentrance", EAutomationTestFlags::ApplicationContextMask | EAutomationTestFlags::EngineFilter)
+IMPLEMENT_SIMPLE_AUTOMATION_TEST(FRWAccessDetector_ReadReentrance, "System.Core.Misc.MTAccessDetector.ReadReentrance", EAutomationTestFlags_ApplicationContextMask | EAutomationTestFlags::EngineFilter)
 bool FRWAccessDetector_ReadReentrance::RunTest(const FString& Parameters)
 {
 	UE_MT_DECLARE_RW_ACCESS_DETECTOR(MTAccessDetector);
@@ -135,7 +135,7 @@ bool FRWAccessDetector_ReadReentrance::RunTest(const FString& Parameters)
 //----------------------------------------------------------------------//
 // FRWRecursiveAccessDetector tests
 //----------------------------------------------------------------------//
-IMPLEMENT_SIMPLE_AUTOMATION_TEST(FRWAccessDetector_RecursiveConcurentReadTest, "System.Core.Misc.MTAccessDetector.RecursiveConcurrentReadAccess", EAutomationTestFlags::ApplicationContextMask | EAutomationTestFlags::EngineFilter)
+IMPLEMENT_SIMPLE_AUTOMATION_TEST(FRWAccessDetector_RecursiveConcurentReadTest, "System.Core.Misc.MTAccessDetector.RecursiveConcurrentReadAccess", EAutomationTestFlags_ApplicationContextMask | EAutomationTestFlags::EngineFilter)
 bool FRWAccessDetector_RecursiveConcurentReadTest::RunTest(const FString& Parameters)
 {
 	bool bReading1 = false;
@@ -160,7 +160,7 @@ bool FRWAccessDetector_RecursiveConcurentReadTest::RunTest(const FString& Parame
 	return Success;
 }
 
-IMPLEMENT_SIMPLE_AUTOMATION_TEST(FRWAccessDetector_RecursiveConcurentWriteTest, "System.Core.Misc.MTAccessDetector.RecursiveConcurrentWriteAccess", EAutomationTestFlags::ApplicationContextMask | EAutomationTestFlags::EngineFilter)
+IMPLEMENT_SIMPLE_AUTOMATION_TEST(FRWAccessDetector_RecursiveConcurentWriteTest, "System.Core.Misc.MTAccessDetector.RecursiveConcurrentWriteAccess", EAutomationTestFlags_ApplicationContextMask | EAutomationTestFlags::EngineFilter)
 bool FRWAccessDetector_RecursiveConcurentWriteTest::RunTest(const FString& Parameters)
 {
 	bool bWriting1 = false;
@@ -185,7 +185,7 @@ bool FRWAccessDetector_RecursiveConcurentWriteTest::RunTest(const FString& Param
 	return Success;
 }
 
-IMPLEMENT_SIMPLE_AUTOMATION_TEST(FRWAccessDetector_RecursiveConcurentReadWriteTest, "System.Core.Misc.MTAccessDetector.RecursiveConcurrentReadWriteAccess", EAutomationTestFlags::ApplicationContextMask | EAutomationTestFlags::EngineFilter)
+IMPLEMENT_SIMPLE_AUTOMATION_TEST(FRWAccessDetector_RecursiveConcurentReadWriteTest, "System.Core.Misc.MTAccessDetector.RecursiveConcurrentReadWriteAccess", EAutomationTestFlags_ApplicationContextMask | EAutomationTestFlags::EngineFilter)
 bool FRWAccessDetector_RecursiveConcurentReadWriteTest::RunTest(const FString& Parameters)
 {
 	bool bReading = false;
@@ -210,7 +210,7 @@ bool FRWAccessDetector_RecursiveConcurentReadWriteTest::RunTest(const FString& P
 	return Success;
 }
 
-IMPLEMENT_SIMPLE_AUTOMATION_TEST(FRWAccessDetector_RecursiveWriteReentrance, "System.Core.Misc.MTAccessDetector.RecursiveWriteReentrance", EAutomationTestFlags::ApplicationContextMask | EAutomationTestFlags::EngineFilter)
+IMPLEMENT_SIMPLE_AUTOMATION_TEST(FRWAccessDetector_RecursiveWriteReentrance, "System.Core.Misc.MTAccessDetector.RecursiveWriteReentrance", EAutomationTestFlags_ApplicationContextMask | EAutomationTestFlags::EngineFilter)
 bool FRWAccessDetector_RecursiveWriteReentrance::RunTest(const FString& Parameters)
 {
 	UE_MT_DECLARE_RW_RECURSIVE_ACCESS_DETECTOR(MTAccessDetector);
@@ -224,7 +224,7 @@ bool FRWAccessDetector_RecursiveWriteReentrance::RunTest(const FString& Paramete
 	return Success;
 }
 
-IMPLEMENT_SIMPLE_AUTOMATION_TEST(FRWAccessDetector_RecursiveReadReentrance, "System.Core.Misc.MTAccessDetector.RecursiveReadReentrance", EAutomationTestFlags::ApplicationContextMask | EAutomationTestFlags::EngineFilter)
+IMPLEMENT_SIMPLE_AUTOMATION_TEST(FRWAccessDetector_RecursiveReadReentrance, "System.Core.Misc.MTAccessDetector.RecursiveReadReentrance", EAutomationTestFlags_ApplicationContextMask | EAutomationTestFlags::EngineFilter)
 bool FRWAccessDetector_RecursiveReadReentrance::RunTest(const FString& Parameters)
 {
 	UE_MT_DECLARE_RW_RECURSIVE_ACCESS_DETECTOR(MTAccessDetector);
@@ -241,7 +241,7 @@ bool FRWAccessDetector_RecursiveReadReentrance::RunTest(const FString& Parameter
 //----------------------------------------------------------------------//
 // FRWFullyRecursiveAccessDetector tests
 //----------------------------------------------------------------------//
-IMPLEMENT_SIMPLE_AUTOMATION_TEST(FRWAccessDetector_FullyRecursiveConcurentReadTest, "System.Core.Misc.MTAccessDetector.FullyRecursiveConcurrentReadAccess", EAutomationTestFlags::ApplicationContextMask | EAutomationTestFlags::EngineFilter)
+IMPLEMENT_SIMPLE_AUTOMATION_TEST(FRWAccessDetector_FullyRecursiveConcurentReadTest, "System.Core.Misc.MTAccessDetector.FullyRecursiveConcurrentReadAccess", EAutomationTestFlags_ApplicationContextMask | EAutomationTestFlags::EngineFilter)
 bool FRWAccessDetector_FullyRecursiveConcurentReadTest::RunTest(const FString& Parameters)
 {
 	bool bReading1 = false;
@@ -266,7 +266,7 @@ bool FRWAccessDetector_FullyRecursiveConcurentReadTest::RunTest(const FString& P
 	return Success;
 }
 
-IMPLEMENT_SIMPLE_AUTOMATION_TEST(FRWAccessDetector_FullyRecursiveConcurentWriteTest, "System.Core.Misc.MTAccessDetector.FullyRecursiveConcurrentWriteAccess", EAutomationTestFlags::ApplicationContextMask | EAutomationTestFlags::EngineFilter)
+IMPLEMENT_SIMPLE_AUTOMATION_TEST(FRWAccessDetector_FullyRecursiveConcurentWriteTest, "System.Core.Misc.MTAccessDetector.FullyRecursiveConcurrentWriteAccess", EAutomationTestFlags_ApplicationContextMask | EAutomationTestFlags::EngineFilter)
 bool FRWAccessDetector_FullyRecursiveConcurentWriteTest::RunTest(const FString& Parameters)
 {
 	bool bWriting1 = false;
@@ -291,7 +291,7 @@ bool FRWAccessDetector_FullyRecursiveConcurentWriteTest::RunTest(const FString&
 	return Success;
 }
 
-IMPLEMENT_SIMPLE_AUTOMATION_TEST(FRWAccessDetector_FullyRecursiveConcurentReadWriteTest, "System.Core.Misc.MTAccessDetector.FullyRecursiveConcurrentReadWriteAccess", EAutomationTestFlags::ApplicationContextMask | EAutomationTestFlags::EngineFilter)
+IMPLEMENT_SIMPLE_AUTOMATION_TEST(FRWAccessDetector_FullyRecursiveConcurentReadWriteTest, "System.Core.Misc.MTAccessDetector.FullyRecursiveConcurrentReadWriteAccess", EAutomationTestFlags_ApplicationContextMask | EAutomationTestFlags::EngineFilter)
 bool FRWAccessDetector_FullyRecursiveConcurentReadWriteTest::RunTest(const FString& Parameters)
 {
 	bool bReading = false;
@@ -316,7 +316,7 @@ bool FRWAccessDetector_FullyRecursiveConcurentReadWriteTest::RunTest(const FStri
 	return Success;
 }
 
-IMPLEMENT_SIMPLE_AUTOMATION_TEST(FRWAccessDetector_FullyRecursiveWriteReentrance, "System.Core.Misc.MTAccessDetector.FullyRecursiveWriteReentrance", EAutomationTestFlags::ApplicationContextMask | EAutomationTestFlags::EngineFilter)
+IMPLEMENT_SIMPLE_AUTOMATION_TEST(FRWAccessDetector_FullyRecursiveWriteReentrance, "System.Core.Misc.MTAccessDetector.FullyRecursiveWriteReentrance", EAutomationTestFlags_ApplicationContextMask | EAutomationTestFlags::EngineFilter)
 bool FRWAccessDetector_FullyRecursiveWriteReentrance::RunTest(const FString& Parameters)
 {
 	UE_MT_DECLARE_RW_FULLY_RECURSIVE_ACCESS_DETECTOR(MTAccessDetector);
@@ -330,7 +330,7 @@ bool FRWAccessDetector_FullyRecursiveWriteReentrance::RunTest(const FString& Par
 	return Success;
 }
 
-IMPLEMENT_SIMPLE_AUTOMATION_TEST(FRWAccessDetector_FullyRecursiveReadReentrance, "System.Core.Misc.MTAccessDetector.FullyRecursiveReadReentrance", EAutomationTestFlags::ApplicationContextMask | EAutomationTestFlags::EngineFilter)
+IMPLEMENT_SIMPLE_AUTOMATION_TEST(FRWAccessDetector_FullyRecursiveReadReentrance, "System.Core.Misc.MTAccessDetector.FullyRecursiveReadReentrance", EAutomationTestFlags_ApplicationContextMask | EAutomationTestFlags::EngineFilter)
 bool FRWAccessDetector_FullyRecursiveReadReentrance::RunTest(const FString& Parameters)
 {
 	UE_MT_DECLARE_RW_FULLY_RECURSIVE_ACCESS_DETECTOR(MTAccessDetector);
@@ -383,7 +383,7 @@ namespace MTAccessDetector_Private
 	};
 }
 
-IMPLEMENT_SIMPLE_AUTOMATION_TEST(FMRSWRecursiveAccessDetectorTest, "System.Core.MRSWRecursiveAccessDetector", EAutomationTestFlags::ApplicationContextMask | EAutomationTestFlags::EngineFilter);
+IMPLEMENT_SIMPLE_AUTOMATION_TEST(FMRSWRecursiveAccessDetectorTest, "System.Core.MRSWRecursiveAccessDetector", EAutomationTestFlags_ApplicationContextMask | EAutomationTestFlags::EngineFilter);
 
 bool FMRSWRecursiveAccessDetectorTest::RunTest(const FString& Parameters)
 {
@@ -459,7 +459,7 @@ bool FMRSWRecursiveAccessDetectorTest::RunTest(const FString& Parameters)
 #pragma warning(pop)
 }
 
-IMPLEMENT_SIMPLE_AUTOMATION_TEST(FMRSWRecursiveAccessDetectorTrivialRelocationTest, "System.Core.MRSWRecursiveAccessDetector.TrivialRelocation", EAutomationTestFlags::ApplicationContextMask | EAutomationTestFlags::EngineFilter);
+IMPLEMENT_SIMPLE_AUTOMATION_TEST(FMRSWRecursiveAccessDetectorTrivialRelocationTest, "System.Core.MRSWRecursiveAccessDetector.TrivialRelocation", EAutomationTestFlags_ApplicationContextMask | EAutomationTestFlags::EngineFilter);
 
 bool FMRSWRecursiveAccessDetectorTrivialRelocationTest::RunTest(const FString& Parameters)
 {
@@ -546,7 +546,7 @@ void RWAccessDetectorConcurrentReadersPerfTest()
 	Wait(Tasks);
 }
 
-IMPLEMENT_SIMPLE_AUTOMATION_TEST(FRWRecursiveAccessDetectorPerfTest, "System.Core.RWRecursiveAccessDetector.Perf", EAutomationTestFlags::ApplicationContextMask | EAutomationTestFlags::EngineFilter);
+IMPLEMENT_SIMPLE_AUTOMATION_TEST(FRWRecursiveAccessDetectorPerfTest, "System.Core.RWRecursiveAccessDetector.Perf", EAutomationTestFlags_ApplicationContextMask | EAutomationTestFlags::EngineFilter);
 
 bool FRWRecursiveAccessDetectorPerfTest::RunTest(const FString& Parameters)
 {
diff --git a/Engine/Source/Runtime/Core/Private/Misc/Parse.cpp b/Engine/Source/Runtime/Core/Private/Misc/Parse.cpp
--- a/Engine/Source/Runtime/Core/Private/Misc/Parse.cpp
+++ b/Engine/Source/Runtime/Core/Private/Misc/Parse.cpp
@@ -1325,7 +1325,7 @@ bool FParse::SchemeNameFromURI(const TCHAR* URI, FString& OutSchemeName)
 #if !(UE_BUILD_SHIPPING || UE_BUILD_TEST)
 
 #include "Misc/AutomationTest.h"
-IMPLEMENT_SIMPLE_AUTOMATION_TEST(FParseLineExtendedTest, "System.Core.Misc.ParseLineExtended", EAutomationTestFlags::ApplicationContextMask | EAutomationTestFlags::SmokeFilter)
+IMPLEMENT_SIMPLE_AUTOMATION_TEST(FParseLineExtendedTest, "System.Core.Misc.ParseLineExtended", EAutomationTestFlags_ApplicationContextMask | EAutomationTestFlags::SmokeFilter)
 bool FParseLineExtendedTest::RunTest(const FString& Parameters)
 {
 	const TCHAR* Tests[] = {
diff --git a/Engine/Source/Runtime/Core/Private/Misc/TextFilterTests.cpp b/Engine/Source/Runtime/Core/Private/Misc/TextFilterTests.cpp
--- a/Engine/Source/Runtime/Core/Private/Misc/TextFilterTests.cpp
+++ b/Engine/Source/Runtime/Core/Private/Misc/TextFilterTests.cpp
@@ -222,7 +222,7 @@ namespace TextFilterTests
 	}
 };
 
-IMPLEMENT_SIMPLE_AUTOMATION_TEST(FTextFilterTests_TextFilterExpressionEvaluator, "System.Core.Text.TextFilterExpressionEvaluator", EAutomationTestFlags::ApplicationContextMask | EAutomationTestFlags::EngineFilter)
+IMPLEMENT_SIMPLE_AUTOMATION_TEST(FTextFilterTests_TextFilterExpressionEvaluator, "System.Core.Text.TextFilterExpressionEvaluator", EAutomationTestFlags_ApplicationContextMask | EAutomationTestFlags::EngineFilter)
 bool FTextFilterTests_TextFilterExpressionEvaluator::RunTest(const FString& Parameters)
 {
 	using namespace TextFilterTests;
@@ -248,7 +248,7 @@ bool FTextFilterTests_TextFilterExpressionEvaluator::RunTest(const FString& Para
 	return bResult;
 }
 
-IMPLEMENT_SIMPLE_AUTOMATION_TEST(FTextFilterTests_TextFilter, "System.Core.Text.TextFilter", EAutomationTestFlags::ApplicationContextMask | EAutomationTestFlags::EngineFilter)
+IMPLEMENT_SIMPLE_AUTOMATION_TEST(FTextFilterTests_TextFilter, "System.Core.Text.TextFilter", EAutomationTestFlags_ApplicationContextMask | EAutomationTestFlags::EngineFilter)
 bool FTextFilterTests_TextFilter::RunTest(const FString& Parameters)
 {
 	using namespace TextFilterTests;
diff --git a/Engine/Source/Runtime/Core/Private/Tests/HAL/FileManagerTest.cpp b/Engine/Source/Runtime/Core/Private/Tests/HAL/FileManagerTest.cpp
--- a/Engine/Source/Runtime/Core/Private/Tests/HAL/FileManagerTest.cpp
+++ b/Engine/Source/Runtime/Core/Private/Tests/HAL/FileManagerTest.cpp
@@ -6,7 +6,7 @@
 
 #if WITH_DEV_AUTOMATION_TESTS
 
-IMPLEMENT_SIMPLE_AUTOMATION_TEST(FFileManagerReaderTest, "System.Core.HAL.FileManager.Reader", EAutomationTestFlags::ApplicationContextMask | EAutomationTestFlags::EngineFilter)
+IMPLEMENT_SIMPLE_AUTOMATION_TEST(FFileManagerReaderTest, "System.Core.HAL.FileManager.Reader", EAutomationTestFlags_ApplicationContextMask | EAutomationTestFlags::EngineFilter)
 
 bool FFileManagerReaderTest::RunTest(const FString& Parameters)
 {
diff --git a/Engine/Source/Runtime/Core/Private/Tests/Serialization/CompactBinaryTest.cpp b/Engine/Source/Runtime/Core/Private/Tests/Serialization/CompactBinaryTest.cpp
--- a/Engine/Source/Runtime/Core/Private/Tests/Serialization/CompactBinaryTest.cpp
+++ b/Engine/Source/Runtime/Core/Private/Tests/Serialization/CompactBinaryTest.cpp
@@ -17,7 +17,7 @@
 ///////////////////////////////////////////////////////////////////////////////////////////////////////////////////////
 
 static constexpr EAutomationTestFlags CompactBinaryTestFlags =
-	EAutomationTestFlags::ApplicationContextMask | EAutomationTestFlags::SmokeFilter;
+	EAutomationTestFlags_ApplicationContextMask | EAutomationTestFlags::SmokeFilter;
 
 ///////////////////////////////////////////////////////////////////////////////////////////////////////////////////////
 
diff --git a/Engine/Source/Runtime/Core/Tests/Math/GuardedIntTest.cpp b/Engine/Source/Runtime/Core/Tests/Math/GuardedIntTest.cpp
--- a/Engine/Source/Runtime/Core/Tests/Math/GuardedIntTest.cpp
+++ b/Engine/Source/Runtime/Core/Tests/Math/GuardedIntTest.cpp
@@ -176,7 +176,7 @@ bool DoGuardedIntTestForType(FAutomationTestBase& Test)
 }
 }
 
-IMPLEMENT_SIMPLE_AUTOMATION_TEST(FGuardedIntTest, "System.Core.GuardedInt", EAutomationTestFlags::ApplicationContextMask | EAutomationTestFlags::EngineFilter)
+IMPLEMENT_SIMPLE_AUTOMATION_TEST(FGuardedIntTest, "System.Core.GuardedInt", EAutomationTestFlags_ApplicationContextMask | EAutomationTestFlags::EngineFilter)
 bool FGuardedIntTest::RunTest(const FString&)
 {
 	bool bResult = true;
diff --git a/Engine/Source/Runtime/CoreOnline/Private/Tests/CoreOnlineTests.cpp b/Engine/Source/Runtime/CoreOnline/Private/Tests/CoreOnlineTests.cpp
--- a/Engine/Source/Runtime/CoreOnline/Private/Tests/CoreOnlineTests.cpp
+++ b/Engine/Source/Runtime/CoreOnline/Private/Tests/CoreOnlineTests.cpp
@@ -6,7 +6,7 @@
 
 IMPLEMENT_SIMPLE_AUTOMATION_TEST(FForeignIdRepTest,
 	"System.Engine.Online.ForeignIdRepTest",
-	EAutomationTestFlags::ApplicationContextMask | EAutomationTestFlags::SmokeFilter)
+	EAutomationTestFlags_ApplicationContextMask | EAutomationTestFlags::SmokeFilter)
 	bool FForeignIdRepTest::RunTest(const FString& Parameters)
 {
 	using namespace UE::Online;
diff --git a/Engine/Source/Runtime/CoreUObject/Private/AssetRegistry/AssetBundleData.cpp b/Engine/Source/Runtime/CoreUObject/Private/AssetRegistry/AssetBundleData.cpp
--- a/Engine/Source/Runtime/CoreUObject/Private/AssetRegistry/AssetBundleData.cpp
+++ b/Engine/Source/Runtime/CoreUObject/Private/AssetRegistry/AssetBundleData.cpp
@@ -589,7 +589,7 @@ FString FAssetBundleData::ToDebugString() const
 
 // Combine import/export tests
 
-IMPLEMENT_SIMPLE_AUTOMATION_TEST(FAssetBundleEntryImportExportTextTest, "System.AssetRegistry.AssetBundleEntry.ImportExportText", EAutomationTestFlags::ApplicationContextMask | EAutomationTestFlags::EngineFilter);
+IMPLEMENT_SIMPLE_AUTOMATION_TEST(FAssetBundleEntryImportExportTextTest, "System.AssetRegistry.AssetBundleEntry.ImportExportText", EAutomationTestFlags_ApplicationContextMask | EAutomationTestFlags::EngineFilter);
 bool FAssetBundleEntryImportExportTextTest::RunTest(const FString& Parameters)
 {
 	FAssetBundleEntry DefaultEntry;
@@ -676,7 +676,7 @@ bool FAssetBundleEntryImportExportTextTest::RunTest(const FString& Parameters)
 }
 
 
-IMPLEMENT_SIMPLE_AUTOMATION_TEST(FLegacyAssetBundleEntryTest, "System.AssetRegistry.AssetBundleEntry.LegacyAssetBundleEntry", EAutomationTestFlags::ApplicationContextMask | EAutomationTestFlags::EngineFilter);
+IMPLEMENT_SIMPLE_AUTOMATION_TEST(FLegacyAssetBundleEntryTest, "System.AssetRegistry.AssetBundleEntry.LegacyAssetBundleEntry", EAutomationTestFlags_ApplicationContextMask | EAutomationTestFlags::EngineFilter);
 // Test ImportText for old asset bundles with BundleAssets field
 bool FLegacyAssetBundleEntryTest::RunTest(const FString& Parameters)
 {
@@ -744,7 +744,7 @@ bool FLegacyAssetBundleEntryTest::RunTest(const FString& Parameters)
 	return true;
 }
 
-IMPLEMENT_SIMPLE_AUTOMATION_TEST(FAssetBundlDataImportExportTextTest, "System.AssetRegistry.AssetBundleData.ImportExportText", EAutomationTestFlags::ApplicationContextMask | EAutomationTestFlags::EngineFilter);
+IMPLEMENT_SIMPLE_AUTOMATION_TEST(FAssetBundlDataImportExportTextTest, "System.AssetRegistry.AssetBundleData.ImportExportText", EAutomationTestFlags_ApplicationContextMask | EAutomationTestFlags::EngineFilter);
 bool FAssetBundlDataImportExportTextTest::RunTest(const FString& Parameters)
 {
 	UScriptStruct* Struct = TBaseStructure<FAssetBundleData>::Get();
diff --git a/Engine/Source/Runtime/CoreUObject/Private/AssetRegistry/AssetData.cpp b/Engine/Source/Runtime/CoreUObject/Private/AssetRegistry/AssetData.cpp
--- a/Engine/Source/Runtime/CoreUObject/Private/AssetRegistry/AssetData.cpp
+++ b/Engine/Source/Runtime/CoreUObject/Private/AssetRegistry/AssetData.cpp
@@ -1257,7 +1257,7 @@ FAssetIdentifier::FAssetIdentifier(UObject* SourceObject, FName InValueName)
 #if WITH_DEV_AUTOMATION_TESTS 
 
 #include "Misc/AutomationTest.h"
-IMPLEMENT_SIMPLE_AUTOMATION_TEST(FAssetDataTests, "System.CoreUObject.AssetData", EAutomationTestFlags::ApplicationContextMask | EAutomationTestFlags::EngineFilter);
+IMPLEMENT_SIMPLE_AUTOMATION_TEST(FAssetDataTests, "System.CoreUObject.AssetData", EAutomationTestFlags_ApplicationContextMask | EAutomationTestFlags::EngineFilter);
 bool FAssetDataTests::RunTest(const FString& Parameters)
 {
 	FAssetData EmptyAssetData;
diff --git a/Engine/Source/Runtime/CoreUObject/Private/AssetRegistry/AssetDataTagMap.cpp b/Engine/Source/Runtime/CoreUObject/Private/AssetRegistry/AssetDataTagMap.cpp
--- a/Engine/Source/Runtime/CoreUObject/Private/AssetRegistry/AssetDataTagMap.cpp
+++ b/Engine/Source/Runtime/CoreUObject/Private/AssetRegistry/AssetDataTagMap.cpp
@@ -1638,7 +1638,7 @@ SIZE_T FAssetDataTagMapSharedView::FMemoryCounter::GetFixedSize() const
 #include "Misc/AutomationTest.h"
 #include "Serialization/MemoryReader.h"
 
-IMPLEMENT_SIMPLE_AUTOMATION_TEST(FAssetRegistryExportPathTest, "System.AssetRegistry.ExportPath", EAutomationTestFlags::ApplicationContextMask | EAutomationTestFlags::SmokeFilter)
+IMPLEMENT_SIMPLE_AUTOMATION_TEST(FAssetRegistryExportPathTest, "System.AssetRegistry.ExportPath", EAutomationTestFlags_ApplicationContextMask | EAutomationTestFlags::SmokeFilter)
 
 bool FAssetRegistryExportPathTest::RunTest(const FString& Parameters)
 {
@@ -1656,7 +1656,7 @@ bool FAssetRegistryExportPathTest::RunTest(const FString& Parameters)
 namespace FixedTagPrivate
 {
 
-IMPLEMENT_SIMPLE_AUTOMATION_TEST(FCompactExportPathTest, "System.AssetRegistry.FixedTag.NumberlessExportPath", EAutomationTestFlags::ApplicationContextMask | EAutomationTestFlags::SmokeFilter)
+IMPLEMENT_SIMPLE_AUTOMATION_TEST(FCompactExportPathTest, "System.AssetRegistry.FixedTag.NumberlessExportPath", EAutomationTestFlags_ApplicationContextMask | EAutomationTestFlags::SmokeFilter)
 
 bool FCompactExportPathTest::RunTest(const FString& Parameters)
 {
@@ -1668,7 +1668,7 @@ bool FCompactExportPathTest::RunTest(const FString& Parameters)
 }
 
 
-IMPLEMENT_SIMPLE_AUTOMATION_TEST(FStoreTest, "System.AssetRegistry.FixedTag.Store", EAutomationTestFlags::ApplicationContextMask | EAutomationTestFlags::SmokeFilter)
+IMPLEMENT_SIMPLE_AUTOMATION_TEST(FStoreTest, "System.AssetRegistry.FixedTag.Store", EAutomationTestFlags_ApplicationContextMask | EAutomationTestFlags::SmokeFilter)
 
 static TRefCountPtr<const FStore> MakeTestStore(FStoreData& Data)
 {
diff --git a/Engine/Source/Runtime/CoreUObject/Private/Misc/PackageName.cpp b/Engine/Source/Runtime/CoreUObject/Private/Misc/PackageName.cpp
--- a/Engine/Source/Runtime/CoreUObject/Private/Misc/PackageName.cpp
+++ b/Engine/Source/Runtime/CoreUObject/Private/Misc/PackageName.cpp
@@ -3177,7 +3177,7 @@ static FAutoConsoleCommand ConsoleCommandConvertLongPackageNameToFilename(
 
 #if WITH_DEV_AUTOMATION_TESTS
 
-IMPLEMENT_SIMPLE_AUTOMATION_TEST(FPackageNameTests, "System.Core.Misc.PackageNames", EAutomationTestFlags::ApplicationContextMask | EAutomationTestFlags::SmokeFilter)
+IMPLEMENT_SIMPLE_AUTOMATION_TEST(FPackageNameTests, "System.Core.Misc.PackageNames", EAutomationTestFlags_ApplicationContextMask | EAutomationTestFlags::SmokeFilter)
 
 bool FPackageNameTests::RunTest(const FString& Parameters)
 {
@@ -3543,7 +3543,7 @@ bool FPackageNameTests::RunTest(const FString& Parameters)
 }
 
 // Tests that are too expensive to run as a SmokeFilter
-IMPLEMENT_SIMPLE_AUTOMATION_TEST(FPackageNameTestsExtended, "System.Core.Misc.PackageNamesExtended", EAutomationTestFlags::ApplicationContextMask | EAutomationTestFlags::EngineFilter)
+IMPLEMENT_SIMPLE_AUTOMATION_TEST(FPackageNameTestsExtended, "System.Core.Misc.PackageNamesExtended", EAutomationTestFlags_ApplicationContextMask | EAutomationTestFlags::EngineFilter)
 
 bool FPackageNameTestsExtended::RunTest(const FString& Parameters)
 {
diff --git a/Engine/Source/Runtime/CoreUObject/Private/Misc/PackagePathTest.cpp b/Engine/Source/Runtime/CoreUObject/Private/Misc/PackagePathTest.cpp
--- a/Engine/Source/Runtime/CoreUObject/Private/Misc/PackagePathTest.cpp
+++ b/Engine/Source/Runtime/CoreUObject/Private/Misc/PackagePathTest.cpp
@@ -128,7 +128,7 @@ public:
 	}
 };
 
-IMPLEMENT_CUSTOM_SIMPLE_AUTOMATION_TEST(FPackagePathTestSubClass, FPackagePathTest, "System.Core.Misc.PackagePath", EAutomationTestFlags::ApplicationContextMask | EAutomationTestFlags::EngineFilter);
+IMPLEMENT_CUSTOM_SIMPLE_AUTOMATION_TEST(FPackagePathTestSubClass, FPackagePathTest, "System.Core.Misc.PackagePath", EAutomationTestFlags_ApplicationContextMask | EAutomationTestFlags::EngineFilter);
 bool FPackagePathTestSubClass::RunTest(const FString& Parameters)
 {
 	return FPackagePathTest::RunTest(Parameters);
diff --git a/Engine/Source/Runtime/CoreUObject/Private/UObject/Class.cpp b/Engine/Source/Runtime/CoreUObject/Private/UObject/Class.cpp
--- a/Engine/Source/Runtime/CoreUObject/Private/UObject/Class.cpp
+++ b/Engine/Source/Runtime/CoreUObject/Private/UObject/Class.cpp
@@ -4189,7 +4189,7 @@ public:
 	virtual bool ElevateLogWarningsToErrors() override { return false; }
 };
 
-IMPLEMENT_CUSTOM_SIMPLE_AUTOMATION_TEST(FAutomationTestAttemptToFindUninitializedScriptStructMembers, FAutomationTestUObjectClassBase, "UObject.Class AttemptToFindUninitializedScriptStructMembers", EAutomationTestFlags::ApplicationContextMask | EAutomationTestFlags::SmokeFilter)
+IMPLEMENT_CUSTOM_SIMPLE_AUTOMATION_TEST(FAutomationTestAttemptToFindUninitializedScriptStructMembers, FAutomationTestUObjectClassBase, "UObject.Class AttemptToFindUninitializedScriptStructMembers", EAutomationTestFlags_ApplicationContextMask | EAutomationTestFlags::SmokeFilter)
 bool FAutomationTestAttemptToFindUninitializedScriptStructMembers::RunTest(const FString& Parameters)
 {
 	// This test fails when running tests under UHT because there is no TestUninitializedScriptStructMembersTest, so just skip it in that config.
@@ -4412,7 +4412,7 @@ FAutoConsoleCommandWithWorldAndArgs GCmdListShortTypeNamesInMetaData(
 			FStructUtils::AttemptToFindShortTypeNamesInMetaData();
 		}));
 
-IMPLEMENT_CUSTOM_SIMPLE_AUTOMATION_TEST(FAutomationTestAttemptToFindShortTypeNamesInMetaData, FAutomationTestUObjectClassBase, "UObject.Class AttemptToFindShortTypeNamesInMetaData", EAutomationTestFlags::EditorContext | EAutomationTestFlags::ApplicationContextMask | EAutomationTestFlags::ServerContext | EAutomationTestFlags::SmokeFilter)
+IMPLEMENT_CUSTOM_SIMPLE_AUTOMATION_TEST(FAutomationTestAttemptToFindShortTypeNamesInMetaData, FAutomationTestUObjectClassBase, "UObject.Class AttemptToFindShortTypeNamesInMetaData", EAutomationTestFlags::EditorContext | EAutomationTestFlags_ApplicationContextMask | EAutomationTestFlags::ServerContext | EAutomationTestFlags::SmokeFilter)
 bool FAutomationTestAttemptToFindShortTypeNamesInMetaData::RunTest(const FString& Parameters)
 {
 	// This test is not necessary when running under UHT so just skip it in that config.
diff --git a/Engine/Source/Runtime/CoreUObject/Private/UObject/CoreRedirects.cpp b/Engine/Source/Runtime/CoreUObject/Private/UObject/CoreRedirects.cpp
--- a/Engine/Source/Runtime/CoreUObject/Private/UObject/CoreRedirects.cpp
+++ b/Engine/Source/Runtime/CoreUObject/Private/UObject/CoreRedirects.cpp
@@ -1775,7 +1775,7 @@ bool FCoreRedirects::RunTests()
 	return bSuccess;
 }
 
-IMPLEMENT_SIMPLE_AUTOMATION_TEST(FCoreRedirectTest, "System.Core.Misc.CoreRedirects", EAutomationTestFlags::ApplicationContextMask | EAutomationTestFlags::EngineFilter)
+IMPLEMENT_SIMPLE_AUTOMATION_TEST(FCoreRedirectTest, "System.Core.Misc.CoreRedirects", EAutomationTestFlags_ApplicationContextMask | EAutomationTestFlags::EngineFilter)
 bool FCoreRedirectTest::RunTest(const FString& Parameters)
 {
 	return FCoreRedirects::RunTests();
diff --git a/Engine/Source/Runtime/CoreUObject/Private/UObject/Enum.cpp b/Engine/Source/Runtime/CoreUObject/Private/UObject/Enum.cpp
--- a/Engine/Source/Runtime/CoreUObject/Private/UObject/Enum.cpp
+++ b/Engine/Source/Runtime/CoreUObject/Private/UObject/Enum.cpp
@@ -1081,7 +1081,7 @@ IMPLEMENT_CORE_INTRINSIC_CLASS(UEnum, UField,
 #if WITH_DEV_AUTOMATION_TESTS 
 
 #include "Misc/AutomationTest.h"
-IMPLEMENT_SIMPLE_AUTOMATION_TEST(FEnumBitfieldTest, "System.CoreUObject.EnumBitfields", EAutomationTestFlags::ApplicationContextMask | EAutomationTestFlags::EngineFilter);
+IMPLEMENT_SIMPLE_AUTOMATION_TEST(FEnumBitfieldTest, "System.CoreUObject.EnumBitfields", EAutomationTestFlags_ApplicationContextMask | EAutomationTestFlags::EngineFilter);
 bool FEnumBitfieldTest::RunTest(const FString& Parameters)
 {
 	UPackage* NativePackage = CreatePackage(TEXT("/Script/TestEnumBitfieldsPackage"));
diff --git a/Engine/Source/Runtime/CoreUObject/Private/UObject/Field.cpp b/Engine/Source/Runtime/CoreUObject/Private/UObject/Field.cpp
--- a/Engine/Source/Runtime/CoreUObject/Private/UObject/Field.cpp
+++ b/Engine/Source/Runtime/CoreUObject/Private/UObject/Field.cpp
@@ -1179,7 +1179,7 @@ FField* FindFPropertyByPath(const TCHAR* InFieldPath)
 
 #if WITH_DEV_AUTOMATION_TESTS
 
-IMPLEMENT_SIMPLE_AUTOMATION_TEST(FFieldCastTest, "UObject.Field Cast", EAutomationTestFlags::EditorContext | EAutomationTestFlags::ApplicationContextMask | EAutomationTestFlags::ServerContext | EAutomationTestFlags::SmokeFilter)
+IMPLEMENT_SIMPLE_AUTOMATION_TEST(FFieldCastTest, "UObject.Field Cast", EAutomationTestFlags::EditorContext | EAutomationTestFlags_ApplicationContextMask | EAutomationTestFlags::ServerContext | EAutomationTestFlags::SmokeFilter)
 
 bool FFieldCastTest::RunTest(const FString& Parameters)
 {
diff --git a/Engine/Source/Runtime/CoreUObject/Private/UObject/LinkerInstancingContext.cpp b/Engine/Source/Runtime/CoreUObject/Private/UObject/LinkerInstancingContext.cpp
--- a/Engine/Source/Runtime/CoreUObject/Private/UObject/LinkerInstancingContext.cpp
+++ b/Engine/Source/Runtime/CoreUObject/Private/UObject/LinkerInstancingContext.cpp
@@ -446,7 +446,7 @@ bool FLinkerInstancedPackageMap::FixupSoftObjectPath(FSoftObjectPath& InOutSoftO
 	return false;
 }
 
-IMPLEMENT_SIMPLE_AUTOMATION_TEST(FLinkerInstancingContextTests, "System.CoreUObject.LinkerInstancingContext", EAutomationTestFlags::ApplicationContextMask | EAutomationTestFlags::EngineFilter);
+IMPLEMENT_SIMPLE_AUTOMATION_TEST(FLinkerInstancingContextTests, "System.CoreUObject.LinkerInstancingContext", EAutomationTestFlags_ApplicationContextMask | EAutomationTestFlags::EngineFilter);
 bool FLinkerInstancingContextTests::RunTest(const FString& Parameters)
 {
 	// Disabled SoftObjectPath Remapping
diff --git a/Engine/Source/Runtime/CoreUObject/Private/UObject/SoftObjectPath.cpp b/Engine/Source/Runtime/CoreUObject/Private/UObject/SoftObjectPath.cpp
--- a/Engine/Source/Runtime/CoreUObject/Private/UObject/SoftObjectPath.cpp
+++ b/Engine/Source/Runtime/CoreUObject/Private/UObject/SoftObjectPath.cpp
@@ -888,7 +888,7 @@ std::ostream& operator<<(std::ostream& Stream, const FSoftObjectPath& Value)
 
 #if WITH_DEV_AUTOMATION_TESTS
 
-IMPLEMENT_SIMPLE_AUTOMATION_TEST(FSoftObjectPathImportTextTests, "System.CoreUObject.SoftObjectPath.ImportText", EAutomationTestFlags::ApplicationContextMask | EAutomationTestFlags::EngineFilter);
+IMPLEMENT_SIMPLE_AUTOMATION_TEST(FSoftObjectPathImportTextTests, "System.CoreUObject.SoftObjectPath.ImportText", EAutomationTestFlags_ApplicationContextMask | EAutomationTestFlags::EngineFilter);
 bool FSoftObjectPathImportTextTests::RunTest(const FString& Parameters)
 {
 	const TCHAR* PackageName = TEXT("/Game/Environments/Sets/Arid/Materials/M_Arid");
@@ -916,7 +916,7 @@ bool FSoftObjectPathImportTextTests::RunTest(const FString& Parameters)
 	return true;
 }
 
-IMPLEMENT_SIMPLE_AUTOMATION_TEST(FSoftObjectPathTrySetPathTests, "System.CoreUObject.SoftObjectPath.TrySetPath", EAutomationTestFlags::ApplicationContextMask | EAutomationTestFlags::EngineFilter);
+IMPLEMENT_SIMPLE_AUTOMATION_TEST(FSoftObjectPathTrySetPathTests, "System.CoreUObject.SoftObjectPath.TrySetPath", EAutomationTestFlags_ApplicationContextMask | EAutomationTestFlags::EngineFilter);
 bool FSoftObjectPathTrySetPathTests::RunTest(const FString& Parameters)
 {
 	FSoftObjectPath Path;
@@ -1045,7 +1045,7 @@ bool FSoftObjectPathTrySetPathTests::RunTest(const FString& Parameters)
 
 #if WITH_EDITOR
 
-IMPLEMENT_SIMPLE_AUTOMATION_TEST(FSoftObjectPathFixupForPIETests, "System.CoreUObject.SoftObjectPath.FixupForPIE", EAutomationTestFlags::ApplicationContextMask | EAutomationTestFlags::EngineFilter);
+IMPLEMENT_SIMPLE_AUTOMATION_TEST(FSoftObjectPathFixupForPIETests, "System.CoreUObject.SoftObjectPath.FixupForPIE", EAutomationTestFlags_ApplicationContextMask | EAutomationTestFlags::EngineFilter);
 bool FSoftObjectPathFixupForPIETests::RunTest(const FString& Parameters)
 {
 	const TCHAR* TestOriginalPath = TEXT("/Game/Maps/Arena.Arena:PersistentLevel.Target");	
diff --git a/Engine/Source/Runtime/CoreUObject/Private/UObject/TopLevelAssetPath.cpp b/Engine/Source/Runtime/CoreUObject/Private/UObject/TopLevelAssetPath.cpp
--- a/Engine/Source/Runtime/CoreUObject/Private/UObject/TopLevelAssetPath.cpp
+++ b/Engine/Source/Runtime/CoreUObject/Private/UObject/TopLevelAssetPath.cpp
@@ -521,7 +521,7 @@ bool FTopLevelAssetPath::SerializeFromMismatchedTag(const FPropertyTag& Tag, FSt
 
 // Combine import/export tests
 
-IMPLEMENT_SIMPLE_AUTOMATION_TEST(FTopLevelAssetPathTest, "System.Core.Misc.TopLevelAssetPath", EAutomationTestFlags::ApplicationContextMask | EAutomationTestFlags::EngineFilter);
+IMPLEMENT_SIMPLE_AUTOMATION_TEST(FTopLevelAssetPathTest, "System.Core.Misc.TopLevelAssetPath", EAutomationTestFlags_ApplicationContextMask | EAutomationTestFlags::EngineFilter);
 
 bool FTopLevelAssetPathTest::RunTest(const FString& Parameters)
 {
diff --git a/Engine/Source/Runtime/Engine/Private/Commandlets/Commandlet.cpp b/Engine/Source/Runtime/Engine/Private/Commandlets/Commandlet.cpp
--- a/Engine/Source/Runtime/Engine/Private/Commandlets/Commandlet.cpp
+++ b/Engine/Source/Runtime/Engine/Private/Commandlets/Commandlet.cpp
@@ -35,7 +35,7 @@ UCommandlet::UCommandlet(const FObjectInitializer& ObjectInitializer)
 
 #if WITH_DEV_AUTOMATION_TESTS
 
-IMPLEMENT_SIMPLE_AUTOMATION_TEST(FCommandletCommandLineParsingTest, "System.Commandlet.ParseCommandLine", EAutomationTestFlags::ApplicationContextMask | EAutomationTestFlags::SmokeFilter)
+IMPLEMENT_SIMPLE_AUTOMATION_TEST(FCommandletCommandLineParsingTest, "System.Commandlet.ParseCommandLine", EAutomationTestFlags_ApplicationContextMask | EAutomationTestFlags::SmokeFilter)
 
 bool FCommandletCommandLineParsingTest::RunTest(const FString& Parameters)
 {
diff --git a/Engine/Source/Runtime/Engine/Private/Net/Tests/Util/SequenceNumberTests.cpp b/Engine/Source/Runtime/Engine/Private/Net/Tests/Util/SequenceNumberTests.cpp
--- a/Engine/Source/Runtime/Engine/Private/Net/Tests/Util/SequenceNumberTests.cpp
+++ b/Engine/Source/Runtime/Engine/Private/Net/Tests/Util/SequenceNumberTests.cpp
@@ -5,7 +5,7 @@
 
 #if WITH_DEV_AUTOMATION_TESTS
 
-IMPLEMENT_SIMPLE_AUTOMATION_TEST(FNetSequenceNumberTest, "Network.SequenceNumberTest", EAutomationTestFlags::ApplicationContextMask | EAutomationTestFlags::EngineFilter)
+IMPLEMENT_SIMPLE_AUTOMATION_TEST(FNetSequenceNumberTest, "Network.SequenceNumberTest", EAutomationTestFlags_ApplicationContextMask | EAutomationTestFlags::EngineFilter)
 
 bool FNetSequenceNumberTest::RunTest(const FString& Parameters)
 {
diff --git a/Engine/Source/Runtime/Engine/Private/PieFixupSerializer.cpp b/Engine/Source/Runtime/Engine/Private/PieFixupSerializer.cpp
--- a/Engine/Source/Runtime/Engine/Private/PieFixupSerializer.cpp
+++ b/Engine/Source/Runtime/Engine/Private/PieFixupSerializer.cpp
@@ -85,7 +85,7 @@ FArchive& FPIEFixupSerializer::operator<<(FSoftObjectPtr& Value)
 
 #include "Tests/PieFixupTestObjects.h"
 
-IMPLEMENT_SIMPLE_AUTOMATION_TEST(FPIEFixupSerializerTest, "System.Engine.PIE.FixupSoftReferences", EAutomationTestFlags::ApplicationContextMask | EAutomationTestFlags::EngineFilter);
+IMPLEMENT_SIMPLE_AUTOMATION_TEST(FPIEFixupSerializerTest, "System.Engine.PIE.FixupSoftReferences", EAutomationTestFlags_ApplicationContextMask | EAutomationTestFlags::EngineFilter);
 bool FPIEFixupSerializerTest::RunTest(const FString& Parameters)
 {
 	FName PackageName = MakeUniqueObjectName(nullptr, UPackage::StaticClass(), FName("/Memory/PieFixupTest"));
diff --git a/Engine/Source/Runtime/Engine/Private/Tests/EngineAutomationTests.cpp b/Engine/Source/Runtime/Engine/Private/Tests/EngineAutomationTests.cpp
--- a/Engine/Source/Runtime/Engine/Private/Tests/EngineAutomationTests.cpp
+++ b/Engine/Source/Runtime/Engine/Private/Tests/EngineAutomationTests.cpp
@@ -231,7 +231,7 @@ bool FSaveGameTest::RunTest(const FString& Parameters)
  * FCVarEnvironmentTest
  * Test makes sure that CVars are set and restore properly
  */
-IMPLEMENT_SIMPLE_AUTOMATION_TEST(FCVarEnvironmentTest, "System.Engine.Automation.Environment.CVar", EAutomationTestFlags::ApplicationContextMask | EAutomationTestFlags::EngineFilter)
+IMPLEMENT_SIMPLE_AUTOMATION_TEST(FCVarEnvironmentTest, "System.Engine.Automation.Environment.CVar", EAutomationTestFlags_ApplicationContextMask | EAutomationTestFlags::EngineFilter)
 
 /**
  * Set and restore a CVar
@@ -303,7 +303,7 @@ bool FCVarEnvironmentTest::RunTest(const FString& Parameters)
  * FCVarEnvironmentReuseTest
  * Test makes sure that a CVar can be set multiple times before restoring back to the original value prior to being set
  */
-IMPLEMENT_SIMPLE_AUTOMATION_TEST(FCVarEnvironmentReuseTest, "System.Engine.Automation.Environment.CVar Reuse", EAutomationTestFlags::ApplicationContextMask | EAutomationTestFlags::EngineFilter)
+IMPLEMENT_SIMPLE_AUTOMATION_TEST(FCVarEnvironmentReuseTest, "System.Engine.Automation.Environment.CVar Reuse", EAutomationTestFlags_ApplicationContextMask | EAutomationTestFlags::EngineFilter)
 
 /**
  * Set a CVar multiple times before restoring to the original value
@@ -354,7 +354,7 @@ bool FCVarEnvironmentReuseTest::RunTest(const FString& Parameters)
 	return !HasAnyErrors();
 }
 
-IMPLEMENT_SIMPLE_AUTOMATION_TEST(FAutomationLogAddMessage, "TestFramework.Log.Add Log Message", EAutomationTestFlags::ApplicationContextMask | EAutomationTestFlags::EngineFilter)
+IMPLEMENT_SIMPLE_AUTOMATION_TEST(FAutomationLogAddMessage, "TestFramework.Log.Add Log Message", EAutomationTestFlags_ApplicationContextMask | EAutomationTestFlags::EngineFilter)
 
 bool FAutomationLogAddMessage::RunTest(const FString& Parameters)
 {
@@ -371,7 +371,7 @@ bool FAutomationLogAddMessage::RunTest(const FString& Parameters)
 	return true;
 }
 
-IMPLEMENT_SIMPLE_AUTOMATION_TEST(FAutomationLogAddWarning, "TestFramework.Log.Add Warning Message", EAutomationTestFlags::ApplicationContextMask | EAutomationTestFlags::EngineFilter)
+IMPLEMENT_SIMPLE_AUTOMATION_TEST(FAutomationLogAddWarning, "TestFramework.Log.Add Warning Message", EAutomationTestFlags_ApplicationContextMask | EAutomationTestFlags::EngineFilter)
 
 bool FAutomationLogAddWarning::RunTest(const FString& Parameters)
 {
@@ -388,7 +388,7 @@ bool FAutomationLogAddWarning::RunTest(const FString& Parameters)
 	return true;
 }
 
-IMPLEMENT_SIMPLE_AUTOMATION_TEST(FAutomationLogAddError, "TestFramework.Log.Add Error Message", EAutomationTestFlags::ApplicationContextMask | EAutomationTestFlags::EngineFilter)
+IMPLEMENT_SIMPLE_AUTOMATION_TEST(FAutomationLogAddError, "TestFramework.Log.Add Error Message", EAutomationTestFlags_ApplicationContextMask | EAutomationTestFlags::EngineFilter)
 
 bool FAutomationLogAddError::RunTest(const FString& Parameters)
 {
@@ -727,7 +727,7 @@ const FString FAutomationNearlyEqualTest::ActualValueIsIncreasedByFormatString(T
 const FString FAutomationNearlyEqualTest::ExpectedValueIsIncreasedByFormatString(TEXT("{0} (expected value is increased by {1})"));
 const FString FAutomationNearlyEqualTest::DifferenceAndOrToleranceAreNotValidFormatString(TEXT("Difference and/or Tolerance are not valid. Difference: {0}, Tolerance: {1}"));
 
-IMPLEMENT_CUSTOM_SIMPLE_AUTOMATION_TEST(FAutomationTestNearlyEqualFloatPositive, FAutomationNearlyEqualTest, "TestFramework.Validation.TestNearlyEqualFloatPositive", EAutomationTestFlags::ApplicationContextMask | EAutomationTestFlags::EngineFilter);
+IMPLEMENT_CUSTOM_SIMPLE_AUTOMATION_TEST(FAutomationTestNearlyEqualFloatPositive, FAutomationNearlyEqualTest, "TestFramework.Validation.TestNearlyEqualFloatPositive", EAutomationTestFlags_ApplicationContextMask | EAutomationTestFlags::EngineFilter);
 bool FAutomationTestNearlyEqualFloatPositive::RunTest(const FString& Parameters)
 {
 	//** TEST **//
@@ -754,7 +754,7 @@ bool FAutomationTestNearlyEqualFloatPositive::RunTest(const FString& Parameters)
 	return true;
 }
 
-IMPLEMENT_CUSTOM_SIMPLE_AUTOMATION_TEST(FAutomationTestNearlyEqualFloatNegative, FAutomationNearlyEqualTest, "TestFramework.Validation.TestNearlyEqualFloatNegative", EAutomationTestFlags::ApplicationContextMask | EAutomationTestFlags::EngineFilter);
+IMPLEMENT_CUSTOM_SIMPLE_AUTOMATION_TEST(FAutomationTestNearlyEqualFloatNegative, FAutomationNearlyEqualTest, "TestFramework.Validation.TestNearlyEqualFloatNegative", EAutomationTestFlags_ApplicationContextMask | EAutomationTestFlags::EngineFilter);
 bool FAutomationTestNearlyEqualFloatNegative::RunTest(const FString& Parameters)
 {
 	int32 CasesCheckedTotal = 0;
@@ -777,7 +777,7 @@ bool FAutomationTestNearlyEqualFloatNegative::RunTest(const FString& Parameters)
 	return true;
 }
 
-IMPLEMENT_CUSTOM_SIMPLE_AUTOMATION_TEST(FAutomationTestNearlyEqualDoublePositive, FAutomationNearlyEqualTest, "TestFramework.Validation.TestNearlyEqualDoublePositive", EAutomationTestFlags::ApplicationContextMask | EAutomationTestFlags::EngineFilter);
+IMPLEMENT_CUSTOM_SIMPLE_AUTOMATION_TEST(FAutomationTestNearlyEqualDoublePositive, FAutomationNearlyEqualTest, "TestFramework.Validation.TestNearlyEqualDoublePositive", EAutomationTestFlags_ApplicationContextMask | EAutomationTestFlags::EngineFilter);
 bool FAutomationTestNearlyEqualDoublePositive::RunTest(const FString& Parameters)
 {
 	//** TEST **//
@@ -803,7 +803,7 @@ bool FAutomationTestNearlyEqualDoublePositive::RunTest(const FString& Parameters
 	return true;
 }
 
-IMPLEMENT_CUSTOM_SIMPLE_AUTOMATION_TEST(FAutomationTestNearlyEqualDoubleNegative, FAutomationNearlyEqualTest, "TestFramework.Validation.TestNearlyEqualDoubleNegative", EAutomationTestFlags::ApplicationContextMask | EAutomationTestFlags::EngineFilter);
+IMPLEMENT_CUSTOM_SIMPLE_AUTOMATION_TEST(FAutomationTestNearlyEqualDoubleNegative, FAutomationNearlyEqualTest, "TestFramework.Validation.TestNearlyEqualDoubleNegative", EAutomationTestFlags_ApplicationContextMask | EAutomationTestFlags::EngineFilter);
 bool FAutomationTestNearlyEqualDoubleNegative::RunTest(const FString& Parameters)
 {
 	int32 CasesCheckedTotal = 0;
@@ -826,7 +826,7 @@ bool FAutomationTestNearlyEqualDoubleNegative::RunTest(const FString& Parameters
 	return true;
 }
 
-IMPLEMENT_CUSTOM_SIMPLE_AUTOMATION_TEST(FAutomationTestNearlyEqualFVectorPositive, FAutomationNearlyEqualTest, "TestFramework.Validation.TestNearlyEqualFVectorPositive", EAutomationTestFlags::ApplicationContextMask | EAutomationTestFlags::EngineFilter);
+IMPLEMENT_CUSTOM_SIMPLE_AUTOMATION_TEST(FAutomationTestNearlyEqualFVectorPositive, FAutomationNearlyEqualTest, "TestFramework.Validation.TestNearlyEqualFVectorPositive", EAutomationTestFlags_ApplicationContextMask | EAutomationTestFlags::EngineFilter);
 bool FAutomationTestNearlyEqualFVectorPositive::RunTest(const FString& Parameters)
 {
 	//** TEST **//
@@ -852,7 +852,7 @@ bool FAutomationTestNearlyEqualFVectorPositive::RunTest(const FString& Parameter
 	return true;
 }
 
-IMPLEMENT_CUSTOM_SIMPLE_AUTOMATION_TEST(FAutomationTestNearlyEqualFVectorNegative, FAutomationNearlyEqualTest, "TestFramework.Validation.TestNearlyEqualFVectorNegative", EAutomationTestFlags::ApplicationContextMask | EAutomationTestFlags::EngineFilter);
+IMPLEMENT_CUSTOM_SIMPLE_AUTOMATION_TEST(FAutomationTestNearlyEqualFVectorNegative, FAutomationNearlyEqualTest, "TestFramework.Validation.TestNearlyEqualFVectorNegative", EAutomationTestFlags_ApplicationContextMask | EAutomationTestFlags::EngineFilter);
 bool FAutomationTestNearlyEqualFVectorNegative::RunTest(const FString& Parameters)
 {
 	int32 CasesCheckedTotal = 0;
@@ -871,7 +871,7 @@ bool FAutomationTestNearlyEqualFVectorNegative::RunTest(const FString& Parameter
 	return true;
 }
 
-IMPLEMENT_CUSTOM_SIMPLE_AUTOMATION_TEST(FAutomationTestNearlyEqualFRotatorPositive, FAutomationNearlyEqualTest, "TestFramework.Validation.TestNearlyEqualFRotatorPositive", EAutomationTestFlags::ApplicationContextMask | EAutomationTestFlags::EngineFilter);
+IMPLEMENT_CUSTOM_SIMPLE_AUTOMATION_TEST(FAutomationTestNearlyEqualFRotatorPositive, FAutomationNearlyEqualTest, "TestFramework.Validation.TestNearlyEqualFRotatorPositive", EAutomationTestFlags_ApplicationContextMask | EAutomationTestFlags::EngineFilter);
 bool FAutomationTestNearlyEqualFRotatorPositive::RunTest(const FString& Parameters)
 {
 	//** TEST **//
@@ -897,7 +897,7 @@ bool FAutomationTestNearlyEqualFRotatorPositive::RunTest(const FString& Paramete
 	return true;
 }
 
-IMPLEMENT_CUSTOM_SIMPLE_AUTOMATION_TEST(FAutomationTestNearlyEqualFRotatorNegative, FAutomationNearlyEqualTest, "TestFramework.Validation.TestNearlyEqualFRotatorNegative", EAutomationTestFlags::ApplicationContextMask | EAutomationTestFlags::EngineFilter);
+IMPLEMENT_CUSTOM_SIMPLE_AUTOMATION_TEST(FAutomationTestNearlyEqualFRotatorNegative, FAutomationNearlyEqualTest, "TestFramework.Validation.TestNearlyEqualFRotatorNegative", EAutomationTestFlags_ApplicationContextMask | EAutomationTestFlags::EngineFilter);
 bool FAutomationTestNearlyEqualFRotatorNegative::RunTest(const FString& Parameters)
 {
 	int32 CasesCheckedTotal = 0;
@@ -916,7 +916,7 @@ bool FAutomationTestNearlyEqualFRotatorNegative::RunTest(const FString& Paramete
 	return true;
 }
 
-IMPLEMENT_CUSTOM_SIMPLE_AUTOMATION_TEST(FAutomationTestNearlyEqualFTransformPositive, FAutomationNearlyEqualTest, "TestFramework.Validation.TestNearlyEqualFTransformPositive", EAutomationTestFlags::ApplicationContextMask | EAutomationTestFlags::EngineFilter);
+IMPLEMENT_CUSTOM_SIMPLE_AUTOMATION_TEST(FAutomationTestNearlyEqualFTransformPositive, FAutomationNearlyEqualTest, "TestFramework.Validation.TestNearlyEqualFTransformPositive", EAutomationTestFlags_ApplicationContextMask | EAutomationTestFlags::EngineFilter);
 bool FAutomationTestNearlyEqualFTransformPositive::RunTest(const FString& Parameters)
 {
 	//** TEST **//
@@ -942,7 +942,7 @@ bool FAutomationTestNearlyEqualFTransformPositive::RunTest(const FString& Parame
 	return true;
 }
 
-IMPLEMENT_CUSTOM_SIMPLE_AUTOMATION_TEST(FAutomationTestNearlyEqualFTransformNegative, FAutomationNearlyEqualTest, "TestFramework.Validation.TestNearlyEqualFTransformNegative", EAutomationTestFlags::ApplicationContextMask | EAutomationTestFlags::EngineFilter);
+IMPLEMENT_CUSTOM_SIMPLE_AUTOMATION_TEST(FAutomationTestNearlyEqualFTransformNegative, FAutomationNearlyEqualTest, "TestFramework.Validation.TestNearlyEqualFTransformNegative", EAutomationTestFlags_ApplicationContextMask | EAutomationTestFlags::EngineFilter);
 bool FAutomationTestNearlyEqualFTransformNegative::RunTest(const FString& Parameters)
 {
 	int32 CasesCheckedTotal = 0;
@@ -961,7 +961,7 @@ bool FAutomationTestNearlyEqualFTransformNegative::RunTest(const FString& Parame
 	return true;
 }
 
-IMPLEMENT_SIMPLE_AUTOMATION_TEST(FAutomationTestInequalityBool, "TestFramework.Validation.TestInequalityBool", EAutomationTestFlags::ApplicationContextMask | EAutomationTestFlags::EngineFilter);
+IMPLEMENT_SIMPLE_AUTOMATION_TEST(FAutomationTestInequalityBool, "TestFramework.Validation.TestInequalityBool", EAutomationTestFlags_ApplicationContextMask | EAutomationTestFlags::EngineFilter);
 bool FAutomationTestInequalityBool::RunTest(const FString& Parameters)
 {
 	TestTrue("True constant", true);
@@ -972,7 +972,7 @@ bool FAutomationTestInequalityBool::RunTest(const FString& Parameters)
 	return true;
 }
 
-IMPLEMENT_SIMPLE_AUTOMATION_TEST(FAutomationTestInequalityPointer, "TestFramework.Validation.TestInequalityPointer", EAutomationTestFlags::ApplicationContextMask | EAutomationTestFlags::EngineFilter);
+IMPLEMENT_SIMPLE_AUTOMATION_TEST(FAutomationTestInequalityPointer, "TestFramework.Validation.TestInequalityPointer", EAutomationTestFlags_ApplicationContextMask | EAutomationTestFlags::EngineFilter);
 bool FAutomationTestInequalityPointer::RunTest(const FString& Parameters)
 {
 	int32 StackValue = 42;
@@ -1000,7 +1000,7 @@ bool FAutomationTestInequalityPointer::RunTest(const FString& Parameters)
 	return true;
 }
 
-IMPLEMENT_SIMPLE_AUTOMATION_TEST(FAutomationTestInequalityReference, "TestFramework.Validation.TestInequalityReference", EAutomationTestFlags::ApplicationContextMask | EAutomationTestFlags::EngineFilter);
+IMPLEMENT_SIMPLE_AUTOMATION_TEST(FAutomationTestInequalityReference, "TestFramework.Validation.TestInequalityReference", EAutomationTestFlags_ApplicationContextMask | EAutomationTestFlags::EngineFilter);
 bool FAutomationTestInequalityReference::RunTest(const FString& Parameters)
 {
 	int32 StackValue = 42;
@@ -1024,7 +1024,7 @@ bool FAutomationTestInequalityReference::RunTest(const FString& Parameters)
 	return true;
 }
 
-IMPLEMENT_SIMPLE_AUTOMATION_TEST(FAutomationTestInequalityInt32, "TestFramework.Validation.TestInequalityInt32", EAutomationTestFlags::ApplicationContextMask | EAutomationTestFlags::EngineFilter);
+IMPLEMENT_SIMPLE_AUTOMATION_TEST(FAutomationTestInequalityInt32, "TestFramework.Validation.TestInequalityInt32", EAutomationTestFlags_ApplicationContextMask | EAutomationTestFlags::EngineFilter);
 bool FAutomationTestInequalityInt32::RunTest(const FString& Parameters)
 {
 	int32 Expected(42);
@@ -1055,7 +1055,7 @@ bool FAutomationTestInequalityInt32::RunTest(const FString& Parameters)
 	return true;
 }
 
-IMPLEMENT_SIMPLE_AUTOMATION_TEST(FAutomationTestInequalityInt64, "TestFramework.Validation.TestInequalityInt64", EAutomationTestFlags::ApplicationContextMask | EAutomationTestFlags::EngineFilter);
+IMPLEMENT_SIMPLE_AUTOMATION_TEST(FAutomationTestInequalityInt64, "TestFramework.Validation.TestInequalityInt64", EAutomationTestFlags_ApplicationContextMask | EAutomationTestFlags::EngineFilter);
 bool FAutomationTestInequalityInt64::RunTest(const FString& Parameters)
 {
 	int64 Expected(42);
@@ -1087,7 +1087,7 @@ bool FAutomationTestInequalityInt64::RunTest(const FString& Parameters)
 }
 
 #if PLATFORM_64BITS
-IMPLEMENT_SIMPLE_AUTOMATION_TEST(FAutomationTestInequalitySizeT, "TestFramework.Validation.TestInequalitySizeT", EAutomationTestFlags::ApplicationContextMask | EAutomationTestFlags::EngineFilter);
+IMPLEMENT_SIMPLE_AUTOMATION_TEST(FAutomationTestInequalitySizeT, "TestFramework.Validation.TestInequalitySizeT", EAutomationTestFlags_ApplicationContextMask | EAutomationTestFlags::EngineFilter);
 bool FAutomationTestInequalitySizeT::RunTest(const FString& Parameters)
 {
 	SIZE_T Expected(42);
@@ -1116,7 +1116,7 @@ bool FAutomationTestInequalitySizeT::RunTest(const FString& Parameters)
 }
 #endif // PLATFORM_64BITS
 
-IMPLEMENT_SIMPLE_AUTOMATION_TEST(FAutomationTestInequalityFloat, "TestFramework.Validation.TestInequalityFloat", EAutomationTestFlags::ApplicationContextMask | EAutomationTestFlags::EngineFilter);
+IMPLEMENT_SIMPLE_AUTOMATION_TEST(FAutomationTestInequalityFloat, "TestFramework.Validation.TestInequalityFloat", EAutomationTestFlags_ApplicationContextMask | EAutomationTestFlags::EngineFilter);
 bool FAutomationTestInequalityFloat::RunTest(const FString& Parameters)
 {
 	float Expected(42);
@@ -1178,7 +1178,7 @@ bool FAutomationTestInequalityFloat::RunTest(const FString& Parameters)
 	return true;
 }
 
-IMPLEMENT_SIMPLE_AUTOMATION_TEST(FAutomationTestInequalityDouble, "TestFramework.Validation.TestInequalityDouble", EAutomationTestFlags::ApplicationContextMask | EAutomationTestFlags::EngineFilter);
+IMPLEMENT_SIMPLE_AUTOMATION_TEST(FAutomationTestInequalityDouble, "TestFramework.Validation.TestInequalityDouble", EAutomationTestFlags_ApplicationContextMask | EAutomationTestFlags::EngineFilter);
 bool FAutomationTestInequalityDouble::RunTest(const FString& Parameters)
 {
 	double Expected(42);
@@ -1240,7 +1240,7 @@ bool FAutomationTestInequalityDouble::RunTest(const FString& Parameters)
 	return true;
 }
 
-IMPLEMENT_SIMPLE_AUTOMATION_TEST(FAutomationTestInequalityString, "TestFramework.Validation.TestInequalityString", EAutomationTestFlags::ApplicationContextMask | EAutomationTestFlags::EngineFilter);
+IMPLEMENT_SIMPLE_AUTOMATION_TEST(FAutomationTestInequalityString, "TestFramework.Validation.TestInequalityString", EAutomationTestFlags_ApplicationContextMask | EAutomationTestFlags::EngineFilter);
 bool FAutomationTestInequalityString::RunTest(const FString& Parameters)
 {
 	FString ExpectedString("Forty-two");
@@ -1320,7 +1320,7 @@ const FString FAutomationUTestMacrosExpr::ExpectedFStringValueLowerCase(TEXT("eq
 const FString FAutomationUTestMacrosExpr::UnexpectedFStringValueLowerCase(TEXT("not-equals"));
 const FString FAutomationUTestMacrosExpr::CustomDescriptionString(TEXT("Error string appears when UTEST_ macro diverges from _EXPR variant"));
 
-IMPLEMENT_CUSTOM_SIMPLE_AUTOMATION_TEST(FAutomationEqualEXPR, FAutomationUTestMacrosExpr, "TestFramework.Validation.UTestEqual", EAutomationTestFlags::ApplicationContextMask | EAutomationTestFlags::EngineFilter);
+IMPLEMENT_CUSTOM_SIMPLE_AUTOMATION_TEST(FAutomationEqualEXPR, FAutomationUTestMacrosExpr, "TestFramework.Validation.UTestEqual", EAutomationTestFlags_ApplicationContextMask | EAutomationTestFlags::EngineFilter);
 bool FAutomationEqualEXPR::RunTest(const FString& Parameters)
 {
 
@@ -1340,7 +1340,7 @@ bool FAutomationEqualEXPR::RunTest(const FString& Parameters)
 	return true;
 }
 
-IMPLEMENT_CUSTOM_SIMPLE_AUTOMATION_TEST(FAutomationSameNotSameEXPR, FAutomationUTestMacrosExpr, "TestFramework.Validation.UTestSameNotSame", EAutomationTestFlags::ApplicationContextMask | EAutomationTestFlags::EngineFilter);
+IMPLEMENT_CUSTOM_SIMPLE_AUTOMATION_TEST(FAutomationSameNotSameEXPR, FAutomationUTestMacrosExpr, "TestFramework.Validation.UTestSameNotSame", EAutomationTestFlags_ApplicationContextMask | EAutomationTestFlags::EngineFilter);
 bool FAutomationSameNotSameEXPR::RunTest(const FString& Parameters)
 {
 
@@ -1352,7 +1352,7 @@ bool FAutomationSameNotSameEXPR::RunTest(const FString& Parameters)
 	return true;
 }
 
-IMPLEMENT_CUSTOM_SIMPLE_AUTOMATION_TEST(FAutomationSameNotSamePtrEXPR, FAutomationUTestMacrosExpr, "TestFramework.Validation.UTestSameNotSamePtr", EAutomationTestFlags::ApplicationContextMask | EAutomationTestFlags::EngineFilter);
+IMPLEMENT_CUSTOM_SIMPLE_AUTOMATION_TEST(FAutomationSameNotSamePtrEXPR, FAutomationUTestMacrosExpr, "TestFramework.Validation.UTestSameNotSamePtr", EAutomationTestFlags_ApplicationContextMask | EAutomationTestFlags::EngineFilter);
 bool FAutomationSameNotSamePtrEXPR::RunTest(const FString& Parameters)
 {
 
@@ -1364,7 +1364,7 @@ bool FAutomationSameNotSamePtrEXPR::RunTest(const FString& Parameters)
 	return true;
 }
 
-IMPLEMENT_CUSTOM_SIMPLE_AUTOMATION_TEST(FAutomationTrueFalseEXPR, FAutomationUTestMacrosExpr, "TestFramework.Validation.UTestTrueFalse", EAutomationTestFlags::ApplicationContextMask | EAutomationTestFlags::EngineFilter);
+IMPLEMENT_CUSTOM_SIMPLE_AUTOMATION_TEST(FAutomationTrueFalseEXPR, FAutomationUTestMacrosExpr, "TestFramework.Validation.UTestTrueFalse", EAutomationTestFlags_ApplicationContextMask | EAutomationTestFlags::EngineFilter);
 bool FAutomationTrueFalseEXPR::RunTest(const FString& Parameters)
 {
 	UTEST_TRUE_EXPR(ActualFloatValue == ExpectedFloatValue);
@@ -1375,7 +1375,7 @@ bool FAutomationTrueFalseEXPR::RunTest(const FString& Parameters)
 	return true;
 }
 
-IMPLEMENT_CUSTOM_SIMPLE_AUTOMATION_TEST(FAutomationValidInvalidEXPR, FAutomationUTestMacrosExpr, "TestFramework.Validation.UTestValidInvalid", EAutomationTestFlags::ApplicationContextMask | EAutomationTestFlags::EngineFilter);
+IMPLEMENT_CUSTOM_SIMPLE_AUTOMATION_TEST(FAutomationValidInvalidEXPR, FAutomationUTestMacrosExpr, "TestFramework.Validation.UTestValidInvalid", EAutomationTestFlags_ApplicationContextMask | EAutomationTestFlags::EngineFilter);
 bool FAutomationValidInvalidEXPR::RunTest(const FString& Parameters)
 {
 	struct FHasIsValid
@@ -1411,7 +1411,7 @@ bool FAutomationValidInvalidEXPR::RunTest(const FString& Parameters)
 	return true;
 }
 
-IMPLEMENT_CUSTOM_SIMPLE_AUTOMATION_TEST(FAutomationNullNotNullPtrEXPR, FAutomationUTestMacrosExpr, "TestFramework.Validation.UTestNullNotNull", EAutomationTestFlags::ApplicationContextMask | EAutomationTestFlags::EngineFilter);
+IMPLEMENT_CUSTOM_SIMPLE_AUTOMATION_TEST(FAutomationNullNotNullPtrEXPR, FAutomationUTestMacrosExpr, "TestFramework.Validation.UTestNullNotNull", EAutomationTestFlags_ApplicationContextMask | EAutomationTestFlags::EngineFilter);
 bool FAutomationNullNotNullPtrEXPR::RunTest(const FString& Parameters)
 {
 	UWorld* World = UWorld::CreateWorld(EWorldType::Game, false);
@@ -1426,7 +1426,7 @@ bool FAutomationNullNotNullPtrEXPR::RunTest(const FString& Parameters)
 	return true;
 }
 
-IMPLEMENT_CUSTOM_SIMPLE_AUTOMATION_TEST(FAutomationLessGreaterEXPR, FAutomationUTestMacrosExpr, "TestFramework.Validation.UTestLessGreater", EAutomationTestFlags::ApplicationContextMask | EAutomationTestFlags::EngineFilter);
+IMPLEMENT_CUSTOM_SIMPLE_AUTOMATION_TEST(FAutomationLessGreaterEXPR, FAutomationUTestMacrosExpr, "TestFramework.Validation.UTestLessGreater", EAutomationTestFlags_ApplicationContextMask | EAutomationTestFlags::EngineFilter);
 bool FAutomationLessGreaterEXPR::RunTest(const FString& Parameters)
 {
 	// inequality
diff --git a/Engine/Source/Runtime/Engine/Private/Tests/IntSerializationTest.cpp b/Engine/Source/Runtime/Engine/Private/Tests/IntSerializationTest.cpp
--- a/Engine/Source/Runtime/Engine/Private/Tests/IntSerializationTest.cpp
+++ b/Engine/Source/Runtime/Engine/Private/Tests/IntSerializationTest.cpp
@@ -7,7 +7,7 @@
 
 #if WITH_DEV_AUTOMATION_TESTS
 
-IMPLEMENT_SIMPLE_AUTOMATION_TEST(FIntSerializationTest, "System.Engine.Int Serialization", EAutomationTestFlags::ApplicationContextMask | EAutomationTestFlags::SmokeFilter)
+IMPLEMENT_SIMPLE_AUTOMATION_TEST(FIntSerializationTest, "System.Engine.Int Serialization", EAutomationTestFlags_ApplicationContextMask | EAutomationTestFlags::SmokeFilter)
 
 bool FIntSerializationTest::RunTest (const FString& Parameters)
 {
diff --git a/Engine/Source/Runtime/Engine/Private/Tests/Internationalization/TextPropertyTest.cpp b/Engine/Source/Runtime/Engine/Private/Tests/Internationalization/TextPropertyTest.cpp
--- a/Engine/Source/Runtime/Engine/Private/Tests/Internationalization/TextPropertyTest.cpp
+++ b/Engine/Source/Runtime/Engine/Private/Tests/Internationalization/TextPropertyTest.cpp
@@ -17,7 +17,7 @@ UTextPropertyTestObject::UTextPropertyTestObject(const FObjectInitializer& Objec
 
 #if WITH_DEV_AUTOMATION_TESTS
 
-IMPLEMENT_SIMPLE_AUTOMATION_TEST(FTextPropertyTest, "System.Engine.Internationalization.Text Property Test", EAutomationTestFlags::ApplicationContextMask | EAutomationTestFlags::SmokeFilter)
+IMPLEMENT_SIMPLE_AUTOMATION_TEST(FTextPropertyTest, "System.Engine.Internationalization.Text Property Test", EAutomationTestFlags_ApplicationContextMask | EAutomationTestFlags::SmokeFilter)
 
 bool FTextPropertyTest::RunTest (const FString& Parameters)
 {
diff --git a/Engine/Source/Runtime/Engine/Private/Tests/NetworkAutomationTests.cpp b/Engine/Source/Runtime/Engine/Private/Tests/NetworkAutomationTests.cpp
--- a/Engine/Source/Runtime/Engine/Private/Tests/NetworkAutomationTests.cpp
+++ b/Engine/Source/Runtime/Engine/Private/Tests/NetworkAutomationTests.cpp
@@ -39,7 +39,7 @@ IMPLEMENT_NETWORKED_AUTOMATION_TEST(FMultiplayer4PlayerTest, "System.Networking.
 bool FMultiplayer4PlayerTest::RunTest(const FString& Parameters)
 {
 	// Accessing the game world is only valid for game-only 
-	check((GetTestFlags() & EAutomationTestFlags::ApplicationContextMask) == EAutomationTestFlags::ClientContext);
+	check((GetTestFlags() & EAutomationTestFlags_ApplicationContextMask) == EAutomationTestFlags::ClientContext);
 	check(GEngine->GetWorldContexts().Num() == 1);
 	check(GEngine->GetWorldContexts()[0].WorldType == EWorldType::Game);
 
diff --git a/Engine/Source/Runtime/Engine/Private/Tests/PropertyPairsMapTests.cpp b/Engine/Source/Runtime/Engine/Private/Tests/PropertyPairsMapTests.cpp
--- a/Engine/Source/Runtime/Engine/Private/Tests/PropertyPairsMapTests.cpp
+++ b/Engine/Source/Runtime/Engine/Private/Tests/PropertyPairsMapTests.cpp
@@ -9,7 +9,7 @@
 
 namespace PropertyPairsMapTests
 {
-	constexpr const EAutomationTestFlags TestFlags = EAutomationTestFlags::ApplicationContextMask | EAutomationTestFlags::EngineFilter;
+	constexpr const EAutomationTestFlags TestFlags = EAutomationTestFlags_ApplicationContextMask | EAutomationTestFlags::EngineFilter;
 
 	IMPLEMENT_SIMPLE_AUTOMATION_TEST(FPropertyPairsMapTests, TEST_NAME_ROOT, TestFlags)
 	bool FPropertyPairsMapTests::RunTest(const FString& Parameters)
diff --git a/Engine/Source/Runtime/Engine/Private/Tests/ReferenceClusterTest.cpp b/Engine/Source/Runtime/Engine/Private/Tests/ReferenceClusterTest.cpp
--- a/Engine/Source/Runtime/Engine/Private/Tests/ReferenceClusterTest.cpp
+++ b/Engine/Source/Runtime/Engine/Private/Tests/ReferenceClusterTest.cpp
@@ -9,7 +9,7 @@
 
 namespace ReferenceClusterTests
 {
-	constexpr const EAutomationTestFlags TestFlags = EAutomationTestFlags::ApplicationContextMask | EAutomationTestFlags::EngineFilter;
+	constexpr const EAutomationTestFlags TestFlags = EAutomationTestFlags_ApplicationContextMask | EAutomationTestFlags::EngineFilter;
 
 	IMPLEMENT_SIMPLE_AUTOMATION_TEST(FReferenceClusterTest, TEST_NAME_ROOT, TestFlags)
 
diff --git a/Engine/Source/Runtime/Engine/Private/Tests/Texture2DTests.cpp b/Engine/Source/Runtime/Engine/Private/Tests/Texture2DTests.cpp
--- a/Engine/Source/Runtime/Engine/Private/Tests/Texture2DTests.cpp
+++ b/Engine/Source/Runtime/Engine/Private/Tests/Texture2DTests.cpp
@@ -9,7 +9,7 @@
 namespace Texture2DTest
 {
 
-constexpr const EAutomationTestFlags TestFlags = EAutomationTestFlags::ApplicationContextMask | EAutomationTestFlags::EngineFilter;
+constexpr const EAutomationTestFlags TestFlags = EAutomationTestFlags_ApplicationContextMask | EAutomationTestFlags::EngineFilter;
 
 // A simple test to make sure that basic functionality in UTexture2D::CreateTransient works as it seems to be a 
 // fairly uncommon code path in our samples/test games etc.
diff --git a/Engine/Source/Runtime/Engine/Private/Tests/TransactionDiffingTests.cpp b/Engine/Source/Runtime/Engine/Private/Tests/TransactionDiffingTests.cpp
--- a/Engine/Source/Runtime/Engine/Private/Tests/TransactionDiffingTests.cpp
+++ b/Engine/Source/Runtime/Engine/Private/Tests/TransactionDiffingTests.cpp
@@ -17,7 +17,7 @@ void UTransactionDiffingTestObject::Serialize(FStructuredArchive::FRecord Record
 namespace TransactionDiffingTests
 {
 
-constexpr const EAutomationTestFlags TestFlags = EAutomationTestFlags::ApplicationContextMask | EAutomationTestFlags::EngineFilter;
+constexpr const EAutomationTestFlags TestFlags = EAutomationTestFlags_ApplicationContextMask | EAutomationTestFlags::EngineFilter;
 
 IMPLEMENT_SIMPLE_AUTOMATION_TEST(FEditPropertyDataTest, "System.Engine.Transactions.EditPropertyData", TestFlags)
 bool FEditPropertyDataTest::RunTest(const FString& Parameters)
diff --git a/Engine/Source/Runtime/FieldNotification/Private/Tests/FieldNotificationMulticastDelegateTest.cpp b/Engine/Source/Runtime/FieldNotification/Private/Tests/FieldNotificationMulticastDelegateTest.cpp
--- a/Engine/Source/Runtime/FieldNotification/Private/Tests/FieldNotificationMulticastDelegateTest.cpp
+++ b/Engine/Source/Runtime/FieldNotification/Private/Tests/FieldNotificationMulticastDelegateTest.cpp
@@ -8,7 +8,7 @@
 
 #define LOCTEXT_NAMESPACE "FieldNotificationMulticastDelegateTest"
 
-IMPLEMENT_SIMPLE_AUTOMATION_TEST(FFieldNotificationMulticastDelegateTest, "System.UMG.FieldNotificationMulticastDelegateTest", EAutomationTestFlags::ApplicationContextMask | EAutomationTestFlags::EngineFilter);
+IMPLEMENT_SIMPLE_AUTOMATION_TEST(FFieldNotificationMulticastDelegateTest, "System.UMG.FieldNotificationMulticastDelegateTest", EAutomationTestFlags_ApplicationContextMask | EAutomationTestFlags::EngineFilter);
 
 
 bool FFieldNotificationMulticastDelegateTest::RunTest(const FString& Parameters)
diff --git a/Engine/Source/Runtime/GameplayTags/Private/Tests/GameplayTagQueryTests.cpp b/Engine/Source/Runtime/GameplayTags/Private/Tests/GameplayTagQueryTests.cpp
--- a/Engine/Source/Runtime/GameplayTags/Private/Tests/GameplayTagQueryTests.cpp
+++ b/Engine/Source/Runtime/GameplayTags/Private/Tests/GameplayTagQueryTests.cpp
@@ -4,7 +4,7 @@
 #include "GameplayTagContainer.h"
 #include "Misc/AutomationTest.h"
 
-IMPLEMENT_SIMPLE_AUTOMATION_TEST(FGameplayTagQueryTest_Empty, "System.GameplayTags.GameplayTagQuery.Empty", EAutomationTestFlags::ApplicationContextMask | EAutomationTestFlags::EngineFilter)
+IMPLEMENT_SIMPLE_AUTOMATION_TEST(FGameplayTagQueryTest_Empty, "System.GameplayTags.GameplayTagQuery.Empty", EAutomationTestFlags_ApplicationContextMask | EAutomationTestFlags::EngineFilter)
 
 #include "GameplayTagsManager.h"
 
diff --git a/Engine/Source/Runtime/GameplayTags/Private/Tests/GameplayTagTests.cpp b/Engine/Source/Runtime/GameplayTags/Private/Tests/GameplayTagTests.cpp
--- a/Engine/Source/Runtime/GameplayTags/Private/Tests/GameplayTagTests.cpp
+++ b/Engine/Source/Runtime/GameplayTags/Private/Tests/GameplayTagTests.cpp
@@ -398,7 +398,7 @@ public:
 
 };
 
-IMPLEMENT_CUSTOM_SIMPLE_AUTOMATION_TEST(FGameplayTagTest, FGameplayTagTestBase, "System.GameplayTags.GameplayTag", EAutomationTestFlags::ApplicationContextMask | EAutomationTestFlags::EngineFilter)
+IMPLEMENT_CUSTOM_SIMPLE_AUTOMATION_TEST(FGameplayTagTest, FGameplayTagTestBase, "System.GameplayTags.GameplayTag", EAutomationTestFlags_ApplicationContextMask | EAutomationTestFlags::EngineFilter)
 bool FGameplayTagTest::RunTest(const FString& Parameters)
 {
 	// Create Test Data 
diff --git a/Engine/Source/Runtime/Json/Private/Tests/JsonPointer.spec.cpp b/Engine/Source/Runtime/Json/Private/Tests/JsonPointer.spec.cpp
--- a/Engine/Source/Runtime/Json/Private/Tests/JsonPointer.spec.cpp
+++ b/Engine/Source/Runtime/Json/Private/Tests/JsonPointer.spec.cpp
@@ -1,4 +1,4 @@
-﻿// Copyright Epic Games, Inc. All Rights Reserved.
+// Copyright Epic Games, Inc. All Rights Reserved.
 
 #include "JsonUtils/JsonPointer.h"
 #include "Misc/AutomationTest.h"
@@ -7,7 +7,7 @@
 
 #if WITH_AUTOMATION_TESTS
 
-BEGIN_DEFINE_SPEC(FJsonPointerSpec, TEXT("JsonPointer"), EAutomationTestFlags::EngineFilter | EAutomationTestFlags::ApplicationContextMask)
+BEGIN_DEFINE_SPEC(FJsonPointerSpec, TEXT("JsonPointer"), EAutomationTestFlags::EngineFilter | EAutomationTestFlags_ApplicationContextMask)
 TSharedPtr<FJsonObject> JsonObject;
 END_DEFINE_SPEC(FJsonPointerSpec)
 
diff --git a/Engine/Source/Runtime/MeshDescription/Private/Tests/MeshElementArrayTests.cpp b/Engine/Source/Runtime/MeshDescription/Private/Tests/MeshElementArrayTests.cpp
--- a/Engine/Source/Runtime/MeshDescription/Private/Tests/MeshElementArrayTests.cpp
+++ b/Engine/Source/Runtime/MeshDescription/Private/Tests/MeshElementArrayTests.cpp
@@ -9,7 +9,7 @@
 
 #if WITH_DEV_AUTOMATION_TESTS
 
-constexpr const EAutomationTestFlags TestFlags = EAutomationTestFlags::ApplicationContextMask | EAutomationTestFlags::EngineFilter;
+constexpr const EAutomationTestFlags TestFlags = EAutomationTestFlags_ApplicationContextMask | EAutomationTestFlags::EngineFilter;
 
 /**
  * A very basic type that can be used as the ElementIDType with TMeshElementArray in our tests
diff --git a/Engine/Source/Runtime/Net/Core/Private/Net/Core/Analytics/NetStatsUtils.cpp b/Engine/Source/Runtime/Net/Core/Private/Net/Core/Analytics/NetStatsUtils.cpp
--- a/Engine/Source/Runtime/Net/Core/Private/Net/Core/Analytics/NetStatsUtils.cpp
+++ b/Engine/Source/Runtime/Net/Core/Private/Net/Core/Analytics/NetStatsUtils.cpp
@@ -135,7 +135,7 @@ void FSampleMinMaxAvg::Reset()
 #include "Misc/AutomationTest.h"
 
 IMPLEMENT_SIMPLE_AUTOMATION_TEST(FBinnedMovingValueTest, "System.Core.Networking.BinnedMovingValue",
-									EAutomationTestFlags::ApplicationContextMask | EAutomationTestFlags::SmokeFilter
+									EAutomationTestFlags_ApplicationContextMask | EAutomationTestFlags::SmokeFilter
 
 );
 
diff --git a/Engine/Source/Runtime/Net/Core/Private/Net/Core/Connection/NetCloseResult.cpp b/Engine/Source/Runtime/Net/Core/Private/Net/Core/Connection/NetCloseResult.cpp
--- a/Engine/Source/Runtime/Net/Core/Private/Net/Core/Connection/NetCloseResult.cpp
+++ b/Engine/Source/Runtime/Net/Core/Private/Net/Core/Connection/NetCloseResult.cpp
@@ -84,7 +84,7 @@ ENetCloseResult FromSecurityEvent(ESecurityEvent::Type Val)
 #include "Misc/AutomationTest.h"
 
 IMPLEMENT_SIMPLE_AUTOMATION_TEST(FNetCloseResultEnumTest, "System.Core.Networking.FNetCloseResult.EnumTest",
-									EAutomationTestFlags::ApplicationContextMask | EAutomationTestFlags::SmokeFilter);
+									EAutomationTestFlags_ApplicationContextMask | EAutomationTestFlags::SmokeFilter);
 
 bool FNetCloseResultEnumTest::RunTest(const FString& Parameters)
 {
diff --git a/Engine/Source/Runtime/Net/Core/Private/Net/Core/Misc/ResizableCircularQueueTests.cpp b/Engine/Source/Runtime/Net/Core/Private/Net/Core/Misc/ResizableCircularQueueTests.cpp
--- a/Engine/Source/Runtime/Net/Core/Private/Net/Core/Misc/ResizableCircularQueueTests.cpp
+++ b/Engine/Source/Runtime/Net/Core/Private/Net/Core/Misc/ResizableCircularQueueTests.cpp
@@ -5,7 +5,7 @@
 
 #if WITH_DEV_AUTOMATION_TESTS
 
-IMPLEMENT_SIMPLE_AUTOMATION_TEST(FResizableCircularQueueTest, "Net.ResizableCircularQueueTest", EAutomationTestFlags::ApplicationContextMask | EAutomationTestFlags::EngineFilter)
+IMPLEMENT_SIMPLE_AUTOMATION_TEST(FResizableCircularQueueTest, "Net.ResizableCircularQueueTest", EAutomationTestFlags_ApplicationContextMask | EAutomationTestFlags::EngineFilter)
 
 struct FResizableCircularQueueTestUtil
 {
diff --git a/Engine/Source/Runtime/Networking/Private/Tests/IPv4AddressTest.cpp b/Engine/Source/Runtime/Networking/Private/Tests/IPv4AddressTest.cpp
--- a/Engine/Source/Runtime/Networking/Private/Tests/IPv4AddressTest.cpp
+++ b/Engine/Source/Runtime/Networking/Private/Tests/IPv4AddressTest.cpp
@@ -6,7 +6,7 @@
 
 #if WITH_DEV_AUTOMATION_TESTS
 
-IMPLEMENT_SIMPLE_AUTOMATION_TEST(FIPv4AddressTest, "System.Engine.Networking.IPv4.IPv4Address", EAutomationTestFlags::ApplicationContextMask | EAutomationTestFlags::SmokeFilter)
+IMPLEMENT_SIMPLE_AUTOMATION_TEST(FIPv4AddressTest, "System.Engine.Networking.IPv4.IPv4Address", EAutomationTestFlags_ApplicationContextMask | EAutomationTestFlags::SmokeFilter)
 
 
 bool FIPv4AddressTest::RunTest( const FString& Parameters )
diff --git a/Engine/Source/Runtime/Online/BuildPatchServices/Private/Tests/Unit/BlockStructure.spec.cpp b/Engine/Source/Runtime/Online/BuildPatchServices/Private/Tests/Unit/BlockStructure.spec.cpp
--- a/Engine/Source/Runtime/Online/BuildPatchServices/Private/Tests/Unit/BlockStructure.spec.cpp
+++ b/Engine/Source/Runtime/Online/BuildPatchServices/Private/Tests/Unit/BlockStructure.spec.cpp
@@ -81,7 +81,7 @@ namespace
 	}
 }
 
-BEGIN_DEFINE_SPEC(FBlockStructureSpec, "BuildPatchServices.Unit", EAutomationTestFlags::ProductFilter | EAutomationTestFlags::ApplicationContextMask)
+BEGIN_DEFINE_SPEC(FBlockStructureSpec, "BuildPatchServices.Unit", EAutomationTestFlags::ProductFilter | EAutomationTestFlags_ApplicationContextMask)
 // Unit.
 BuildPatchServices::FBlockStructure BlockStructure;
 BuildPatchServices::FBlockStructure OtherStructure;
diff --git a/Engine/Source/Runtime/Online/BuildPatchServices/Private/Tests/Unit/ChainedChunkSource.spec.cpp b/Engine/Source/Runtime/Online/BuildPatchServices/Private/Tests/Unit/ChainedChunkSource.spec.cpp
--- a/Engine/Source/Runtime/Online/BuildPatchServices/Private/Tests/Unit/ChainedChunkSource.spec.cpp
+++ b/Engine/Source/Runtime/Online/BuildPatchServices/Private/Tests/Unit/ChainedChunkSource.spec.cpp
@@ -30,7 +30,7 @@ namespace BuildPatchServices
 	};
 }
 
-BEGIN_DEFINE_SPEC(FChainedChunkSourceSpec, "BuildPatchServices.Unit", EAutomationTestFlags::ProductFilter | EAutomationTestFlags::ApplicationContextMask)
+BEGIN_DEFINE_SPEC(FChainedChunkSourceSpec, "BuildPatchServices.Unit", EAutomationTestFlags::ProductFilter | EAutomationTestFlags_ApplicationContextMask)
 // Unit
 TUniquePtr<BuildPatchServices::IChainedChunkSource> ChainedChunkSource;
 // Mock
diff --git a/Engine/Source/Runtime/Online/BuildPatchServices/Private/Tests/Unit/ChunkEvictionPolicy.spec.cpp b/Engine/Source/Runtime/Online/BuildPatchServices/Private/Tests/Unit/ChunkEvictionPolicy.spec.cpp
--- a/Engine/Source/Runtime/Online/BuildPatchServices/Private/Tests/Unit/ChunkEvictionPolicy.spec.cpp
+++ b/Engine/Source/Runtime/Online/BuildPatchServices/Private/Tests/Unit/ChunkEvictionPolicy.spec.cpp
@@ -8,7 +8,7 @@
 
 #if WITH_DEV_AUTOMATION_TESTS
 
-BEGIN_DEFINE_SPEC(FChunkEvictionPolicySpec, "BuildPatchServices.Unit", EAutomationTestFlags::ProductFilter | EAutomationTestFlags::ApplicationContextMask)
+BEGIN_DEFINE_SPEC(FChunkEvictionPolicySpec, "BuildPatchServices.Unit", EAutomationTestFlags::ProductFilter | EAutomationTestFlags_ApplicationContextMask)
 const uint32 TestChunkSize = 128 * 1024;
 // Unit
 TUniquePtr<BuildPatchServices::IChunkEvictionPolicy> ChunkEvictionPolicy;
diff --git a/Engine/Source/Runtime/Online/BuildPatchServices/Private/Tests/Unit/ChunkReferenceTracker.spec.cpp b/Engine/Source/Runtime/Online/BuildPatchServices/Private/Tests/Unit/ChunkReferenceTracker.spec.cpp
--- a/Engine/Source/Runtime/Online/BuildPatchServices/Private/Tests/Unit/ChunkReferenceTracker.spec.cpp
+++ b/Engine/Source/Runtime/Online/BuildPatchServices/Private/Tests/Unit/ChunkReferenceTracker.spec.cpp
@@ -10,7 +10,7 @@
 
 #if WITH_DEV_AUTOMATION_TESTS
 
-BEGIN_DEFINE_SPEC(FChunkReferenceTrackerSpec, "BuildPatchServices.Unit", EAutomationTestFlags::ProductFilter | EAutomationTestFlags::ApplicationContextMask)
+BEGIN_DEFINE_SPEC(FChunkReferenceTrackerSpec, "BuildPatchServices.Unit", EAutomationTestFlags::ProductFilter | EAutomationTestFlags_ApplicationContextMask)
 const uint32 TestChunkSize = 128 * 1024;
 // Unit
 TUniquePtr<BuildPatchServices::IChunkReferenceTracker> ChunkReferenceTracker;
diff --git a/Engine/Source/Runtime/Online/BuildPatchServices/Private/Tests/Unit/CloudChunkSource.spec.cpp b/Engine/Source/Runtime/Online/BuildPatchServices/Private/Tests/Unit/CloudChunkSource.spec.cpp
--- a/Engine/Source/Runtime/Online/BuildPatchServices/Private/Tests/Unit/CloudChunkSource.spec.cpp
+++ b/Engine/Source/Runtime/Online/BuildPatchServices/Private/Tests/Unit/CloudChunkSource.spec.cpp
@@ -29,7 +29,7 @@ namespace BuildPatchServices
 {
 	typedef TPlatform<FMockPlatformProcess, FMockPlatformMisc> FMockPlatform;
 }
-BEGIN_DEFINE_SPEC(FCloudChunkSourceSpec, "BuildPatchServices.Unit", EAutomationTestFlags::ProductFilter | EAutomationTestFlags::ApplicationContextMask)
+BEGIN_DEFINE_SPEC(FCloudChunkSourceSpec, "BuildPatchServices.Unit", EAutomationTestFlags::ProductFilter | EAutomationTestFlags_ApplicationContextMask)
 // Unit.
 TUniquePtr<BuildPatchServices::ICloudChunkSource> CloudChunkSource;
 TUniquePtr<BuildPatchServices::FBuildInstallerSharedContext> InstallerSharedContext;
diff --git a/Engine/Source/Runtime/Online/BuildPatchServices/Private/Tests/Unit/DiskChunkStore.spec.cpp b/Engine/Source/Runtime/Online/BuildPatchServices/Private/Tests/Unit/DiskChunkStore.spec.cpp
--- a/Engine/Source/Runtime/Online/BuildPatchServices/Private/Tests/Unit/DiskChunkStore.spec.cpp
+++ b/Engine/Source/Runtime/Online/BuildPatchServices/Private/Tests/Unit/DiskChunkStore.spec.cpp
@@ -12,7 +12,7 @@
 
 #if WITH_DEV_AUTOMATION_TESTS
 
-BEGIN_DEFINE_SPEC(FDiskChunkStoreSpec, "BuildPatchServices.Unit", EAutomationTestFlags::ProductFilter | EAutomationTestFlags::ApplicationContextMask)
+BEGIN_DEFINE_SPEC(FDiskChunkStoreSpec, "BuildPatchServices.Unit", EAutomationTestFlags::ProductFilter | EAutomationTestFlags_ApplicationContextMask)
 // Unit
 TUniquePtr<BuildPatchServices::IDiskChunkStore> DiskChunkStore;
 TUniquePtr<BuildPatchServices::FBuildInstallerSharedContext> InstallerSharedContext;
diff --git a/Engine/Source/Runtime/Online/BuildPatchServices/Private/Tests/Unit/DownloadService.spec.cpp b/Engine/Source/Runtime/Online/BuildPatchServices/Private/Tests/Unit/DownloadService.spec.cpp
--- a/Engine/Source/Runtime/Online/BuildPatchServices/Private/Tests/Unit/DownloadService.spec.cpp
+++ b/Engine/Source/Runtime/Online/BuildPatchServices/Private/Tests/Unit/DownloadService.spec.cpp
@@ -12,7 +12,7 @@
 
 #if WITH_DEV_AUTOMATION_TESTS
 
-BEGIN_DEFINE_SPEC(FDownloadServiceSpec, "BuildPatchServices.Unit", EAutomationTestFlags::ProductFilter | EAutomationTestFlags::ApplicationContextMask)
+BEGIN_DEFINE_SPEC(FDownloadServiceSpec, "BuildPatchServices.Unit", EAutomationTestFlags::ProductFilter | EAutomationTestFlags_ApplicationContextMask)
 // Unit.
 TUniquePtr<BuildPatchServices::IDownloadService> DownloadService;
 // Mock/Fake.
diff --git a/Engine/Source/Runtime/Online/BuildPatchServices/Private/Tests/Unit/FileAttribution.spec.cpp b/Engine/Source/Runtime/Online/BuildPatchServices/Private/Tests/Unit/FileAttribution.spec.cpp
--- a/Engine/Source/Runtime/Online/BuildPatchServices/Private/Tests/Unit/FileAttribution.spec.cpp
+++ b/Engine/Source/Runtime/Online/BuildPatchServices/Private/Tests/Unit/FileAttribution.spec.cpp
@@ -13,7 +13,7 @@
 
 #if WITH_DEV_AUTOMATION_TESTS
 
-BEGIN_DEFINE_SPEC(FFileAttributionSpec, "BuildPatchServices.Unit", EAutomationTestFlags::ProductFilter | EAutomationTestFlags::ApplicationContextMask)
+BEGIN_DEFINE_SPEC(FFileAttributionSpec, "BuildPatchServices.Unit", EAutomationTestFlags::ProductFilter | EAutomationTestFlags_ApplicationContextMask)
 // Unit.
 TUniquePtr<BuildPatchServices::IFileAttribution> FileAttribution;
 // Mock.
diff --git a/Engine/Source/Runtime/Online/BuildPatchServices/Private/Tests/Unit/InstallChunkSource.spec.cpp b/Engine/Source/Runtime/Online/BuildPatchServices/Private/Tests/Unit/InstallChunkSource.spec.cpp
--- a/Engine/Source/Runtime/Online/BuildPatchServices/Private/Tests/Unit/InstallChunkSource.spec.cpp
+++ b/Engine/Source/Runtime/Online/BuildPatchServices/Private/Tests/Unit/InstallChunkSource.spec.cpp
@@ -19,7 +19,7 @@
 
 #if WITH_DEV_AUTOMATION_TESTS
 
-BEGIN_DEFINE_SPEC(FInstallChunkSourceSpec, "BuildPatchServices.Unit", EAutomationTestFlags::ProductFilter | EAutomationTestFlags::ApplicationContextMask)
+BEGIN_DEFINE_SPEC(FInstallChunkSourceSpec, "BuildPatchServices.Unit", EAutomationTestFlags::ProductFilter | EAutomationTestFlags_ApplicationContextMask)
 const uint32 TestChunkSize = 128 * 1024;
 // Unit.
 TUniquePtr<BuildPatchServices::IInstallChunkSource> InstallChunkSource;
diff --git a/Engine/Source/Runtime/Online/BuildPatchServices/Private/Tests/Unit/MemoryChunkStore.spec.cpp b/Engine/Source/Runtime/Online/BuildPatchServices/Private/Tests/Unit/MemoryChunkStore.spec.cpp
--- a/Engine/Source/Runtime/Online/BuildPatchServices/Private/Tests/Unit/MemoryChunkStore.spec.cpp
+++ b/Engine/Source/Runtime/Online/BuildPatchServices/Private/Tests/Unit/MemoryChunkStore.spec.cpp
@@ -11,7 +11,7 @@
 
 #if WITH_DEV_AUTOMATION_TESTS
 
-BEGIN_DEFINE_SPEC(FMemoryChunkStoreSpec, "BuildPatchServices.Unit", EAutomationTestFlags::ProductFilter | EAutomationTestFlags::ApplicationContextMask)
+BEGIN_DEFINE_SPEC(FMemoryChunkStoreSpec, "BuildPatchServices.Unit", EAutomationTestFlags::ProductFilter | EAutomationTestFlags_ApplicationContextMask)
 // Unit
 TUniquePtr<BuildPatchServices::IMemoryChunkStore> MemoryChunkStore;
 // Mock
diff --git a/Engine/Source/Runtime/Online/BuildPatchServices/Private/Tests/Unit/ProcessTimer.spec.cpp b/Engine/Source/Runtime/Online/BuildPatchServices/Private/Tests/Unit/ProcessTimer.spec.cpp
--- a/Engine/Source/Runtime/Online/BuildPatchServices/Private/Tests/Unit/ProcessTimer.spec.cpp
+++ b/Engine/Source/Runtime/Online/BuildPatchServices/Private/Tests/Unit/ProcessTimer.spec.cpp
@@ -7,7 +7,7 @@
 
 #if WITH_DEV_AUTOMATION_TESTS
 
-BEGIN_DEFINE_SPEC(FProcessTimerSpec, "BuildPatchServices.Unit", EAutomationTestFlags::ProductFilter | EAutomationTestFlags::ApplicationContextMask)
+BEGIN_DEFINE_SPEC(FProcessTimerSpec, "BuildPatchServices.Unit", EAutomationTestFlags::ProductFilter | EAutomationTestFlags_ApplicationContextMask)
 // Unit.
 typedef BuildPatchServices::TProcessTimer<BuildPatchServices::FMockCyclesProvider> FTestProcessTimer;
 TUniquePtr<FTestProcessTimer> ProcessTimer;
diff --git a/Engine/Source/Runtime/Online/BuildPatchServices/Private/Tests/Unit/Verifier.spec.cpp b/Engine/Source/Runtime/Online/BuildPatchServices/Private/Tests/Unit/Verifier.spec.cpp
--- a/Engine/Source/Runtime/Online/BuildPatchServices/Private/Tests/Unit/Verifier.spec.cpp
+++ b/Engine/Source/Runtime/Online/BuildPatchServices/Private/Tests/Unit/Verifier.spec.cpp
@@ -16,7 +16,7 @@
 
 #if WITH_DEV_AUTOMATION_TESTS
 
-BEGIN_DEFINE_SPEC(FVerifierSpec, "BuildPatchServices.Unit", EAutomationTestFlags::ProductFilter | EAutomationTestFlags::ApplicationContextMask)
+BEGIN_DEFINE_SPEC(FVerifierSpec, "BuildPatchServices.Unit", EAutomationTestFlags::ProductFilter | EAutomationTestFlags_ApplicationContextMask)
 // Unit.
 TUniquePtr<BuildPatchServices::IVerifier> Verifier;
 // Mock.
diff --git a/Engine/Source/Runtime/PropertyPath/Private/Tests/PropertyPathHelpersTest.cpp b/Engine/Source/Runtime/PropertyPath/Private/Tests/PropertyPathHelpersTest.cpp
--- a/Engine/Source/Runtime/PropertyPath/Private/Tests/PropertyPathHelpersTest.cpp
+++ b/Engine/Source/Runtime/PropertyPath/Private/Tests/PropertyPathHelpersTest.cpp
@@ -50,7 +50,7 @@ DEFINE_LOG_CATEGORY_STATIC(LogPropertyPathHelpersTest, Log, All);
 
 // Begin Tests
 
-IMPLEMENT_SIMPLE_AUTOMATION_TEST(FSetPrimitivesTest, "System.PropertyPath.SetPrimitivesTest", EAutomationTestFlags::ApplicationContextMask | EAutomationTestFlags::EngineFilter)
+IMPLEMENT_SIMPLE_AUTOMATION_TEST(FSetPrimitivesTest, "System.PropertyPath.SetPrimitivesTest", EAutomationTestFlags_ApplicationContextMask | EAutomationTestFlags::EngineFilter)
 bool FSetPrimitivesTest::RunTest(const FString& Parameters)
 {
 	bool GPassing = true;
@@ -97,7 +97,7 @@ bool FSetPrimitivesTest::RunTest(const FString& Parameters)
 	return GPassing;
 }
 
-IMPLEMENT_SIMPLE_AUTOMATION_TEST(FSetInvalidInputIsNoOp, "System.PropertyPath.SetInvalidInputIsNoOp", EAutomationTestFlags::ApplicationContextMask | EAutomationTestFlags::EngineFilter)
+IMPLEMENT_SIMPLE_AUTOMATION_TEST(FSetInvalidInputIsNoOp, "System.PropertyPath.SetInvalidInputIsNoOp", EAutomationTestFlags_ApplicationContextMask | EAutomationTestFlags::EngineFilter)
 bool FSetInvalidInputIsNoOp::RunTest(const FString& Parameters)
 {
 	bool GPassing = true;
@@ -147,7 +147,7 @@ bool FSetInvalidInputIsNoOp::RunTest(const FString& Parameters)
 	return GPassing;
 }
 
-IMPLEMENT_SIMPLE_AUTOMATION_TEST(FSetStringTest, "System.PropertyPath.SetStringTest", EAutomationTestFlags::ApplicationContextMask | EAutomationTestFlags::EngineFilter)
+IMPLEMENT_SIMPLE_AUTOMATION_TEST(FSetStringTest, "System.PropertyPath.SetStringTest", EAutomationTestFlags_ApplicationContextMask | EAutomationTestFlags::EngineFilter)
 bool FSetStringTest::RunTest(const FString& Parameters)
 {
 	bool GPassing = true;
@@ -168,7 +168,7 @@ bool FSetStringTest::RunTest(const FString& Parameters)
 	return GPassing;
 }
 
-IMPLEMENT_SIMPLE_AUTOMATION_TEST(FSetEnumTest, "System.PropertyPath.SetEnumTest", EAutomationTestFlags::ApplicationContextMask | EAutomationTestFlags::EngineFilter)
+IMPLEMENT_SIMPLE_AUTOMATION_TEST(FSetEnumTest, "System.PropertyPath.SetEnumTest", EAutomationTestFlags_ApplicationContextMask | EAutomationTestFlags::EngineFilter)
 bool FSetEnumTest::RunTest(const FString& Parameters)
 {
 	bool GPassing = true;
@@ -245,7 +245,7 @@ bool FSetEnumTest::RunTest(const FString& Parameters)
 	return GPassing;
 }
 
-IMPLEMENT_SIMPLE_AUTOMATION_TEST(FSetStructTest, "System.PropertyPath.SetStructTest", EAutomationTestFlags::ApplicationContextMask | EAutomationTestFlags::EngineFilter)
+IMPLEMENT_SIMPLE_AUTOMATION_TEST(FSetStructTest, "System.PropertyPath.SetStructTest", EAutomationTestFlags_ApplicationContextMask | EAutomationTestFlags::EngineFilter)
 bool FSetStructTest::RunTest(const FString& Parameters)
 {
 	bool GPassing = true;
@@ -292,7 +292,7 @@ bool FSetStructTest::RunTest(const FString& Parameters)
 	return GPassing;
 }
 
-IMPLEMENT_SIMPLE_AUTOMATION_TEST(FSetInnerObjectTest, "System.PropertyPath.SetInnerObjectTest", EAutomationTestFlags::ApplicationContextMask | EAutomationTestFlags::EngineFilter)
+IMPLEMENT_SIMPLE_AUTOMATION_TEST(FSetInnerObjectTest, "System.PropertyPath.SetInnerObjectTest", EAutomationTestFlags_ApplicationContextMask | EAutomationTestFlags::EngineFilter)
 bool FSetInnerObjectTest::RunTest(const FString& Parameters)
 {
 	bool GPassing = true;
@@ -313,7 +313,7 @@ bool FSetInnerObjectTest::RunTest(const FString& Parameters)
 	return GPassing;
 }
 
-IMPLEMENT_SIMPLE_AUTOMATION_TEST(FSetStructDataTest, "System.PropertyPath.SetStructDataTest", EAutomationTestFlags::ApplicationContextMask | EAutomationTestFlags::EngineFilter)
+IMPLEMENT_SIMPLE_AUTOMATION_TEST(FSetStructDataTest, "System.PropertyPath.SetStructDataTest", EAutomationTestFlags_ApplicationContextMask | EAutomationTestFlags::EngineFilter)
 bool FSetStructDataTest::RunTest(const FString& Parameters)
 {
 	bool GPassing = true;
@@ -363,7 +363,7 @@ bool FSetStructDataTest::RunTest(const FString& Parameters)
 	return GPassing;
 }
 
-IMPLEMENT_SIMPLE_AUTOMATION_TEST(FSetStructMemberTest, "System.PropertyPath.SetStructMemberTest", EAutomationTestFlags::ApplicationContextMask | EAutomationTestFlags::EngineFilter)
+IMPLEMENT_SIMPLE_AUTOMATION_TEST(FSetStructMemberTest, "System.PropertyPath.SetStructMemberTest", EAutomationTestFlags_ApplicationContextMask | EAutomationTestFlags::EngineFilter)
 bool FSetStructMemberTest::RunTest(const FString& Parameters)
 {
 	bool GPassing = true;
@@ -422,7 +422,7 @@ bool FSetStructMemberTest::RunTest(const FString& Parameters)
 	return GPassing;
 }
 
-IMPLEMENT_SIMPLE_AUTOMATION_TEST(FSetStructInnerObjectMemberTest, "System.PropertyPath.SetStructInnerObjectMemberTest", EAutomationTestFlags::ApplicationContextMask | EAutomationTestFlags::EngineFilter)
+IMPLEMENT_SIMPLE_AUTOMATION_TEST(FSetStructInnerObjectMemberTest, "System.PropertyPath.SetStructInnerObjectMemberTest", EAutomationTestFlags_ApplicationContextMask | EAutomationTestFlags::EngineFilter)
 bool FSetStructInnerObjectMemberTest::RunTest(const FString& Parameters)
 {
 	bool GPassing = true;
@@ -481,7 +481,7 @@ bool FSetStructInnerObjectMemberTest::RunTest(const FString& Parameters)
 	return GPassing;
 }
 
-IMPLEMENT_SIMPLE_AUTOMATION_TEST(FSetInnerStructMemberTest, "System.PropertyPath.SetInnerStructMemberTest", EAutomationTestFlags::ApplicationContextMask | EAutomationTestFlags::EngineFilter)
+IMPLEMENT_SIMPLE_AUTOMATION_TEST(FSetInnerStructMemberTest, "System.PropertyPath.SetInnerStructMemberTest", EAutomationTestFlags_ApplicationContextMask | EAutomationTestFlags::EngineFilter)
 bool FSetInnerStructMemberTest::RunTest(const FString& Parameters)
 {
 	bool GPassing = true;
@@ -506,7 +506,7 @@ bool FSetInnerStructMemberTest::RunTest(const FString& Parameters)
 	return GPassing;
 }
 
-IMPLEMENT_SIMPLE_AUTOMATION_TEST(FSetPrimitivesInnerObjectTest, "System.PropertyPath.SetPrimitivesInnerObjectTest", EAutomationTestFlags::ApplicationContextMask | EAutomationTestFlags::EngineFilter)
+IMPLEMENT_SIMPLE_AUTOMATION_TEST(FSetPrimitivesInnerObjectTest, "System.PropertyPath.SetPrimitivesInnerObjectTest", EAutomationTestFlags_ApplicationContextMask | EAutomationTestFlags::EngineFilter)
 bool FSetPrimitivesInnerObjectTest::RunTest(const FString& Parameters)
 {
 	bool GPassing = true;
@@ -562,7 +562,7 @@ bool FSetPrimitivesInnerObjectTest::RunTest(const FString& Parameters)
 	return GPassing;
 }
 
-IMPLEMENT_SIMPLE_AUTOMATION_TEST(FSetStringInnerObjectTest, "System.PropertyPath.SetStringInnerObjectTest", EAutomationTestFlags::ApplicationContextMask | EAutomationTestFlags::EngineFilter)
+IMPLEMENT_SIMPLE_AUTOMATION_TEST(FSetStringInnerObjectTest, "System.PropertyPath.SetStringInnerObjectTest", EAutomationTestFlags_ApplicationContextMask | EAutomationTestFlags::EngineFilter)
 bool FSetStringInnerObjectTest::RunTest(const FString& Parameters)
 {
 	bool GPassing = true;
@@ -586,7 +586,7 @@ bool FSetStringInnerObjectTest::RunTest(const FString& Parameters)
 	return GPassing;
 }
 
-IMPLEMENT_SIMPLE_AUTOMATION_TEST(FSetStructInnerObjectTest, "System.PropertyPath.SetStructInnerObjectTest", EAutomationTestFlags::ApplicationContextMask | EAutomationTestFlags::EngineFilter)
+IMPLEMENT_SIMPLE_AUTOMATION_TEST(FSetStructInnerObjectTest, "System.PropertyPath.SetStructInnerObjectTest", EAutomationTestFlags_ApplicationContextMask | EAutomationTestFlags::EngineFilter)
 bool FSetStructInnerObjectTest::RunTest(const FString& Parameters)
 {
 	bool GPassing = true;
@@ -609,7 +609,7 @@ bool FSetStructInnerObjectTest::RunTest(const FString& Parameters)
 	return GPassing;
 }
 
-IMPLEMENT_SIMPLE_AUTOMATION_TEST(FSetStructMemberInnerObjectTest, "System.PropertyPath.SetStructMemberInnerObjectTest", EAutomationTestFlags::ApplicationContextMask | EAutomationTestFlags::EngineFilter)
+IMPLEMENT_SIMPLE_AUTOMATION_TEST(FSetStructMemberInnerObjectTest, "System.PropertyPath.SetStructMemberInnerObjectTest", EAutomationTestFlags_ApplicationContextMask | EAutomationTestFlags::EngineFilter)
 bool FSetStructMemberInnerObjectTest::RunTest(const FString& Parameters)
 {
 	bool GPassing = true;
@@ -636,7 +636,7 @@ bool FSetStructMemberInnerObjectTest::RunTest(const FString& Parameters)
 	return GPassing;
 }
 
-IMPLEMENT_SIMPLE_AUTOMATION_TEST(FSetInnerStructMemberInnerObjectTest, "System.PropertyPath.SetInnerStructMemberInnerObjectTest", EAutomationTestFlags::ApplicationContextMask | EAutomationTestFlags::EngineFilter)
+IMPLEMENT_SIMPLE_AUTOMATION_TEST(FSetInnerStructMemberInnerObjectTest, "System.PropertyPath.SetInnerStructMemberInnerObjectTest", EAutomationTestFlags_ApplicationContextMask | EAutomationTestFlags::EngineFilter)
 bool FSetInnerStructMemberInnerObjectTest::RunTest(const FString& Parameters)
 {
 	bool GPassing = true;
@@ -663,7 +663,7 @@ bool FSetInnerStructMemberInnerObjectTest::RunTest(const FString& Parameters)
 	return GPassing;
 }
 
-IMPLEMENT_SIMPLE_AUTOMATION_TEST(FSetInnerStructMemberEnumInnerObjectTest, "System.PropertyPath.SetInnerStructMemberEnumInnerObjectTest", EAutomationTestFlags::ApplicationContextMask | EAutomationTestFlags::EngineFilter)
+IMPLEMENT_SIMPLE_AUTOMATION_TEST(FSetInnerStructMemberEnumInnerObjectTest, "System.PropertyPath.SetInnerStructMemberEnumInnerObjectTest", EAutomationTestFlags_ApplicationContextMask | EAutomationTestFlags::EngineFilter)
 bool FSetInnerStructMemberEnumInnerObjectTest::RunTest(const FString& Parameters)
 {
 	bool GPassing = true;
@@ -690,7 +690,7 @@ bool FSetInnerStructMemberEnumInnerObjectTest::RunTest(const FString& Parameters
 	return GPassing;
 }
 
- IMPLEMENT_SIMPLE_AUTOMATION_TEST(FGetPrimitivesTest, "System.PropertyPath.GetPrimitivesTest", EAutomationTestFlags::ApplicationContextMask | EAutomationTestFlags::EngineFilter)
+ IMPLEMENT_SIMPLE_AUTOMATION_TEST(FGetPrimitivesTest, "System.PropertyPath.GetPrimitivesTest", EAutomationTestFlags_ApplicationContextMask | EAutomationTestFlags::EngineFilter)
  bool FGetPrimitivesTest::RunTest(const FString& Parameters)
  {
 	 bool GPassing = true;
@@ -738,7 +738,7 @@ bool FSetInnerStructMemberEnumInnerObjectTest::RunTest(const FString& Parameters
  }
 
 
-IMPLEMENT_SIMPLE_AUTOMATION_TEST(FGetStringTest, "System.PropertyPath.GetStringTest", EAutomationTestFlags::ApplicationContextMask | EAutomationTestFlags::EngineFilter)
+IMPLEMENT_SIMPLE_AUTOMATION_TEST(FGetStringTest, "System.PropertyPath.GetStringTest", EAutomationTestFlags_ApplicationContextMask | EAutomationTestFlags::EngineFilter)
 bool FGetStringTest::RunTest(const FString& Parameters)
 {
 	bool GPassing = true;
@@ -759,7 +759,7 @@ bool FGetStringTest::RunTest(const FString& Parameters)
 	return GPassing;
 }
 
-IMPLEMENT_SIMPLE_AUTOMATION_TEST(FGetEnumTest, "System.PropertyPath.GetEnumTest", EAutomationTestFlags::ApplicationContextMask | EAutomationTestFlags::EngineFilter)
+IMPLEMENT_SIMPLE_AUTOMATION_TEST(FGetEnumTest, "System.PropertyPath.GetEnumTest", EAutomationTestFlags_ApplicationContextMask | EAutomationTestFlags::EngineFilter)
 bool FGetEnumTest::RunTest(const FString& Parameters)
 {
 	bool GPassing = true;
@@ -780,7 +780,7 @@ bool FGetEnumTest::RunTest(const FString& Parameters)
 	return GPassing;
 }
 
-IMPLEMENT_SIMPLE_AUTOMATION_TEST(FGetStructTest, "System.PropertyPath.GetStructTest", EAutomationTestFlags::ApplicationContextMask | EAutomationTestFlags::EngineFilter)
+IMPLEMENT_SIMPLE_AUTOMATION_TEST(FGetStructTest, "System.PropertyPath.GetStructTest", EAutomationTestFlags_ApplicationContextMask | EAutomationTestFlags::EngineFilter)
 bool FGetStructTest::RunTest(const FString& Parameters)
 {
 	bool GPassing = true;
@@ -827,7 +827,7 @@ bool FGetStructTest::RunTest(const FString& Parameters)
 	return GPassing;
 }
 
-IMPLEMENT_SIMPLE_AUTOMATION_TEST(FGetInnerObjectTest, "System.PropertyPath.GetInnerObjectTest", EAutomationTestFlags::ApplicationContextMask | EAutomationTestFlags::EngineFilter)
+IMPLEMENT_SIMPLE_AUTOMATION_TEST(FGetInnerObjectTest, "System.PropertyPath.GetInnerObjectTest", EAutomationTestFlags_ApplicationContextMask | EAutomationTestFlags::EngineFilter)
 bool FGetInnerObjectTest::RunTest(const FString& Parameters)
 {
 	bool GPassing = true;
@@ -848,7 +848,7 @@ bool FGetInnerObjectTest::RunTest(const FString& Parameters)
 	return GPassing;
 }
 
-IMPLEMENT_SIMPLE_AUTOMATION_TEST(FGetStructMemberTest, "System.PropertyPath.GetStructMemberTest", EAutomationTestFlags::ApplicationContextMask | EAutomationTestFlags::EngineFilter)
+IMPLEMENT_SIMPLE_AUTOMATION_TEST(FGetStructMemberTest, "System.PropertyPath.GetStructMemberTest", EAutomationTestFlags_ApplicationContextMask | EAutomationTestFlags::EngineFilter)
 bool FGetStructMemberTest::RunTest(const FString& Parameters)
 {
 	bool GPassing = true;
@@ -898,7 +898,7 @@ bool FGetStructMemberTest::RunTest(const FString& Parameters)
 	return GPassing;
 }
 
-IMPLEMENT_SIMPLE_AUTOMATION_TEST(FGetStructInnerObjectMemberTest, "System.PropertyPath.GetStructInnerObjectMemberTest", EAutomationTestFlags::ApplicationContextMask | EAutomationTestFlags::EngineFilter)
+IMPLEMENT_SIMPLE_AUTOMATION_TEST(FGetStructInnerObjectMemberTest, "System.PropertyPath.GetStructInnerObjectMemberTest", EAutomationTestFlags_ApplicationContextMask | EAutomationTestFlags::EngineFilter)
 bool FGetStructInnerObjectMemberTest::RunTest(const FString& Parameters)
 {
 	bool GPassing = true;
@@ -948,7 +948,7 @@ bool FGetStructInnerObjectMemberTest::RunTest(const FString& Parameters)
 	return GPassing;
 }
 
-IMPLEMENT_SIMPLE_AUTOMATION_TEST(FGetInnerStructMemberTest, "System.PropertyPath.GetInnerStructMemberTest", EAutomationTestFlags::ApplicationContextMask | EAutomationTestFlags::EngineFilter)
+IMPLEMENT_SIMPLE_AUTOMATION_TEST(FGetInnerStructMemberTest, "System.PropertyPath.GetInnerStructMemberTest", EAutomationTestFlags_ApplicationContextMask | EAutomationTestFlags::EngineFilter)
 bool FGetInnerStructMemberTest::RunTest(const FString& Parameters)
 {
 	bool GPassing = true;
@@ -970,7 +970,7 @@ bool FGetInnerStructMemberTest::RunTest(const FString& Parameters)
 	return GPassing;
 }
 
-IMPLEMENT_SIMPLE_AUTOMATION_TEST(FGetPrimitivesInnerObjectTest, "System.PropertyPath.GetPrimitivesInnerObjectTest", EAutomationTestFlags::ApplicationContextMask | EAutomationTestFlags::EngineFilter)
+IMPLEMENT_SIMPLE_AUTOMATION_TEST(FGetPrimitivesInnerObjectTest, "System.PropertyPath.GetPrimitivesInnerObjectTest", EAutomationTestFlags_ApplicationContextMask | EAutomationTestFlags::EngineFilter)
 bool FGetPrimitivesInnerObjectTest::RunTest(const FString& Parameters)
 {
 	bool GPassing = true;
@@ -1026,7 +1026,7 @@ bool FGetPrimitivesInnerObjectTest::RunTest(const FString& Parameters)
 	return GPassing;
 }
 
-IMPLEMENT_SIMPLE_AUTOMATION_TEST(FGetStringInnerObjectTest, "System.PropertyPath.GetStringInnerObjectTest", EAutomationTestFlags::ApplicationContextMask | EAutomationTestFlags::EngineFilter)
+IMPLEMENT_SIMPLE_AUTOMATION_TEST(FGetStringInnerObjectTest, "System.PropertyPath.GetStringInnerObjectTest", EAutomationTestFlags_ApplicationContextMask | EAutomationTestFlags::EngineFilter)
 bool FGetStringInnerObjectTest::RunTest(const FString& Parameters)
 {
 	bool GPassing = true;
@@ -1050,7 +1050,7 @@ bool FGetStringInnerObjectTest::RunTest(const FString& Parameters)
 	return GPassing;
 }
 
-IMPLEMENT_SIMPLE_AUTOMATION_TEST(FGetStructInnerObjectTest, "System.PropertyPath.GetStructInnerObjectTest", EAutomationTestFlags::ApplicationContextMask | EAutomationTestFlags::EngineFilter)
+IMPLEMENT_SIMPLE_AUTOMATION_TEST(FGetStructInnerObjectTest, "System.PropertyPath.GetStructInnerObjectTest", EAutomationTestFlags_ApplicationContextMask | EAutomationTestFlags::EngineFilter)
 bool FGetStructInnerObjectTest::RunTest(const FString& Parameters)
 {
 	bool GPassing = true;
@@ -1073,7 +1073,7 @@ bool FGetStructInnerObjectTest::RunTest(const FString& Parameters)
 	return GPassing;
 }
 
-IMPLEMENT_SIMPLE_AUTOMATION_TEST(FGetStructMemberInnerObjectTest, "System.PropertyPath.GetStructMemberInnerObjectTest", EAutomationTestFlags::ApplicationContextMask | EAutomationTestFlags::EngineFilter)
+IMPLEMENT_SIMPLE_AUTOMATION_TEST(FGetStructMemberInnerObjectTest, "System.PropertyPath.GetStructMemberInnerObjectTest", EAutomationTestFlags_ApplicationContextMask | EAutomationTestFlags::EngineFilter)
 bool FGetStructMemberInnerObjectTest::RunTest(const FString& Parameters)
 {
 	bool GPassing = true;
@@ -1096,7 +1096,7 @@ bool FGetStructMemberInnerObjectTest::RunTest(const FString& Parameters)
 	return GPassing;
 }
 
-IMPLEMENT_SIMPLE_AUTOMATION_TEST(FGetInnerStructMemberInnerObjectTest, "System.PropertyPath.GetInnerStructMemberInnerObjectTest", EAutomationTestFlags::ApplicationContextMask | EAutomationTestFlags::EngineFilter)
+IMPLEMENT_SIMPLE_AUTOMATION_TEST(FGetInnerStructMemberInnerObjectTest, "System.PropertyPath.GetInnerStructMemberInnerObjectTest", EAutomationTestFlags_ApplicationContextMask | EAutomationTestFlags::EngineFilter)
 bool FGetInnerStructMemberInnerObjectTest::RunTest(const FString& Parameters)
 {
 	bool GPassing = true;
@@ -1119,7 +1119,7 @@ bool FGetInnerStructMemberInnerObjectTest::RunTest(const FString& Parameters)
 	return GPassing;
 }
 
-IMPLEMENT_SIMPLE_AUTOMATION_TEST(FGetInnerStructMemberEnumInnerObjectTest, "System.PropertyPath.GetInnerStructMemberEnumInnerObjectTest", EAutomationTestFlags::ApplicationContextMask | EAutomationTestFlags::EngineFilter)
+IMPLEMENT_SIMPLE_AUTOMATION_TEST(FGetInnerStructMemberEnumInnerObjectTest, "System.PropertyPath.GetInnerStructMemberEnumInnerObjectTest", EAutomationTestFlags_ApplicationContextMask | EAutomationTestFlags::EngineFilter)
 bool FGetInnerStructMemberEnumInnerObjectTest::RunTest(const FString& Parameters)
 {
 	bool GPassing = true;
diff --git a/Engine/Source/Runtime/SlateCore/Private/Test/SlateAttributeTest.cpp b/Engine/Source/Runtime/SlateCore/Private/Test/SlateAttributeTest.cpp
--- a/Engine/Source/Runtime/SlateCore/Private/Test/SlateAttributeTest.cpp
+++ b/Engine/Source/Runtime/SlateCore/Private/Test/SlateAttributeTest.cpp
@@ -14,8 +14,8 @@
 
 #define LOCTEXT_NAMESPACE "Slate.Attribute"
 
-//IMPLEMENT_SIMPLE_AUTOMATION_TEST(FSlateAttributeTest, "Slate.Attribute", EAutomationTestFlags::ApplicationContextMask | EAutomationTestFlags::SmokeFilter)
-IMPLEMENT_SIMPLE_AUTOMATION_TEST(FSlateAttributeTest, "Slate.Attribute", EAutomationTestFlags::ApplicationContextMask | EAutomationTestFlags::EngineFilter)
+//IMPLEMENT_SIMPLE_AUTOMATION_TEST(FSlateAttributeTest, "Slate.Attribute", EAutomationTestFlags_ApplicationContextMask | EAutomationTestFlags::SmokeFilter)
+IMPLEMENT_SIMPLE_AUTOMATION_TEST(FSlateAttributeTest, "Slate.Attribute", EAutomationTestFlags_ApplicationContextMask | EAutomationTestFlags::EngineFilter)
 
 namespace UE::Slate::Private
 {
diff --git a/Engine/Source/Runtime/SlateCore/Private/Test/SlateChildrenTest.cpp b/Engine/Source/Runtime/SlateCore/Private/Test/SlateChildrenTest.cpp
--- a/Engine/Source/Runtime/SlateCore/Private/Test/SlateChildrenTest.cpp
+++ b/Engine/Source/Runtime/SlateCore/Private/Test/SlateChildrenTest.cpp
@@ -13,7 +13,7 @@
 
 #define LOCTEXT_NAMESPACE "SlateChildrenTest"
 
-IMPLEMENT_SIMPLE_AUTOMATION_TEST(FSlateChildrenForEachTest, "Slate.Children.ForEach", EAutomationTestFlags::ApplicationContextMask | EAutomationTestFlags::EngineFilter)
+IMPLEMENT_SIMPLE_AUTOMATION_TEST(FSlateChildrenForEachTest, "Slate.Children.ForEach", EAutomationTestFlags_ApplicationContextMask | EAutomationTestFlags::EngineFilter)
 
 namespace UE::Slate::Test::Private
 {
diff --git a/Engine/Source/Runtime/SlateCore/Private/Test/SlateInvalidationWidgetListTest.cpp b/Engine/Source/Runtime/SlateCore/Private/Test/SlateInvalidationWidgetListTest.cpp
--- a/Engine/Source/Runtime/SlateCore/Private/Test/SlateInvalidationWidgetListTest.cpp
+++ b/Engine/Source/Runtime/SlateCore/Private/Test/SlateInvalidationWidgetListTest.cpp
@@ -14,8 +14,8 @@
 
 #define LOCTEXT_NAMESPACE "Slate.FastPath.InvalidationWidgetList"
 
-//IMPLEMENT_SIMPLE_AUTOMATION_TEST(FSlateInvalidationWidgetListTest, "Slate.InvalidationWidgetList.AddBuildRemove", EAutomationTestFlags::ApplicationContextMask | EAutomationTestFlags::SmokeFilter)
-IMPLEMENT_SIMPLE_AUTOMATION_TEST(FSlateInvalidationWidgetListTest, "Slate.InvalidationWidgetList.AddBuildRemove", EAutomationTestFlags::ApplicationContextMask | EAutomationTestFlags::EngineFilter)
+//IMPLEMENT_SIMPLE_AUTOMATION_TEST(FSlateInvalidationWidgetListTest, "Slate.InvalidationWidgetList.AddBuildRemove", EAutomationTestFlags_ApplicationContextMask | EAutomationTestFlags::SmokeFilter)
+IMPLEMENT_SIMPLE_AUTOMATION_TEST(FSlateInvalidationWidgetListTest, "Slate.InvalidationWidgetList.AddBuildRemove", EAutomationTestFlags_ApplicationContextMask | EAutomationTestFlags::EngineFilter)
 
 namespace UE::Slate::Private
 {
diff --git a/Engine/Source/Runtime/SlateCore/Private/Test/SlateInvalidationWidgetReasonText.cpp b/Engine/Source/Runtime/SlateCore/Private/Test/SlateInvalidationWidgetReasonText.cpp
--- a/Engine/Source/Runtime/SlateCore/Private/Test/SlateInvalidationWidgetReasonText.cpp
+++ b/Engine/Source/Runtime/SlateCore/Private/Test/SlateInvalidationWidgetReasonText.cpp
@@ -9,7 +9,7 @@
 
 #define LOCTEXT_NAMESPACE "Slate.InvalidationWidgetReason.Lex"
 
-IMPLEMENT_SIMPLE_AUTOMATION_TEST(FSlateInvalidationWidgetReasonLexTest, "Slate.ParsingInvalidationWidgetReason", EAutomationTestFlags::ApplicationContextMask | EAutomationTestFlags::EngineFilter)
+IMPLEMENT_SIMPLE_AUTOMATION_TEST(FSlateInvalidationWidgetReasonLexTest, "Slate.ParsingInvalidationWidgetReason", EAutomationTestFlags_ApplicationContextMask | EAutomationTestFlags::EngineFilter)
 
 bool FSlateInvalidationWidgetReasonLexTest::RunTest(const FString& Parameters)
 {
diff --git a/Engine/Source/Runtime/TimeManagement/Private/TimeManagementTests.cpp b/Engine/Source/Runtime/TimeManagement/Private/TimeManagementTests.cpp
--- a/Engine/Source/Runtime/TimeManagement/Private/TimeManagementTests.cpp
+++ b/Engine/Source/Runtime/TimeManagement/Private/TimeManagementTests.cpp
@@ -29,7 +29,7 @@ namespace
 	}
 }
 
-IMPLEMENT_SIMPLE_AUTOMATION_TEST(FFrameTimeComparisonTest, "System.Core.Time.Comparison", EAutomationTestFlags::ApplicationContextMask | EAutomationTestFlags::SmokeFilter | EAutomationTestFlags::HighPriority)
+IMPLEMENT_SIMPLE_AUTOMATION_TEST(FFrameTimeComparisonTest, "System.Core.Time.Comparison", EAutomationTestFlags_ApplicationContextMask | EAutomationTestFlags::SmokeFilter | EAutomationTestFlags::HighPriority)
 bool FFrameTimeComparisonTest::RunTest(const FString& Parameters)
 {
 	const int32 NumFrames = UE_ARRAY_COUNT(TestTimes);
@@ -80,7 +80,7 @@ bool FFrameTimeComparisonTest::RunTest(const FString& Parameters)
 	return true;
 }
 
-IMPLEMENT_SIMPLE_AUTOMATION_TEST(FFrameTimeToSecondsTest, "System.Core.Time.ToSeconds", EAutomationTestFlags::ApplicationContextMask | EAutomationTestFlags::SmokeFilter | EAutomationTestFlags::HighPriority)
+IMPLEMENT_SIMPLE_AUTOMATION_TEST(FFrameTimeToSecondsTest, "System.Core.Time.ToSeconds", EAutomationTestFlags_ApplicationContextMask | EAutomationTestFlags::SmokeFilter | EAutomationTestFlags::HighPriority)
 bool FFrameTimeToSecondsTest::RunTest(const FString& Parameters)
 {
 	const int32 NumFrames = UE_ARRAY_COUNT(TestTimes);
@@ -113,7 +113,7 @@ bool FFrameTimeToSecondsTest::RunTest(const FString& Parameters)
 	return true;
 }
 
-IMPLEMENT_SIMPLE_AUTOMATION_TEST(FFrameTimeAdditionTest, "System.Core.Time.Addition", EAutomationTestFlags::ApplicationContextMask | EAutomationTestFlags::SmokeFilter | EAutomationTestFlags::HighPriority)
+IMPLEMENT_SIMPLE_AUTOMATION_TEST(FFrameTimeAdditionTest, "System.Core.Time.Addition", EAutomationTestFlags_ApplicationContextMask | EAutomationTestFlags::SmokeFilter | EAutomationTestFlags::HighPriority)
 bool FFrameTimeAdditionTest::RunTest(const FString& Parameters)
 {
 	const int32 NumFrames = UE_ARRAY_COUNT(TestTimes);
@@ -236,7 +236,7 @@ bool FFrameTimeAdditionTest::RunTest(const FString& Parameters)
 	return true;
 }
 
-IMPLEMENT_SIMPLE_AUTOMATION_TEST(FFrameTimeSubtractionTest, "System.Core.Time.Subtraction", EAutomationTestFlags::ApplicationContextMask | EAutomationTestFlags::SmokeFilter | EAutomationTestFlags::HighPriority)
+IMPLEMENT_SIMPLE_AUTOMATION_TEST(FFrameTimeSubtractionTest, "System.Core.Time.Subtraction", EAutomationTestFlags_ApplicationContextMask | EAutomationTestFlags::SmokeFilter | EAutomationTestFlags::HighPriority)
 bool FFrameTimeSubtractionTest::RunTest(const FString& Parameters)
 {
 	const int32 NumFrames = UE_ARRAY_COUNT(TestTimes);
@@ -360,7 +360,7 @@ bool FFrameTimeSubtractionTest::RunTest(const FString& Parameters)
 }
 
 
-IMPLEMENT_SIMPLE_AUTOMATION_TEST(FFrameTimeConversionTest, "System.Core.Time.Conversion", EAutomationTestFlags::ApplicationContextMask | EAutomationTestFlags::SmokeFilter | EAutomationTestFlags::HighPriority)
+IMPLEMENT_SIMPLE_AUTOMATION_TEST(FFrameTimeConversionTest, "System.Core.Time.Conversion", EAutomationTestFlags_ApplicationContextMask | EAutomationTestFlags::SmokeFilter | EAutomationTestFlags::HighPriority)
 bool FFrameTimeConversionTest::RunTest(const FString& Parameters)
 {
 	const int32 NumFrames = UE_ARRAY_COUNT(TestTimes);
@@ -448,7 +448,7 @@ bool FFrameTimeConversionTest::RunTest(const FString& Parameters)
 	return true;
 }
 
-IMPLEMENT_SIMPLE_AUTOMATION_TEST(FFrameRateMultiplesTest, "System.Core.Time.FrameRateMultiples", EAutomationTestFlags::ApplicationContextMask | EAutomationTestFlags::SmokeFilter | EAutomationTestFlags::HighPriority)
+IMPLEMENT_SIMPLE_AUTOMATION_TEST(FFrameRateMultiplesTest, "System.Core.Time.FrameRateMultiples", EAutomationTestFlags_ApplicationContextMask | EAutomationTestFlags::SmokeFilter | EAutomationTestFlags::HighPriority)
 bool FFrameRateMultiplesTest::RunTest(const FString& Parameters)
 {
 	FFrameRate TestRates[] = {
@@ -559,7 +559,7 @@ bool FFrameRateMultiplesTest::RunTest(const FString& Parameters)
 	return true;
 }
 
-IMPLEMENT_SIMPLE_AUTOMATION_TEST(FFrameTimeToSecondsConversionTest, "System.Core.Time.FrameConversion", EAutomationTestFlags::ApplicationContextMask | EAutomationTestFlags::EngineFilter | EAutomationTestFlags::HighPriority)
+IMPLEMENT_SIMPLE_AUTOMATION_TEST(FFrameTimeToSecondsConversionTest, "System.Core.Time.FrameConversion", EAutomationTestFlags_ApplicationContextMask | EAutomationTestFlags::EngineFilter | EAutomationTestFlags::HighPriority)
 bool FFrameTimeToSecondsConversionTest::RunTest(const FString& Parameters)
 {
 	FFrameRate TestRates[] = {
diff --git a/Engine/Source/Runtime/UMG/Private/Binding/States/Tests/WidgetStateBitfieldTests.cpp b/Engine/Source/Runtime/UMG/Private/Binding/States/Tests/WidgetStateBitfieldTests.cpp
--- a/Engine/Source/Runtime/UMG/Private/Binding/States/Tests/WidgetStateBitfieldTests.cpp
+++ b/Engine/Source/Runtime/UMG/Private/Binding/States/Tests/WidgetStateBitfieldTests.cpp
@@ -26,7 +26,7 @@
 
 DEFINE_LOG_CATEGORY_STATIC(LogWidgetStateBitfieldTest, Log, All);
 
-IMPLEMENT_SIMPLE_AUTOMATION_TEST(FOperatorBoolTest, "Slate.WidgetState.OperatorBoolTest", EAutomationTestFlags::ApplicationContextMask | EAutomationTestFlags::EngineFilter)
+IMPLEMENT_SIMPLE_AUTOMATION_TEST(FOperatorBoolTest, "Slate.WidgetState.OperatorBoolTest", EAutomationTestFlags_ApplicationContextMask | EAutomationTestFlags::EngineFilter)
 bool FOperatorBoolTest::RunTest(const FString& Parameters)
 {
 	bool GPassing = true;
@@ -53,7 +53,7 @@ bool FOperatorBoolTest::RunTest(const FString& Parameters)
 	return GPassing;
 }
 
-IMPLEMENT_SIMPLE_AUTOMATION_TEST(FOperatorBitwiseAndTest, "Slate.WidgetState.OperatorBitwiseAndTest", EAutomationTestFlags::ApplicationContextMask | EAutomationTestFlags::EngineFilter)
+IMPLEMENT_SIMPLE_AUTOMATION_TEST(FOperatorBitwiseAndTest, "Slate.WidgetState.OperatorBitwiseAndTest", EAutomationTestFlags_ApplicationContextMask | EAutomationTestFlags::EngineFilter)
 bool FOperatorBitwiseAndTest::RunTest(const FString& Parameters)
 {
 	bool GPassing = true;
@@ -93,7 +93,7 @@ bool FOperatorBitwiseAndTest::RunTest(const FString& Parameters)
 	return GPassing;
 }
 
-IMPLEMENT_SIMPLE_AUTOMATION_TEST(FOperatorBitwiseOrTest, "Slate.WidgetState.OperatorBitwiseOrTest", EAutomationTestFlags::ApplicationContextMask | EAutomationTestFlags::EngineFilter)
+IMPLEMENT_SIMPLE_AUTOMATION_TEST(FOperatorBitwiseOrTest, "Slate.WidgetState.OperatorBitwiseOrTest", EAutomationTestFlags_ApplicationContextMask | EAutomationTestFlags::EngineFilter)
 bool FOperatorBitwiseOrTest::RunTest(const FString& Parameters)
 {
 	bool GPassing = true;
@@ -124,7 +124,7 @@ bool FOperatorBitwiseOrTest::RunTest(const FString& Parameters)
 	return GPassing;
 }
 
-IMPLEMENT_SIMPLE_AUTOMATION_TEST(FOperatorBitwiseNotTest, "Slate.WidgetState.OperatorBitwiseNotTest", EAutomationTestFlags::ApplicationContextMask | EAutomationTestFlags::EngineFilter)
+IMPLEMENT_SIMPLE_AUTOMATION_TEST(FOperatorBitwiseNotTest, "Slate.WidgetState.OperatorBitwiseNotTest", EAutomationTestFlags_ApplicationContextMask | EAutomationTestFlags::EngineFilter)
 bool FOperatorBitwiseNotTest::RunTest(const FString& Parameters)
 {
 	bool GPassing = true;
@@ -148,7 +148,7 @@ bool FOperatorBitwiseNotTest::RunTest(const FString& Parameters)
 	return GPassing;
 }
 
-IMPLEMENT_SIMPLE_AUTOMATION_TEST(FHasFlagTest, "Slate.WidgetState.HasFlagTest", EAutomationTestFlags::ApplicationContextMask | EAutomationTestFlags::EngineFilter)
+IMPLEMENT_SIMPLE_AUTOMATION_TEST(FHasFlagTest, "Slate.WidgetState.HasFlagTest", EAutomationTestFlags_ApplicationContextMask | EAutomationTestFlags::EngineFilter)
 bool FHasFlagTest::RunTest(const FString& Parameters)
 {
 	bool GPassing = true;
diff --git a/Engine/Source/Runtime/VectorVM/Private/Tests/VectorVMTests.cpp b/Engine/Source/Runtime/VectorVM/Private/Tests/VectorVMTests.cpp
--- a/Engine/Source/Runtime/VectorVM/Private/Tests/VectorVMTests.cpp
+++ b/Engine/Source/Runtime/VectorVM/Private/Tests/VectorVMTests.cpp
@@ -15,7 +15,7 @@
 
 #if WITH_DEV_AUTOMATION_TESTS
 
-IMPLEMENT_SIMPLE_AUTOMATION_TEST(FVectorVMTest, "System.Core.Math.Vector VM", EAutomationTestFlags::ApplicationContextMask | EAutomationTestFlags::SmokeFilter)
+IMPLEMENT_SIMPLE_AUTOMATION_TEST(FVectorVMTest, "System.Core.Math.Vector VM", EAutomationTestFlags_ApplicationContextMask | EAutomationTestFlags::SmokeFilter)
 
 /*------------------------------------------------------------------------------
 Automation test for the VM.
