Set the type for INVALID_NAVNODEREF like the other navigation macros. This is needed to let the compiler resolve some templated methods like AutomationTestBase::TestNotEqual

diff --git a/Engine/Source/Runtime/Engine/Classes/AI/Navigation/NavigationTypes.h b/Engine/Source/Runtime/Engine/Classes/AI/Navigation/NavigationTypes.h
--- a/Engine/Source/Runtime/Engine/Classes/AI/Navigation/NavigationTypes.h
+++ b/Engine/Source/Runtime/Engine/Classes/AI/Navigation/NavigationTypes.h
@@ -14,7 +14,6 @@
 #endif
 #include "NavigationTypes.generated.h"
 
-#define INVALID_NAVNODEREF (0)
 #define INVALID_NAVQUERYID uint32(0)
 #define INVALID_NAVDATA uint32(0)
 #define INVALID_NAVEXTENT (FVector::ZeroVector)
@@ -30,6 +29,7 @@ struct FNavigationPath;
 
 /** uniform identifier type for navigation data elements may it be a polygon or graph node */
 typedef uint64 NavNodeRef;
+#define INVALID_NAVNODEREF NavNodeRef(0)
 
 namespace FNavigationSystem
 {
