Fix for crash in blend-profile when skeleton mapping is incomplete.

diff --git a/Engine/Source/Runtime/Engine/Private/Animation/BlendProfile.cpp b/Engine/Source/Runtime/Engine/Private/Animation/BlendProfile.cpp
--- a/Engine/Source/Runtime/Engine/Private/Animation/BlendProfile.cpp
+++ b/Engine/Source/Runtime/Engine/Private/Animation/BlendProfile.cpp
@@ -313,7 +313,10 @@ void UBlendProfile::FillSkeletonBoneDurationsArray(TCustomBoneIndexArrayView<flo
 				for (const FBlendProfileBoneEntry& Entry : ProfileEntries)
 				{
 					const FSkeletonPoseBoneIndex SkeletonBoneIndex(SkeletonRemapping.GetTargetSkeletonBoneIndex(Entry.BoneReference.BoneIndex));
-					OutDurationPerBone[SkeletonBoneIndex] *= Entry.BlendScale;
+					if (SkeletonBoneIndex != INDEX_NONE)
+					{
+						OutDurationPerBone[SkeletonBoneIndex] *= Entry.BlendScale;
+					}
 				}
 			}
 			else
@@ -321,7 +324,10 @@ void UBlendProfile::FillSkeletonBoneDurationsArray(TCustomBoneIndexArrayView<flo
 				for (const FBlendProfileBoneEntry& Entry : ProfileEntries)
 				{
 					const FSkeletonPoseBoneIndex SkeletonBoneIndex(Entry.BoneReference.BoneIndex);
-					OutDurationPerBone[SkeletonBoneIndex] *= Entry.BlendScale;
+					if (SkeletonBoneIndex != INDEX_NONE)
+					{
+						OutDurationPerBone[SkeletonBoneIndex] *= Entry.BlendScale;
+					}
 				}
 			}
 		}
@@ -334,7 +340,7 @@ void UBlendProfile::FillSkeletonBoneDurationsArray(TCustomBoneIndexArrayView<flo
 				for (const FBlendProfileBoneEntry& Entry : ProfileEntries)
 				{
 					const FSkeletonPoseBoneIndex SkeletonBoneIndex(SkeletonRemapping.GetTargetSkeletonBoneIndex(Entry.BoneReference.BoneIndex));
-					if (Entry.BlendScale > UE_SMALL_NUMBER)
+					if (Entry.BlendScale > UE_SMALL_NUMBER && SkeletonBoneIndex != INDEX_NONE)
 					{
 						OutDurationPerBone[SkeletonBoneIndex] /= Entry.BlendScale;
 					}
@@ -345,7 +351,7 @@ void UBlendProfile::FillSkeletonBoneDurationsArray(TCustomBoneIndexArrayView<flo
 				for (const FBlendProfileBoneEntry& Entry : ProfileEntries)
 				{
 					const FSkeletonPoseBoneIndex SkeletonBoneIndex(Entry.BoneReference.BoneIndex);
-					if (Entry.BlendScale > UE_SMALL_NUMBER)
+					if (Entry.BlendScale > UE_SMALL_NUMBER && SkeletonBoneIndex != INDEX_NONE)
 					{
 						OutDurationPerBone[SkeletonBoneIndex] /= Entry.BlendScale;
 					}
