Fixed crash when force deleting a level. Removed faulty code that was already deprecated since 5.0.

diff --git a/Engine/Source/Runtime/Engine/Private/WorldPartition/DataLayer/WorldDataLayers.cpp b/Engine/Source/Runtime/Engine/Private/WorldPartition/DataLayer/WorldDataLayers.cpp
--- a/Engine/Source/Runtime/Engine/Private/WorldPartition/DataLayer/WorldDataLayers.cpp
+++ b/Engine/Source/Runtime/Engine/Private/WorldPartition/DataLayer/WorldDataLayers.cpp
@@ -233,24 +233,6 @@ void AWorldDataLayers::SetDataLayerRuntimeState(const UDataLayerInstance* InData
 
 	if (CurrentState != InState)
 	{
-		if (GetWorld()->IsGameWorld())
-		{
-			PRAGMA_DISABLE_DEPRECATION_WARNINGS
-			if (DataLayersFilterDelegate.IsBound())
-			{
-				const FString DataLayerShortName(InDataLayerInstance->GetDataLayerShortName());
-				if (!DataLayersFilterDelegate.Execute(FName(*DataLayerShortName), CurrentState, InState))
-				{
-					UE_LOG(LogWorldPartition, Log, TEXT("Data Layer '%s' state change was filtered out: %s -> %s"),
-						*DataLayerShortName,
-						*StaticEnum<EDataLayerRuntimeState>()->GetDisplayNameTextByValue((int64)CurrentState).ToString(),
-						*StaticEnum<EDataLayerRuntimeState>()->GetDisplayNameTextByValue((int64)InState).ToString());
-					return;
-				}
-			}
-			PRAGMA_ENABLE_DEPRECATION_WARNINGS
-		}
-
 		TargetLoadedDataLayerNames.Remove(InDataLayerInstance->GetDataLayerFName());
 		TargetActiveDataLayerNames.Remove(InDataLayerInstance->GetDataLayerFName());
 
diff --git a/Engine/Source/Runtime/Engine/Public/WorldPartition/DataLayer/WorldDataLayers.h b/Engine/Source/Runtime/Engine/Public/WorldPartition/DataLayer/WorldDataLayers.h
--- a/Engine/Source/Runtime/Engine/Public/WorldPartition/DataLayer/WorldDataLayers.h
+++ b/Engine/Source/Runtime/Engine/Public/WorldPartition/DataLayer/WorldDataLayers.h
@@ -342,12 +342,6 @@ private:
 	friend class UDataLayerManager;
 	friend class UExternalDataLayerManager;
 	friend class UDataLayerEditorSubsystem;
-
-public:
-	DECLARE_DELEGATE_RetVal_ThreeParams(bool, FDataLayersFilterDelegate, FName /*DataLayerName*/, EDataLayerRuntimeState /*CurrentState*/, EDataLayerRuntimeState /*TargetState*/);
-
-	UE_DEPRECATED(5.0, "do not use, will be replaced by another mechanism for initial release.")
-	FDataLayersFilterDelegate DataLayersFilterDelegate;
 };
 
 DEFINE_ACTORDESC_TYPE(AWorldDataLayers, FWorldDataLayersActorDesc);
