Fix SkeletalMeshes rendering as static not showing up in ray tracing when r.RayTracing.UseReferenceBasedResidency is enabled.
- corresponding ray tracing geometry was always evicted due to incorrect CurrentFirstLODIdx in FRayTracingGeometryManager.
- changes:
    - provide CurrentFirstLODIdx to RegisterRayTracingGeometryGroup(...) so it's immediately initialized
    - added logic in SkeletalMeshUpdate.cpp to update CurrentFirstLODIdx in FRayTracingGeometryManager when LODs are streamed in/out.

diff --git a/Engine/Source/Runtime/Engine/Internal/Streaming/SkeletalMeshUpdate.cpp b/Engine/Source/Runtime/Engine/Internal/Streaming/SkeletalMeshUpdate.cpp
--- a/Engine/Source/Runtime/Engine/Internal/Streaming/SkeletalMeshUpdate.cpp
+++ b/Engine/Source/Runtime/Engine/Internal/Streaming/SkeletalMeshUpdate.cpp
@@ -11,6 +11,7 @@ SkeletalMeshUpdate.cpp: Helpers to stream in and out skeletal mesh LODs.
 #include "Streaming/TextureStreamingHelpers.h"
 #include "Serialization/MemoryReader.h"
 #include "Rendering/SkeletalMeshRenderData.h"
+#include "Rendering/RayTracingGeometryManager.h"
 #include "Components/SkinnedMeshComponent.h"
 #include "Streaming/RenderAssetUpdate.inl"
 #include "RHIResourceReplace.h"
@@ -235,6 +236,14 @@ void FSkeletalMeshStreamIn::DoFinishUpdate(const FContext& Context)
 #endif
 
 	Context.RenderData->PendingFirstLODIdx = Context.RenderData->CurrentFirstLODIdx = ResourceState.LODCountToAssetFirstLODIdx(ResourceState.NumRequestedLODs);
+
+#if RHI_RAYTRACING
+	if (IsRayTracingAllowed() && Context.RenderData->bSupportRayTracing)
+	{
+		((FRayTracingGeometryManager*)GRayTracingGeometryManager)->SetRayTracingGeometryGroupCurrentFirstLODIndex(FRHICommandListImmediate::Get(), Context.RenderData->RayTracingGeometryGroupHandle, Context.RenderData->CurrentFirstLODIdx);
+	}
+#endif
+
 	MarkAsSuccessfullyFinished();
 }
 
@@ -384,6 +393,13 @@ void FSkeletalMeshStreamOut::ReleaseBuffers(const FContext& Context)
 			}
 		}
 
+#if RHI_RAYTRACING
+		if (IsRayTracingAllowed() && Context.RenderData->bSupportRayTracing)
+		{
+			((FRayTracingGeometryManager*)GRayTracingGeometryManager)->SetRayTracingGeometryGroupCurrentFirstLODIndex(FRHICommandListImmediate::Get(), Context.RenderData->RayTracingGeometryGroupHandle, Context.RenderData->CurrentFirstLODIdx);
+		}
+#endif
+
 		MarkAsSuccessfullyFinished();
 	}
 }
diff --git a/Engine/Source/Runtime/Engine/Private/Rendering/RayTracingGeometryManager.cpp b/Engine/Source/Runtime/Engine/Private/Rendering/RayTracingGeometryManager.cpp
--- a/Engine/Source/Runtime/Engine/Private/Rendering/RayTracingGeometryManager.cpp
+++ b/Engine/Source/Runtime/Engine/Private/Rendering/RayTracingGeometryManager.cpp
@@ -144,14 +144,14 @@ void FRayTracingGeometryManager::RemoveBuildRequest(BuildRequestIndex InRequestI
 	GeometryBuildRequests.RemoveAt(InRequestIndex);
 }
 
-RayTracing::GeometryGroupHandle FRayTracingGeometryManager::RegisterRayTracingGeometryGroup(uint32 NumLODs)
+RayTracing::GeometryGroupHandle FRayTracingGeometryManager::RegisterRayTracingGeometryGroup(uint32 NumLODs, uint32 CurrentFirstLODIdx)
 {
 	FScopeLock ScopeLock(&MainCS);
 
 	FRayTracingGeometryGroup Group;
 	Group.Geometries.AddDefaulted(NumLODs);
 	Group.NumReferences = 1;
-	Group.CurrentFirstLODIdx = NumLODs;
+	Group.CurrentFirstLODIdx = CurrentFirstLODIdx;
 
 	RayTracing::GeometryGroupHandle Handle = RegisteredGroups.Add(MoveTemp(Group));
 
@@ -284,7 +284,12 @@ void FRayTracingGeometryManager::SetRayTracingGeometryGroupCurrentFirstLODIndex(
 		FRHIResourceReplaceBatcher Batcher(RHICmdList, NewCurrentFirstLODIdx - Group.CurrentFirstLODIdx);
 		for (int32 LODIdx = Group.CurrentFirstLODIdx; LODIdx < NewCurrentFirstLODIdx; ++LODIdx)
 		{
-			Group.Geometries[LODIdx]->ReleaseRHIForStreaming(Batcher);
+			// some LODs might be stripped during cook
+			// skeletal meshes only create static LOD when rendering as static
+			if (Group.Geometries[LODIdx] != nullptr)
+			{
+				Group.Geometries[LODIdx]->ReleaseRHIForStreaming(Batcher);
+			}
 		}
 	}
 
diff --git a/Engine/Source/Runtime/Engine/Private/SkeletalMeshRenderData.cpp b/Engine/Source/Runtime/Engine/Private/SkeletalMeshRenderData.cpp
--- a/Engine/Source/Runtime/Engine/Private/SkeletalMeshRenderData.cpp
+++ b/Engine/Source/Runtime/Engine/Private/SkeletalMeshRenderData.cpp
@@ -625,7 +625,7 @@ void FSkeletalMeshRenderData::InitResources(bool bNeedsVertexColors, TArray<UMor
 			ENQUEUE_RENDER_COMMAND(SkeletalMeshInitRayTracingGeometryGroup)(UE::RenderCommandPipe::SkeletalMesh,
 				[this, NumLODs = LODRenderData.Num()]
 				{
-					RayTracingGeometryGroupHandle = GRayTracingGeometryManager->RegisterRayTracingGeometryGroup(NumLODs);
+					RayTracingGeometryGroupHandle = GRayTracingGeometryManager->RegisterRayTracingGeometryGroup(NumLODs, CurrentFirstLODIdx);
 				});
 		}
 #endif
diff --git a/Engine/Source/Runtime/Engine/Private/StaticMesh.cpp b/Engine/Source/Runtime/Engine/Private/StaticMesh.cpp
--- a/Engine/Source/Runtime/Engine/Private/StaticMesh.cpp
+++ b/Engine/Source/Runtime/Engine/Private/StaticMesh.cpp
@@ -2294,9 +2294,7 @@ void FStaticMeshRenderData::InitResources(ERHIFeatureLevel::Type InFeatureLevel,
 		ENQUEUE_RENDER_COMMAND(InitStaticMeshRayTracingGeometry)(
 			[this, Owner](FRHICommandListImmediate& RHICmdList)
 			{
-				RayTracingGeometryGroupHandle = GRayTracingGeometryManager->RegisterRayTracingGeometryGroup(LODResources.Num());
-
-				((FRayTracingGeometryManager*)GRayTracingGeometryManager)->SetRayTracingGeometryGroupCurrentFirstLODIndex(RHICmdList, RayTracingGeometryGroupHandle, RayTracingProxy->bUsingRenderingLODs ? CurrentFirstLODIdx : 0);
+				RayTracingGeometryGroupHandle = GRayTracingGeometryManager->RegisterRayTracingGeometryGroup(LODResources.Num(), RayTracingProxy->bUsingRenderingLODs ? CurrentFirstLODIdx : 0);
 				
 				FStaticMeshRayTracingProxyLODArray& RayTracingLODs = RayTracingProxy->LODs;
 
diff --git a/Engine/Source/Runtime/Engine/Public/Rendering/RayTracingGeometryManager.h b/Engine/Source/Runtime/Engine/Public/Rendering/RayTracingGeometryManager.h
--- a/Engine/Source/Runtime/Engine/Public/Rendering/RayTracingGeometryManager.h
+++ b/Engine/Source/Runtime/Engine/Public/Rendering/RayTracingGeometryManager.h
@@ -30,7 +30,7 @@ public:
 	ENGINE_API virtual RayTracingGeometryHandle RegisterRayTracingGeometry(FRayTracingGeometry* InGeometry) override;
 	ENGINE_API virtual void ReleaseRayTracingGeometryHandle(RayTracingGeometryHandle Handle) override;
 
-	ENGINE_API virtual RayTracing::GeometryGroupHandle RegisterRayTracingGeometryGroup(uint32 NumLODs) override;
+	ENGINE_API virtual RayTracing::GeometryGroupHandle RegisterRayTracingGeometryGroup(uint32 NumLODs, uint32 CurrentFirstLODIdx = 0) override;
 	ENGINE_API virtual void ReleaseRayTracingGeometryGroup(RayTracing::GeometryGroupHandle Handle) override;
 
 	ENGINE_API virtual void RefreshRegisteredGeometry(RayTracingGeometryHandle Handle) override;
diff --git a/Engine/Source/Runtime/RenderCore/Public/RayTracingGeometryManagerInterface.h b/Engine/Source/Runtime/RenderCore/Public/RayTracingGeometryManagerInterface.h
--- a/Engine/Source/Runtime/RenderCore/Public/RayTracingGeometryManagerInterface.h
+++ b/Engine/Source/Runtime/RenderCore/Public/RayTracingGeometryManagerInterface.h
@@ -48,7 +48,7 @@ public:
 	* For example, the FRayTracingGeometry of all the LODs of UStaticMesh should use the same RayTracing::GeometryGroupHandle.
 	* This grouping is useful to keep track which proxies need to be invalidated when a FRayTracingGeometry is built or made resident.
 	*/
-	RENDERCORE_API virtual RayTracing::GeometryGroupHandle RegisterRayTracingGeometryGroup(uint32 NumLODs) = 0;
+	RENDERCORE_API virtual RayTracing::GeometryGroupHandle RegisterRayTracingGeometryGroup(uint32 NumLODs, uint32 CurrentFirstLODIdx = 0) = 0;
 	RENDERCORE_API virtual void ReleaseRayTracingGeometryGroup(RayTracing::GeometryGroupHandle Handle) = 0;
 
 	RENDERCORE_API virtual void RequestUpdateCachedRenderState(RayTracing::GeometryGroupHandle InRayTracingGeometryGroupHandle) = 0;
