Temporary fix to prevent clients from disconnecting when they attempt to update visibility on a level the server doesn't have

#sh-06-13-desert-biome-server-kicks-ext
#testing Tested locally on cooked windows client/server that pc players do not disconnect when entereing desert biome

diff --git a/Engine/Source/Runtime/Engine/Private/NetConnection.cpp b/Engine/Source/Runtime/Engine/Private/NetConnection.cpp
--- a/Engine/Source/Runtime/Engine/Private/NetConnection.cpp
+++ b/Engine/Source/Runtime/Engine/Private/NetConnection.cpp
@@ -112,6 +112,8 @@ static FAutoConsoleVariableRef CVarCloseTimingDebug(TEXT("net.CloseTimingDebug")
 static bool GCVarLogAllConnectionCleanup = false;
 static FAutoConsoleVariableRef CVarLogAllConnectionCleanup(TEXT("net.Connection.LogAllCleanup"), GCVarLogAllConnectionCleanup, TEXT("When true log every connection CleanUp even when it was a normal socket closure."));
 
+static TAutoConsoleVariable<bool> CVarSkipMissingLevelDisconnect( TEXT("net.SkipMissingLevelDisconnect"), false, TEXT("If true skip disconnecting a player if they have a level that doesn't exist on server"));
+
 extern int32 GNetDormancyValidate;
 extern bool GbNetReuseReplicatorsForDormantObjects;
 
@@ -1729,7 +1731,8 @@ void UNetConnection::UpdateLevelVisibilityInternal(const FUpdateLevelVisibilityL
 			UE_LOG(LogPlayerController, Warning, TEXT("ServerUpdateLevelVisibility() ignored non-existant package. PackageName='%s', FileName='%s'"),
 					ToCStr(PackageNameStr), ToCStr(FileNameStr));
 
-			if (!LevelVisibility.bSkipCloseOnError)
+			// Added CVAR to address #sh-06-13-desert-biome-server-kicks-ext. Should be removed once a proper solution is found
+			if (!CVarSkipMissingLevelDisconnect.GetValueOnAnyThread() && !LevelVisibility.bSkipCloseOnError)
 			{
 				TStringBuilder<1024> VisibilityParms;
 
