Fixed edge case in UWorldPartitionRuntimeHashSet where nothing was loaded if the only spatially loaded  content of the map was coming from external streaming objects (EDL or ContentBundles).

diff --git a/Engine/Source/Runtime/Engine/Private/WorldPartition/RuntimeHashSet/WorldPartitionRuntimeHashSet.cpp b/Engine/Source/Runtime/Engine/Private/WorldPartition/RuntimeHashSet/WorldPartitionRuntimeHashSet.cpp
--- a/Engine/Source/Runtime/Engine/Private/WorldPartition/RuntimeHashSet/WorldPartitionRuntimeHashSet.cpp
+++ b/Engine/Source/Runtime/Engine/Private/WorldPartition/RuntimeHashSet/WorldPartitionRuntimeHashSet.cpp
@@ -1083,30 +1083,30 @@ void UWorldPartitionRuntimeHashSet::ForEachStreamingData(TFunctionRef<bool(const
 
 void UWorldPartitionRuntimeHashSet::UpdateRuntimeDataGridMap()
 {
-	if (!RuntimeStreamingData.IsEmpty())
-	{
-		RuntimeSpatiallyLoadedDataGridMap.Reset();
-		RuntimeNonSpatiallyLoadedDataGridList.Reset();
+	RuntimeSpatiallyLoadedDataGridMap.Reset();
+	RuntimeNonSpatiallyLoadedDataGridList.Reset();
 
-		ForEachStreamingData([this](const FRuntimePartitionStreamingData& StreamingData)
+	ForEachStreamingData([this](const FRuntimePartitionStreamingData& StreamingData)
+	{
+		if (StreamingData.SpatiallyLoadedCells.Num())
 		{
-			if (StreamingData.SpatiallyLoadedCells.Num())
-			{
-				TArray<const FRuntimePartitionStreamingData*>& StreamingDataList = RuntimeSpatiallyLoadedDataGridMap.FindOrAdd(StreamingData.Name);
-				StreamingDataList.Add(&StreamingData);
-			}
-
-			if (StreamingData.NonSpatiallyLoadedCells.Num())
-			{
-				RuntimeNonSpatiallyLoadedDataGridList.Add(&StreamingData);
-			}
+			TArray<const FRuntimePartitionStreamingData*>& StreamingDataList = RuntimeSpatiallyLoadedDataGridMap.FindOrAdd(StreamingData.Name);
+			StreamingDataList.Add(&StreamingData);
+		}
 
-			return true;
-		});
-	}
-	else
+		if (StreamingData.NonSpatiallyLoadedCells.Num())
+		{
+			RuntimeNonSpatiallyLoadedDataGridList.Add(&StreamingData);
+		}
+		return true;
+	});
+	
+	if (RuntimeSpatiallyLoadedDataGridMap.IsEmpty())
 	{
 		RuntimeSpatiallyLoadedDataGridMap.Empty();
+	}
+	if (RuntimeNonSpatiallyLoadedDataGridList.IsEmpty())
+	{
 		RuntimeNonSpatiallyLoadedDataGridList.Empty();
 	}
 }
