ODSC support for targetting single material with ODSC.recompilemode = 1
This works like in editor, recompileshaders material [material_name]

This allows to only recompile the permutations needed by rendering for a given material, making iteration faster for changes impacting multiple materials

#rnx

diff --git a/Engine/Source/Runtime/Engine/Private/ODSC/ODSCManager.cpp b/Engine/Source/Runtime/Engine/Private/ODSC/ODSCManager.cpp
--- a/Engine/Source/Runtime/Engine/Private/ODSC/ODSCManager.cpp
+++ b/Engine/Source/Runtime/Engine/Private/ODSC/ODSCManager.cpp
@@ -153,7 +153,8 @@ void FODSCManager::AddThreadedRequest(
 	EShaderPlatform ShaderPlatform,
 	ERHIFeatureLevel::Type FeatureLevel,
 	EMaterialQualityLevel::Type QualityLevel,
-	ODSCRecompileCommand RecompileCommandType
+	ODSCRecompileCommand RecompileCommandType,
+	const FString& RequestedMaterialName
 )
 {
 	if (IsHandlingRequests())
@@ -164,6 +165,9 @@ void FODSCManager::AddThreadedRequest(
 			&& (CVarODSCRecompileMode.GetValueOnAnyThread() > 0))
 		{
 			Thread->ResetMaterialsODSCData(FeatureLevel);
+
+			// Rendering commands got flushed by ResetMaterialsODSCData
+			MaterialNameToRecompile = FName(*RequestedMaterialName);	
 			
 			// when we ask for "changed", we want both materials and global
 			if (RecompileCommandType == ODSCRecompileCommand::Changed)
@@ -265,6 +269,33 @@ bool FODSCManager::ShouldForceRecompileInternal(const FMaterialShaderMap* Materi
 		return false;
 	}
 
+	if (!MaterialNameToRecompile.IsNone())
+	{
+		EODSCMetaDataType ODSCMetaData = (EODSCMetaDataType)Material->GetODSCMetaData();
+
+		// No locking on the material ODSC metadata: the dependency chain/material names are not supposed to be changing on a per frame basis
+		if (ODSCMetaData == EODSCMetaDataType::Default)
+		{
+			UMaterialInterface* EngineMaterialInterface = Material->GetMaterialInterface();
+			TSet<UMaterialInterface*> MaterialDependencies;
+			EngineMaterialInterface->GetDependencies(MaterialDependencies);
+
+			bool bHasDependencies = false;
+			for (auto Iter : MaterialDependencies)
+			{
+				UMaterialInterface* MIDep = Iter;
+				if (MIDep && MIDep->GetFName() == MaterialNameToRecompile)
+				{
+					bHasDependencies = true;
+				}
+			}
+
+			ODSCMetaData = (bHasDependencies) ? EODSCMetaDataType::IsDependentOnMaterialName : EODSCMetaDataType::IsNotDependentOnMaterialName;
+			Material->SetODSCMetaData((uint8)ODSCMetaData);
+		}
+		return ODSCMetaData == EODSCMetaDataType::IsDependentOnMaterialName;
+	}
+
 	return true;
 #else
 	return false;
diff --git a/Engine/Source/Runtime/Engine/Private/ODSC/ODSCThread.cpp b/Engine/Source/Runtime/Engine/Private/ODSC/ODSCThread.cpp
--- a/Engine/Source/Runtime/Engine/Private/ODSC/ODSCThread.cpp
+++ b/Engine/Source/Runtime/Engine/Private/ODSC/ODSCThread.cpp
@@ -183,6 +183,9 @@ void FODSCThread::ResetMaterialsODSCData(ERHIFeatureLevel::Type FeatureLevel)
 				if (MaterialResource && MaterialResource->GetGameThreadShaderMap())
 				{
 					MaterialResource->GetGameThreadShaderMap()->SetIsFromODSC(false);
+#if WITH_ODSC
+					MaterialResource->SetODSCMetaData((uint8)0);
+#endif
 				}
 				UpdateContext.AddMaterialInterface(Material);
 			}
diff --git a/Engine/Source/Runtime/Engine/Private/ODSC/ODSCThread.h b/Engine/Source/Runtime/Engine/Private/ODSC/ODSCThread.h
--- a/Engine/Source/Runtime/Engine/Private/ODSC/ODSCThread.h
+++ b/Engine/Source/Runtime/Engine/Private/ODSC/ODSCThread.h
@@ -23,6 +23,13 @@ namespace UE
 	}
 }
 
+enum class EODSCMetaDataType
+{
+	Default=0, // The material hasn't been seen by ODSCManager yet
+	IsDependentOnMaterialName,
+	IsNotDependentOnMaterialName,
+};
+
 class FODSCMessageHandler : public IPlatformFile::IFileServerMessageHandler
 {
 public:
diff --git a/Engine/Source/Runtime/Engine/Private/ShaderCompiler/ShaderCompiler.cpp b/Engine/Source/Runtime/Engine/Private/ShaderCompiler/ShaderCompiler.cpp
--- a/Engine/Source/Runtime/Engine/Private/ShaderCompiler/ShaderCompiler.cpp
+++ b/Engine/Source/Runtime/Engine/Private/ShaderCompiler/ShaderCompiler.cpp
@@ -8649,7 +8649,7 @@ namespace
 		}
 	}
 
-	ODSCRecompileCommand ParseRecompileCommandString(const TCHAR* CmdString, TArray<FString>& OutMaterialsToLoad, FString& OutShaderTypesToLoad)
+	ODSCRecompileCommand ParseRecompileCommandString(const TCHAR* CmdString, TArray<FString>& OutMaterialsToLoad, FString& OutShaderTypesToLoad, FString& OutRequestedMaterialName)
 	{
 		FString CmdName = FParse::Token(CmdString, 0);
 
@@ -8662,6 +8662,7 @@ namespace
 
 			// tell other side the material to load, by pathname
 			FString RequestedMaterialName( FParse::Token( CmdString, 0 ) );
+			OutRequestedMaterialName = RequestedMaterialName;
 			UMaterialInterface* MatchingMaterial = nullptr;
 			for (TObjectIterator<UMaterialInterface> It; It; ++It)
 			{
@@ -8875,11 +8876,12 @@ bool RecompileShaders(const TCHAR* Cmd, FOutputDevice& Ar)
 #if WITH_ODSC
 		TArray<FString> MaterialsToLoad;
 		FString ShaderTypesToLoad;
-		ODSCRecompileCommand CommandType = ParseRecompileCommandString(Cmd, MaterialsToLoad, ShaderTypesToLoad);
+		FString RequestedMaterialName;
+		ODSCRecompileCommand CommandType = ParseRecompileCommandString(Cmd, MaterialsToLoad, ShaderTypesToLoad, RequestedMaterialName);
 
 		ERHIFeatureLevel::Type TargetFeatureLevel = GetMaxSupportedFeatureLevel(GMaxRHIShaderPlatform);
 		const EMaterialQualityLevel::Type ActiveQualityLevel = GetCachedScalabilityCVars().MaterialQualityLevel;
-		GODSCManager->AddThreadedRequest(MaterialsToLoad, ShaderTypesToLoad, GMaxRHIShaderPlatform, TargetFeatureLevel, ActiveQualityLevel, CommandType);
+		GODSCManager->AddThreadedRequest(MaterialsToLoad, ShaderTypesToLoad, GMaxRHIShaderPlatform, TargetFeatureLevel, ActiveQualityLevel, CommandType, RequestedMaterialName);
 #endif
 		return true;
 	}
diff --git a/Engine/Source/Runtime/Engine/Public/MaterialShared.h b/Engine/Source/Runtime/Engine/Public/MaterialShared.h
--- a/Engine/Source/Runtime/Engine/Public/MaterialShared.h
+++ b/Engine/Source/Runtime/Engine/Public/MaterialShared.h
@@ -2564,6 +2564,11 @@ public:
 	virtual void GatherExpressionsForCustomInterpolators(TArray<class UMaterialExpression*>& OutExpressions) const {}
 #endif // WITH_EDITORONLY_DATA
 
+#if WITH_ODSC
+	uint8 GetODSCMetaData() const { return ODSCMetaData.load(std::memory_order_relaxed); };
+	void SetODSCMetaData(uint8 InODSCMetaData) const { ODSCMetaData.store(InODSCMetaData, std::memory_order_relaxed); };
+#endif
+
 protected:
 	// shared code needed for GetUniformScalarParameterExpressions, GetUniformVectorParameterExpressions, GetUniformCubeTextureExpressions..
 	// @return can be 0
@@ -2711,6 +2716,10 @@ private:
 	/** Quality level that this material is representing, may be EMaterialQualityLevel::Num if material doesn't depend on current quality level */
 	EMaterialQualityLevel::Type QualityLevel;
 
+#if WITH_ODSC
+	mutable std::atomic<uint8> ODSCMetaData = {};
+#endif
+
 	/** Feature level that this material is representing. */
 	ERHIFeatureLevel::Type FeatureLevel;
 
diff --git a/Engine/Source/Runtime/Engine/Public/ODSC/ODSCManager.h b/Engine/Source/Runtime/Engine/Public/ODSC/ODSCManager.h
--- a/Engine/Source/Runtime/Engine/Public/ODSC/ODSCManager.h
+++ b/Engine/Source/Runtime/Engine/Public/ODSC/ODSCManager.h
@@ -53,7 +53,7 @@ public:
 	 *
 	 * @return false if no longer needs ticking
 	 */
-	ENGINE_API void AddThreadedRequest(const TArray<FString>& MaterialsToCompile, const FString& ShaderTypesToLoad, EShaderPlatform ShaderPlatform, ERHIFeatureLevel::Type FeatureLevel, EMaterialQualityLevel::Type QualityLevel, ODSCRecompileCommand RecompileCommandType);
+	ENGINE_API void AddThreadedRequest(const TArray<FString>& MaterialsToCompile, const FString& ShaderTypesToLoad, EShaderPlatform ShaderPlatform, ERHIFeatureLevel::Type FeatureLevel, EMaterialQualityLevel::Type QualityLevel, ODSCRecompileCommand RecompileCommandType, const FString& RequestedMaterialName = FString());
 
 	/**
 	 * Add a request to compile a pipeline of shaders.  The results are submitted and processed in an async manner.
@@ -129,6 +129,8 @@ private:
 
 	FCriticalSection ErrorMessageCS;
 	FString ErrorMessage;
+
+	FName MaterialNameToRecompile;
 };
 
 struct FODSCSuspendForceRecompileScope
