Fixes LandscapeStreamingProxy EditorBounds in ActorDescs.
Resaves OpenWorld template to update streaming proxies EditorBounds.
#virtualized

diff --git a/Engine/Source/Runtime/Engine/Private/WorldPartition/Landscape/LandscapeActorDesc.cpp b/Engine/Source/Runtime/Engine/Private/WorldPartition/Landscape/LandscapeActorDesc.cpp
--- a/Engine/Source/Runtime/Engine/Private/WorldPartition/Landscape/LandscapeActorDesc.cpp
+++ b/Engine/Source/Runtime/Engine/Private/WorldPartition/Landscape/LandscapeActorDesc.cpp
@@ -24,6 +24,8 @@ void FLandscapeActorDesc::Init(const AActor* InActor)
 		{
 			LandscapeActorGuid = LandscapeActor->GetActorGuid();
 		}
+
+		EditorBounds = InActor->GetStreamingBoundsEditor();
 	}
 }
 
