Actor Label: Fix discrepancies around WITH_EDITOR usage when interacting with ActorLabel. Add GetActorLabelView, and update GetActorNameOrLabel to return a non-const FString.
- WITH_EDITOR and WITH_EDITORONLY_DATA mismatches were fixed up in CL 24269601, but certain uses of WITH_EDITOR still remain. Fixed this up to simplify things and ensure consistency.
- GetActorLabelView allows accessing ActorLabel without producing a copy, and returns an empty view in build configurations where this member doesn't exist.
- GetActorNameOrLabel doesn't need to return a const FString since we're returning by value. Returning by const value can prevent it from binding to functions accepting rvalue references (e.g. move assignment). See https://gcc.godbolt.org/z/Kjf1KsaTf

diff --git a/Engine/Source/Runtime/Engine/Classes/GameFramework/Actor.h b/Engine/Source/Runtime/Engine/Classes/GameFramework/Actor.h
--- a/Engine/Source/Runtime/Engine/Classes/GameFramework/Actor.h
+++ b/Engine/Source/Runtime/Engine/Classes/GameFramework/Actor.h
@@ -1146,9 +1146,9 @@ private:
 #endif
 
 public:
-	const FString GetActorNameOrLabel() const
+	FString GetActorNameOrLabel() const
 	{
-#if WITH_EDITORONLY_DATA || (!WITH_EDITOR && ACTOR_HAS_LABELS)
+#if WITH_EDITORONLY_DATA || ACTOR_HAS_LABELS
 		if (!ActorLabel.IsEmpty())
 		{
 			return ActorLabel;
@@ -1157,6 +1157,15 @@ public:
 		return GetName();
 	}
 
+	FStringView GetActorLabelView() const
+	{
+#if WITH_EDITORONLY_DATA || ACTOR_HAS_LABELS
+		return ActorLabel;
+#else
+		return FStringView();
+#endif
+	}
+
 #if WITH_EDITORONLY_DATA
 private:
 	/** 
diff --git a/Engine/Source/Runtime/Engine/Private/Actor.cpp b/Engine/Source/Runtime/Engine/Private/Actor.cpp
--- a/Engine/Source/Runtime/Engine/Private/Actor.cpp
+++ b/Engine/Source/Runtime/Engine/Private/Actor.cpp
@@ -884,7 +884,7 @@ void AActor::Serialize(FArchive& Ar)
 
 		if (bIsCooked)
 		{
-#if !(WITH_EDITORONLY_DATA || (!WITH_EDITOR && ACTOR_HAS_LABELS))
+#if !(WITH_EDITORONLY_DATA || ACTOR_HAS_LABELS)
 			// In non-development builds, just skip over the actor labels. We need to figure out a way to either strip that data from shipping builds,
 			// or skip over the string data without doing any memory allocations, probably with a custom FString::SerializeToNull implementation.
 			FString ActorLabel;
