[World Partition] Fix FWorldPartitionExternalDirtyActorsTracker holding duplicate entries for non dirty actors resulting in bad pinning logic.

diff --git a/Engine/Source/Runtime/Engine/Private/WorldPartition/WorldPartition.cpp b/Engine/Source/Runtime/Engine/Private/WorldPartition/WorldPartition.cpp
--- a/Engine/Source/Runtime/Engine/Private/WorldPartition/WorldPartition.cpp
+++ b/Engine/Source/Runtime/Engine/Private/WorldPartition/WorldPartition.cpp
@@ -270,7 +270,7 @@ UWorldPartition::FWorldPartitionExternalDirtyActorsTracker::FWorldPartitionExter
 void UWorldPartition::FWorldPartitionExternalDirtyActorsTracker::OnRemoveNonDirtyActor(TWeakObjectPtr<AActor> InActor, FWorldPartitionReference& InValue)
 {
 	check(InActor.IsValid());
-	NonDirtyActors.Emplace(InActor, InValue);
+	NonDirtyActors.Emplace({ InActor, InValue });
 }
 
 void UWorldPartition::FWorldPartitionExternalDirtyActorsTracker::Tick(float InDeltaSeconds)
diff --git a/Engine/Source/Runtime/Engine/Public/WorldPartition/WorldPartition.h b/Engine/Source/Runtime/Engine/Public/WorldPartition/WorldPartition.h
--- a/Engine/Source/Runtime/Engine/Public/WorldPartition/WorldPartition.h
+++ b/Engine/Source/Runtime/Engine/Public/WorldPartition/WorldPartition.h
@@ -650,7 +650,7 @@ private:
 		//~ End TExternalDirtyActorsTracker interface
 
 	private:
-		TArray<TPair<TWeakObjectPtr<AActor>, FWorldPartitionReference>> NonDirtyActors;
+		TSet<TPair<TWeakObjectPtr<AActor>, FWorldPartitionReference>> NonDirtyActors;
 	};
 
 	TUniquePtr<FWorldPartitionExternalDirtyActorsTracker> ExternalDirtyActorsTracker;
