Add a new gracefully-closing state to NetConnection that waits for reliable data to be acknowledged before actually closing communication, and support this in Iris.
This helps make sure reliable client RPCs sent right before the server gracefully disconnects a client by destroying their PlayerController are received, especially under packet loss.
Can be disabled by setting the cvar net.GracefulCloseEnabled.

diff --git a/Engine/Plugins/Runtime/ReplicationGraph/Source/Private/ReplicationGraph.cpp b/Engine/Plugins/Runtime/ReplicationGraph/Source/Private/ReplicationGraph.cpp
--- a/Engine/Plugins/Runtime/ReplicationGraph/Source/Private/ReplicationGraph.cpp
+++ b/Engine/Plugins/Runtime/ReplicationGraph/Source/Private/ReplicationGraph.cpp
@@ -2431,6 +2431,11 @@ bool UReplicationGraph::ProcessRemoteFunction(class AActor* Actor, UFunction* Fu
 			FConnectionReplicationActorInfo& ConnectionActorInfo = Manager->ActorInfoMap.FindOrAdd(Actor);
 			UNetConnection* NetConnection = Manager->NetConnection;
 
+			if (NetConnection->IsClosingOrClosed())
+			{
+				return true;
+			}
+
 			// This connection isn't ready yet
 			if (NetConnection->ViewTarget == nullptr)
 			{
@@ -2562,7 +2567,7 @@ bool UReplicationGraph::ProcessRemoteFunction(class AActor* Actor, UFunction* Fu
 			Connection = ((UChildConnection*)Connection)->Parent;
 		}
 	
-		if (Connection->GetConnectionState() == USOCK_Closed)
+		if (Connection->IsClosingOrClosed())
 		{
 			return true;
 		}
@@ -2887,8 +2892,8 @@ bool UNetReplicationGraphConnection::PrepareForReplication()
 	
 	UWorld* CurrentWorld = GetWorld();
 	UPackage* CurrentWorldPackage = CurrentWorld ? CurrentWorld->GetPackage() : nullptr;
-	bool bConnectionHasCorrectWorld = CurrentWorldPackage ? NetConnection->GetClientWorldPackageName() == CurrentWorldPackage->GetFName() : true;
-	
+	const bool bConnectionHasCorrectWorld = CurrentWorldPackage ? NetConnection->GetClientWorldPackageName() == CurrentWorldPackage->GetFName() : true;
+
 	// Set any children viewtargets
 	for (int32 i = 0; i < NetConnection->Children.Num(); ++i)
 	{
@@ -2923,7 +2928,7 @@ bool UNetReplicationGraphConnection::PrepareForReplication()
 		}
 	}
 
-	return (NetConnection->GetConnectionState() != USOCK_Closed) && (NetConnection->ViewTarget != nullptr) && bConnectionHasCorrectWorld;
+	return !NetConnection->IsClosingOrClosed() && (NetConnection->ViewTarget != nullptr) && bConnectionHasCorrectWorld;
 }
 
 void UNetReplicationGraphConnection::BuildVisibleLevels()
diff --git a/Engine/Plugins/Runtime/ReplicationSystemTestPlugin/Source/Private/Tests/DataStream/MockDataStream.cpp b/Engine/Plugins/Runtime/ReplicationSystemTestPlugin/Source/Private/Tests/DataStream/MockDataStream.cpp
--- a/Engine/Plugins/Runtime/ReplicationSystemTestPlugin/Source/Private/Tests/DataStream/MockDataStream.cpp
+++ b/Engine/Plugins/Runtime/ReplicationSystemTestPlugin/Source/Private/Tests/DataStream/MockDataStream.cpp
@@ -53,3 +53,8 @@ void UMockDataStream::ProcessPacketDeliveryStatus(UE::Net::EPacketDeliveryStatus
 	Records.Pop();
 	CallStatus.ProcessPacketDeliveryStatusMagicValue = Record->MagicValue;
 }
+
+bool UMockDataStream::HasAcknowledgedAllReliableData() const
+{
+	return true;
+}
diff --git a/Engine/Plugins/Runtime/ReplicationSystemTestPlugin/Source/Private/Tests/DataStream/MockDataStream.h b/Engine/Plugins/Runtime/ReplicationSystemTestPlugin/Source/Private/Tests/DataStream/MockDataStream.h
--- a/Engine/Plugins/Runtime/ReplicationSystemTestPlugin/Source/Private/Tests/DataStream/MockDataStream.h
+++ b/Engine/Plugins/Runtime/ReplicationSystemTestPlugin/Source/Private/Tests/DataStream/MockDataStream.h
@@ -39,6 +39,7 @@ protected:
 	virtual EWriteResult WriteData(UE::Net::FNetSerializationContext& context, FDataStreamRecord const*& OutRecord) override;
 	virtual void ReadData(UE::Net::FNetSerializationContext& context) override;
 	virtual void ProcessPacketDeliveryStatus(UE::Net::EPacketDeliveryStatus Status, FDataStreamRecord const* Record) override;
+	virtual bool HasAcknowledgedAllReliableData() const override;
 
 private:
 	struct FRecord : public FDataStreamRecord
diff --git a/Engine/Source/Runtime/Engine/Classes/Engine/Channel.h b/Engine/Source/Runtime/Engine/Classes/Engine/Channel.h
--- a/Engine/Source/Runtime/Engine/Classes/Engine/Channel.h
+++ b/Engine/Source/Runtime/Engine/Classes/Engine/Channel.h
@@ -168,6 +168,9 @@ public:
 	/* Notification that this channel has been placed in a channel pool and needs to reset to its original state so it can be used again like a new channel */
 	ENGINE_API virtual void AddedToChannelPool();
 
+	/** Returns true if this channel has flushed and acknowledged pending reliable data during a graceful close (when net.EnableGracefulClose is true) */
+	ENGINE_API virtual bool HasAcknowledgedAllReliableData() const;
+
 protected:
 
 	/** Closes the actor channel but with a 'dormant' flag set so it can be reopened */
diff --git a/Engine/Source/Runtime/Engine/Classes/Engine/NetConnection.h b/Engine/Source/Runtime/Engine/Classes/Engine/NetConnection.h
--- a/Engine/Source/Runtime/Engine/Classes/Engine/NetConnection.h
+++ b/Engine/Source/Runtime/Engine/Classes/Engine/NetConnection.h
@@ -91,6 +91,7 @@ enum EConnectionState
 	USOCK_Closed    = 1, // Connection permanently closed.
 	USOCK_Pending	= 2, // Connection is awaiting connection.
 	USOCK_Open      = 3, // Connection is open.
+	USOCK_Closing   = 4, // Connection is closing and waiting for all reliable data to be acked. No new data will be sent.
 };
 ENGINE_API const TCHAR* LexToString(const EConnectionState Value);
 
@@ -406,6 +407,9 @@ public:
 	ENGINE_API const EConnectionState GetConnectionState() const;
 	ENGINE_API void SetConnectionState(EConnectionState ConnectionState);
 	
+	/** Returns true if this connection's state is USOCK_Closing or USOCK_Closed. No new data should be sent in this case. */
+	ENGINE_API bool IsClosingOrClosed() const;
+
 	uint32 bPendingDestroy:1;    // when true, playercontroller or beaconclient is being destroyed
 
 
@@ -539,6 +543,9 @@ private:
 	/** total packets received on this connection, including PacketHandler */
 	int32 InTotalHandlerPackets;
 
+	/** Driver ElapsedTime at which a graceful close will time out and terminate. */
+	double GracefulCloseTimeoutDeadline = 0.0;
+
 public:
 	int32 GetInTotalHandlerPackets() const
 	{
@@ -1038,6 +1045,25 @@ public:
 	ENGINE_API void Close(FNetResult&& CloseReason);
 
 
+	/**
+	 * Starts a graceful close. If net.EnableGracefulClose is true, waits longer for reliable RPCs to be acknowledged before fully cleaning up the connection.
+	 * If net.EnableGracefulClose is false, immediately sends the control channel close bunch. Reliable RPCs may be lost.
+	 *
+	 * @param CloseReason	Specifies the reason for the Close
+	 */
+	ENGINE_API void GracefulClose(FNetResult&& CloseReason);
+
+	/**
+	 * Starts a graceful close. If net.EnableGracefulClose is true, waits longer for reliable RPCs to be acknowledged before fully cleaning up the connection.
+	 * If net.EnableGracefulClose is false, immediately sends the control channel close bunch. Reliable RPCs may be lost.
+	 *
+	 * @param CloseReason	Specifies the reason for the Close
+	 */
+	void GracefulClose(FNetCloseResult&& CloseReason)
+	{
+		GracefulClose(static_cast<FNetResult&&>(MoveTemp(CloseReason)));
+	}
+
 	/** closes the control channel, cleans up structures, and prepares for deletion */
 	ENGINE_API virtual void CleanUp();
 
@@ -1760,6 +1786,9 @@ private:
 	/** Whether or not this NetConnection has already received an NMT_CloseReason message */
 	bool bReceivedCloseReason = false;
 
+	/** Stored close reason from a GracefulClose call, to be used in the actual Close */
+	TPimplPtr<UE::Net::FNetResult, EPimplPtrMode::DeepCopy> PendingGracefulCloseResult;
+
 	/** Ping collection and calculation */
 	TPimplPtr<UE::Net::FNetPing> NetPing;
 
@@ -1833,6 +1862,11 @@ private:
 	 */
 	void HandleNetResultOrClose(ENetCloseResult InResult);
 
+	/**
+	 * If this connection is in the USOCK_Closing state, check whether all reliables have been acknowledged, if so, actually Close (send control channel close bunch).
+	 */
+	void TryClosePendingGracefulClose();
+
 protected:
 	TOptional<FNetworkCongestionControl> NetworkCongestionControl;
 
@@ -1898,24 +1932,29 @@ struct FNetConnectionSettings
 	{
 #if DO_ENABLE_NET_TEST
 		PacketLag = InConnection->PacketSimulationSettings.PktLag;
+		PacketLoss = InConnection->PacketSimulationSettings.PktLoss;
 #else
 		PacketLag = 0;
+		PacketLoss = 0;
 #endif
 	}
 
 	FNetConnectionSettings( int32 InPacketLag )
 	{
 		PacketLag = InPacketLag;
+		PacketLoss = 0;
 	}
 
 	void ApplyTo(UNetConnection* Connection)
 	{
 #if DO_ENABLE_NET_TEST
 		Connection->PacketSimulationSettings.PktLag = PacketLag;
+		Connection->PacketSimulationSettings.PktLoss = PacketLoss;
 #endif
 	}
 
 	int32 PacketLag;
+	int32 PacketLoss;
 };
 
 /** Allows you to temporarily set connection settings within a scape. This will also force flush the connection before/after.
@@ -1929,19 +1968,20 @@ struct FScopedNetConnectionSettings
 		if (ShouldApply)
 		{
 			Connection->FlushNet();
-			NewSettings.ApplyTo(Connection);
+			NewSettings.ApplyTo(Connection.Get());
 		}
 	}
 	~FScopedNetConnectionSettings()
 	{
-		if (ShouldApply)
+		UNetConnection* LocalConnection = Connection.Get();
+		if (LocalConnection && ShouldApply)
 		{
-			Connection->FlushNet();
-			OldSettings.ApplyTo(Connection);
+			LocalConnection->FlushNet();
+			OldSettings.ApplyTo(LocalConnection);
 		}
 	}
 
-	UNetConnection * Connection;
+	TWeakObjectPtr<UNetConnection> Connection;
 	FNetConnectionSettings OldSettings;
 	bool ShouldApply;
 };
diff --git a/Engine/Source/Runtime/Engine/Classes/Engine/NetDriver.h b/Engine/Source/Runtime/Engine/Classes/Engine/NetDriver.h
--- a/Engine/Source/Runtime/Engine/Classes/Engine/NetDriver.h
+++ b/Engine/Source/Runtime/Engine/Classes/Engine/NetDriver.h
@@ -910,6 +910,13 @@ public:
 	UPROPERTY(Config)
 	float ConnectionTimeout;
 
+	/** 
+	 * Amount of time to wait for a graceful close/bPendingDestroy to complete before considering the connection timed out.  
+	 * This is the time used to allow any existing, pending reliable data to be acknowledged.
+	 */
+	UPROPERTY(Config)
+	float GracefulCloseConnectionTimeout = 2.0f;
+
 	/**
 	* A multiplier that is applied to the above values when we are running with unoptimized builds (debug)
 	* or data (uncooked). This allows us to retain normal timeout behavior while debugging without resorting
diff --git a/Engine/Source/Runtime/Engine/Private/DataChannel.cpp b/Engine/Source/Runtime/Engine/Private/DataChannel.cpp
--- a/Engine/Source/Runtime/Engine/Private/DataChannel.cpp
+++ b/Engine/Source/Runtime/Engine/Private/DataChannel.cpp
@@ -302,7 +302,7 @@ int64 UChannel::Close(EChannelCloseReason Reason)
 
 	int64 NumBits = 0;
 
-	if ( !Closing && ( Connection->GetConnectionState() == USOCK_Open || Connection->GetConnectionState() == USOCK_Pending ) && !SentClosingBunch)
+	if ( !Closing && ( Connection->GetConnectionState() == USOCK_Open || Connection->GetConnectionState() == USOCK_Pending || (Connection->GetConnectionState() == USOCK_Closing && ChIndex == 0)) && !SentClosingBunch)
 	{
 		if ( ChIndex == 0 )
 		{
@@ -1619,6 +1619,11 @@ void UChannel::AddedToChannelPool()
 	NumOutRec = 0;
 }
 
+bool UChannel::HasAcknowledgedAllReliableData() const
+{
+	return NumOutRec == 0;
+}
+
 /*-----------------------------------------------------------------------------
 	UControlChannel implementation.
 -----------------------------------------------------------------------------*/
diff --git a/Engine/Source/Runtime/Engine/Private/Net/Experimental/Iris/DataStreamChannel.cpp b/Engine/Source/Runtime/Engine/Private/Net/Experimental/Iris/DataStreamChannel.cpp
--- a/Engine/Source/Runtime/Engine/Private/Net/Experimental/Iris/DataStreamChannel.cpp
+++ b/Engine/Source/Runtime/Engine/Private/Net/Experimental/Iris/DataStreamChannel.cpp
@@ -499,3 +499,13 @@ void UDataStreamChannel::AppendExportBunches(TArray<FOutBunch*>& OutExportBunche
 void UDataStreamChannel::AppendMustBeMappedGuids(FOutBunch* Bunch)
 {
 }
+
+bool UDataStreamChannel::HasAcknowledgedAllReliableData() const
+{
+	if (DataStreamManager)
+	{
+		return DataStreamManager->HasAcknowledgedAllReliableData();
+	}
+
+	return true;
+}
diff --git a/Engine/Source/Runtime/Engine/Private/Net/Experimental/Iris/DataStreamChannel.h b/Engine/Source/Runtime/Engine/Private/Net/Experimental/Iris/DataStreamChannel.h
--- a/Engine/Source/Runtime/Engine/Private/Net/Experimental/Iris/DataStreamChannel.h
+++ b/Engine/Source/Runtime/Engine/Private/Net/Experimental/Iris/DataStreamChannel.h
@@ -57,7 +57,9 @@ private:
 	/** Packet delivery status handling */
 	ENGINE_API virtual void ReceivedAck(int32 PacketId) override;
 	ENGINE_API virtual void ReceivedNak(int32 PacketId) override;
-	
+
+	ENGINE_API virtual bool HasAcknowledgedAllReliableData() const override;
+
 private:
 	enum : uint32
 	{
diff --git a/Engine/Source/Runtime/Engine/Private/NetConnection.cpp b/Engine/Source/Runtime/Engine/Private/NetConnection.cpp
--- a/Engine/Source/Runtime/Engine/Private/NetConnection.cpp
+++ b/Engine/Source/Runtime/Engine/Private/NetConnection.cpp
@@ -209,6 +209,9 @@ namespace UE::Net::Connection::Private
 		return false;
 	}
 
+	bool bGracefulCloseEnabled = true;
+	FAutoConsoleVariableRef CVarNetGracefulCloseEnabled(TEXT("net.GracefulCloseEnabled"), bGracefulCloseEnabled, TEXT("If enabled, connections will wait for reliable bunches to be acked before cleaning up."));
+
 	int32 bTrackFlushedDormantObjects = true;
 	FAutoConsoleVariableRef CVarNetTrackFlushedDormantObjects(TEXT("net.TrackFlushedDormantObjects"), bTrackFlushedDormantObjects, TEXT("If enabled, track dormant subobjects when dormancy is flushed, so they can be properly deleted if destroyed prior to the next ReplicateActor."));
 
@@ -1045,6 +1048,12 @@ const TCHAR* LexToString(const EConnectionState Value)
 	}
 }
 
+bool UNetConnection::IsClosingOrClosed() const
+{
+	const EConnectionState CurrentState = GetConnectionState();
+	return CurrentState == USOCK_Closing || CurrentState == USOCK_Closed;
+}
+
 void UNetConnection::Close(FNetResult&& CloseReason)
 {
 	if (IsInternalAck())
@@ -1112,6 +1121,61 @@ void UNetConnection::Close(FNetResult&& CloseReason)
 	LogSustainedCount	= 0;
 }
 
+void UNetConnection::GracefulClose(FNetResult&& CloseReason)
+{
+	bPendingDestroy = true;
+
+	if (Channels[0] && GetConnectionState() != USOCK_Closed)
+	{
+		if (UE::Net::Connection::Private::bGracefulCloseEnabled && GetConnectionState() != USOCK_Closing)
+		{
+			UE_LOG(LogNet, Log, TEXT("Initiating graceful close. CloseReason: %s. %s"), ToCStr(CloseReason.DynamicToString(UE::Net::ENetResultString::WithChain)), *Describe());
+
+			PendingGracefulCloseResult = MakePimpl<FNetResult, EPimplPtrMode::DeepCopy>(MoveTemp(CloseReason));
+
+			SetConnectionState(EConnectionState::USOCK_Closing);
+			GracefulCloseTimeoutDeadline = Driver->GetElapsedTime() + Driver->GracefulCloseConnectionTimeout;
+
+#if UE_WITH_IRIS
+			if (UReplicationSystem* ReplicationSystem = Driver->GetReplicationSystem())
+			{
+				ReplicationSystem->SetConnectionGracefullyClosing(ConnectionId);
+			}
+#endif
+
+			// Might as well try to close immediately in case there isn't any pending reliable data.
+			TryClosePendingGracefulClose();
+		}
+		else
+		{
+			Close(MoveTemp(CloseReason));
+		}
+	}
+}
+
+void UNetConnection::TryClosePendingGracefulClose()
+{
+	if (GetConnectionState() == USOCK_Closing)
+	{
+		bool bAllChannelsReadyToClose = true;
+
+		for (TObjectPtr<class UChannel> Channel : OpenChannels)
+		{
+			if (Channel != nullptr && !Channel->HasAcknowledgedAllReliableData())
+			{
+				bAllChannelsReadyToClose = false;
+				break;
+			}
+		}
+
+		if (bAllChannelsReadyToClose)
+		{
+			UE_LOG(LogNet, Log, TEXT("Graceful close complete, closing connection: %s"), *Describe());
+			Close(MoveTemp(*PendingGracefulCloseResult));
+		}
+	}
+}
+
 void UNetConnection::HandleNetResultOrClose(ENetCloseResult InResult)
 {
 	using namespace UE::Net;
@@ -1506,7 +1570,7 @@ bool UNetConnection::Exec( UWorld* InWorld, const TCHAR* Cmd, FOutputDevice& Ar
 void UNetConnection::AssertValid()
 {
 	// Make sure this connection is in a reasonable state.
-	check(GetConnectionState()==USOCK_Closed || GetConnectionState()==USOCK_Pending || GetConnectionState()==USOCK_Open);
+	check(IsClosingOrClosed() || GetConnectionState() == USOCK_Pending || GetConnectionState() == USOCK_Open);
 
 }
 
@@ -4503,11 +4567,21 @@ void UNetConnection::Tick(float DeltaSeconds)
 	// Handle timeouts.
 	const float Timeout = GetTimeoutValue();
 
-	if ((CurrentRealtimeSeconds - LastReceiveRealtime) > Timeout)
+	const bool bReceiveTimedOut = (CurrentRealtimeSeconds - LastReceiveRealtime) > Timeout;
+	const bool bGracefulCloseTimedOut = (GetConnectionState() == USOCK_Closing) && (DriverElapsedTime > GracefulCloseTimeoutDeadline);
+
+	if (bReceiveTimedOut || bGracefulCloseTimedOut)
 	{
 		const TCHAR* const TimeoutString = TEXT("UNetConnection::Tick: Connection TIMED OUT. Closing connection.");
 		const TCHAR* const DestroyString = TEXT("UNetConnection::Tick: Connection closing during pending destroy, not all shutdown traffic may have been negotiated");
-		
+
+		const TCHAR* ErrorString = bPendingDestroy ? DestroyString : TimeoutString;
+
+		if (bGracefulCloseTimedOut)
+		{
+			ErrorString = TEXT("UNetConnection::Tick: Connection graceful close timed out, not all shutdown traffic was negotiated");
+		}
+
 		// Compute true realtime since packet was received (as well as truly processed)
 		const double Seconds = FPlatformTime::Seconds();
 
@@ -4516,12 +4590,12 @@ void UNetConnection::Tick(float DeltaSeconds)
 
 		// Timeout.
 		FString Error = FString::Printf(TEXT("%s. Elapsed: %2.2f, Real: %2.2f, Good: %2.2f, DriverTime: %2.2f, Threshold: %2.2f, %s"),
-			bPendingDestroy ? DestroyString : TimeoutString,
+			ErrorString,
 			DriverElapsedTime - LastReceiveTime,
 			ReceiveRealtimeDelta,
 			GoodRealtimeDelta,
 			DriverElapsedTime,
-			Timeout,
+			bGracefulCloseTimedOut ? Driver->GracefulCloseConnectionTimeout : Timeout,
 			*Describe());
 		
 		static double LastTimePrinted = 0.0f;
@@ -4615,6 +4689,9 @@ void UNetConnection::Tick(float DeltaSeconds)
 			}
 		}
 
+		// If gracefully closing, check whether all channels are ready for close
+		TryClosePendingGracefulClose();
+
 		// If channel 0 has closed, mark the connection as closed.
 		if (Channels[0] == nullptr && (OutReliable[0] != InitOutReliable || InReliable[0] != InitInReliable))
 		{
diff --git a/Engine/Source/Runtime/Engine/Private/NetDriver.cpp b/Engine/Source/Runtime/Engine/Private/NetDriver.cpp
--- a/Engine/Source/Runtime/Engine/Private/NetDriver.cpp
+++ b/Engine/Source/Runtime/Engine/Private/NetDriver.cpp
@@ -2586,8 +2586,16 @@ void UNetDriver::InternalProcessRemoteFunctionPrivate(
 		Connection = ((UChildConnection*)Connection)->Parent;
 	}
 
+	// Prevent RPC calls to gracefully-closing connections
+	const EConnectionState ConnectionState = Connection->GetConnectionState();
+	if (ConnectionState == USOCK_Closing)
+	{
+		DEBUG_REMOTEFUNCTION(TEXT("Attempting to call RPC on a closing connection. Not calling %s::%s"), *GetNameSafe(Actor), *GetNameSafe(Function));
+		return;
+	}
+
 	// Prevent RPC calls to closed connections
-	if (Connection->GetConnectionState() == USOCK_Closed)
+	if (ConnectionState == USOCK_Closed)
 	{
 		DEBUG_REMOTEFUNCTION(TEXT("Attempting to call RPC on a closed connection. Not calling %s::%s"), *GetNameSafe(Actor), *GetNameSafe(Function));
 		return;
@@ -4935,7 +4943,7 @@ int32 UNetDriver::ServerReplicateActors_PrepConnections( const float DeltaSecond
 	{
 		UNetConnection* Connection = ClientConnections[ConnIdx];
 		check( Connection );
-		check( Connection->GetConnectionState() == USOCK_Pending || Connection->GetConnectionState() == USOCK_Open || Connection->GetConnectionState() == USOCK_Closed );
+		check( Connection->GetConnectionState() == USOCK_Pending || Connection->GetConnectionState() == USOCK_Open || Connection->IsClosingOrClosed());
 		checkSlow( Connection->GetUChildConnection() == NULL );
 
 		// Handle not ready channels.
diff --git a/Engine/Source/Runtime/Engine/Private/PlayerController.cpp b/Engine/Source/Runtime/Engine/Private/PlayerController.cpp
--- a/Engine/Source/Runtime/Engine/Private/PlayerController.cpp
+++ b/Engine/Source/Runtime/Engine/Private/PlayerController.cpp
@@ -277,14 +277,10 @@ UNetConnection* APlayerController::GetNetConnection() const
 
 bool APlayerController::DestroyNetworkActorHandled()
 {
-	UNetConnection* C = Cast<UNetConnection>(Player);
-	if (C)
+	UNetConnection* Connection = Cast<UNetConnection>(Player);
+	if (Connection)
 	{
-		if (C->Channels[0] && C->GetConnectionState() != USOCK_Closed)
-		{
-			C->bPendingDestroy = true;
-			C->Channels[0]->Close(EChannelCloseReason::Destroyed);
-		}
+		Connection->GracefulClose(ENetCloseResult::ControlChannelClose);
 		return true;
 	}
 
diff --git a/Engine/Source/Runtime/Experimental/Iris/Core/Private/Iris/DataStream/DataStreamManager.cpp b/Engine/Source/Runtime/Experimental/Iris/Core/Private/Iris/DataStream/DataStreamManager.cpp
--- a/Engine/Source/Runtime/Experimental/Iris/Core/Private/Iris/DataStream/DataStreamManager.cpp
+++ b/Engine/Source/Runtime/Experimental/Iris/Core/Private/Iris/DataStream/DataStreamManager.cpp
@@ -25,6 +25,7 @@ public:
 	void EndWrite();
 	void ReadData(UE::Net::FNetSerializationContext& context);
 	void ProcessPacketDeliveryStatus(UE::Net::EPacketDeliveryStatus Status, FDataStreamRecord const* Record);
+	bool HasAcknowledgedAllReliableData() const;
 
 	ECreateDataStreamResult CreateStream(const FName StreamName);
 
@@ -120,6 +121,11 @@ void UDataStreamManager::ProcessPacketDeliveryStatus(UE::Net::EPacketDeliverySta
 	return Impl->ProcessPacketDeliveryStatus(Status, Record);
 }
 
+bool UDataStreamManager::HasAcknowledgedAllReliableData() const
+{
+	return Impl->HasAcknowledgedAllReliableData();
+}
+
 ECreateDataStreamResult UDataStreamManager::CreateStream(const FName StreamName)
 {
 	return Impl->CreateStream(StreamName);
@@ -409,6 +415,19 @@ void UDataStreamManager::FImpl::ProcessPacketDeliveryStatus(UE::Net::EPacketDeli
 	Records.Pop();
 }
 
+bool UDataStreamManager::FImpl::HasAcknowledgedAllReliableData() const
+{
+	for (TObjectPtr<const UDataStream> Stream : Streams)
+	{
+		if (Stream && !Stream->HasAcknowledgedAllReliableData())
+		{
+			return false;
+		}
+	}
+
+	return true;
+}
+
 ECreateDataStreamResult UDataStreamManager::FImpl::CreateStream(const FName StreamName)
 {
 	// Bumping MaxStreamCount may require modifying the FRecord and WriteData/ReadData.
diff --git a/Engine/Source/Runtime/Experimental/Iris/Core/Private/Iris/ReplicationSystem/AttachmentReplication.cpp b/Engine/Source/Runtime/Experimental/Iris/Core/Private/Iris/ReplicationSystem/AttachmentReplication.cpp
--- a/Engine/Source/Runtime/Experimental/Iris/Core/Private/Iris/ReplicationSystem/AttachmentReplication.cpp
+++ b/Engine/Source/Runtime/Experimental/Iris/Core/Private/Iris/ReplicationSystem/AttachmentReplication.cpp
@@ -503,6 +503,27 @@ bool FNetObjectAttachmentsWriter::IsAllReliableSentAndAcked(ENetObjectAttachment
 	return Queue->IsAllSentAndAcked();
 }
 
+bool FNetObjectAttachmentsWriter::AreAllObjectsReliableSentAndAcked() const
+{
+	for (const auto& ObjectQueuePair : ObjectToQueue)
+	{
+		if (!ObjectQueuePair.Value.IsAllReliableSentAndAcked())
+		{
+			return false;
+		}
+	}
+
+	for (const TUniquePtr<FNetObjectAttachmentSendQueue>& SpecialQueue : SpecialQueues)
+	{
+		if (SpecialQueue.IsValid() && !SpecialQueue->IsAllReliableSentAndAcked())
+		{
+			return false;
+		}
+	}
+
+	return true;
+}
+
 bool FNetObjectAttachmentsWriter::CanSendMoreReliableAttachments(ENetObjectAttachmentType Type, uint32 ObjectIndex) const
 {
 	const FNetObjectAttachmentSendQueue* Queue = GetQueue(Type, ObjectIndex);
diff --git a/Engine/Source/Runtime/Experimental/Iris/Core/Private/Iris/ReplicationSystem/AttachmentReplication.h b/Engine/Source/Runtime/Experimental/Iris/Core/Private/Iris/ReplicationSystem/AttachmentReplication.h
--- a/Engine/Source/Runtime/Experimental/Iris/Core/Private/Iris/ReplicationSystem/AttachmentReplication.h
+++ b/Engine/Source/Runtime/Experimental/Iris/Core/Private/Iris/ReplicationSystem/AttachmentReplication.h
@@ -123,6 +123,9 @@ public:
 	// Whether all queued reliable attachments have been sent and acked
 	bool IsAllReliableSentAndAcked(ENetObjectAttachmentType Type, uint32 ObjectIndex) const;
 
+	// Whether all queued reliable attachments have been sent and acked for all types and objects
+	bool AreAllObjectsReliableSentAndAcked() const;
+
 	// Whether more reliable attachments can be sent now. It's possible to queue up as many attachments as you see fit, but if the queue is full it can take a while before more attachments will be replicated.
 	bool CanSendMoreReliableAttachments(ENetObjectAttachmentType Type, uint32 ObjectIndex) const;
 
diff --git a/Engine/Source/Runtime/Experimental/Iris/Core/Private/Iris/ReplicationSystem/NetBlob/NetBlobManager.cpp b/Engine/Source/Runtime/Experimental/Iris/Core/Private/Iris/ReplicationSystem/NetBlob/NetBlobManager.cpp
--- a/Engine/Source/Runtime/Experimental/Iris/Core/Private/Iris/ReplicationSystem/NetBlob/NetBlobManager.cpp
+++ b/Engine/Source/Runtime/Experimental/Iris/Core/Private/Iris/ReplicationSystem/NetBlob/NetBlobManager.cpp
@@ -137,7 +137,7 @@ bool FNetBlobManager::SendRPC(const UObject* Object, const UObject* SubObject, c
 	}
 
 	RPC->SetNetObjectReference(OwnerInfo.CallerRef, OwnerInfo.TargetRef);
-	AttachmentSendQueue.Enqueue(OwnerInfo.RootObjectIndex, OwnerInfo.SubObjectIndex, reinterpret_cast<const TRefCountPtr<FNetObjectAttachment>&>(RPC), SendFlags);
+	AttachmentSendQueue.Enqueue(OwnerInfo.RootObjectIndex, OwnerInfo.SubObjectIndex, reinterpret_cast<const TRefCountPtr<FNetObjectAttachment>&>(RPC), SendFlags, Connections->GetOpenConnections());
 	return true;
 }
 
@@ -166,6 +166,12 @@ bool FNetBlobManager::SendRPC(uint32 ConnectionId, const UObject* Object, const
 		return true;
 	}
 
+	if (!Connections->IsOpenConnection(ConnectionId))
+	{
+		// This connection is shutting down and only flushing existing reliable data, not sending new RPCs.
+		return true;
+	}
+
 	FRPCOwner OwnerInfo;
 	if (!GetRPCOwner(OwnerInfo, Object, SubObject, Function))
 	{
@@ -423,7 +429,7 @@ void FNetBlobManager::FNetObjectAttachmentSendQueue::Enqueue(uint32 ConnectionId
 	QueueEntry.Attachment = Attachment;
 }
 
-void FNetBlobManager::FNetObjectAttachmentSendQueue::Enqueue(FInternalNetRefIndex OwnerIndex, FInternalNetRefIndex SubObjectIndex, const TRefCountPtr<FNetObjectAttachment>& Attachment, ENetObjectAttachmentSendPolicyFlags SendFlags)
+void FNetBlobManager::FNetObjectAttachmentSendQueue::Enqueue(FInternalNetRefIndex OwnerIndex, FInternalNetRefIndex SubObjectIndex, const TRefCountPtr<FNetObjectAttachment>& Attachment, ENetObjectAttachmentSendPolicyFlags SendFlags, FNetBitArray OpenConnections)
 {
 	const bool bScheduleUsingOOBAttachmentQueue = EnumHasAnyFlags(SendFlags, ENetObjectAttachmentSendPolicyFlags::ScheduleAsOOB);
 	FQueue& TargetQueue = bScheduleUsingOOBAttachmentQueue ? ScheduleAsOOBAttachmentQueue : AttachmentQueue;
@@ -434,6 +440,7 @@ void FNetBlobManager::FNetObjectAttachmentSendQueue::Enqueue(FInternalNetRefInde
 	QueueEntry.SubObjectIndex = SubObjectIndex;
 	QueueEntry.SendFlags = SendFlags;
 	QueueEntry.Attachment = Attachment;
+	QueueEntry.MulticastConnections = MoveTemp(OpenConnections);
 
 	bHasMulticastAttachments = true;
 }
@@ -565,12 +572,23 @@ bool FNetBlobManager::HasUnprocessedReliableAttachments(FInternalNetRefIndex Int
 	return AttachmentSendQueue.HasUnprocessedReliableAttachments(InternalIndex);
 }
 
+bool FNetBlobManager::HasAnyUnprocessedReliableAttachments() const
+{
+	return AttachmentSendQueue.HasAnyUnprocessedReliableAttachments();
+}
+
 bool FNetBlobManager::FNetObjectAttachmentSendQueue::HasUnprocessedReliableAttachments(FInternalNetRefIndex InternalIndex)  const
 {
 	// For the moment we only need to check the AttachmentQueue as reliable attachments are not schedules as immediate.
 	return AttachmentQueue.ContainsByPredicate([&InternalIndex](const FNetObjectAttachmentQueueEntry& Entry) { return (Entry.OwnerIndex == InternalIndex || Entry.SubObjectIndex == InternalIndex) && Entry.Attachment->IsReliable();} );
 }
 
+bool FNetBlobManager::FNetObjectAttachmentSendQueue::HasAnyUnprocessedReliableAttachments()  const
+{
+	// For the moment we only need to check the AttachmentQueue as reliable attachments are not schedules as immediate.
+	return AttachmentQueue.ContainsByPredicate([](const FNetObjectAttachmentQueueEntry& Entry) { return Entry.Attachment->IsReliable(); });
+}
+
 void FNetBlobManager::FNetObjectAttachmentSendQueue::ProcessQueue(EProcessMode ProcessMode)
 {
 	if (!ProcessContext.IsValid())
@@ -627,6 +645,12 @@ void FNetBlobManager::FNetObjectAttachmentSendQueue::ProcessQueue(EProcessMode P
 					continue;
 				}
 
+				// Don't send RPCs to connections that were already closing when the RPC was called/queued
+				if (!Entry.MulticastConnections.IsBitSet(ConnectionId))
+				{
+					continue;
+				}
+
 				if (bHasConnectionSpecificSerialization)
 				{
 					PartialNetBlobs.Reset();
diff --git a/Engine/Source/Runtime/Experimental/Iris/Core/Private/Iris/ReplicationSystem/NetBlob/NetBlobManager.h b/Engine/Source/Runtime/Experimental/Iris/Core/Private/Iris/ReplicationSystem/NetBlob/NetBlobManager.h
--- a/Engine/Source/Runtime/Experimental/Iris/Core/Private/Iris/ReplicationSystem/NetBlob/NetBlobManager.h
+++ b/Engine/Source/Runtime/Experimental/Iris/Core/Private/Iris/ReplicationSystem/NetBlob/NetBlobManager.h
@@ -49,7 +49,8 @@ public:
 	bool SendRPC(const UObject* Object, const UObject* SubObject, const UFunction* Function, const void* Parameters, ENetObjectAttachmentSendPolicyFlags SendFlags = ENetObjectAttachmentSendPolicyFlags::None);
 	bool SendRPC(uint32 ConnectionId, const UObject* Object, const UObject* SubObject, const UFunction* Function, const void* Parameters, ENetObjectAttachmentSendPolicyFlags SendFlags = ENetObjectAttachmentSendPolicyFlags::None);
 
-	bool HasUnprocessedReliableAttachments(FInternalNetRefIndex InternalIndex)  const;
+	bool HasUnprocessedReliableAttachments(FInternalNetRefIndex InternalIndex) const;
+	bool HasAnyUnprocessedReliableAttachments() const;
 
 	enum class EProcessMode 
 	{
@@ -118,7 +119,7 @@ private:
 		void Enqueue(uint32 ConnectionId, FInternalNetRefIndex OwnerIndex, FInternalNetRefIndex SubObjectIndex, const TRefCountPtr<FNetObjectAttachment>& Attachment, ENetObjectAttachmentSendPolicyFlags SendFlags);
 
 		// Multicast
-		void Enqueue(FInternalNetRefIndex OwnerIndex, FInternalNetRefIndex SubObjectIndex, const TRefCountPtr<FNetObjectAttachment>& Attachment, ENetObjectAttachmentSendPolicyFlags SendFlags);
+		void Enqueue(FInternalNetRefIndex OwnerIndex, FInternalNetRefIndex SubObjectIndex, const TRefCountPtr<FNetObjectAttachment>& Attachment, ENetObjectAttachmentSendPolicyFlags SendFlags, FNetBitArray OpenConnections);
 
 		void PrepareProcessQueue(FReplicationConnections* InConnections, const FNetRefHandleManager* InNetRefHandleManager);
 		void ProcessQueue(EProcessMode ProcessMode);
@@ -126,6 +127,7 @@ private:
 		void PrepareAndProcessOOBAttachmentQueue(FReplicationConnections* InConnections, const FNetRefHandleManager* InNetRefHandleManager, FNetBitArray& OutConnetionsPendingImmediateSend);
 
 		bool HasUnprocessedReliableAttachments(FInternalNetRefIndex InternalIndex)  const;
+		bool HasAnyUnprocessedReliableAttachments()  const;
 	
 	private:
 		struct FNetObjectAttachmentQueueEntry
@@ -135,6 +137,7 @@ private:
 			FInternalNetRefIndex SubObjectIndex;
 			ENetObjectAttachmentSendPolicyFlags SendFlags;
 			TRefCountPtr<FNetObjectAttachment> Attachment;
+			FNetBitArray MulticastConnections;
 		};
 		typedef TArray<FNetObjectAttachmentQueueEntry> FQueue;
 
diff --git a/Engine/Source/Runtime/Experimental/Iris/Core/Private/Iris/ReplicationSystem/NetTokenDataStream.cpp b/Engine/Source/Runtime/Experimental/Iris/Core/Private/Iris/ReplicationSystem/NetTokenDataStream.cpp
--- a/Engine/Source/Runtime/Experimental/Iris/Core/Private/Iris/ReplicationSystem/NetTokenDataStream.cpp
+++ b/Engine/Source/Runtime/Experimental/Iris/Core/Private/Iris/ReplicationSystem/NetTokenDataStream.cpp
@@ -231,4 +231,9 @@ void UNetTokenDataStream::ProcessPacketDeliveryStatus(UE::Net::EPacketDeliverySt
 			--RecordCount;
 		}
 	}
+}
+
+bool UNetTokenDataStream::HasAcknowledgedAllReliableData() const
+{
+	return NetTokensPendingExport.Num() == 0 && NetTokenExports.Num() == 0;
 }
\ No newline at end of file
diff --git a/Engine/Source/Runtime/Experimental/Iris/Core/Private/Iris/ReplicationSystem/NetTokenDataStream.h b/Engine/Source/Runtime/Experimental/Iris/Core/Private/Iris/ReplicationSystem/NetTokenDataStream.h
--- a/Engine/Source/Runtime/Experimental/Iris/Core/Private/Iris/ReplicationSystem/NetTokenDataStream.h
+++ b/Engine/Source/Runtime/Experimental/Iris/Core/Private/Iris/ReplicationSystem/NetTokenDataStream.h
@@ -48,6 +48,7 @@ private:
 	virtual EWriteResult WriteData(UE::Net::FNetSerializationContext& Context, FDataStreamRecord const*& OutRecord) override;
 	virtual void ReadData(UE::Net::FNetSerializationContext& Context) override;
 	virtual void ProcessPacketDeliveryStatus(UE::Net::EPacketDeliveryStatus Status, FDataStreamRecord const* Record) override;
+	virtual bool HasAcknowledgedAllReliableData() const override;
 
 private:
 
diff --git a/Engine/Source/Runtime/Experimental/Iris/Core/Private/Iris/ReplicationSystem/ReplicationConnections.cpp b/Engine/Source/Runtime/Experimental/Iris/Core/Private/Iris/ReplicationSystem/ReplicationConnections.cpp
--- a/Engine/Source/Runtime/Experimental/Iris/Core/Private/Iris/ReplicationSystem/ReplicationConnections.cpp
+++ b/Engine/Source/Runtime/Experimental/Iris/Core/Private/Iris/ReplicationSystem/ReplicationConnections.cpp
@@ -65,6 +65,21 @@ void FReplicationConnections::RemoveConnection(uint32 ConnectionId)
 	ValidConnections.ClearBit(ConnectionId);
 }
 
+FNetBitArray FReplicationConnections::GetOpenConnections() const
+{
+	FNetBitArray OpenConnections(ValidConnections.GetNumBits());
+
+	for (int32 ConnectionId = 0; ConnectionId < Connections.Num(); ++ConnectionId)
+	{
+		if (ValidConnections.IsBitSet(ConnectionId) && !Connections[ConnectionId].bIsClosing)
+		{
+			OpenConnections.SetBit(ConnectionId);
+		}
+	}
+
+	return OpenConnections;
+}
+
 void FReplicationConnections::DestroyReplicationReaderAndWriter(uint32 ConnectionId)
 {
 	FReplicationConnection* Connection = GetConnection(ConnectionId);
diff --git a/Engine/Source/Runtime/Experimental/Iris/Core/Private/Iris/ReplicationSystem/ReplicationConnections.h b/Engine/Source/Runtime/Experimental/Iris/Core/Private/Iris/ReplicationSystem/ReplicationConnections.h
--- a/Engine/Source/Runtime/Experimental/Iris/Core/Private/Iris/ReplicationSystem/ReplicationConnections.h
+++ b/Engine/Source/Runtime/Experimental/Iris/Core/Private/Iris/ReplicationSystem/ReplicationConnections.h
@@ -26,6 +26,7 @@ struct FReplicationConnection
 	FReplicationWriter* ReplicationWriter = nullptr;
 	FReplicationReader* ReplicationReader = nullptr;
 	FObjectPtr UserData = nullptr;
+	bool bIsClosing = false; // Should be set when a connection starts the graceful close process to finish flushing reliable data
 };
 
 class FReplicationConnections
@@ -66,6 +67,11 @@ public:
 		return ConnectionId < GetMaxConnectionCount() && ValidConnections.GetBit(ConnectionId);
 	}
 
+	bool IsOpenConnection(uint32 ConnectionId) const
+	{
+		return ConnectionId < GetMaxConnectionCount() && ValidConnections.GetBit(ConnectionId) && !Connections[ConnectionId].bIsClosing;
+	}
+
 	void AddConnection(uint32 ConnectionId)
 	{
 		check(ValidConnections.GetBit(ConnectionId) == false);
@@ -78,6 +84,9 @@ public:
 
 	const FNetBitArray& GetValidConnections() const { return ValidConnections; }
 
+	// Returns connections that are not in the closing state
+	FNetBitArray GetOpenConnections() const;
+
 	void InitDataStreams(uint32 ReplicationSystemId, uint32 ConnectionId, UDataStreamManager* DataStreamManager);
 
 	void SetReplicationView(uint32 ConnectionId, const FReplicationView& ViewInfo);
@@ -86,6 +95,13 @@ public:
 	const FNetTokenStoreState& GetRemoteNetTokenStoreState(uint32 ConnectionId) const { return RemoteNetTokenStoreStates[ConnectionId]; }
 	FNetTokenStoreState& GetRemoteNetTokenStoreState(uint32 ConnectionId) { return RemoteNetTokenStoreStates[ConnectionId]; }
 
+	// Flag a connection as being in a graceful-close state meant to flush pending reliable data.
+	void SetConnectionIsClosing(uint32 ConnectionId)
+	{
+		check(ValidConnections.GetBit(ConnectionId) == true);
+		Connections[ConnectionId].bIsClosing = true;
+	}
+
 private:
 	void DestroyReplicationReaderAndWriter(uint32 ConnectionId);
 
diff --git a/Engine/Source/Runtime/Experimental/Iris/Core/Private/Iris/ReplicationSystem/ReplicationDataStream.cpp b/Engine/Source/Runtime/Experimental/Iris/Core/Private/Iris/ReplicationSystem/ReplicationDataStream.cpp
--- a/Engine/Source/Runtime/Experimental/Iris/Core/Private/Iris/ReplicationSystem/ReplicationDataStream.cpp
+++ b/Engine/Source/Runtime/Experimental/Iris/Core/Private/Iris/ReplicationSystem/ReplicationDataStream.cpp
@@ -58,6 +58,16 @@ void UReplicationDataStream::ProcessPacketDeliveryStatus(UE::Net::EPacketDeliver
 	}
 }
 
+bool UReplicationDataStream::HasAcknowledgedAllReliableData() const
+{
+	if (ReplicationWriter != nullptr)
+	{
+		return ReplicationWriter->AreAllReliableAttachmentsSentAndAcked();
+	}
+
+	return true;
+}
+
 void UReplicationDataStream::SetReaderAndWriter(UE::Net::Private::FReplicationReader* Reader, UE::Net::Private::FReplicationWriter* Writer)
 {
 	ReplicationReader = Reader;
diff --git a/Engine/Source/Runtime/Experimental/Iris/Core/Private/Iris/ReplicationSystem/ReplicationDataStream.h b/Engine/Source/Runtime/Experimental/Iris/Core/Private/Iris/ReplicationSystem/ReplicationDataStream.h
--- a/Engine/Source/Runtime/Experimental/Iris/Core/Private/Iris/ReplicationSystem/ReplicationDataStream.h
+++ b/Engine/Source/Runtime/Experimental/Iris/Core/Private/Iris/ReplicationSystem/ReplicationDataStream.h
@@ -29,6 +29,7 @@ private:
 	virtual void EndWrite() override;
 	virtual void ReadData(UE::Net::FNetSerializationContext& Context) override;
 	virtual void ProcessPacketDeliveryStatus(UE::Net::EPacketDeliveryStatus Status, FDataStreamRecord const* Record) override;
+	virtual bool HasAcknowledgedAllReliableData() const override;
 
 private:
 	UE::Net::Private::FReplicationReader* ReplicationReader;
diff --git a/Engine/Source/Runtime/Experimental/Iris/Core/Private/Iris/ReplicationSystem/ReplicationSystem.cpp b/Engine/Source/Runtime/Experimental/Iris/Core/Private/Iris/ReplicationSystem/ReplicationSystem.cpp
--- a/Engine/Source/Runtime/Experimental/Iris/Core/Private/Iris/ReplicationSystem/ReplicationSystem.cpp
+++ b/Engine/Source/Runtime/Experimental/Iris/Core/Private/Iris/ReplicationSystem/ReplicationSystem.cpp
@@ -449,7 +449,13 @@ public:
 		auto UpdateDirtyChangeMasks = [&Connections, &UpdatedChangeMasks](uint32 ConnectionId)
 		{
 			FReplicationConnection* Conn = Connections.GetConnection(ConnectionId);
-			Conn->ReplicationWriter->UpdateDirtyChangeMasks(UpdatedChangeMasks);
+			
+			// Only update open connections, as closing connections are only
+			// flushing reliable data and we shouldn't send new state data to them.
+			if (!Conn->bIsClosing)
+			{
+				Conn->ReplicationWriter->UpdateDirtyChangeMasks(UpdatedChangeMasks);
+			}
 		};
 		const FNetBitArray& ValidConnections = Connections.GetValidConnections();
 		ValidConnections.ForAllSetBits(UpdateDirtyChangeMasks);
@@ -946,6 +952,15 @@ bool UReplicationSystem::IsValidConnection(uint32 ConnectionId) const
 	return Connections.GetConnection(ConnectionId) != nullptr;
 }
 
+
+void UReplicationSystem::SetConnectionGracefullyClosing(uint32 ConnectionId) const
+{
+	UE::Net::Private::FReplicationConnections& Connections = Impl->ReplicationSystemInternal.GetConnections();
+	check(Connections.IsValidConnection(ConnectionId));
+
+	Connections.SetConnectionIsClosing(ConnectionId);
+}
+
 void UReplicationSystem::SetReplicationEnabledForConnection(uint32 ConnectionId, bool bReplicationEnabled)
 {
 	UE::Net::Private::FReplicationConnections& Connections = Impl->ReplicationSystemInternal.GetConnections();
diff --git a/Engine/Source/Runtime/Experimental/Iris/Core/Private/Iris/ReplicationSystem/ReplicationWriter.cpp b/Engine/Source/Runtime/Experimental/Iris/Core/Private/Iris/ReplicationSystem/ReplicationWriter.cpp
--- a/Engine/Source/Runtime/Experimental/Iris/Core/Private/Iris/ReplicationSystem/ReplicationWriter.cpp
+++ b/Engine/Source/Runtime/Experimental/Iris/Core/Private/Iris/ReplicationSystem/ReplicationWriter.cpp
@@ -345,6 +345,12 @@ bool FReplicationWriter::QueueNetObjectAttachments(FInternalNetRefIndex OwnerInt
 	return true;
 }
 
+bool FReplicationWriter::AreAllReliableAttachmentsSentAndAcked() const
+{
+	const bool bHasUnprocessedReliables = ReplicationSystemInternal->GetNetBlobManager().HasAnyUnprocessedReliableAttachments();
+	return !bHasUnprocessedReliables && Attachments.AreAllObjectsReliableSentAndAcked();
+}
+
 void FReplicationWriter::SetState(uint32 InternalIndex, EReplicatedObjectState NewState)
 {
 	FReplicationInfo& Info = GetReplicationInfo(InternalIndex);
diff --git a/Engine/Source/Runtime/Experimental/Iris/Core/Private/Iris/ReplicationSystem/ReplicationWriter.h b/Engine/Source/Runtime/Experimental/Iris/Core/Private/Iris/ReplicationSystem/ReplicationWriter.h
--- a/Engine/Source/Runtime/Experimental/Iris/Core/Private/Iris/ReplicationSystem/ReplicationWriter.h
+++ b/Engine/Source/Runtime/Experimental/Iris/Core/Private/Iris/ReplicationSystem/ReplicationWriter.h
@@ -193,6 +193,8 @@ public:
 	// Queue NetObjectAttachments, returns whether the attachments was enqueued or not.
 	bool QueueNetObjectAttachments(FInternalNetRefIndex OwnerInternalIndex, FInternalNetRefIndex SubObjectInternalIndex, TArrayView<const TRefCountPtr<FNetBlob>> Attachments, ENetObjectAttachmentSendPolicyFlags SendFlags);
 
+	bool AreAllReliableAttachmentsSentAndAcked() const;
+
 private:
 	// Various types
 
diff --git a/Engine/Source/Runtime/Experimental/Iris/Core/Public/Iris/DataStream/DataStream.h b/Engine/Source/Runtime/Experimental/Iris/Core/Public/Iris/DataStream/DataStream.h
--- a/Engine/Source/Runtime/Experimental/Iris/Core/Public/Iris/DataStream/DataStream.h
+++ b/Engine/Source/Runtime/Experimental/Iris/Core/Public/Iris/DataStream/DataStream.h
@@ -110,4 +110,9 @@ public:
 	 * @param Record The record which was set by this stream during a WriteData call.
 	 */
 	IRISCORE_API virtual void ProcessPacketDeliveryStatus(UE::Net::EPacketDeliveryStatus Status, FDataStreamRecord const* Record) PURE_VIRTUAL(ProcessPacketDeliveryStatus,);
+
+	/**
+	 * @return true if the stream has no pending reliable data for which it is waiting on an acknowledgement.
+	 */
+	IRISCORE_API virtual bool HasAcknowledgedAllReliableData() const PURE_VIRTUAL(HasAcknowledgedAllReliableData, return true;);
 };
diff --git a/Engine/Source/Runtime/Experimental/Iris/Core/Public/Iris/DataStream/DataStreamManager.h b/Engine/Source/Runtime/Experimental/Iris/Core/Public/Iris/DataStream/DataStreamManager.h
--- a/Engine/Source/Runtime/Experimental/Iris/Core/Public/Iris/DataStream/DataStreamManager.h
+++ b/Engine/Source/Runtime/Experimental/Iris/Core/Public/Iris/DataStream/DataStreamManager.h
@@ -71,6 +71,9 @@ public:
 	/** Called for all data streams that wrote to a packet whose delivery status is now known. */
 	IRISCORE_API virtual void ProcessPacketDeliveryStatus(UE::Net::EPacketDeliveryStatus Status, FDataStreamRecord const* Record) override;
 
+	/** Returns true if calling HasAcknowledgedAllReliableData on all data streams also returns true. */
+	IRISCORE_API virtual bool HasAcknowledgedAllReliableData() const override;
+
 	// DataStreamManager specifics
 
 	/**
diff --git a/Engine/Source/Runtime/Experimental/Iris/Core/Public/Iris/ReplicationSystem/ReplicationSystem.h b/Engine/Source/Runtime/Experimental/Iris/Core/Public/Iris/ReplicationSystem/ReplicationSystem.h
--- a/Engine/Source/Runtime/Experimental/Iris/Core/Public/Iris/ReplicationSystem/ReplicationSystem.h
+++ b/Engine/Source/Runtime/Experimental/Iris/Core/Public/Iris/ReplicationSystem/ReplicationSystem.h
@@ -191,6 +191,13 @@ public:
 	 */
 	IRISCORE_API bool IsValidConnection(uint32 ConnectionId) const;
 
+	/**
+	 * Sets a connection as gracefully closing, where it will flush all pending reliable data
+	 * before completely shutting down.
+	 * @param ConnectionId The ID of the connection that's closing.
+	 */
+	IRISCORE_API void SetConnectionGracefullyClosing(uint32 ConnectionId) const;
+
 	/**
 	 * Enable or disable the ReplicationDataStream to transmit data for a particular connection.
 	 * @param ConnectionId The ID of the connection to enable or disable object replication for.
