Fix device profile reloading from dynamic config layers

Previously passed a set of full section names to a method that expects a set of device profile names

#rnx

diff --git a/Engine/Source/Runtime/Engine/Private/DeviceProfiles/DeviceProfileManager.cpp b/Engine/Source/Runtime/Engine/Private/DeviceProfiles/DeviceProfileManager.cpp
--- a/Engine/Source/Runtime/Engine/Private/DeviceProfiles/DeviceProfileManager.cpp
+++ b/Engine/Source/Runtime/Engine/Private/DeviceProfiles/DeviceProfileManager.cpp
@@ -105,7 +105,14 @@ UDeviceProfileManager& UDeviceProfileManager::Get(bool bFromPostCDOContruct)
 		// when we load/unload dynamic configs, we may need to propagate those changes to the active DP
 		UE::DynamicConfig::UpdateDeviceProfiles.AddLambda([](const TSet<FString>& ModifiedSections)
 			{
-				if (UDeviceProfileManager::Get().DoActiveProfilesReference(ModifiedSections))
+				TSet<FString> DeviceProfilesToQuery;
+				int32 DeviceProfilePatternLength = UDeviceProfile::StaticClass()->GetName().Len() + 1;
+				for (const FString& Section : ModifiedSections)
+				{
+					DeviceProfilesToQuery.Add(Section.LeftChop(DeviceProfilePatternLength));
+				}
+
+				if (UDeviceProfileManager::Get().DoActiveProfilesReference(DeviceProfilesToQuery))
 				{
 					UDeviceProfileManager::Get().ReapplyDeviceProfile();
 				}
