Fix stale references in FRayTracingGeometryManager::PendingStreamingRequests
- remove pending request when ray tracing geometry is unregistered.
- also provide GeometryHandle instead of FRayTracingGeometry*  to RequestRayTracingGeometryStreamIn(...)

diff --git a/Engine/Source/Runtime/Engine/Private/Rendering/RayTracingGeometryManager.cpp b/Engine/Source/Runtime/Engine/Private/Rendering/RayTracingGeometryManager.cpp
--- a/Engine/Source/Runtime/Engine/Private/Rendering/RayTracingGeometryManager.cpp
+++ b/Engine/Source/Runtime/Engine/Private/Rendering/RayTracingGeometryManager.cpp
@@ -255,6 +255,7 @@ void FRayTracingGeometryManager::ReleaseRayTracingGeometryHandle(RayTracingGeome
 
 	RegisteredGeometries.RemoveAt(Handle);
 	ReferencedGeometryHandles.Remove(Handle);
+	PendingStreamingRequests.Remove(Handle);
 
 	DEC_DWORD_STAT(STAT_RayTracingGeometryCount);
 }
@@ -471,7 +472,7 @@ void FRayTracingGeometryManager::Tick(FRHICommandList& RHICmdList)
 				Geometry->MakeResident(RHICmdList);
 			}
 
-			RequestRayTracingGeometryStreamIn(RHICmdList, Geometry);
+			RequestRayTracingGeometryStreamIn(RHICmdList, Geometry->RayTracingGeometryHandle);
 		}
 
 		// 4th step
@@ -501,9 +502,9 @@ void FRayTracingGeometryManager::Tick(FRHICommandList& RHICmdList)
 					RegisteredGeometry.Geometry->MakeResident(RHICmdList);
 				}
 
-				if (!RequestRayTracingGeometryStreamIn(RHICmdList, RegisteredGeometry.Geometry))
+				if (!RequestRayTracingGeometryStreamIn(RHICmdList, RegisteredGeometry.Geometry->RayTracingGeometryHandle))
 				{
-					PendingStreamingRequests.Add(RegisteredGeometry.Geometry);
+					PendingStreamingRequests.Add(RegisteredGeometry.Geometry->RayTracingGeometryHandle);
 				}
 			}
 		}
@@ -517,15 +518,15 @@ void FRayTracingGeometryManager::Tick(FRHICommandList& RHICmdList)
 			}
 #endif
 			
-			TArray<FRayTracingGeometry*> CurrentPendingStreamingRequests;
+			TSet<RayTracingGeometryHandle> CurrentPendingStreamingRequests;
 			Swap(CurrentPendingStreamingRequests, PendingStreamingRequests);
 			PendingStreamingRequests.Reserve(CurrentPendingStreamingRequests.Num());
 
-			for (FRayTracingGeometry* Geometry : CurrentPendingStreamingRequests)
+			for (RayTracingGeometryHandle GeometryHandle : CurrentPendingStreamingRequests)
 			{
-				if (!RequestRayTracingGeometryStreamIn(RHICmdList, Geometry))
+				if (!RequestRayTracingGeometryStreamIn(RHICmdList, GeometryHandle))
 				{
-					PendingStreamingRequests.Add(Geometry);
+					PendingStreamingRequests.Add(GeometryHandle);
 				}
 			}
 		}
@@ -539,11 +540,12 @@ void FRayTracingGeometryManager::Tick(FRHICommandList& RHICmdList)
 	SET_MEMORY_STAT(STAT_RayTracingGeometryResidentMemory, TotalResidentSize);
 }
 
-bool FRayTracingGeometryManager::RequestRayTracingGeometryStreamIn(FRHICommandList& RHICmdList, FRayTracingGeometry* Geometry)
+bool FRayTracingGeometryManager::RequestRayTracingGeometryStreamIn(FRHICommandList& RHICmdList, RayTracingGeometryHandle GeometryHandle)
 {
-	FRegisteredGeometry& RegisteredGeometry = RegisteredGeometries[Geometry->RayTracingGeometryHandle];
+	FRegisteredGeometry& RegisteredGeometry = RegisteredGeometries[GeometryHandle];
+	FRayTracingGeometry* Geometry = RegisteredGeometry.Geometry;
 
-	if (RegisteredGeometry.Geometry->Initializer.Type != ERayTracingGeometryInitializerType::StreamingDestination
+	if (Geometry->Initializer.Type != ERayTracingGeometryInitializerType::StreamingDestination
 		|| RegisteredGeometry.Status == FRegisteredGeometry::FStatus::Streaming)
 	{
 		// no streaming required or streaming request already in-flight
diff --git a/Engine/Source/Runtime/Engine/Public/Rendering/RayTracingGeometryManager.h b/Engine/Source/Runtime/Engine/Public/Rendering/RayTracingGeometryManager.h
--- a/Engine/Source/Runtime/Engine/Public/Rendering/RayTracingGeometryManager.h
+++ b/Engine/Source/Runtime/Engine/Public/Rendering/RayTracingGeometryManager.h
@@ -71,7 +71,7 @@ private:
 
 	void ReleaseRayTracingGeometryGroupReference(RayTracing::GeometryGroupHandle Handle);
 
-	bool RequestRayTracingGeometryStreamIn(FRHICommandList& RHICmdList, FRayTracingGeometry* Geometry);
+	bool RequestRayTracingGeometryStreamIn(FRHICommandList& RHICmdList, RayTracingGeometryHandle GeometryHandle);
 	void ProcessCompletedStreamingRequests(FRHICommandList& RHICmdList);
 
 	FCriticalSection RequestCS;
@@ -140,7 +140,7 @@ private:
 	TSet<RayTracingGeometryHandle> ReferencedGeometryHandles;
 	TSet<RayTracing::GeometryGroupHandle> ReferencedGeometryGroups;
 
-	TArray<FRayTracingGeometry*> PendingStreamingRequests;
+	TSet<RayTracingGeometryHandle> PendingStreamingRequests;
 
 	TArray<FStreamingRequest> StreamingRequests;
 	int32 NumStreamingRequests = 0;
