Changes to r.RayTracing.EnableOnDemand to work with offline data.
- moved ownership of RayTracingGeometry from FStaticMeshLODResources to FStaticMeshRayTracingProxyLOD so offline data can be streamed-in when ray tracing is toggled on.
- request streaming of ray tracing geometry when it is initally registered or FirstLODIndex changes.
- moved call to ProcessCompletedStreamingRequests(...) so it in every code path of FRayTracingGeometryManager::Tick(...)
- enabled r.RayTracing.EnableOnDemand by default on all platforms.

diff --git a/Engine/Config/Windows/WindowsEngine.ini b/Engine/Config/Windows/WindowsEngine.ini
--- a/Engine/Config/Windows/WindowsEngine.ini
+++ b/Engine/Config/Windows/WindowsEngine.ini
@@ -25,8 +25,6 @@ D3D12.Bindless.ResourceDescriptorHeapSize=100000
 D3D12.Bindless.SamplerDescriptorHeapSize=2048
 r.PSOPrecache.GlobalComputeShaders=1
 r.DynamicRes.DynamicFrameTime=1
-; Allow ray tracing to be toggled using r.RayTracing.Enable=1/0 without restarting
-r.RayTracing.EnableOnDemand=True
 
 [PlatformCrypto]
 PlatformRequiresDataCrypto=True
diff --git a/Engine/Source/Runtime/Engine/Private/Rendering/RayTracingGeometryManager.cpp b/Engine/Source/Runtime/Engine/Private/Rendering/RayTracingGeometryManager.cpp
--- a/Engine/Source/Runtime/Engine/Private/Rendering/RayTracingGeometryManager.cpp
+++ b/Engine/Source/Runtime/Engine/Private/Rendering/RayTracingGeometryManager.cpp
@@ -213,6 +213,11 @@ FRayTracingGeometryManager::RayTracingGeometryHandle FRayTracingGeometryManager:
 
 		Group.Geometries[InGeometry->LODIndex] = InGeometry;
 		++Group.NumReferences;
+
+		if (InGeometry->LODIndex >= Group.CurrentFirstLODIdx && IsRayTracingEnabled() && !IsRayTracingUsingReferenceBasedResidency())
+		{
+			PendingStreamingRequests.Add(Handle);
+		}
 	}
 		
 	INC_DWORD_STAT(STAT_RayTracingGeometryCount);
@@ -293,6 +298,16 @@ void FRayTracingGeometryManager::SetRayTracingGeometryGroupCurrentFirstLODIndex(
 			}
 		}
 	}
+	else if(IsRayTracingEnabled() && !IsRayTracingUsingReferenceBasedResidency())
+	{
+		for (int32 LODIdx = NewCurrentFirstLODIdx; LODIdx < Group.CurrentFirstLODIdx; ++LODIdx)
+		{
+			if (Group.Geometries[LODIdx])
+			{
+				PendingStreamingRequests.Add(Group.Geometries[LODIdx]->RayTracingGeometryHandle);
+			}
+		}
+	}
 
 	Group.CurrentFirstLODIdx = NewCurrentFirstLODIdx;
 }
@@ -485,8 +500,6 @@ void FRayTracingGeometryManager::Tick(FRHICommandList& RHICmdList)
 				Geometry->Evict();
 			}
 		}
-
-		ProcessCompletedStreamingRequests(RHICmdList);
 	}
 	else
 	{
@@ -532,6 +545,8 @@ void FRayTracingGeometryManager::Tick(FRHICommandList& RHICmdList)
 		}
 	}
 
+	ProcessCompletedStreamingRequests(RHICmdList);
+
 	ReferencedGeometryHandles.Reset();
 	ReferencedGeometryGroups.Reset();
 
diff --git a/Engine/Source/Runtime/Engine/Private/StaticMesh.cpp b/Engine/Source/Runtime/Engine/Private/StaticMesh.cpp
--- a/Engine/Source/Runtime/Engine/Private/StaticMesh.cpp
+++ b/Engine/Source/Runtime/Engine/Private/StaticMesh.cpp
@@ -2398,9 +2398,11 @@ void FStaticMeshRenderData::AllocateLODResources(int32 NumLODs)
 
 void FStaticMeshRenderData::InitializeRayTracingRepresentationFromRenderingLODs()
 {
+	const bool bProxyOwnsRayTracingGeometry = IsRayTracingEnableOnDemandSupported();
+
 	const int32 NumLODs = LODResources.Num();
 
-	if (!IsRayTracingUsingReferenceBasedResidency())
+	if (!bProxyOwnsRayTracingGeometry)
 	{
 		for (int32 LODIndex = 0; LODIndex < NumLODs; ++LODIndex)
 		{
@@ -2425,7 +2427,7 @@ void FStaticMeshRenderData::InitializeRayTracingRepresentationFromRenderingLODs(
 
 		FStaticMeshRayTracingProxyLOD* RayTracingLOD = new FStaticMeshRayTracingProxyLOD;
 
-		RayTracingLOD->bOwnsRayTracingGeometry = IsRayTracingUsingReferenceBasedResidency();
+		RayTracingLOD->bOwnsRayTracingGeometry = bProxyOwnsRayTracingGeometry;
 		RayTracingLOD->RayTracingGeometry = RayTracingLOD->bOwnsRayTracingGeometry ? new FRayTracingGeometry() : LODModel.RayTracingGeometry;
 
 		RayTracingLOD->Sections = &LODModel.Sections;
@@ -3254,7 +3256,7 @@ static FString BuildStaticMeshDerivedDataKeySuffix(const ITargetPlatform* Target
 	{
 		KeySuffix += TEXT("RT1");
 
-		if (IsRayTracingUsingReferenceBasedResidency())
+		if (IsRayTracingEnableOnDemandSupported())
 		{
 			KeySuffix += TEXT("_RBR");
 		}
diff --git a/Engine/Source/Runtime/RenderCore/Private/RenderUtils.cpp b/Engine/Source/Runtime/RenderCore/Private/RenderUtils.cpp
--- a/Engine/Source/Runtime/RenderCore/Private/RenderUtils.cpp
+++ b/Engine/Source/Runtime/RenderCore/Private/RenderUtils.cpp
@@ -41,7 +41,7 @@ FAutoConsoleVariableRef CVarAllowTranslucencyShadowsInProject(
 	ECVF_ReadOnly | ECVF_RenderThreadSafe
 );
 
-static int32 GRayTracingEnableOnDemand = 0;
+static int32 GRayTracingEnableOnDemand = 1;
 static FAutoConsoleVariableRef CVarRayTracingEnableOnDemand(
 	TEXT("r.RayTracing.EnableOnDemand"),
 	GRayTracingEnableOnDemand,
@@ -1664,6 +1664,11 @@ bool IsRayTracingUsingReferenceBasedResidency()
 	return CVar && CVar->GetValueOnAnyThread();
 }
 
+bool IsRayTracingEnableOnDemandSupported()
+{
+	return GRayTracingEnableOnDemand != 0;
+}
+
 bool UseSplineMeshSceneResources(const FStaticShaderPlatform Platform)
 {
 	// This feature currently requires GPU Scene to be supported and enabled.
diff --git a/Engine/Source/Runtime/RenderCore/Public/RenderUtils.h b/Engine/Source/Runtime/RenderCore/Public/RenderUtils.h
--- a/Engine/Source/Runtime/RenderCore/Public/RenderUtils.h
+++ b/Engine/Source/Runtime/RenderCore/Public/RenderUtils.h
@@ -540,6 +540,9 @@ extern RENDERCORE_API ERayTracingMode GetRayTracingMode();
 // In which case the runtime should track which ray tracing geometries are referenced in the TLAS
 RENDERCORE_API bool IsRayTracingUsingReferenceBasedResidency();
 
+// Returns 'true' when ray tracing can be toggled on/off at runtime.
+RENDERCORE_API bool IsRayTracingEnableOnDemandSupported();
+
 // Returns 'true' when static lighting is enabled for the project
 inline bool IsStaticLightingAllowed()
 {
