Stats to track number of pending and in-flight streaming requests in FRayTracingGeometryManager.

diff --git a/Engine/Source/Runtime/Engine/Private/Rendering/RayTracingGeometryManager.cpp b/Engine/Source/Runtime/Engine/Private/Rendering/RayTracingGeometryManager.cpp
--- a/Engine/Source/Runtime/Engine/Private/Rendering/RayTracingGeometryManager.cpp
+++ b/Engine/Source/Runtime/Engine/Private/Rendering/RayTracingGeometryManager.cpp
@@ -86,6 +86,8 @@ DECLARE_MEMORY_STAT(TEXT("Resident Memory"), STAT_RayTracingGeometryResidentMemo
 DECLARE_DWORD_ACCUMULATOR_STAT(TEXT("Pending Builds"), STAT_RayTracingPendingBuilds, STATGROUP_RayTracingGeometry);
 DECLARE_DWORD_ACCUMULATOR_STAT(TEXT("Pending Build Primitives"), STAT_RayTracingPendingBuildPrimitives, STATGROUP_RayTracingGeometry);
 
+DECLARE_DWORD_ACCUMULATOR_STAT(TEXT("Pending Streaming Requests"), STAT_RayTracingPendingStreamingRequests, STATGROUP_RayTracingGeometry);
+DECLARE_DWORD_ACCUMULATOR_STAT(TEXT("In-flight Streaming Requests"), STAT_RayTracingInflightStreamingRequests, STATGROUP_RayTracingGeometry);
 
 FRayTracingGeometryManager::FRayTracingGeometryManager()
 {
@@ -217,6 +219,7 @@ FRayTracingGeometryManager::RayTracingGeometryHandle FRayTracingGeometryManager:
 		if (InGeometry->LODIndex >= Group.CurrentFirstLODIdx && IsRayTracingEnabled() && !IsRayTracingUsingReferenceBasedResidency())
 		{
 			PendingStreamingRequests.Add(Handle);
+			INC_DWORD_STAT(STAT_RayTracingPendingStreamingRequests);
 		}
 	}
 		
@@ -260,7 +263,10 @@ void FRayTracingGeometryManager::ReleaseRayTracingGeometryHandle(RayTracingGeome
 
 	RegisteredGeometries.RemoveAt(Handle);
 	ReferencedGeometryHandles.Remove(Handle);
-	PendingStreamingRequests.Remove(Handle);
+	if (PendingStreamingRequests.Remove(Handle) > 0)
+	{
+		DEC_DWORD_STAT(STAT_RayTracingPendingStreamingRequests);
+	}
 
 	DEC_DWORD_STAT(STAT_RayTracingGeometryCount);
 }
@@ -305,6 +311,7 @@ void FRayTracingGeometryManager::SetRayTracingGeometryGroupCurrentFirstLODIndex(
 			if (Group.Geometries[LODIdx])
 			{
 				PendingStreamingRequests.Add(Group.Geometries[LODIdx]->RayTracingGeometryHandle);
+				INC_DWORD_STAT(STAT_RayTracingPendingStreamingRequests);
 			}
 		}
 	}
@@ -400,6 +407,8 @@ void FRayTracingGeometryManager::Tick(FRHICommandList& RHICmdList)
 			}
 
 			PendingStreamingRequests.Empty();
+
+			SET_DWORD_STAT(STAT_RayTracingPendingStreamingRequests, 0);
 		}
 		else
 		{
@@ -518,6 +527,8 @@ void FRayTracingGeometryManager::Tick(FRHICommandList& RHICmdList)
 				if (!RequestRayTracingGeometryStreamIn(RHICmdList, RegisteredGeometry.Geometry->RayTracingGeometryHandle))
 				{
 					PendingStreamingRequests.Add(RegisteredGeometry.Geometry->RayTracingGeometryHandle);
+
+					INC_DWORD_STAT(STAT_RayTracingPendingStreamingRequests);
 				}
 			}
 		}
@@ -542,6 +553,8 @@ void FRayTracingGeometryManager::Tick(FRHICommandList& RHICmdList)
 					PendingStreamingRequests.Add(GeometryHandle);
 				}
 			}
+
+			SET_DWORD_STAT(STAT_RayTracingPendingStreamingRequests, PendingStreamingRequests.Num());
 		}
 	}
 
@@ -626,6 +639,8 @@ bool FRayTracingGeometryManager::RequestRayTracingGeometryStreamIn(FRHICommandLi
 		NextStreamingRequestIndex = (NextStreamingRequestIndex + 1) % GRayTracingStreamingMaxPendingRequests;
 		++NumStreamingRequests;
 
+		INC_DWORD_STAT(STAT_RayTracingInflightStreamingRequests);
+
 		StreamingRequest.GeometryHandle = Geometry->RayTracingGeometryHandle;
 		StreamingRequest.RequestBuffer = FIoBuffer(RegisteredGeometry.StreamableDataSize); // TODO: Use FIoBuffer::Wrap with preallocated memory
 
@@ -740,6 +755,8 @@ void FRayTracingGeometryManager::ProcessCompletedStreamingRequests(FRHICommandLi
 	}
 
 	NumStreamingRequests -= NumCompletedRequests;
+
+	SET_DWORD_STAT(STAT_RayTracingInflightStreamingRequests, NumStreamingRequests);
 }
 
 void FRayTracingGeometryManager::BoostPriority(BuildRequestIndex InRequestIndex, float InBoostValue)
