[World Partition] Implemented asynchronous update streaming state.
Can be turned on using wp.Runtime.UpdateStreaming.EnableAsyncUpdate true (disabled by default).
This allows to run the computation of the world partition streaming cells to load/activate/unload outside of the game thread / critical path.
Update of the cells underlying level streaming state and priority is still done in game thread.

diff --git a/Engine/Source/Runtime/Engine/Private/LevelInstance/LevelInstanceSubsystem.cpp b/Engine/Source/Runtime/Engine/Private/LevelInstance/LevelInstanceSubsystem.cpp
--- a/Engine/Source/Runtime/Engine/Private/LevelInstance/LevelInstanceSubsystem.cpp
+++ b/Engine/Source/Runtime/Engine/Private/LevelInstance/LevelInstanceSubsystem.cpp
@@ -257,6 +257,8 @@ bool ULevelInstanceSubsystem::IsLoading(const ILevelInstanceInterface* LevelInst
 
 void ULevelInstanceSubsystem::UpdateStreamingState()
 {
+	TRACE_CPUPROFILER_EVENT_SCOPE(ULevelInstanceSubsystem::UpdateStreamingState);
+
 	if (!LevelInstancesToUnload.Num() && !LevelInstancesToLoadOrUpdate.Num())
 	{
 		return;
diff --git a/Engine/Source/Runtime/Engine/Private/WorldPartition/DataLayer/DataLayerManager.cpp b/Engine/Source/Runtime/Engine/Private/WorldPartition/DataLayer/DataLayerManager.cpp
--- a/Engine/Source/Runtime/Engine/Private/WorldPartition/DataLayer/DataLayerManager.cpp
+++ b/Engine/Source/Runtime/Engine/Private/WorldPartition/DataLayer/DataLayerManager.cpp
@@ -364,13 +364,19 @@ const TSet<FName>& UDataLayerManager::GetEffectiveLoadedDataLayerNames() const
 }
 
 bool UDataLayerManager::IsAnyDataLayerInEffectiveRuntimeState(TArrayView<const FName> InDataLayerNames, EDataLayerRuntimeState InState) const
+{
+	static FWorldDataLayersEffectiveStates EmptyStates;
+	AWorldDataLayers* WorldDataLayers = GetWorldDataLayers();
+	return UDataLayerManager::IsAnyDataLayerInEffectiveRuntimeState(InDataLayerNames, InState, WorldDataLayers ? FWorldDataLayersEffectiveStatesAccessor::Get(WorldDataLayers) : EmptyStates);
+}
+
+bool UDataLayerManager::IsAnyDataLayerInEffectiveRuntimeState(TArrayView<const FName> InDataLayerNames, EDataLayerRuntimeState InState, const FWorldDataLayersEffectiveStates& InEffectiveStates)
 {
 	if (InState == EDataLayerRuntimeState::Activated)
 	{
-		const TSet<FName>& Activated = GetEffectiveActiveDataLayerNames();
 		for (const FName& DataLayerName : InDataLayerNames)
 		{
-			if (Activated.Contains(DataLayerName))
+			if (InEffectiveStates.GetAllEffectiveActiveDataLayerNames().Contains(DataLayerName))
 			{
 				return true;
 			}
@@ -378,10 +384,9 @@ bool UDataLayerManager::IsAnyDataLayerInEffectiveRuntimeState(TArrayView<const F
 	}
 	else if (InState == EDataLayerRuntimeState::Loaded)
 	{
-		const TSet<FName>& Loaded = GetEffectiveLoadedDataLayerNames();
 		for (const FName& DataLayerName : InDataLayerNames)
 		{
-			if (Loaded.Contains(DataLayerName))
+			if (InEffectiveStates.GetAllEffectiveLoadedDataLayerNames().Contains(DataLayerName))
 			{
 				return true;
 			}
@@ -392,13 +397,18 @@ bool UDataLayerManager::IsAnyDataLayerInEffectiveRuntimeState(TArrayView<const F
 
 bool UDataLayerManager::IsAllDataLayerInEffectiveRuntimeState(TArrayView<const FName> InDataLayerNames, EDataLayerRuntimeState InState) const
 {
-	const TSet<FName>& Activated = GetEffectiveActiveDataLayerNames();
+	static FWorldDataLayersEffectiveStates EmptyStates;
+	AWorldDataLayers* WorldDataLayers = GetWorldDataLayers();
+	return UDataLayerManager::IsAllDataLayerInEffectiveRuntimeState(InDataLayerNames, InState, WorldDataLayers ? FWorldDataLayersEffectiveStatesAccessor::Get(WorldDataLayers) : EmptyStates);
+}
 
+bool UDataLayerManager::IsAllDataLayerInEffectiveRuntimeState(TArrayView<const FName> InDataLayerNames, EDataLayerRuntimeState InState, const FWorldDataLayersEffectiveStates& InEffectiveStates)
+{
 	if (InState == EDataLayerRuntimeState::Activated)
 	{
 		for (const FName& DataLayerName : InDataLayerNames)
 		{
-			if (!Activated.Contains(DataLayerName))
+			if (!InEffectiveStates.GetAllEffectiveActiveDataLayerNames().Contains(DataLayerName))
 			{
 				return false;
 			}
@@ -406,10 +416,9 @@ bool UDataLayerManager::IsAllDataLayerInEffectiveRuntimeState(TArrayView<const F
 	}
 	else if (InState == EDataLayerRuntimeState::Loaded)
 	{
-		const TSet<FName>& Loaded = GetEffectiveLoadedDataLayerNames();
 		for (const FName& DataLayerName : InDataLayerNames)
 		{
-			if (!Loaded.Contains(DataLayerName) && !Activated.Contains(DataLayerName))
+			if (!InEffectiveStates.GetAllEffectiveLoadedDataLayerNames().Contains(DataLayerName) && !InEffectiveStates.GetAllEffectiveActiveDataLayerNames().Contains(DataLayerName))
 			{
 				return false;
 			}
diff --git a/Engine/Source/Runtime/Engine/Private/WorldPartition/DataLayer/WorldDataLayers.cpp b/Engine/Source/Runtime/Engine/Private/WorldPartition/DataLayer/WorldDataLayers.cpp
--- a/Engine/Source/Runtime/Engine/Private/WorldPartition/DataLayer/WorldDataLayers.cpp
+++ b/Engine/Source/Runtime/Engine/Private/WorldPartition/DataLayer/WorldDataLayers.cpp
@@ -15,6 +15,7 @@
 #include "WorldPartition/WorldPartition.h"
 #include "WorldPartition/WorldPartitionLog.h"
 #include "WorldPartition/WorldPartitionRuntimeCellInterface.h"
+#include "WorldPartition/WorldPartitionStreamingPolicy.h"
 #include "Engine/Level.h"
 #include "ProfilingDebugging/CsvProfiler.h"
 #include "Net/UnrealNetwork.h"
@@ -67,9 +68,6 @@ AWorldDataLayers::AWorldDataLayers(const FObjectInitializer& ObjectInitializer)
 	// Avoid actor from being Destroyed/Recreated when scrubbing a replay
 	// instead AWorldDataLayers::RewindForReplay() gets called to reset this actors state
 	bReplayRewindable = true;
-
-	AllEffectiveActiveDataLayerNamesEpoch = MAX_int32;
-	AllEffectiveLoadedDataLayerNamesEpoch = MAX_int32;
 }
 
 void AWorldDataLayers::GetLifetimeReplicatedProps(TArray< FLifetimeProperty >& OutLifetimeProps) const
@@ -157,8 +155,8 @@ void AWorldDataLayers::InitializeDataLayerRuntimeStates()
 			ResolveEffectiveRuntimeState(DataLayerInstance, bNotifyChange);
 		}
 
-		AWORLDDATALAYERS_UPDATE_REPLICATED_DATALAYERS(RepEffectiveActiveDataLayerNames, EffectiveActiveDataLayerNames.Array());
-		AWORLDDATALAYERS_UPDATE_REPLICATED_DATALAYERS(RepEffectiveLoadedDataLayerNames, EffectiveLoadedDataLayerNames.Array());
+		AWORLDDATALAYERS_UPDATE_REPLICATED_DATALAYERS(RepEffectiveActiveDataLayerNames, EffectiveStates.GetReplicatedEffectiveActiveDataLayerNames().Array());
+		AWORLDDATALAYERS_UPDATE_REPLICATED_DATALAYERS(RepEffectiveLoadedDataLayerNames, EffectiveStates.GetReplicatedEffectiveLoadedDataLayerNames().Array());
 
 		UE_CLOG(RepEffectiveActiveDataLayerNames.Num() || RepEffectiveLoadedDataLayerNames.Num(), LogWorldPartition, Log, TEXT("Initial Data Layer Effective States Activated(%s) Loaded(%s)"), *JoinDataLayerShortNamesFromInstanceNames(this, RepEffectiveActiveDataLayerNames), *JoinDataLayerShortNamesFromInstanceNames(this, RepEffectiveLoadedDataLayerNames));
 	}
@@ -171,11 +169,8 @@ void AWorldDataLayers::ResetDataLayerRuntimeStates()
 	LocalActiveDataLayerNames.Reset();
 	LocalLoadedDataLayerNames.Reset();
 
-	EffectiveActiveDataLayerNames.Reset();
-	EffectiveLoadedDataLayerNames.Reset();
-	LocalEffectiveActiveDataLayerNames.Reset();
-	LocalEffectiveLoadedDataLayerNames.Reset();
-
+	EffectiveStates.Reset();
+	
 	static const TArray<FName> Empty;
 	FlushNetDormancy();
 	AWORLDDATALAYERS_UPDATE_REPLICATED_DATALAYERS(RepActiveDataLayerNames, Empty);
@@ -326,63 +321,36 @@ EDataLayerRuntimeState AWorldDataLayers::GetDataLayerRuntimeStateByName(FName In
 void AWorldDataLayers::OnRep_EffectiveActiveDataLayerNames()
 {
 	++DataLayersStateEpoch;
-	EffectiveActiveDataLayerNames.Reset();
-	EffectiveActiveDataLayerNames.Append(RepEffectiveActiveDataLayerNames);
+	// TSet do not support replication so we replicate an array and update the set here
+	EffectiveStates.SetReplicatedEffectiveActiveDataLayerNames(RepEffectiveActiveDataLayerNames);
 }
 
 void AWorldDataLayers::OnRep_EffectiveLoadedDataLayerNames()
 {
 	++DataLayersStateEpoch;
-	EffectiveLoadedDataLayerNames.Reset();
-	EffectiveLoadedDataLayerNames.Append(RepEffectiveLoadedDataLayerNames);
+	// TSet do not support replication so we replicate an array and update the set here
+	EffectiveStates.SetReplicatedEffectiveLoadedDataLayerNames(RepEffectiveLoadedDataLayerNames);
 }
 
 EDataLayerRuntimeState AWorldDataLayers::GetDataLayerEffectiveRuntimeStateByName(FName InDataLayerName) const
 {
-	if (EffectiveActiveDataLayerNames.Contains(InDataLayerName))
-	{
-		check(!EffectiveLoadedDataLayerNames.Contains(InDataLayerName));
-		return EDataLayerRuntimeState::Activated;
-	}
-	else if (EffectiveLoadedDataLayerNames.Contains(InDataLayerName))
-	{
-		check(!EffectiveActiveDataLayerNames.Contains(InDataLayerName));
-		return EDataLayerRuntimeState::Loaded;
-	}
-	else if (LocalEffectiveActiveDataLayerNames.Contains(InDataLayerName))
-	{
-		check(!LocalEffectiveLoadedDataLayerNames.Contains(InDataLayerName));
-		return EDataLayerRuntimeState::Activated;
-	}
-	else if (LocalEffectiveLoadedDataLayerNames.Contains(InDataLayerName))
-	{
-		check(!LocalEffectiveActiveDataLayerNames.Contains(InDataLayerName));
-		return EDataLayerRuntimeState::Loaded;
-	}
+	return EffectiveStates.GetDataLayerEffectiveRuntimeStateByName(InDataLayerName);
+}
 
-	return EDataLayerRuntimeState::Unloaded;
+const FWorldDataLayersEffectiveStates& AWorldDataLayers::GetEffectiveStates() const
+{
+	check(IsInGameThread());
+	return EffectiveStates;
 }
 
 const TSet<FName>& AWorldDataLayers::GetEffectiveActiveDataLayerNames() const
 {
-	if (AllEffectiveActiveDataLayerNamesEpoch != DataLayersStateEpoch)
-	{
-		AllEffectiveActiveDataLayerNames = EffectiveActiveDataLayerNames;
-		AllEffectiveActiveDataLayerNames.Append(LocalEffectiveActiveDataLayerNames);
-		AllEffectiveActiveDataLayerNamesEpoch = DataLayersStateEpoch;
-	}
-	return AllEffectiveActiveDataLayerNames;
+	return EffectiveStates.GetAllEffectiveActiveDataLayerNames();
 }
 
 const TSet<FName>& AWorldDataLayers::GetEffectiveLoadedDataLayerNames() const
 {
-	if (AllEffectiveLoadedDataLayerNamesEpoch != DataLayersStateEpoch)
-	{
-		AllEffectiveLoadedDataLayerNames = EffectiveLoadedDataLayerNames;
-		AllEffectiveLoadedDataLayerNames.Append(LocalEffectiveLoadedDataLayerNames);
-		AllEffectiveLoadedDataLayerNamesEpoch = DataLayersStateEpoch;
-	}
-	return AllEffectiveLoadedDataLayerNames;
+	return EffectiveStates.GetAllEffectiveLoadedDataLayerNames();
 }
 
 void AWorldDataLayers::ResolveEffectiveRuntimeState(const UDataLayerInstance* InDataLayerInstance, bool bInNotifyChange)
@@ -393,7 +361,6 @@ void AWorldDataLayers::ResolveEffectiveRuntimeState(const UDataLayerInstance* In
 	const bool bDataLayerServerOnly = InDataLayerInstance->IsServerOnly();
 	const bool bIsLocalDataLayer = (bDataLayerClientOnly || bDataLayerServerOnly);
 	const FName DataLayerName = InDataLayerInstance->GetDataLayerFName();
-	EDataLayerRuntimeState CurrentEffectiveRuntimeState = GetDataLayerEffectiveRuntimeStateByName(DataLayerName);
 	EDataLayerRuntimeState NewEffectiveRuntimeState = GetDataLayerRuntimeStateByName(DataLayerName);
 	const UDataLayerInstance* Parent = InDataLayerInstance->GetParent();
 
@@ -407,29 +374,15 @@ void AWorldDataLayers::ResolveEffectiveRuntimeState(const UDataLayerInstance* In
 		Parent = Parent->GetParent();
 	}
 
-	if (CurrentEffectiveRuntimeState != NewEffectiveRuntimeState)
+	EDataLayerRuntimeState OldEffectiveRuntimeState;
+	if (EffectiveStates.SetDataLayerEffectiveRuntimeState(DataLayerName, bIsLocalDataLayer, NewEffectiveRuntimeState, OldEffectiveRuntimeState))
 	{
-		TSet<FName>& TargetEffectiveLoadedDataLayerNames = bIsLocalDataLayer ? LocalEffectiveLoadedDataLayerNames : EffectiveLoadedDataLayerNames;
-		TSet<FName>& TargetEffectiveActiveDataLayerNames = bIsLocalDataLayer ? LocalEffectiveActiveDataLayerNames : EffectiveActiveDataLayerNames;
-
-		TargetEffectiveLoadedDataLayerNames.Remove(DataLayerName);
-		TargetEffectiveActiveDataLayerNames.Remove(DataLayerName);
-
-		if (NewEffectiveRuntimeState == EDataLayerRuntimeState::Loaded)
-		{
-			TargetEffectiveLoadedDataLayerNames.Add(DataLayerName);
-		}
-		else if (NewEffectiveRuntimeState == EDataLayerRuntimeState::Activated)
-		{
-			TargetEffectiveActiveDataLayerNames.Add(DataLayerName);
-		}
-
 		// Update Replicated Properties
 		if (!bIsLocalDataLayer && (NetMode == NM_DedicatedServer || NetMode == NM_ListenServer))
 		{
 			FlushNetDormancy();
-			AWORLDDATALAYERS_UPDATE_REPLICATED_DATALAYERS(RepEffectiveActiveDataLayerNames, EffectiveActiveDataLayerNames.Array());
-			AWORLDDATALAYERS_UPDATE_REPLICATED_DATALAYERS(RepEffectiveLoadedDataLayerNames, EffectiveLoadedDataLayerNames.Array());
+			AWORLDDATALAYERS_UPDATE_REPLICATED_DATALAYERS(RepEffectiveActiveDataLayerNames, EffectiveStates.GetReplicatedEffectiveActiveDataLayerNames().Array());
+			AWORLDDATALAYERS_UPDATE_REPLICATED_DATALAYERS(RepEffectiveLoadedDataLayerNames, EffectiveStates.GetReplicatedEffectiveLoadedDataLayerNames().Array());
 		}
 
 		++DataLayersStateEpoch;
@@ -438,7 +391,7 @@ void AWorldDataLayers::ResolveEffectiveRuntimeState(const UDataLayerInstance* In
 		{
 			UE_LOG(LogWorldPartition, Log, TEXT("Data Layer Instance '%s' effective state changed: %s -> %s"),
 				*InDataLayerInstance->GetDataLayerShortName(),
-				*StaticEnum<EDataLayerRuntimeState>()->GetDisplayNameTextByValue((int64)CurrentEffectiveRuntimeState).ToString(),
+				*StaticEnum<EDataLayerRuntimeState>()->GetDisplayNameTextByValue((int64)OldEffectiveRuntimeState).ToString(),
 				*StaticEnum<EDataLayerRuntimeState>()->GetDisplayNameTextByValue((int64)NewEffectiveRuntimeState).ToString());
 
 			OnDataLayerRuntimeStateChanged(InDataLayerInstance, NewEffectiveRuntimeState);
@@ -571,12 +524,12 @@ void AWorldDataLayers::DumpDataLayers(FOutputDevice& OutputDevice) const
 			}
 		};
 
-		if (EffectiveLoadedDataLayerNames.Num() || EffectiveActiveDataLayerNames.Num())
+		if (EffectiveStates.GetAllEffectiveLoadedDataLayerNames().Num() || EffectiveStates.GetAllEffectiveActiveDataLayerNames().Num())
 		{
 			OutputDevice.Logf(TEXT("----------------------------------------------------"));
 			OutputDevice.Logf(TEXT(" Data Layers Runtime States"));
-			DumpDataLayersRuntimeState(TEXT("Loaded"), EffectiveLoadedDataLayerNames);
-			DumpDataLayersRuntimeState(TEXT("Active"), EffectiveActiveDataLayerNames);
+			DumpDataLayersRuntimeState(TEXT("Loaded"), EffectiveStates.GetAllEffectiveLoadedDataLayerNames());
+			DumpDataLayersRuntimeState(TEXT("Active"), EffectiveStates.GetAllEffectiveActiveDataLayerNames());
 			OutputDevice.Logf(TEXT("----------------------------------------------------"));
 			OutputDevice.Logf(TEXT(""));
 		}
@@ -1693,4 +1646,107 @@ TArray<const UDataLayerInstance*> AWorldDataLayers::GetDataLayerInstances(const
 
 //~ End Deprecated
 
+FWorldDataLayersEffectiveStates::FWorldDataLayersEffectiveStates()
+	: UpdateEpoch(0)
+	, AllEffectiveActiveDataLayerNamesEpoch(MAX_int32)
+	, AllEffectiveLoadedDataLayerNamesEpoch(MAX_int32)
+{}
+
+void FWorldDataLayersEffectiveStates::SetReplicatedEffectiveActiveDataLayerNames(const TArray<FName>& InRepEffectiveActiveDataLayerNames)
+{
+	ReplicatedEffectiveActiveDataLayerNames.Reset();
+	ReplicatedEffectiveActiveDataLayerNames.Append(InRepEffectiveActiveDataLayerNames);
+	++UpdateEpoch;
+}
+
+void FWorldDataLayersEffectiveStates::SetReplicatedEffectiveLoadedDataLayerNames(const TArray<FName>& InRepEffectiveLoadedDataLayerNames)
+{
+	ReplicatedEffectiveLoadedDataLayerNames.Reset();
+	ReplicatedEffectiveLoadedDataLayerNames.Append(InRepEffectiveLoadedDataLayerNames);
+	++UpdateEpoch;
+}
+
+bool FWorldDataLayersEffectiveStates::SetDataLayerEffectiveRuntimeState(FName InDataLayerName, bool bIsLocalDataLayer, EDataLayerRuntimeState NewEffectiveRuntimeState, EDataLayerRuntimeState& OutOldEffectiveRuntimeState)
+{
+	OutOldEffectiveRuntimeState = GetDataLayerEffectiveRuntimeStateByName(InDataLayerName);
+	if (OutOldEffectiveRuntimeState != NewEffectiveRuntimeState)
+	{
+		TSet<FName>& TargetEffectiveLoadedDataLayerNames = bIsLocalDataLayer ? LocalEffectiveLoadedDataLayerNames : ReplicatedEffectiveLoadedDataLayerNames;
+		TSet<FName>& TargetEffectiveActiveDataLayerNames = bIsLocalDataLayer ? LocalEffectiveActiveDataLayerNames : ReplicatedEffectiveActiveDataLayerNames;
+
+		TargetEffectiveLoadedDataLayerNames.Remove(InDataLayerName);
+		TargetEffectiveActiveDataLayerNames.Remove(InDataLayerName);
+
+		if (NewEffectiveRuntimeState == EDataLayerRuntimeState::Loaded)
+		{
+			TargetEffectiveLoadedDataLayerNames.Add(InDataLayerName);
+		}
+		else if (NewEffectiveRuntimeState == EDataLayerRuntimeState::Activated)
+		{
+			TargetEffectiveActiveDataLayerNames.Add(InDataLayerName);
+		}
+
+		++UpdateEpoch;
+		return true;
+	}
+	return false;
+}
+
+void FWorldDataLayersEffectiveStates::Reset()
+{
+	ReplicatedEffectiveActiveDataLayerNames.Reset();
+	ReplicatedEffectiveLoadedDataLayerNames.Reset();
+	LocalEffectiveActiveDataLayerNames.Reset();
+	LocalEffectiveLoadedDataLayerNames.Reset();
+	++UpdateEpoch;
+}
+
+EDataLayerRuntimeState FWorldDataLayersEffectiveStates::GetDataLayerEffectiveRuntimeStateByName(FName InDataLayerName) const
+{
+	if (ReplicatedEffectiveActiveDataLayerNames.Contains(InDataLayerName))
+	{
+		check(!ReplicatedEffectiveLoadedDataLayerNames.Contains(InDataLayerName));
+		return EDataLayerRuntimeState::Activated;
+	}
+	else if (ReplicatedEffectiveLoadedDataLayerNames.Contains(InDataLayerName))
+	{
+		check(!ReplicatedEffectiveActiveDataLayerNames.Contains(InDataLayerName));
+		return EDataLayerRuntimeState::Loaded;
+	}
+	else if (LocalEffectiveActiveDataLayerNames.Contains(InDataLayerName))
+	{
+		check(!LocalEffectiveLoadedDataLayerNames.Contains(InDataLayerName));
+		return EDataLayerRuntimeState::Activated;
+	}
+	else if (LocalEffectiveLoadedDataLayerNames.Contains(InDataLayerName))
+	{
+		check(!LocalEffectiveActiveDataLayerNames.Contains(InDataLayerName));
+		return EDataLayerRuntimeState::Loaded;
+	}
+
+	return EDataLayerRuntimeState::Unloaded;
+}
+
+const TSet<FName>& FWorldDataLayersEffectiveStates::GetAllEffectiveActiveDataLayerNames() const
+{
+	if (AllEffectiveActiveDataLayerNamesEpoch != UpdateEpoch)
+	{
+		AllEffectiveActiveDataLayerNames = ReplicatedEffectiveActiveDataLayerNames;
+		AllEffectiveActiveDataLayerNames.Append(LocalEffectiveActiveDataLayerNames);
+		AllEffectiveActiveDataLayerNamesEpoch = UpdateEpoch;
+	}
+	return AllEffectiveActiveDataLayerNames;
+}
+
+const TSet<FName>& FWorldDataLayersEffectiveStates::GetAllEffectiveLoadedDataLayerNames() const
+{
+	if (AllEffectiveLoadedDataLayerNamesEpoch != UpdateEpoch)
+	{
+		AllEffectiveLoadedDataLayerNames = ReplicatedEffectiveLoadedDataLayerNames;
+		AllEffectiveLoadedDataLayerNames.Append(LocalEffectiveLoadedDataLayerNames);
+		AllEffectiveLoadedDataLayerNamesEpoch = UpdateEpoch;
+	}
+	return AllEffectiveLoadedDataLayerNames;
+}
+
 #undef LOCTEXT_NAMESPACE
diff --git a/Engine/Source/Runtime/Engine/Private/WorldPartition/RuntimeHashSet/WorldPartitionRuntimeHashSet.cpp b/Engine/Source/Runtime/Engine/Private/WorldPartition/RuntimeHashSet/WorldPartitionRuntimeHashSet.cpp
--- a/Engine/Source/Runtime/Engine/Private/WorldPartition/RuntimeHashSet/WorldPartitionRuntimeHashSet.cpp
+++ b/Engine/Source/Runtime/Engine/Private/WorldPartition/RuntimeHashSet/WorldPartitionRuntimeHashSet.cpp
@@ -7,7 +7,6 @@
 #include "WorldPartition/RuntimeHashSet/RuntimePartitionPersistent.h"
 #include "WorldPartition/WorldPartitionRuntimeLevelStreamingCell.h"
 #include "WorldPartition/WorldPartitionLevelStreamingDynamic.h"
-#include "WorldPartition/DataLayer/WorldDataLayers.h"
 #include "WorldPartition/HLOD/HLODLayer.h"
 #include "WorldPartition/ContentBundle/ContentBundleDescriptor.h"
 #include "WorldPartition/DataLayer/DataLayersID.h"
@@ -517,11 +516,12 @@ void UWorldPartitionRuntimeHashSet::ForEachStreamingCellsQuery(const FWorldParti
 	});
 }
 
-void UWorldPartitionRuntimeHashSet::ForEachStreamingCellsSources(const TArray<FWorldPartitionStreamingSource>& Sources, TFunctionRef<bool(const UWorldPartitionRuntimeCell*, EStreamingSourceTargetState)> Func) const
+void UWorldPartitionRuntimeHashSet::ForEachStreamingCellsSources(const TArray<FWorldPartitionStreamingSource>& Sources, TFunctionRef<bool(const UWorldPartitionRuntimeCell*, EStreamingSourceTargetState)> Func, const FWorldPartitionStreamingContext& InContext) const
 {
-	const UWorld* OuterWorld = GetTypedOuter<UWorld>();
-	const AWorldDataLayers* WorldDataLayers = OuterWorld->GetWorldDataLayers();
-	const int32 DataLayersStateEpoch = WorldDataLayers->GetDataLayersStateEpoch();
+	// Build a context when none is provided (for backward compatibility)
+	const FWorldPartitionStreamingContext StackContext = !InContext.IsValid() ? FWorldPartitionStreamingContext::Create(GetTypedOuter<UWorld>()) : FWorldPartitionStreamingContext();
+	const FWorldPartitionStreamingContext& Context = InContext.IsValid() ? InContext : StackContext;
+	check(Context.IsValid());
 
 	// Non-spatially loaded cells
 	for (const FRuntimePartitionStreamingData* StreamingData : RuntimeNonSpatiallyLoadedDataGridList)
@@ -536,7 +536,7 @@ void UWorldPartitionRuntimeHashSet::ForEachStreamingCellsSources(const TArray<FW
 #else
 			check(IsCellRelevantFor(Cell->GetClientOnlyVisible()));
 #endif
-			const EDataLayerRuntimeState CellEffectiveWantedState = Cell->GetCellEffectiveWantedState(DataLayersStateEpoch);
+			const EDataLayerRuntimeState CellEffectiveWantedState = Cell->GetCellEffectiveWantedState(Context);
 			if (CellEffectiveWantedState != EDataLayerRuntimeState::Unloaded)
 			{
 				Func(Cell, (CellEffectiveWantedState == EDataLayerRuntimeState::Loaded) ? EStreamingSourceTargetState::Loaded : EStreamingSourceTargetState::Activated);
@@ -579,11 +579,11 @@ void UWorldPartitionRuntimeHashSet::ForEachStreamingCellsSources(const TArray<FW
 
 				for (const FRuntimePartitionStreamingData* StreamingData : *StreamingDataList)
 				{
-					Source.ForEachShape(StreamingData->GetLoadingRange(), false, [this, &Source, StreamingData, DataLayersStateEpoch, &Func](const FSphericalSector& Shape)
+					Source.ForEachShape(StreamingData->GetLoadingRange(), false, [this, &Source, StreamingData, &Context, &Func](const FSphericalSector& Shape)
 					{
 						const FSphere ShapeSphere(Shape.GetCenter(), Shape.GetRadius());
 
-						auto ForEachIntersectingElementFunc = [this, &Source, &Shape, DataLayersStateEpoch, &Func](UWorldPartitionRuntimeCell* Cell)
+						auto ForEachIntersectingElementFunc = [this, &Source, &Shape, &Context, &Func](UWorldPartitionRuntimeCell* Cell)
 						{
 #if WITH_EDITOR
 							if (!IsCellRelevantFor(Cell->GetClientOnlyVisible()))
@@ -593,10 +593,10 @@ void UWorldPartitionRuntimeHashSet::ForEachStreamingCellsSources(const TArray<FW
 #else
 							check(IsCellRelevantFor(Cell->GetClientOnlyVisible()));
 #endif
-							const EDataLayerRuntimeState CellEffectiveWantedState = Cell->GetCellEffectiveWantedState(DataLayersStateEpoch);
+							const EDataLayerRuntimeState CellEffectiveWantedState = Cell->GetCellEffectiveWantedState(Context);
 							if (CellEffectiveWantedState != EDataLayerRuntimeState::Unloaded)
 							{
-								Cell->AppendStreamingSourceInfo(Source, Shape);
+								Cell->AppendStreamingSourceInfo(Source, Shape, Context);
 								Func(Cell, ((CellEffectiveWantedState == EDataLayerRuntimeState::Loaded) || (Source.TargetState == EStreamingSourceTargetState::Loaded)) ? EStreamingSourceTargetState::Loaded : EStreamingSourceTargetState::Activated);
 							}
 						};
diff --git a/Engine/Source/Runtime/Engine/Private/WorldPartition/WorldPartition.cpp b/Engine/Source/Runtime/Engine/Private/WorldPartition/WorldPartition.cpp
--- a/Engine/Source/Runtime/Engine/Private/WorldPartition/WorldPartition.cpp
+++ b/Engine/Source/Runtime/Engine/Private/WorldPartition/WorldPartition.cpp
@@ -977,6 +977,8 @@ void UWorldPartition::RegisterDelegates()
 			FCoreDelegates::OnGetOnScreenMessages.AddUObject(this, &UWorldPartition::GetOnScreenMessages);
 #endif
 		}
+
+		World->GetSubsystem<UWorldPartitionSubsystem>()->OnStreamingStateUpdated().AddUObject(this, &UWorldPartition::OnStreamingStateUpdated);
 	}
 }
 
@@ -1023,6 +1025,8 @@ void UWorldPartition::UnregisterDelegates()
 			FCoreDelegates::OnGetOnScreenMessages.RemoveAll(this);
 #endif
 		}
+
+		World->GetSubsystem<UWorldPartitionSubsystem>()->OnStreamingStateUpdated().RemoveAll(this);
 	}
 }
 
@@ -1052,6 +1056,28 @@ void UWorldPartition::OnWorldMatchStarting()
 	GetWorld()->BlockTillLevelStreamingCompleted();
 }
 
+void UWorldPartition::OnStreamingStateUpdated()
+{
+	check(GetWorld()->IsGameWorld());
+	if (StreamingPolicy)
+	{
+		StreamingPolicy->OnStreamingStateUpdated();
+	}
+}
+
+void UWorldPartition::OnPreChangeStreamingContent()
+{
+	if (StreamingPolicy)
+	{
+		StreamingPolicy->OnPreChangeStreamingContent();
+	}
+}
+
+int32 UWorldPartition::GetUpdateStreamingStateEpoch() const
+{
+	return StreamingPolicy ? StreamingPolicy->UpdateStreamingStateCounter : 0;
+}
+
 #if WITH_EDITOR
 UWorldPartition* UWorldPartition::CreateOrRepairWorldPartition(AWorldSettings* WorldSettings, TSubclassOf<UWorldPartitionEditorHash> EditorHashClass, TSubclassOf<UWorldPartitionRuntimeHash> RuntimeHashClass)
 {
@@ -1730,6 +1756,7 @@ bool UWorldPartition::IsExternalStreamingObjectInjected(URuntimeHashExternalStre
 
 bool UWorldPartition::InjectExternalStreamingObject(URuntimeHashExternalStreamingObjectBase* InExternalStreamingObject)
 {
+	OnPreChangeStreamingContent();
 	bool bInjected = RuntimeHash->InjectExternalStreamingObject(InExternalStreamingObject);
 	if (bInjected)
 	{
@@ -1751,6 +1778,7 @@ bool UWorldPartition::InjectExternalStreamingObject(URuntimeHashExternalStreamin
 
 bool UWorldPartition::RemoveExternalStreamingObject(URuntimeHashExternalStreamingObjectBase* InExternalStreamingObject)
 {
+	OnPreChangeStreamingContent();
 	bool bRemoved = RuntimeHash->RemoveExternalStreamingObject(InExternalStreamingObject);
 	if (bRemoved)
 	{
@@ -2329,11 +2357,13 @@ bool UWorldPartition::SupportsWorldAssetStreaming(const FName& InTargetGrid)
 
 FGuid UWorldPartition::RegisterWorldAssetStreaming(const FRegisterWorldAssetStreamingParams& InParams)
 {
+	OnPreChangeStreamingContent();
 	return RuntimeHash ? RuntimeHash->RegisterWorldAssetStreaming(InParams) : FGuid();
 }
 
 bool UWorldPartition::UnregisterWorldAssetStreaming(const FGuid& InWorldAssetStreamingGuid)
 {
+	OnPreChangeStreamingContent();
 	return RuntimeHash ? RuntimeHash->UnregisterWorldAssetStreaming(InWorldAssetStreamingGuid) : false;
 }
 
diff --git a/Engine/Source/Runtime/Engine/Private/WorldPartition/WorldPartitionReplay.cpp b/Engine/Source/Runtime/Engine/Private/WorldPartition/WorldPartitionReplay.cpp
--- a/Engine/Source/Runtime/Engine/Private/WorldPartition/WorldPartitionReplay.cpp
+++ b/Engine/Source/Runtime/Engine/Private/WorldPartition/WorldPartitionReplay.cpp
@@ -124,7 +124,7 @@ void AWorldPartitionReplay::Uninitialize(UWorld* World)
 	}
 }
 
-bool AWorldPartitionReplay::IsPlaybackEnabled(UWorld* World)
+bool AWorldPartitionReplay::IsPlaybackEnabled(const UWorld* World)
 {
 	if (GUseReplayStreamingSources && World->IsPlayingReplay())
 	{
@@ -133,11 +133,11 @@ bool AWorldPartitionReplay::IsPlaybackEnabled(UWorld* World)
 			return !!WorldPartition->Replay;
 		}
 	}
-	
+
 	return false;
 }
 
-bool AWorldPartitionReplay::IsRecordingEnabled(UWorld* World)
+bool AWorldPartitionReplay::IsRecordingEnabled(const UWorld* World)
 {
 	if (GRecordReplayStreamingSources && World->IsRecordingReplay())
 	{
@@ -146,7 +146,7 @@ bool AWorldPartitionReplay::IsRecordingEnabled(UWorld* World)
 			return !!WorldPartition->Replay;
 		}
 	}
-		
+
 	return false;
 }
 
@@ -181,7 +181,7 @@ void AWorldPartitionReplay::PreReplication(IRepChangedPropertyTracker& ChangedPr
 
 bool AWorldPartitionReplay::GetReplayStreamingSources(TArray<FWorldPartitionStreamingSource>& OutStreamingSources)
 {
-	UWorld* World = GetWorld();
+	const UWorld* World = GetWorld();
 	verify(IsPlaybackEnabled(World));
 
 	UDemoNetDriver* DemoNetDriver = World->GetDemoNetDriver();
diff --git a/Engine/Source/Runtime/Engine/Private/WorldPartition/WorldPartitionRuntimeCell.cpp b/Engine/Source/Runtime/Engine/Private/WorldPartition/WorldPartitionRuntimeCell.cpp
--- a/Engine/Source/Runtime/Engine/Private/WorldPartition/WorldPartitionRuntimeCell.cpp
+++ b/Engine/Source/Runtime/Engine/Private/WorldPartition/WorldPartitionRuntimeCell.cpp
@@ -148,7 +148,7 @@ UDataLayerManager* UWorldPartitionRuntimeCell::GetDataLayerManager() const
 	return GetOuterWorld()->GetWorldPartition()->GetDataLayerManager();
 }
 
-EDataLayerRuntimeState UWorldPartitionRuntimeCell::GetCellEffectiveWantedState() const
+EDataLayerRuntimeState UWorldPartitionRuntimeCell::GetCellEffectiveWantedState(const FWorldPartitionStreamingContext& Context) const
 {
 	if (!HasDataLayers())
 	{
@@ -156,59 +156,11 @@ EDataLayerRuntimeState UWorldPartitionRuntimeCell::GetCellEffectiveWantedState()
 	}
 	else
 	{
-		const UWorld* OuterWorld = GetOuterWorld();
-		const AWorldDataLayers* WorldDataLayers = OuterWorld->GetWorldDataLayers();
-		check(WorldDataLayers);
-		if (EffectiveWantedStateEpoch != WorldDataLayers->GetDataLayersStateEpoch())
+		const int32 ResolvingDataLayersRuntimeStateEpoch = Context.GetResolvingDataLayersRuntimeStateEpoch();
+		if (EffectiveWantedStateEpoch != ResolvingDataLayersRuntimeStateEpoch)
 		{
-			EffectiveWantedState = EDataLayerRuntimeState::Unloaded;
-
-			UWorldPartition* WorldPartition = OuterWorld->GetWorldPartition();
-			if (const UDataLayerManager* DataLayerManager = WorldPartition->GetDataLayerManager())
-			{
-				// Determine the maximum runtime state the cell can have based on its External Data Layer. If none, maximum is Activated.
-				const UExternalDataLayerInstance* ExternalDatalayerInstance = GetExternalDataLayerInstance();
-				EDataLayerRuntimeState MaxEffectiveRuntimeState = ExternalDatalayerInstance ? ExternalDatalayerInstance->GetEffectiveRuntimeState() : EDataLayerRuntimeState::Activated;
-
-				if (MaxEffectiveRuntimeState > EDataLayerRuntimeState::Unloaded)
-				{
-					TArrayView<const FName> NonExternalDataLayers = DataLayers.GetNonExternalDataLayers();
-					if (NonExternalDataLayers.IsEmpty())
-					{
-						EffectiveWantedState = MaxEffectiveRuntimeState;
-					}
-					else
-					{
-						switch (WorldPartition->GetDataLayersLogicOperator())
-						{
-						case EWorldPartitionDataLayersLogicOperator::Or:
-							if (DataLayerManager->IsAnyDataLayerInEffectiveRuntimeState(NonExternalDataLayers, EDataLayerRuntimeState::Activated))
-							{
-								EffectiveWantedState = MaxEffectiveRuntimeState;
-							}
-							else if (DataLayerManager->IsAnyDataLayerInEffectiveRuntimeState(NonExternalDataLayers, EDataLayerRuntimeState::Loaded))
-							{
-								EffectiveWantedState = EDataLayerRuntimeState::Loaded;
-							}
-							break;
-						case EWorldPartitionDataLayersLogicOperator::And:
-							if (DataLayerManager->IsAllDataLayerInEffectiveRuntimeState(NonExternalDataLayers, EDataLayerRuntimeState::Activated))
-							{
-								EffectiveWantedState = MaxEffectiveRuntimeState;
-							}
-							else if (DataLayerManager->IsAllDataLayerInEffectiveRuntimeState(NonExternalDataLayers, EDataLayerRuntimeState::Loaded))
-							{
-								EffectiveWantedState = EDataLayerRuntimeState::Loaded;
-							}
-							break;
-						default:
-							checkNoEntry();
-						}
-					}
-				}
-			}
-
-			EffectiveWantedStateEpoch = WorldDataLayers->GetDataLayersStateEpoch();
+			EffectiveWantedState = Context.ResolveDataLayerRuntimeState(DataLayers);
+			EffectiveWantedStateEpoch = ResolvingDataLayersRuntimeStateEpoch;
 		}
 	}
 
diff --git a/Engine/Source/Runtime/Engine/Private/WorldPartition/WorldPartitionRuntimeCellData.cpp b/Engine/Source/Runtime/Engine/Private/WorldPartition/WorldPartitionRuntimeCellData.cpp
--- a/Engine/Source/Runtime/Engine/Private/WorldPartition/WorldPartitionRuntimeCellData.cpp
+++ b/Engine/Source/Runtime/Engine/Private/WorldPartition/WorldPartitionRuntimeCellData.cpp
@@ -2,10 +2,9 @@
 
 #include "WorldPartition/WorldPartitionRuntimeCellData.h"
 #include "WorldPartition/WorldPartitionLog.h"
+#include "WorldPartition/WorldPartitionStreamingPolicy.h"
 #include "Misc/HierarchicalLogArchive.h"
 
-int32 UWorldPartitionRuntimeCellData::StreamingSourceCacheEpoch = 0;
-
 UWorldPartitionRuntimeCellData::UWorldPartitionRuntimeCellData(const FObjectInitializer& ObjectInitializer)
 	: Super(ObjectInitializer)
 	, CachedMinSourcePriority(MAX_uint8)
@@ -37,25 +36,25 @@ void UWorldPartitionRuntimeCellData::DumpStateLog(FHierarchicalLogArchive& Ar) c
 }
 #endif
 
-void UWorldPartitionRuntimeCellData::ResetStreamingSourceInfo() const
+void UWorldPartitionRuntimeCellData::ResetStreamingSourceInfo(const FWorldPartitionStreamingContext& Context) const
 {
 	CachedMinSourcePriority = MAX_uint8;
 	bCachedWasRequestedByBlockingSource = false;
 	CachedMinSquareDistanceToBlockingSource = MAX_dbl;
 	CachedMinBlockOnSlowStreamingRatio = MAX_flt;
 	CachedMinSpatialSortingPriority = MAX_dbl;
-	CachedSourceInfoEpoch = StreamingSourceCacheEpoch;	
+	CachedSourceInfoEpoch = Context.GetUpdateStreamingStateEpoch();
 }
 
 DECLARE_CYCLE_STAT(TEXT("Append Streaming Source Info"), STAT_WorldPartitionAppendStreamingSourceInfo, STATGROUP_WorldPartition);
-void UWorldPartitionRuntimeCellData::AppendStreamingSourceInfo(const FWorldPartitionStreamingSource& Source, const FSphericalSector& SourceShape) const
+void UWorldPartitionRuntimeCellData::AppendStreamingSourceInfo(const FWorldPartitionStreamingSource& Source, const FSphericalSector& SourceShape, const FWorldPartitionStreamingContext& Context) const
 {
 	SCOPE_CYCLE_COUNTER(STAT_WorldPartitionAppendStreamingSourceInfo);
 
-	if (CachedSourceInfoEpoch != StreamingSourceCacheEpoch)
+	if (CachedSourceInfoEpoch != Context.GetUpdateStreamingStateEpoch())
 	{
-		ResetStreamingSourceInfo();
-		check(CachedSourceInfoEpoch == StreamingSourceCacheEpoch);
+		ResetStreamingSourceInfo(Context);
+		check(CachedSourceInfoEpoch == Context.GetUpdateStreamingStateEpoch());
 	}
 
 	// Compute cosine angle from cell to source direction ratio
diff --git a/Engine/Source/Runtime/Engine/Private/WorldPartition/WorldPartitionRuntimeCellDataSpatialHash.cpp b/Engine/Source/Runtime/Engine/Private/WorldPartition/WorldPartitionRuntimeCellDataSpatialHash.cpp
--- a/Engine/Source/Runtime/Engine/Private/WorldPartition/WorldPartitionRuntimeCellDataSpatialHash.cpp
+++ b/Engine/Source/Runtime/Engine/Private/WorldPartition/WorldPartitionRuntimeCellDataSpatialHash.cpp
@@ -32,9 +32,9 @@ UWorldPartitionRuntimeCellDataSpatialHash::UWorldPartitionRuntimeCellDataSpatial
 	, CachedMinSquareDistanceToBlockingSource2D(MAX_dbl)
 {}
 
-void UWorldPartitionRuntimeCellDataSpatialHash::ResetStreamingSourceInfo() const
+void UWorldPartitionRuntimeCellDataSpatialHash::ResetStreamingSourceInfo(const FWorldPartitionStreamingContext& Context) const
 {
-	Super::ResetStreamingSourceInfo();
+	Super::ResetStreamingSourceInfo(Context);
 
 	CachedMinSquareDistanceToSource = MAX_dbl;
 	CachedSourcePriorityWeights.Reset();
@@ -83,9 +83,9 @@ float UWorldPartitionRuntimeCellDataSpatialHash::ComputeSourceToCellAngleFactor(
 	return AngleFactor;
 }
 
-void UWorldPartitionRuntimeCellDataSpatialHash::AppendStreamingSourceInfo(const FWorldPartitionStreamingSource& Source, const FSphericalSector& SourceShape) const
+void UWorldPartitionRuntimeCellDataSpatialHash::AppendStreamingSourceInfo(const FWorldPartitionStreamingSource& Source, const FSphericalSector& SourceShape, const FWorldPartitionStreamingContext& Context) const
 {
-	Super::AppendStreamingSourceInfo(Source, SourceShape);
+	Super::AppendStreamingSourceInfo(Source, SourceShape, Context);
 
 	const double SquareDistance2D = FVector::DistSquared2D(SourceShape.GetCenter(), Position);
 
diff --git a/Engine/Source/Runtime/Engine/Private/WorldPartition/WorldPartitionRuntimeHash.cpp b/Engine/Source/Runtime/Engine/Private/WorldPartition/WorldPartitionRuntimeHash.cpp
--- a/Engine/Source/Runtime/Engine/Private/WorldPartition/WorldPartitionRuntimeHash.cpp
+++ b/Engine/Source/Runtime/Engine/Private/WorldPartition/WorldPartitionRuntimeHash.cpp
@@ -10,6 +10,8 @@
 #include "WorldPartition/WorldPartitionLevelStreamingDynamic.h"
 #include "WorldPartition/DataLayer/ExternalDataLayerAsset.h"
 #include "WorldPartition/DataLayer/ExternalDataLayerHelper.h"
+#include "WorldPartition/DataLayer/DataLayerInstanceNames.h"
+#include "WorldPartition/DataLayer/DataLayerManager.h"
 #include "Misc/ArchiveMD5.h"
 #if WITH_EDITOR
 #include "WorldPartition/Cook/WorldPartitionCookPackage.h"
@@ -656,9 +658,9 @@ void UWorldPartitionRuntimeHash::ExecutePostSetupHLODActors(const UWorldPartitio
 
 #endif
 
-void UWorldPartitionRuntimeHash::FStreamingSourceCells::AddCell(const UWorldPartitionRuntimeCell* Cell, const FWorldPartitionStreamingSource& Source, const FSphericalSector& SourceShape)
+void UWorldPartitionRuntimeHash::FStreamingSourceCells::AddCell(const UWorldPartitionRuntimeCell* Cell, const FWorldPartitionStreamingSource& Source, const FSphericalSector& SourceShape, const FWorldPartitionStreamingContext& Context)
 {
-	Cell->AppendStreamingSourceInfo(Source, SourceShape);
+	Cell->AppendStreamingSourceInfo(Source, SourceShape, Context);
 	Cells.Add(Cell);
 }
 
@@ -681,4 +683,87 @@ double FWorldPartitionQueryCache::GetCellMinSquareDist(const UWorldPartitionRunt
 	return Dist ? *Dist : MAX_dbl;
 }
 
+FWorldPartitionStreamingContext FWorldPartitionStreamingContext::Create(const UWorld* InWorld)
+{
+	if (InWorld && InWorld->GetWorldPartition() && InWorld->GetWorldDataLayers())
+	{
+		check(!InWorld->IsGameWorld() || IsInGameThread());
+		return FWorldPartitionStreamingContext(InWorld);
+	}
+	return FWorldPartitionStreamingContext();
+}
+
+FWorldPartitionStreamingContext::FWorldPartitionStreamingContext()
+	: bIsValid(false)
+	, DataLayersLogicOperator(EWorldPartitionDataLayersLogicOperator::Or)
+	, DataLayerEffectiveStates(nullptr)
+	, UpdateStreamingStateEpoch(0)
+{}
+
+FWorldPartitionStreamingContext::FWorldPartitionStreamingContext(const UWorld* InWorld)
+	: FWorldPartitionStreamingContext(InWorld->GetWorldPartition()->GetDataLayersLogicOperator(), FWorldDataLayersEffectiveStatesAccessor::Get(InWorld->GetWorldDataLayers()), InWorld->GetWorldPartition()->GetUpdateStreamingStateEpoch())
+{
+}
+
+FWorldPartitionStreamingContext::FWorldPartitionStreamingContext(EWorldPartitionDataLayersLogicOperator InDataLayersLogicOperator, const FWorldDataLayersEffectiveStates& InDataLayerEffectiveStates, int32 InUpdateStreamingStateEpoch)
+	: bIsValid(true)
+	, DataLayersLogicOperator(InDataLayersLogicOperator)
+	, DataLayerEffectiveStates(&InDataLayerEffectiveStates)
+	, UpdateStreamingStateEpoch(InUpdateStreamingStateEpoch)
+{}
+
+EDataLayerRuntimeState FWorldPartitionStreamingContext::ResolveDataLayerRuntimeState(const FDataLayerInstanceNames& InDataLayers) const
+{
+	if (InDataLayers.IsEmpty())
+	{
+		return EDataLayerRuntimeState::Activated;
+	}
+
+	check(IsValid());
+	check(DataLayerEffectiveStates);
+	EDataLayerRuntimeState Result = EDataLayerRuntimeState::Unloaded;
+
+	// Determine the maximum runtime state the cell can have based on its External Data Layer. If none, maximum is Activated.
+	FName ExternalDatalayerName = InDataLayers.GetExternalDataLayer();
+	EDataLayerRuntimeState MaxEffectiveRuntimeState = !ExternalDatalayerName.IsNone() ? DataLayerEffectiveStates->GetDataLayerEffectiveRuntimeStateByName(ExternalDatalayerName) : EDataLayerRuntimeState::Activated;
+
+	if (MaxEffectiveRuntimeState > EDataLayerRuntimeState::Unloaded)
+	{
+		TArrayView<const FName> NonExternalDataLayers = InDataLayers.GetNonExternalDataLayers();
+		if (NonExternalDataLayers.IsEmpty())
+		{
+			Result = MaxEffectiveRuntimeState;
+		}
+		else
+		{
+			switch (DataLayersLogicOperator)
+			{
+			case EWorldPartitionDataLayersLogicOperator::Or:
+				if (UDataLayerManager::IsAnyDataLayerInEffectiveRuntimeState(NonExternalDataLayers, EDataLayerRuntimeState::Activated, *DataLayerEffectiveStates))
+				{
+					Result = MaxEffectiveRuntimeState;
+				}
+				else if (UDataLayerManager::IsAnyDataLayerInEffectiveRuntimeState(NonExternalDataLayers, EDataLayerRuntimeState::Loaded, *DataLayerEffectiveStates))
+				{
+					Result = EDataLayerRuntimeState::Loaded;
+				}
+				break;
+			case EWorldPartitionDataLayersLogicOperator::And:
+				if (UDataLayerManager::IsAllDataLayerInEffectiveRuntimeState(NonExternalDataLayers, EDataLayerRuntimeState::Activated, *DataLayerEffectiveStates))
+				{
+					Result = MaxEffectiveRuntimeState;
+				}
+				else if (UDataLayerManager::IsAllDataLayerInEffectiveRuntimeState(NonExternalDataLayers, EDataLayerRuntimeState::Loaded, *DataLayerEffectiveStates))
+				{
+					Result = EDataLayerRuntimeState::Loaded;
+				}
+				break;
+			default:
+				checkNoEntry();
+			}
+		}
+	}
+	return Result;
+}
+
 #undef LOCTEXT_NAMESPACE
diff --git a/Engine/Source/Runtime/Engine/Private/WorldPartition/WorldPartitionRuntimeSpatialHash.cpp b/Engine/Source/Runtime/Engine/Private/WorldPartition/WorldPartitionRuntimeSpatialHash.cpp
--- a/Engine/Source/Runtime/Engine/Private/WorldPartition/WorldPartitionRuntimeSpatialHash.cpp
+++ b/Engine/Source/Runtime/Engine/Private/WorldPartition/WorldPartitionRuntimeSpatialHash.cpp
@@ -362,7 +362,7 @@ void FSpatialHashStreamingGrid::GetCells(const FWorldPartitionStreamingQuerySour
 	}
 }
 
-void FSpatialHashStreamingGrid::GetCells(const TArray<FWorldPartitionStreamingSource>& Sources, UWorldPartitionRuntimeHash::FStreamingSourceCells& OutActivateCells, UWorldPartitionRuntimeHash::FStreamingSourceCells& OutLoadCells, bool bEnableZCulling) const
+void FSpatialHashStreamingGrid::GetCells(const TArray<FWorldPartitionStreamingSource>& Sources, UWorldPartitionRuntimeHash::FStreamingSourceCells& OutActivateCells, UWorldPartitionRuntimeHash::FStreamingSourceCells& OutLoadCells, bool bEnableZCulling, const FWorldPartitionStreamingContext& Context) const
 {
 	TRACE_CPUPROFILER_EVENT_SCOPE(FSpatialHashStreamingGrid::GetCells);
 
@@ -407,19 +407,19 @@ void FSpatialHashStreamingGrid::GetCells(const TArray<FWorldPartitionStreamingSo
 						
 						if (bIncludeCell)
 						{
-							switch (Cell->GetCellEffectiveWantedState())
+							switch (Cell->GetCellEffectiveWantedState(Context))
 							{
 							case EDataLayerRuntimeState::Loaded:
-								OutLoadCells.AddCell(Cell, Source, Shape);
+								OutLoadCells.AddCell(Cell, Source, Shape, Context);
 								break;
 							case EDataLayerRuntimeState::Activated:
 								switch (Source.TargetState)
 								{
 								case EStreamingSourceTargetState::Loaded:
-									OutLoadCells.AddCell(Cell, Source, Shape);
+									OutLoadCells.AddCell(Cell, Source, Shape, Context);
 									break;
 								case EStreamingSourceTargetState::Activated:
-									OutActivateCells.AddCell(Cell, Source, Shape);
+									OutActivateCells.AddCell(Cell, Source, Shape, Context);
 									bAddedActivatedCell = !Settings.bUseAlignedGridLevels && Settings.bSnapNonAlignedGridLevelsToLowerLevels;
 									break;
 								default:
@@ -442,7 +442,7 @@ void FSpatialHashStreamingGrid::GetCells(const TArray<FWorldPartitionStreamingSo
 		});
 	}
 
-	GetNonSpatiallyLoadedCells(OutActivateCells.GetCells(), OutLoadCells.GetCells());
+	GetNonSpatiallyLoadedCells(OutActivateCells.GetCells(), OutLoadCells.GetCells(), Context);
 
 	if (!Settings.bUseAlignedGridLevels && Settings.bSnapNonAlignedGridLevelsToLowerLevels)
 	{
@@ -499,14 +499,14 @@ void FSpatialHashStreamingGrid::GetCells(const TArray<FWorldPartitionStreamingSo
 		{
 			ForEachRuntimeCell(ParentCell.Key, [&](const UWorldPartitionRuntimeCell* Cell)
 			{
-				switch (Cell->GetCellEffectiveWantedState())
+				switch (Cell->GetCellEffectiveWantedState(Context))
 				{
 				case EDataLayerRuntimeState::Loaded:
 					break;
 				case EDataLayerRuntimeState::Activated:
 					for (const auto& Info : ParentCell.Value)
 					{
-						OutActivateCells.AddCell(Cell, Info.Source, Info.SourceShape);
+						OutActivateCells.AddCell(Cell, Info.Source, Info.SourceShape, Context);
 					}
 					break;
 				case EDataLayerRuntimeState::Unloaded:
@@ -730,17 +730,17 @@ void FSpatialHashStreamingGrid::ForEachRuntimeCell(TFunctionRef<bool(const UWorl
 	ForEachGridLevel(InjectedGridLevels);
 }
 
-void FSpatialHashStreamingGrid::GetNonSpatiallyLoadedCells(TSet<const UWorldPartitionRuntimeCell*>& OutActivateCells, TSet<const UWorldPartitionRuntimeCell*>& OutLoadCells) const
+void FSpatialHashStreamingGrid::GetNonSpatiallyLoadedCells(TSet<const UWorldPartitionRuntimeCell*>& OutActivateCells, TSet<const UWorldPartitionRuntimeCell*>& OutLoadCells, const FWorldPartitionStreamingContext& Context) const
 {
 	if (GridLevels.Num() > 0)
 	{
-		auto ForEachGridLevelCells = [&OutActivateCells, &OutLoadCells](const FSpatialHashStreamingGridLevel& InGridLevel)
+		auto ForEachGridLevelCells = [&OutActivateCells, &OutLoadCells, &Context](const FSpatialHashStreamingGridLevel& InGridLevel)
 		{
 			for (const FSpatialHashStreamingGridLayerCell& LayerCell : InGridLevel.LayerCells)
 			{
 				for (const UWorldPartitionRuntimeCell* Cell : LayerCell.GridCells)
 				{
-					switch (Cell->GetCellEffectiveWantedState())
+					switch (Cell->GetCellEffectiveWantedState(Context))
 					{
 					case EDataLayerRuntimeState::Loaded:
 						check(Cell->HasDataLayers());
@@ -779,7 +779,7 @@ void FSpatialHashStreamingGrid::GetFilteredCellsForDebugDraw(const FSpatialHashS
 				EStreamingStatus StreamingStatus = Cell->GetStreamingStatus();
 				const TArray<FName>& DataLayers = Cell->GetDataLayers();
 
-				switch (Cell->GetCellEffectiveWantedState())
+				switch (Cell->GetCellEffectiveWantedStateRaw())
 				{
 				case EDataLayerRuntimeState::Loaded:
 				case EDataLayerRuntimeState::Activated:
@@ -1744,8 +1744,13 @@ uint32 UWorldPartitionRuntimeSpatialHash::ComputeUpdateStreamingHash() const
 	return HashBuilder.GetHash();
 }
 
-void UWorldPartitionRuntimeSpatialHash::ForEachStreamingCellsSources(const TArray<FWorldPartitionStreamingSource>& Sources, TFunctionRef<bool(const UWorldPartitionRuntimeCell*, EStreamingSourceTargetState)> Func) const
+void UWorldPartitionRuntimeSpatialHash::ForEachStreamingCellsSources(const TArray<FWorldPartitionStreamingSource>& Sources, TFunctionRef<bool(const UWorldPartitionRuntimeCell*, EStreamingSourceTargetState)> Func, const FWorldPartitionStreamingContext& InContext) const
 {
+	// Build a context when none is provided (for backward compatibility)
+	const FWorldPartitionStreamingContext StackContext = !InContext.IsValid() ? FWorldPartitionStreamingContext::Create(GetTypedOuter<UWorld>()) : FWorldPartitionStreamingContext();
+	const FWorldPartitionStreamingContext& Context = InContext.IsValid() ? InContext : StackContext;
+	check(Context.IsValid());
+
 	FStreamingSourceCells ActivateStreamingSourceCells;
 	FStreamingSourceCells LoadStreamingSourceCells;
 
@@ -1756,7 +1761,7 @@ void UWorldPartitionRuntimeSpatialHash::ForEachStreamingCellsSources(const TArra
 		{
 			if (IsCellRelevantFor(StreamingGrid.bClientOnlyVisible))
 			{
-				StreamingGrid.GetNonSpatiallyLoadedCells(ActivateStreamingSourceCells.GetCells(), LoadStreamingSourceCells.GetCells());
+				StreamingGrid.GetNonSpatiallyLoadedCells(ActivateStreamingSourceCells.GetCells(), LoadStreamingSourceCells.GetCells(), Context);
 			}
 		});
 	}
@@ -1767,7 +1772,7 @@ void UWorldPartitionRuntimeSpatialHash::ForEachStreamingCellsSources(const TArra
 		{
 			if (IsCellRelevantFor(StreamingGrid.bClientOnlyVisible))
 			{
-				StreamingGrid.GetCells(Sources, ActivateStreamingSourceCells, LoadStreamingSourceCells, GetEffectiveEnableZCulling(bEnableZCulling));
+				StreamingGrid.GetCells(Sources, ActivateStreamingSourceCells, LoadStreamingSourceCells, GetEffectiveEnableZCulling(bEnableZCulling), Context);
 			}
 		});
 	}
diff --git a/Engine/Source/Runtime/Engine/Private/WorldPartition/WorldPartitionStreamingPolicy.cpp b/Engine/Source/Runtime/Engine/Private/WorldPartition/WorldPartitionStreamingPolicy.cpp
--- a/Engine/Source/Runtime/Engine/Private/WorldPartition/WorldPartitionStreamingPolicy.cpp
+++ b/Engine/Source/Runtime/Engine/Private/WorldPartition/WorldPartitionStreamingPolicy.cpp
@@ -6,21 +6,23 @@
 
 #include "WorldPartition/WorldPartitionStreamingPolicy.h"
 #include "Algo/Find.h"
+#include "Async/TaskGraphInterfaces.h"
+#include "Engine/Level.h"
+#include "Engine/NetConnection.h"
 #include "Logging/LogScopedCategoryAndVerbosityOverride.h"
+#include "Misc/HashBuilder.h"
 #include "Misc/ScopeExit.h"
+#include "Stats/Stats2.h"
 #include "UObject/Package.h"
-#include "WorldPartition/WorldPartitionSubsystem.h"
-#include "WorldPartition/WorldPartitionLog.h"
-#include "WorldPartition/WorldPartitionReplay.h"
 #include "WorldPartition/DataLayer/DataLayerManager.h"
 #include "WorldPartition/DataLayer/WorldDataLayers.h"
 #include "WorldPartition/HLOD/HLODRuntimeSubsystem.h"
 #include "WorldPartition/ContentBundle/ContentBundle.h"
+#include "WorldPartition/WorldPartitionSubsystem.h"
+#include "WorldPartition/WorldPartitionLog.h"
+#include "WorldPartition/WorldPartitionReplay.h"
 #include "WorldPartition/WorldPartitionDebugHelper.h"
-#include "Engine/Level.h"
-#include "Engine/NetConnection.h"
 #include "WorldPartition/WorldPartitionRuntimeHash.h"
-#include "Misc/HashBuilder.h"
 
 #include UE_INLINE_GENERATED_CPP_BY_NAME(WorldPartitionStreamingPolicy)
 
@@ -54,29 +56,32 @@ FAutoConsoleVariableRef UWorldPartitionStreamingPolicy::CVarForceUpdateFrameCoun
 	TEXT("Frequency (in frames) at which world partition streaming update will be executed regardless if no changes are detected."),
 	ECVF_Default);
 
-static void SortStreamingCellsByImportance(TArray<const UWorldPartitionRuntimeCell*>& InOutCells)
-{
-	if (InOutCells.Num() > 1)
-	{
-		TRACE_CPUPROFILER_EVENT_SCOPE(SortStreamingCellsByImportance);
-		Algo::Sort(InOutCells, [](const UWorldPartitionRuntimeCell* CellA, const UWorldPartitionRuntimeCell* CellB) { return CellA->SortCompare(CellB) < 0; });
-	}
-}
+bool UWorldPartitionStreamingPolicy::IsAsyncUpdateStreamingStateEnabled = false;
+FAutoConsoleVariableRef UWorldPartitionStreamingPolicy::CVarAsyncUpdateStreamingStateEnabled(
+	TEXT("wp.Runtime.UpdateStreaming.EnableAsyncUpdate"),
+	UWorldPartitionStreamingPolicy::IsAsyncUpdateStreamingStateEnabled,
+	TEXT("Set to enable asynchronous World Partition UpdateStreamingState."),
+	ECVF_Default);
+
+DECLARE_CYCLE_STAT(TEXT("WorldPartition_AsyncUpdateStreamingState"), WPStreamingPolicy_AsyncUpdateStreamingState, STATGROUP_TaskGraphTasks);
+DECLARE_CYCLE_STAT(TEXT("WorldPartition_AsyncPostUpdateStreamingState"), WPStreamingPolicy_AsyncPostUpdateStreamingState, STATGROUP_TaskGraphTasks);
 
 UWorldPartitionStreamingPolicy::UWorldPartitionStreamingPolicy(const FObjectInitializer& ObjectInitializer)
 	: Super(ObjectInitializer) 
 	, WorldPartition(nullptr)
-	, ProcessedToActivateCells(0)
-	, ProcessedToLoadCells(0)
-	, bCriticalPerformanceRequestedBlockTillOnWorld(false)
+	, CurrentStreamingPerformance(EWorldPartitionStreamingPerformance::Good)
 	, bShouldMergeStreamingSourceInfo(false)
+	, bCriticalPerformanceRequestedBlockTillOnWorld(false)
 	, CriticalPerformanceBlockTillLevelStreamingCompletedEpoch(0)
+	, ProcessedToLoadCells(0)
+	, ProcessedToActivateCells(0)
 	, ServerStreamingStateEpoch(INT_MIN)
 	, ServerStreamingEnabledEpoch(INT_MIN)
 	, UpdateStreamingHash(0)
 	, UpdateStreamingSourcesHash(0)
-	, UpdateStreamingStateCalls(0)
-	, StreamingPerformance(EWorldPartitionStreamingPerformance::Good)
+	, UpdateStreamingStateCounter(0)
+	, AsyncUpdateTaskState(EAsyncUpdateTaskState::None)
+	, AsyncShouldSkipUpdateCounter(0)
 #if !UE_BUILD_SHIPPING
 	, OnScreenMessageStartTime(0.0)
 	, OnScreenMessageStreamingPerformance(EWorldPartitionStreamingPerformance::Good)
@@ -91,10 +96,13 @@ UWorldPartitionStreamingPolicy::UWorldPartitionStreamingPolicy(const FObjectInit
 
 void UWorldPartitionStreamingPolicy::UpdateStreamingSources(bool bCanOptimizeUpdate)
 {
+	check(IsInGameThread());
+
 	TRACE_CPUPROFILER_EVENT_SCOPE(UWorldPartitionStreamingPolicy::UpdateStreamingSources);
+
 	if (!WorldPartition->CanStream())
 	{
-		StreamingSources.Reset();
+		CurrentState.StreamingSources.Reset();
 		return;
 	}
 
@@ -104,29 +112,31 @@ void UWorldPartitionStreamingPolicy::UpdateStreamingSources(bool bCanOptimizeUpd
 	{
 		TArray<FWorldPartitionStreamingSource> LocalStreamingSources;
 		WorldPartitionSubsystem->GetStreamingSources(WorldPartition, LocalStreamingSources);
-		check(LocalStreamingSources.Num() == StreamingSources.Num());
+		check(LocalStreamingSources.Num() == CurrentState.StreamingSources.Num());
 		const FTransform WorldToLocal = WorldPartition->GetInstanceTransform().Inverse();
 		for (int32 i=0; i<LocalStreamingSources.Num(); i++)
 		{
-			check(StreamingSources[i].Name == LocalStreamingSources[i].Name);
-			StreamingSources[i].Velocity = WorldToLocal.TransformVector(LocalStreamingSources[i].Velocity);
+			check(CurrentState.StreamingSources[i].Name == LocalStreamingSources[i].Name);
+			CurrentState.StreamingSources[i].Velocity = WorldToLocal.TransformVector(LocalStreamingSources[i].Velocity);
 		}
 		return;
 	}
 
-	StreamingSources.Reset();
-	WorldPartitionSubsystem->GetStreamingSources(WorldPartition, StreamingSources);
+	CurrentState.StreamingSources.Reset();
+	WorldPartitionSubsystem->GetStreamingSources(WorldPartition, CurrentState.StreamingSources);
 	UpdateStreamingSourcesHash = NewUpdateStreamingSourcesHash;
 }
 
 bool UWorldPartitionStreamingPolicy::IsInBlockTillLevelStreamingCompleted(bool bIsCausedByBadStreamingPerformance /* = false*/) const
 {
+	check(IsInGameThread());
+
 	const UWorld* World = GetWorld();
 	const bool bIsInBlockTillLevelStreamingCompleted = World->GetIsInBlockTillLevelStreamingCompleted();
 	if (bIsCausedByBadStreamingPerformance)
 	{
 		return bIsInBlockTillLevelStreamingCompleted &&
-				(StreamingPerformance != EWorldPartitionStreamingPerformance::Good) &&
+				(CurrentStreamingPerformance != EWorldPartitionStreamingPerformance::Good) &&
 				(CriticalPerformanceBlockTillLevelStreamingCompletedEpoch == World->GetBlockTillLevelStreamingCompletedEpoch());
 	}
 	return bIsInBlockTillLevelStreamingCompleted;
@@ -146,6 +156,7 @@ bool UWorldPartitionStreamingPolicy::IsUpdateStreamingOptimEnabled()
 
 uint32 UWorldPartitionStreamingPolicy::ComputeUpdateStreamingHash(bool bCanOptimizeUpdate) const
 {
+	TRACE_CPUPROFILER_EVENT_SCOPE(UWorldPartitionStreamingPolicy::ComputeUpdateStreamingHash);
 	if (bCanOptimizeUpdate)
 	{
 		const bool bIsStreaming3D = WorldPartition->RuntimeHash && WorldPartition->RuntimeHash->IsStreaming3D();
@@ -162,7 +173,7 @@ uint32 UWorldPartitionStreamingPolicy::ComputeUpdateStreamingHash(bool bCanOptim
 		HashBuilder << ComputeServerStreamingEnabledEpoch();
 		HashBuilder << WorldPartition->GetStreamingStateEpoch();
 		HashBuilder << bIsStreaming3D;
-		for (const FWorldPartitionStreamingSource& Source : StreamingSources)
+		for (const FWorldPartitionStreamingSource& Source : CurrentState.StreamingSources)
 		{
 			HashBuilder << Source.GetHash(bIsStreaming3D);
 		}
@@ -222,7 +233,7 @@ bool UWorldPartitionStreamingPolicy::GetIntersectingCells(const TArray<FWorldPar
 	return true;
 }
 
-const TSet<FName>& UWorldPartitionStreamingPolicy::GetServerDisallowedStreamingOutDataLayers()
+const TSet<FName>& UWorldPartitionStreamingPolicy::GetServerDisallowedStreamingOutDataLayers() const
 {
 	if (!CachedServerDisallowStreamingOutDataLayers.IsSet())
 	{
@@ -255,19 +266,51 @@ void UWorldPartitionStreamingPolicy::UpdateStreamingState()
 {
 	TRACE_CPUPROFILER_EVENT_SCOPE(UWorldPartitionStreamingPolicy::UpdateStreamingState);
 
-	++UpdateStreamingStateCalls;
+	++UpdateStreamingStateCounter;
 
 	UWorld* World = GetWorld();
 	check(World);
 	check(World->IsGameWorld());
 
-	const bool bLastUpdateCompletedLoadingAndActivation = ((ProcessedToActivateCells + ProcessedToLoadCells) == (ToActivateCells.Num() + ToLoadCells.Num()));
-	ProcessedToActivateCells = 0;
+	const bool bLastUpdateCompletedLoadingAndActivation = ((ProcessedToActivateCells + ProcessedToLoadCells) == (TargetState.ToActivateCells.Num() + TargetState.ToLoadCells.Num()));
+
 	ProcessedToLoadCells = 0;
-	ToActivateCells.Reset();
-	ToLoadCells.Reset();
+	ProcessedToActivateCells = 0;
+	TargetState.Reset();
+
+	// Last update was asynchronous
+	if (WaitForAsyncUpdateStreamingState())
+	{
+		TRACE_CPUPROFILER_EVENT_SCOPE(UWorldPartitionStreamingPolicy::UpdateTargetStateFromAsyncTask);
+		// Update Target State using asynchronous task results
+		// Filter result (AsyncTaskTargetState) using CurrentState as asynchronous tasks started with a snapshot version of ActivatedCells/LoadedCells
+		for (const UWorldPartitionRuntimeCell* Cell : AsyncTaskTargetState.ToActivateCells)
+		{
+			if (!CurrentState.ActivatedCells.Contains(Cell))
+			{
+				TargetState.ToActivateCells.Add(Cell);
+			}
+		}
+		for (const UWorldPartitionRuntimeCell* Cell : AsyncTaskTargetState.ToLoadCells)
+		{
+			if (!CurrentState.LoadedCells.Contains(Cell))
+			{
+				TargetState.ToLoadCells.Add(Cell);
+			}
+		}
 
-	// Dermine if the World's BlockTillLevelStreamingCompleted was triggered by WorldPartitionStreamingPolicy
+		// Reset everything related to last asynchronous tasks
+		check(AsyncUpdateTaskState == EAsyncUpdateTaskState::Started);
+		AsyncUpdateTaskState = EAsyncUpdateTaskState::None;
+		AsyncUpdateStreamingStateTask = nullptr;
+		AsyncPostUpdateStreamingStateTask = nullptr;
+		AsyncTaskCurrentState.Reset();
+		AsyncTaskTargetState.Reset();
+	}
+
+	check(AsyncUpdateTaskState == EAsyncUpdateTaskState::None);
+
+	// Determine if the World's BlockTillLevelStreamingCompleted was triggered by WorldPartitionStreamingPolicy
 	if (bCriticalPerformanceRequestedBlockTillOnWorld && IsInBlockTillLevelStreamingCompleted())
 	{
 		bCriticalPerformanceRequestedBlockTillOnWorld = false;
@@ -276,22 +319,40 @@ void UWorldPartitionStreamingPolicy::UpdateStreamingState()
 
 	const bool bIsServer = WorldPartition->IsServer();
 	const bool bCanStream = WorldPartition->CanStream();
-	const bool bForceFrameUpdate = (UWorldPartitionStreamingPolicy::ForceUpdateFrameCount > 0) ? ((UpdateStreamingStateCalls % UWorldPartitionStreamingPolicy::ForceUpdateFrameCount) == 0) : false;
+
+	// If server (non-streaming) has nothing to do, early out
+	if (bIsServer &&
+		bCanStream &&
+		bLastUpdateCompletedLoadingAndActivation &&
+		!WorldPartition->IsServerStreamingEnabled() &&
+		(ServerStreamingEnabledEpoch == ComputeServerStreamingEnabledEpoch()) &&
+		(ServerStreamingStateEpoch == WorldPartition->GetStreamingStateEpoch()))
+	{
+		return;
+	}
+
+	const bool bForceFrameUpdate = (UWorldPartitionStreamingPolicy::ForceUpdateFrameCount > 0) ? ((UpdateStreamingStateCounter % UWorldPartitionStreamingPolicy::ForceUpdateFrameCount) == 0) : false;
 	const bool bCanOptimizeUpdate =
 		WorldPartition->RuntimeHash &&
 		bCanStream &&
-		!bForceFrameUpdate &&										// We garantee to update every N frame to force some internal updates like UpdateStreamingPerformance
-		IsUpdateStreamingOptimEnabled() &&							// Check CVars to see if optimization is enabled
-		bLastUpdateCompletedLoadingAndActivation &&					// Don't optimize if last frame didn't process all cells to load/activate
-		!IsInBlockTillLevelStreamingCompleted() &&					// Don't optimize when inside UWorld::BlockTillLevelStreamingCompleted
-		ActivatedCells.GetPendingAddToWorldCells().IsEmpty(); 		// Don't optimize when remaining cells to add to world
-	
+		!bForceFrameUpdate &&												// We guarantee to update every N frame to force some internal updates like UpdateStreamingPerformance
+		IsUpdateStreamingOptimEnabled() &&									// Check CVars to see if optimization is enabled
+		bLastUpdateCompletedLoadingAndActivation &&							// Don't optimize if last frame didn't process all cells to load/activate
+		!IsInBlockTillLevelStreamingCompleted() &&							// Don't optimize when inside UWorld::BlockTillLevelStreamingCompleted
+		CurrentState.ActivatedCells.GetPendingAddToWorldCells().IsEmpty(); 	// Don't optimize when remaining cells to add to world
+
 	// Update streaming sources
 	UpdateStreamingSources(bCanOptimizeUpdate);
 
+	// Determine if update will be async or not
+	const bool bCanUpdateAsync = UWorldPartitionStreamingPolicy::IsAsyncUpdateStreamingStateEnabled && !bIsServer && bCanStream && WorldPartition->IsInitialized() && !IsInBlockTillLevelStreamingCompleted();
+
 	// Detect if nothing relevant changed and early out
 	const uint32 NewUpdateStreamingHash = ComputeUpdateStreamingHash(bCanOptimizeUpdate);
-	const bool bShouldSkipUpdate = NewUpdateStreamingHash && (UpdateStreamingHash == NewUpdateStreamingHash);
+	const bool bIsUpdateStreamingHashIdentical = NewUpdateStreamingHash && (UpdateStreamingHash == NewUpdateStreamingHash);
+	AsyncShouldSkipUpdateCounter = (bIsUpdateStreamingHashIdentical && bCanUpdateAsync) ? (AsyncShouldSkipUpdateCounter + 1) : 0;
+	// Since the asynchronous update is working with a snapshot of the last frame, wait for 2 consecutive update without any changes before deciding to skip the update
+	const bool bShouldSkipUpdate = bIsUpdateStreamingHashIdentical && (!bCanUpdateAsync || (AsyncShouldSkipUpdateCounter >= 2));
 	if (bShouldSkipUpdate)
 	{
 		return;
@@ -300,108 +361,115 @@ void UWorldPartitionStreamingPolicy::UpdateStreamingState()
 	// Update new streaming sources hash
 	UpdateStreamingHash = NewUpdateStreamingHash;
 
-	bool bUpdateServerEpoch = false;
+	// UpdateStreamingStateInternal asynchronously or synchronously
+	if (bCanUpdateAsync)
+	{
+		// Put state to Pending
+		// Wait for WorldPartitionSubsystem to finish accessing streaming cells (used to sort)
+		// Once it's the case, OnStreamingStateUpdated will be called and async tasks will be created and dispatched
+		AsyncUpdateTaskState = EAsyncUpdateTaskState::Pending;
+	}
+	else
+	{
+		TargetState.Reset();
+		UWorldPartitionStreamingPolicy::UpdateStreamingStateInternal(FUpdateStreamingStateParams(this, CurrentState), TargetState);
+		PostUpdateStreamingStateInternal_GameThread(TargetState);
+	}
+}
+
+void UWorldPartitionStreamingPolicy::UpdateStreamingStateInternal(const FUpdateStreamingStateParams& InParams, FWorldPartitionUpdateStreamingTargetState& OutTargetState)
+{
+	TRACE_CPUPROFILER_EVENT_SCOPE(UWorldPartitionStreamingPolicy::UpdateStreamingStateInternal);
+
+	OutTargetState.bUpdateServerEpoch = false;
 
-	check(FrameActivateCells.IsEmpty());
-	check(FrameLoadCells.IsEmpty());
+	check(InParams.World);
+	check(InParams.World->IsGameWorld());
+	check(InParams.FrameActivateCells.IsEmpty());
+	check(InParams.FrameLoadCells.IsEmpty());
+	check(OutTargetState.IsEmpty());
+
+	FWorldPartitionStreamingContext Context(InParams.DataLayersLogicOperator, InParams.GetWorldDataLayersEffectiveStates(), InParams.PolicyUpdateStreamingStateEpoch);
 
 	ON_SCOPE_EXIT
 	{
-		// Reset frame StreamingSourceCells (optimization to avoid reallocation at every call to UpdateStreamingState)
-		FrameActivateCells.Reset();
-		FrameLoadCells.Reset();
+		// Reset frame cells to avoid reallocation at every Update call
+		InParams.FrameActivateCells.Reset();
+		InParams.FrameLoadCells.Reset();
 	};
 
-	const bool bIsServerStreamingEnabled = WorldPartition->IsServerStreamingEnabled();
-	const int32 NewServerStreamingEnabledEpoch = ComputeServerStreamingEnabledEpoch();
-
-	const bool bIsStreamingInEnabled = WorldPartition->IsStreamingInEnabled();
-
-	if (bCanStream)
+	if (InParams.bCanStream)
 	{
-		if (!bIsServer || bIsServerStreamingEnabled || AWorldPartitionReplay::IsPlaybackEnabled(World))
+		if (!InParams.bIsServer || InParams.bIsServerStreamingEnabled || AWorldPartitionReplay::IsPlaybackEnabled(InParams.World))
 		{
 			// When world partition can't stream, all cells must be unloaded
-			if (WorldPartition->RuntimeHash)
+			if (InParams.RuntimeHash)
 			{
 				TRACE_CPUPROFILER_EVENT_SCOPE(UWorldPartitionStreamingPolicy::UpdateStreamingState_ForEachStreamingCellsSources);
 
-				UWorldPartitionRuntimeCellData::DirtyStreamingSourceCacheEpoch();
-
-				WorldPartition->RuntimeHash->ForEachStreamingCellsSources(StreamingSources, [this](const UWorldPartitionRuntimeCell* Cell, EStreamingSourceTargetState TargetState)
+				InParams.RuntimeHash->ForEachStreamingCellsSources(InParams.CurrentState.StreamingSources, [&InParams](const UWorldPartitionRuntimeCell* Cell, EStreamingSourceTargetState SourceTargetState)
 				{
-					switch (TargetState)
+					switch (SourceTargetState)
 					{
 					case EStreamingSourceTargetState::Loaded:
-						FrameLoadCells.Add(Cell);
+						InParams.FrameLoadCells.Add(Cell);
 						break;
 					case EStreamingSourceTargetState::Activated:
-						FrameActivateCells.Add(Cell);
+						InParams.FrameActivateCells.Add(Cell);
 						break;
 					default:
 						check(0);
 					}
-
 					return true;
-				});
+				}, Context);
 			}
 		}
 
-		if (bIsServer)
+		if (InParams.bIsServer)
 		{
-			const bool bCanServerDeactivateOrUnloadCells = WorldPartition->IsServerStreamingOutEnabled();
-			const TSet<FName>& ServerDisallowStreamingOutDataLayers = GetServerDisallowedStreamingOutDataLayers();
-
 			TRACE_CPUPROFILER_EVENT_SCOPE(UWorldPartitionStreamingPolicy::UpdateStreamingState_ServerUpdate);
 
 			// Server will activate all non data layer cells at first and then load/activate/unload data layer cells only when the data layer states change
-			if (!bIsServerStreamingEnabled && 
-				bLastUpdateCompletedLoadingAndActivation && 
-				(ServerStreamingEnabledEpoch == NewServerStreamingEnabledEpoch) &&
-				(ServerStreamingStateEpoch == WorldPartition->GetStreamingStateEpoch()))
-			{
-				// Server as nothing to do early out
-				return; 
-			}
+			const bool bCanServerDeactivateOrUnloadCells = InParams.bIsServerStreamingOutEnabled;
 
-			bUpdateServerEpoch = true;
+			OutTargetState.bUpdateServerEpoch = true;
 
-			auto CanServerDeactivateOrUnloadDataLayerCell = [&ServerDisallowStreamingOutDataLayers](const UWorldPartitionRuntimeCell* Cell)
+			auto CanServerDeactivateOrUnloadDataLayerCell = [&InParams](const UWorldPartitionRuntimeCell* Cell)
 			{
-				return !Cell->HasDataLayers() || !Cell->HasAnyDataLayer(ServerDisallowStreamingOutDataLayers);
+				return !Cell->HasDataLayers() || !Cell->HasAnyDataLayer(InParams.ServerDisallowedStreamingOutDataLayers);
 			};
 
-			auto AddServerFrameCell = [this, CanServerDeactivateOrUnloadDataLayerCell](const UWorldPartitionRuntimeCell* Cell)
+			auto AddServerFrameCell = [&InParams, &Context, CanServerDeactivateOrUnloadDataLayerCell](const UWorldPartitionRuntimeCell* Cell)
 			{
 				// Keep Data Layer cells in their current state if server cannot deactivate/unload data layer cells
 				if (!CanServerDeactivateOrUnloadDataLayerCell(Cell))
 				{
 					// If cell was activated, keep it activated
-					if (ActivatedCells.Contains(Cell))
+					if (InParams.CurrentState.ActivatedCells.Contains(Cell))
 					{
-						FrameActivateCells.Add(Cell);
+						InParams.FrameActivateCells.Add(Cell);
 						return;
 					}
 					else
 					{
 						// If cell was loaded, keep it loaded except if it should become activated.
 						// In the second case, let the standard code path process it and add it to FrameActivateCells.
-						const bool bIsAnActivatedDataLayerCell = Cell->HasDataLayers() && (Cell->GetCellEffectiveWantedState() == EDataLayerRuntimeState::Activated); 
-						if (LoadedCells.Contains(Cell) && !bIsAnActivatedDataLayerCell)
+						const bool bIsAnActivatedDataLayerCell = Cell->HasDataLayers() && (Cell->GetCellEffectiveWantedState(Context) == EDataLayerRuntimeState::Activated); 
+						if (InParams.CurrentState.LoadedCells.Contains(Cell) && !bIsAnActivatedDataLayerCell)
 						{
-							FrameLoadCells.Add(Cell);
+							InParams.FrameLoadCells.Add(Cell);
 							return;
 						}
 					}
 				}
 				
-				switch (Cell->GetCellEffectiveWantedState())
+				switch (Cell->GetCellEffectiveWantedState(Context))
 				{
 				case EDataLayerRuntimeState::Loaded:
-					FrameLoadCells.Add(Cell);
+					InParams.FrameLoadCells.Add(Cell);
 					break;
 				case EDataLayerRuntimeState::Activated:
-					FrameActivateCells.Add(Cell);
+					InParams.FrameActivateCells.Add(Cell);
 					break;
 				case EDataLayerRuntimeState::Unloaded:
 					break;
@@ -410,11 +478,11 @@ void UWorldPartitionStreamingPolicy::UpdateStreamingState()
 				}
 			};
 
-			if (!bIsServerStreamingEnabled)
+			if (!InParams.bIsServerStreamingEnabled)
 			{
-				if (WorldPartition->RuntimeHash)
+				if (InParams.RuntimeHash)
 				{
-					WorldPartition->RuntimeHash->ForEachStreamingCells([this, &AddServerFrameCell](const UWorldPartitionRuntimeCell* Cell)
+					InParams.RuntimeHash->ForEachStreamingCells([&AddServerFrameCell](const UWorldPartitionRuntimeCell* Cell)
 					{
 						AddServerFrameCell(Cell);
 						return true;
@@ -424,11 +492,11 @@ void UWorldPartitionStreamingPolicy::UpdateStreamingState()
 			else if (!bCanServerDeactivateOrUnloadCells)
 			{
 				// When server streaming-out is disabled, revisit existing loaded/activated cells and add them in the proper FrameLoadCells/FrameActivateCells
-				for (const UWorldPartitionRuntimeCell* Cell : ActivatedCells.GetCells())
+				for (const UWorldPartitionRuntimeCell* Cell : InParams.CurrentState.ActivatedCells.GetCells())
 				{
 					AddServerFrameCell(Cell);
 				}
-				for (const UWorldPartitionRuntimeCell* Cell : LoadedCells)
+				for (const UWorldPartitionRuntimeCell* Cell : InParams.CurrentState.LoadedCells)
 				{
 					AddServerFrameCell(Cell);
 				}
@@ -436,18 +504,18 @@ void UWorldPartitionStreamingPolicy::UpdateStreamingState()
 		}
 	}
 
-	const TSet<FName>& ServerClientsVisibleLevelNames = GetWorld()->GetSubsystem<UWorldPartitionSubsystem>()->ServerClientsVisibleLevelNames;
-	auto ShouldWaitForClientVisibility = [bIsServer, this, &bUpdateServerEpoch, &ServerClientsVisibleLevelNames](const UWorldPartitionRuntimeCell* Cell)
+	const TSet<FName>& ServerClientsVisibleLevelNames = InParams.World->GetSubsystem<UWorldPartitionSubsystem>()->ServerClientsVisibleLevelNames;
+	auto ShouldWaitForClientVisibility = [&InParams, &OutTargetState, &ServerClientsVisibleLevelNames](const UWorldPartitionRuntimeCell* Cell)
 	{
-		check(bIsServer);
+		check(InParams.bIsServer);
 		if (Cell->ShouldServerWaitForClientLevelVisibility())
 		{
 			if (ULevel* Level = Cell->GetLevel())
 			{
 				if (ServerClientsVisibleLevelNames.Contains(Level->GetPackage()->GetFName()))
 				{
-					UE_CLOG(bUpdateServerEpoch, LogWorldPartition, Verbose, TEXT("Server epoch update delayed by client visibility"));
-					bUpdateServerEpoch = false;
+					UE_CLOG(OutTargetState.bUpdateServerEpoch, LogWorldPartition, Verbose, TEXT("Server epoch update delayed by client visibility"));
+					OutTargetState.bUpdateServerEpoch = false;
 					return true;
 				}
 			}
@@ -460,45 +528,56 @@ void UWorldPartitionStreamingPolicy::UpdateStreamingState()
 		return Cell->GetIsHLOD() && !UWorldPartitionHLODRuntimeSubsystem::IsHLODEnabled();
 	};
 
-	// Activation superseeds Loading
-	FrameLoadCells = FrameLoadCells.Difference(FrameActivateCells);
+	auto ShouldSkipCellForPerformance = [&InParams](const UWorldPartitionRuntimeCell* Cell)
+	{
+		// When performance is degrading start skipping non blocking cells
+		return !InParams.bIsServer && InParams.bIsBlockingCausedByBadStreamingPerformance && !Cell->GetBlockOnSlowLoading();
+	};
+
+	// Activation supersedes Loading
+	{
+		TRACE_CPUPROFILER_EVENT_SCOPE(UWorldPartitionStreamingPolicy::UpdateStreamingState_FrameLoadCells);
+		if (InParams.FrameLoadCells.Num() && InParams.FrameActivateCells.Num())
+		{
+			InParams.FrameLoadCells = InParams.FrameLoadCells.Difference(InParams.FrameActivateCells);
+		}
+	}
 
 	// Determine cells to activate
-	if (bCanStream)
+	if (InParams.bCanStream)
 	{
 		TRACE_CPUPROFILER_EVENT_SCOPE(UWorldPartitionStreamingPolicy::UpdateStreamingState_ToActivateCells);
-		for (const UWorldPartitionRuntimeCell* Cell : FrameActivateCells)
+		for (const UWorldPartitionRuntimeCell* Cell : InParams.FrameActivateCells)
 		{
-			if (ActivatedCells.Contains(Cell))
+			if (InParams.CurrentState.ActivatedCells.Contains(Cell))
 			{
 				// Update streaming source info for pending add to world cells
-				if (bShouldMergeStreamingSourceInfo && ActivatedCells.GetPendingAddToWorldCells().Contains(Cell))
+				if (InParams.bShouldMergeStreamingSourceInfo && InParams.CurrentState.ActivatedCells.GetPendingAddToWorldCells().Contains(Cell))
 				{
 					Cell->MergeStreamingSourceInfo();
 				}
 			}
 			else if (!ShouldSkipCellForPerformance(Cell) && !ShouldSkipDisabledHLODCell(Cell))
 			{
-				if (bShouldMergeStreamingSourceInfo)
+				if (InParams.bShouldMergeStreamingSourceInfo)
 				{
 					Cell->MergeStreamingSourceInfo();
 				}
-				ToActivateCells.Add(Cell);
+				OutTargetState.ToActivateCells.Add(Cell);
 			}
 		}
 	}
 
 	// Determine cells to load and server cells to deactivate
-	TArray<const UWorldPartitionRuntimeCell*> ServerToDeactivateCells;
-	if (bCanStream)
+	if (InParams.bCanStream)
 	{
 		TRACE_CPUPROFILER_EVENT_SCOPE(UWorldPartitionStreamingPolicy::UpdateStreamingState_ToLoadCells);
-		for (const UWorldPartitionRuntimeCell* Cell : FrameLoadCells)
+		for (const UWorldPartitionRuntimeCell* Cell : InParams.FrameLoadCells)
 		{
-			if (LoadedCells.Contains(Cell))
+			if (InParams.CurrentState.LoadedCells.Contains(Cell))
 			{
 				// Update streaming source info for pending load cells
-				if (bShouldMergeStreamingSourceInfo && !Cell->GetLevel())
+				if (InParams.bShouldMergeStreamingSourceInfo && !Cell->GetLevel())
 				{
 					Cell->MergeStreamingSourceInfo();
 				}
@@ -508,21 +587,21 @@ void UWorldPartitionStreamingPolicy::UpdateStreamingState()
 				if (!ShouldSkipCellForPerformance(Cell) && !ShouldSkipDisabledHLODCell(Cell))
 				{
 					// Server deactivated cells are processed right away (see below for details)
-					if (const bool bIsServerCellToDeactivate = bIsServer && ActivatedCells.Contains(Cell))
+					if (const bool bIsServerCellToDeactivate = InParams.bIsServer && InParams.CurrentState.ActivatedCells.Contains(Cell))
 					{
 						// Only deactivated server cells need to call ShouldWaitForClientVisibility (those part of ActivatedCells)
 						if (!ShouldWaitForClientVisibility(Cell))
 						{
-							ServerToDeactivateCells.Add(Cell);
+							OutTargetState.ToDeactivateCells.Add(Cell);
 						}
 					}
 					else
 					{
-						if (bShouldMergeStreamingSourceInfo)
+						if (InParams.bShouldMergeStreamingSourceInfo)
 						{
 							Cell->MergeStreamingSourceInfo();
 						}
-						ToLoadCells.Add(Cell);
+						OutTargetState.ToLoadCells.Add(Cell);
 					}
 				}
 			}
@@ -530,72 +609,86 @@ void UWorldPartitionStreamingPolicy::UpdateStreamingState()
 	}
 
 	// Determine cells to unload
-	TArray<const UWorldPartitionRuntimeCell*> ToUnloadCells;
 	{
 		TRACE_CPUPROFILER_EVENT_SCOPE(UWorldPartitionStreamingPolicy::UpdateStreamingState_ToUnloadCells);
-		auto BuildCellsToUnload = [this, &ToUnloadCells, bCanStream, bIsServer, ShouldWaitForClientVisibility](const TSet<TObjectPtr<const UWorldPartitionRuntimeCell>>& InCells)
+		auto BuildCellsToUnload = [&InParams, &OutTargetState, ShouldWaitForClientVisibility](const TSet<TObjectPtr<const UWorldPartitionRuntimeCell>>& InCells)
 		{
 			for (const UWorldPartitionRuntimeCell* Cell : InCells)
 			{
-				if (!FrameActivateCells.Contains(Cell) && !FrameLoadCells.Contains(Cell))
+				if (!InParams.FrameActivateCells.Contains(Cell) && !InParams.FrameLoadCells.Contains(Cell))
 				{
-					if (!bCanStream || !bIsServer || !ShouldWaitForClientVisibility(Cell))
+					if (!InParams.bCanStream || !InParams.bIsServer || !ShouldWaitForClientVisibility(Cell))
 					{
-						ToUnloadCells.Add(Cell);
+						OutTargetState.ToUnloadCells.Add(Cell);
 					}
 				}
 			}
 		};
 
-		BuildCellsToUnload(ActivatedCells.GetCells());
-		BuildCellsToUnload(LoadedCells);
+		BuildCellsToUnload(InParams.CurrentState.ActivatedCells.GetCells());
+		BuildCellsToUnload(InParams.CurrentState.LoadedCells);
 	}
 
 	UE_SUPPRESS(LogWorldPartition, Verbose,
-		if ((bIsStreamingInEnabled && (ToActivateCells.Num() > 0 || ToLoadCells.Num() > 0)) || ToUnloadCells.Num() > 0)
+	if ((InParams.bIsStreamingInEnabled && (OutTargetState.ToActivateCells.Num() > 0 || OutTargetState.ToLoadCells.Num() > 0)) || OutTargetState.ToUnloadCells.Num() > 0)
+	{
+		UE_LOG(LogWorldPartition, Verbose, TEXT("UWorldPartitionStreamingPolicy: CellsToActivate(%d), CellsToLoad(%d), CellsToUnload(%d)"), OutTargetState.ToActivateCells.Num(), OutTargetState.ToLoadCells.Num(), OutTargetState.ToUnloadCells.Num());
+		for (int i = 0; i < InParams.CurrentState.StreamingSources.Num(); ++i)
 		{
-			UE_LOG(LogWorldPartition, Verbose, TEXT("UWorldPartitionStreamingPolicy: CellsToActivate(%d), CellsToLoad(%d), CellsToUnload(%d)"), ToActivateCells.Num(), ToLoadCells.Num(), ToUnloadCells.Num());
-			FTransform LocalToWorld = WorldPartition->GetInstanceTransform();
-			for (int i = 0; i < StreamingSources.Num(); ++i)
-			{
-				FVector ViewLocation = LocalToWorld.TransformPosition(StreamingSources[i].Location);
-				FRotator ViewRotation = LocalToWorld.TransformRotation(StreamingSources[i].Rotation.Quaternion()).Rotator();
-				UE_LOG(LogWorldPartition, Verbose, TEXT("UWorldPartitionStreamingPolicy: Sources[%d] = %s,%s"), i, *ViewLocation.ToString(), *ViewRotation.ToString());
-			}
-		});
+			FVector ViewLocation = InParams.WorldPartitionInstanceTransform.TransformPosition(InParams.CurrentState.StreamingSources[i].Location);
+			FRotator ViewRotation = InParams.WorldPartitionInstanceTransform.TransformRotation(InParams.CurrentState.StreamingSources[i].Rotation.Quaternion()).Rotator();
+			UE_LOG(LogWorldPartition, Verbose, TEXT("UWorldPartitionStreamingPolicy: Sources[%d] = %s,%s"), i, *ViewLocation.ToString(), *ViewRotation.ToString());
+		}
+	});
 
 #if !UE_BUILD_SHIPPING
-	UpdateDebugCellsStreamingPriority(FrameActivateCells, FrameLoadCells);
+	UWorldPartitionStreamingPolicy::UpdateDebugCellsStreamingPriority(InParams.FrameActivateCells, InParams.FrameLoadCells, InParams.bShouldMergeStreamingSourceInfo);
 #endif
 
-	// Unloaded cells
-	if (ToUnloadCells.Num() > 0)
+	if (InParams.RuntimeHash)
 	{
-		SetCellsStateToUnloaded(ToUnloadCells);
+		OutTargetState.StreamingPerformance = InParams.RuntimeHash->GetStreamingPerformance(InParams.FrameActivateCells);
 	}
-	
-	// Server deactivated cells (activated -> loaded)
-	// 
-	// Server deactivation is handle right away to ensure that even if WorldPartitionSubsystem::UpdateStreamingState
-	// is running in incremental mode, server deactivated cells will make their streaming level ShouldBeVisible() 
-	// return false. This way, UNetConnection::UpdateLevelVisibilityInternal will not allow clients to make their 
-	// streaming level visible (see LevelVisibility.bTryMakeVisible).
-	for (const UWorldPartitionRuntimeCell* ServerCellToDeactivate : ServerToDeactivateCells)
+}
+
+void UWorldPartitionStreamingPolicy::PostUpdateStreamingStateInternal_GameThread(FWorldPartitionUpdateStreamingTargetState& InOutTargetState)
+{
+	TRACE_CPUPROFILER_EVENT_SCOPE(UWorldPartitionStreamingPolicy::PostUpdateStreamingStateInternal_GameThread);
+
+	check(IsInGameThread());
+
+	// Unloaded cells
+	if (!InOutTargetState.ToUnloadCells.IsEmpty())
 	{
-		int32 DummyMaxCellToLoad = 0; // Deactivating is not concerned by MaxCellsToLoad
-		SetCellStateToLoaded(ServerCellToDeactivate, DummyMaxCellToLoad);
+		SetCellsStateToUnloaded(InOutTargetState.ToUnloadCells);
+		InOutTargetState.ToUnloadCells.Reset();
 	}
 
-	// Evaluate streaming performance based on cells that should be activated
-	UpdateStreamingPerformance(FrameActivateCells);
-	
-	// Update Epoch if we aren't waiting for clients anymore
+	const bool bIsServer = WorldPartition->IsServer();
+	check(bIsServer || InOutTargetState.ToDeactivateCells.IsEmpty());
 	if (bIsServer)
 	{
-		if (bUpdateServerEpoch)
+		if (!InOutTargetState.ToDeactivateCells.IsEmpty())
+		{
+			// Server deactivated cells (activated -> loaded)
+			// 
+			// Server deactivation is handle right away to ensure that even if WorldPartitionSubsystem::UpdateStreamingState
+			// is running in incremental mode, server deactivated cells will make their streaming level ShouldBeVisible() 
+			// return false. This way, UNetConnection::UpdateLevelVisibilityInternal will not allow clients to make their 
+			// streaming level visible (see LevelVisibility.bTryMakeVisible).
+			for (const UWorldPartitionRuntimeCell* ServerCellToDeactivate : InOutTargetState.ToDeactivateCells)
+			{
+				int32 DummyMaxCellToLoad = 0; // Deactivating is not concerned by MaxCellsToLoad
+				SetCellStateToLoaded(ServerCellToDeactivate, DummyMaxCellToLoad);
+			}
+			InOutTargetState.ToDeactivateCells.Reset();
+		}
+
+		// Update Epoch if we aren't waiting for clients anymore
+		if (InOutTargetState.bUpdateServerEpoch)
 		{
 			ServerStreamingStateEpoch = WorldPartition->GetStreamingStateEpoch();
-			ServerStreamingEnabledEpoch = NewServerStreamingEnabledEpoch;
+			ServerStreamingEnabledEpoch = ComputeServerStreamingEnabledEpoch();
 			UE_LOG(LogWorldPartition, Verbose, TEXT("Server epoch updated"));
 		}
 		else
@@ -605,17 +698,88 @@ void UWorldPartitionStreamingPolicy::UpdateStreamingState()
 			UpdateStreamingHash = 0;
 		}
 	}
+
+	// Evaluate streaming performance based on cells that should be activated
+	UpdateStreamingPerformance(InOutTargetState.StreamingPerformance);
+}
+
+void UWorldPartitionStreamingPolicy::OnPreChangeStreamingContent()
+{
+	WaitForAsyncUpdateStreamingState();
+}
+
+bool UWorldPartitionStreamingPolicy::WaitForAsyncUpdateStreamingState()
+{
+	if (AsyncPostUpdateStreamingStateTask.IsValid())
+	{
+		// Wait for completion
+		TRACE_CPUPROFILER_EVENT_SCOPE(UWorldPartitionStreamingPolicy::WaitForAsyncUpdateStreamingState);
+		check(AsyncUpdateTaskState == EAsyncUpdateTaskState::Started);
+		AsyncPostUpdateStreamingStateTask->Wait();
+		check(AsyncUpdateStreamingStateTask.IsValid());
+		check(AsyncUpdateStreamingStateTask->IsComplete());
+		return true;
+	}
+	return false;
+}
+
+void UWorldPartitionStreamingPolicy::OnStreamingStateUpdated()
+{
+	if (AsyncUpdateTaskState == EAsyncUpdateTaskState::Pending)
+	{
+		// Here, it's considered safe to start the asynchronous call to UWorldPartitionStreamingPolicy::UpdateStreamingStateInternal since the
+		// WorldPartitionSubsystem is done working on the world partition streaming cells returned by UWorldPartitionStreamingPolicy::GetCellsToUpdate.
+		// 
+		// Any call that modifies the streaming content should first call UWorldPartitionStreamingPolicy::OnPreChangeStreamingContent to make sure
+		// that any asynchronous update task completes before modifying the streaming content.
+		// 
+		// All the required input is prepared and copied in the FUpdateStreamingStateParams structure. 
+		// Some members (like FrameActivateCells/FrameLoadCells) are direct references to UWorldPartitionStreamingPolicy members, these are considered 
+		// safe to access read/write from the asynchronous task.
+		// 
+		// Note that some calls to world partition cells will cache information into the cell (thus modify it).
+		// Some of this information is either cached for performance reasons, some is used to prioritize cells (see UWorldPartitionRuntimeCell::SortCompare)
+		// Here's the list of calls that modify the cell:
+		// - UWorldPartitionRuntimeCell::GetCellEffectiveWantedState
+		// - UWorldPartitionRuntimeCellData::ResetStreamingSourceInfo
+		// - UWorldPartitionRuntimeCellData::AppendStreamingSourceInfo
+		// - UWorldPartitionRuntimeCellData::MergeStreamingSourceInfo
+		// 
+		// Note: This could be revisited at some point (FWorldPartitionStreamingContext could store this information).
+
+		TRACE_CPUPROFILER_EVENT_SCOPE(UWorldPartitionStreamingPolicy::OnStreamingStateUpdated);
+		{
+			// Prepare async task input payload
+			TRACE_CPUPROFILER_EVENT_SCOPE(UWorldPartitionStreamingPolicy::PrepareAsyncTaskPayloads);
+			AsyncTaskCurrentState.CopyFrom(CurrentState);
+			check(AsyncTaskTargetState.IsEmpty());
+		}
+		{
+			// Create async tasks
+			TRACE_CPUPROFILER_EVENT_SCOPE(UWorldPartitionStreamingPolicy::CreateAndDispatchAsyncTasks);
+			FUpdateStreamingStateParams InputParams = FUpdateStreamingStateParams(this, AsyncTaskCurrentState).SetRequiredWorldDataLayersEffectiveStatesCopy(true);
+			AsyncUpdateStreamingStateTask = FFunctionGraphTask::CreateAndDispatchWhenReady([InputParams = MoveTemp(InputParams), this]()
+			{
+				TRACE_CPUPROFILER_EVENT_SCOPE(UWorldPartitionStreamingPolicy::AsyncUpdateStreamingState);
+				UWorldPartitionStreamingPolicy::UpdateStreamingStateInternal(InputParams, AsyncTaskTargetState);
+			}, GET_STATID(WPStreamingPolicy_AsyncUpdateStreamingState), NULL, ENamedThreads::AnyBackgroundHiPriTask);
+			AsyncPostUpdateStreamingStateTask = FFunctionGraphTask::CreateAndDispatchWhenReady([this]() { PostUpdateStreamingStateInternal_GameThread(AsyncTaskTargetState); }, GET_STATID(WPStreamingPolicy_AsyncPostUpdateStreamingState), AsyncUpdateStreamingStateTask, ENamedThreads::GameThread);
+			AsyncUpdateTaskState = EAsyncUpdateTaskState::Started;
+		}
+	}
 }
 
 #if !UE_BUILD_SHIPPING
-void UWorldPartitionStreamingPolicy::UpdateDebugCellsStreamingPriority(const TSet<const UWorldPartitionRuntimeCell*>& ActivateStreamingCells, const TSet<const UWorldPartitionRuntimeCell*>& LoadStreamingCells)
+void UWorldPartitionStreamingPolicy::UpdateDebugCellsStreamingPriority(const TSet<const UWorldPartitionRuntimeCell*>& InActivateStreamingCells, const TSet<const UWorldPartitionRuntimeCell*>& InLoadStreamingCells, bool bInShouldMergeStreamingSourceInfo)
 {
+	// @todo_ow: This code generates debug priority values local to this partitioned world.
+	// To properly support multiple partitioned worlds, move the sorting pass and the priority update in the WorldPartitionSubsystem.
 	if (FWorldPartitionDebugHelper::IsRuntimeSpatialHashCellStreamingPriorityShown())
 	{
-		TArray<const UWorldPartitionRuntimeCell*> Cells = ActivateStreamingCells.Array();
-		Cells.Append(LoadStreamingCells.Array());
+		TArray<const UWorldPartitionRuntimeCell*> Cells = InActivateStreamingCells.Array();
+		Cells.Append(InLoadStreamingCells.Array());
 
-		if (bShouldMergeStreamingSourceInfo)
+		if (bInShouldMergeStreamingSourceInfo)
 		{
 			for (const UWorldPartitionRuntimeCell* Cell : Cells)
 			{
@@ -623,7 +787,11 @@ void UWorldPartitionStreamingPolicy::UpdateDebugCellsStreamingPriority(const TSe
 			}
 		}
 
-		SortStreamingCellsByImportance(Cells);
+		if (Cells.Num() > 1)
+		{
+			TRACE_CPUPROFILER_EVENT_SCOPE(SortStreamingCellsByImportance);
+			Algo::Sort(Cells, [](const UWorldPartitionRuntimeCell* CellA, const UWorldPartitionRuntimeCell* CellB) { return CellA->SortCompare(CellB) < 0; });
+		}
 
 		const int32 CellCount = Cells.Num();
 		int32 CellPrio = 0;
@@ -635,12 +803,14 @@ void UWorldPartitionStreamingPolicy::UpdateDebugCellsStreamingPriority(const TSe
 }
 #endif
 
-void UWorldPartitionStreamingPolicy::UpdateStreamingPerformance(const TSet<const UWorldPartitionRuntimeCell*>& CellsToActivate)
-{		
+void UWorldPartitionStreamingPolicy::UpdateStreamingPerformance(EWorldPartitionStreamingPerformance NewStreamingPerformance)
+{
+	check(IsInGameThread());
+
 	TRACE_CPUPROFILER_EVENT_SCOPE(UWorldPartitionStreamingPolicy::UpdateStreamingPerformance);
 	UWorld* World = GetWorld();
 	// If we are currently in a blocked loading just reset the on screen message time and return
-	if (StreamingPerformance == EWorldPartitionStreamingPerformance::Critical && IsInBlockTillLevelStreamingCompleted())
+	if (CurrentStreamingPerformance == EWorldPartitionStreamingPerformance::Critical && IsInBlockTillLevelStreamingCompleted())
 	{
 #if !UE_BUILD_SHIPPING
 		OnScreenMessageStartTime = FPlatformTime::Seconds();
@@ -650,28 +820,26 @@ void UWorldPartitionStreamingPolicy::UpdateStreamingPerformance(const TSet<const
 
 	if (WorldPartition->RuntimeHash)
 	{
-		EWorldPartitionStreamingPerformance NewStreamingPerformance = WorldPartition->RuntimeHash->GetStreamingPerformance(CellsToActivate);
-
-		if (StreamingPerformance != NewStreamingPerformance)
+		if (CurrentStreamingPerformance != NewStreamingPerformance)
 		{
 			UE_LOG(LogWorldPartition, Log, TEXT("Streaming performance changed: %s -> %s"),
-				*StaticEnum<EWorldPartitionStreamingPerformance>()->GetDisplayNameTextByValue((int64)StreamingPerformance).ToString(),
+				*StaticEnum<EWorldPartitionStreamingPerformance>()->GetDisplayNameTextByValue((int64)CurrentStreamingPerformance).ToString(),
 				*StaticEnum<EWorldPartitionStreamingPerformance>()->GetDisplayNameTextByValue((int64)NewStreamingPerformance).ToString());
 
-			StreamingPerformance = NewStreamingPerformance;
+			CurrentStreamingPerformance = NewStreamingPerformance;
 		}
 	}
 
 #if !UE_BUILD_SHIPPING
-	if (StreamingPerformance != EWorldPartitionStreamingPerformance::Good)
+	if (CurrentStreamingPerformance != EWorldPartitionStreamingPerformance::Good)
 	{
 		// performance still bad keep message alive
 		OnScreenMessageStartTime = FPlatformTime::Seconds();
-		OnScreenMessageStreamingPerformance = StreamingPerformance;
+		OnScreenMessageStreamingPerformance = CurrentStreamingPerformance;
 	}
 #endif
 	
-	if (StreamingPerformance == EWorldPartitionStreamingPerformance::Critical)
+	if (CurrentStreamingPerformance == EWorldPartitionStreamingPerformance::Critical)
 	{
 		const bool bIsServer = WorldPartition->IsServer();
 		const bool bIsServerStreamingEnabled = WorldPartition->IsServerStreamingEnabled();
@@ -712,28 +880,19 @@ void UWorldPartitionStreamingPolicy::GetOnScreenMessages(FCoreDelegates::FSeveri
 }
 #endif
 
-bool UWorldPartitionStreamingPolicy::ShouldSkipCellForPerformance(const UWorldPartitionRuntimeCell* Cell) const
-{
-	// When performance is degrading start skipping non blocking cells
-	if (!Cell->GetBlockOnSlowLoading())
-	{
-		if (!WorldPartition->IsServer())
-		{
-			return IsInBlockTillLevelStreamingCompleted(/*bIsCausedByBadStreamingPerformance*/true);
-		}
-	}
-	return false;
-}
-
 void UWorldPartitionStreamingPolicy::GetCellsToUpdate(TArray<const UWorldPartitionRuntimeCell*>& OutToLoadCells, TArray<const UWorldPartitionRuntimeCell*>& OutToActivateCells)
 {
-	OutToLoadCells.Append(ToLoadCells);
-	OutToActivateCells.Append(ToActivateCells);
+	check(IsInGameThread());
+
+	OutToLoadCells.Append(TargetState.ToLoadCells);
+	OutToActivateCells.Append(TargetState.ToActivateCells);
 }
 
 void UWorldPartitionStreamingPolicy::GetCellsToReprioritize(TArray<const UWorldPartitionRuntimeCell*>& OutToReprioritizeLoadCells, TArray<const UWorldPartitionRuntimeCell*>& OutToReprioritizeActivateCells)
 {
-	for (const UWorldPartitionRuntimeCell* Cell : LoadedCells)
+	check(IsInGameThread());
+
+	for (const UWorldPartitionRuntimeCell* Cell : CurrentState.LoadedCells)
 	{
 		if (!Cell->GetLevel())
 		{
@@ -741,7 +900,7 @@ void UWorldPartitionStreamingPolicy::GetCellsToReprioritize(TArray<const UWorldP
 		}
 	}
 
-	for (const UWorldPartitionRuntimeCell* Cell : ActivatedCells.GetPendingAddToWorldCells())
+	for (const UWorldPartitionRuntimeCell* Cell : CurrentState.ActivatedCells.GetPendingAddToWorldCells())
 	{
 		OutToReprioritizeActivateCells.Add(Cell);
 	}
@@ -749,13 +908,14 @@ void UWorldPartitionStreamingPolicy::GetCellsToReprioritize(TArray<const UWorldP
 
 void UWorldPartitionStreamingPolicy::SetCellStateToLoaded(const UWorldPartitionRuntimeCell* InCell, int32& InOutMaxCellsToLoad)
 {
+	check(IsInGameThread());
+
 	bool bLoadCell = false;
-	if (ActivatedCells.Contains(InCell))
+	if (CurrentState.ActivatedCells.Contains(InCell))
 	{
 		InCell->Deactivate();
-		ActivatedCells.Remove(InCell);
+		CurrentState.ActivatedCells.Remove(InCell);
 		bLoadCell = true;
-		
 	}
 	else if (WorldPartition->IsStreamingInEnabled())
 	{
@@ -773,22 +933,24 @@ void UWorldPartitionStreamingPolicy::SetCellStateToLoaded(const UWorldPartitionR
 	if (bLoadCell)
 	{
 		UE_LOG(LogWorldPartition, Verbose, TEXT("UWorldPartitionStreamingPolicy::SetCellStateToLoaded %s"), *InCell->GetName());
-		LoadedCells.Add(InCell);
+		CurrentState.LoadedCells.Add(InCell);
 		++ProcessedToLoadCells;
 	}
 }
 
 void UWorldPartitionStreamingPolicy::SetCellStateToActivated(const UWorldPartitionRuntimeCell* InCell, int32& InOutMaxCellsToLoad)
 {
+	check(IsInGameThread());
+
 	if (!WorldPartition->IsStreamingInEnabled())
 	{
 		return;
 	}
 
 	bool bActivateCell = false;
-	if (LoadedCells.Contains(InCell))
+	if (CurrentState.LoadedCells.Contains(InCell))
 	{
-		LoadedCells.Remove(InCell);
+		CurrentState.LoadedCells.Remove(InCell);
 		bActivateCell = true;
 	}
 	else if (InOutMaxCellsToLoad > 0)
@@ -803,13 +965,13 @@ void UWorldPartitionStreamingPolicy::SetCellStateToActivated(const UWorldPartiti
 	if (bActivateCell)
 	{
 		UE_LOG(LogWorldPartition, Verbose, TEXT("UWorldPartitionStreamingPolicy::SetCellStateToActivated %s"), *InCell->GetName());
-		ActivatedCells.Add(InCell);
+		CurrentState.ActivatedCells.Add(InCell);
 		InCell->Activate();
 		++ProcessedToActivateCells;
 	}
 }
 
-void UWorldPartitionStreamingPolicy::SetCellsStateToUnloaded(const TArray<const UWorldPartitionRuntimeCell*>& InToUnloadCells)
+void UWorldPartitionStreamingPolicy::SetCellsStateToUnloaded(const TArray<TObjectPtr<const UWorldPartitionRuntimeCell>>& InToUnloadCells)
 {
 	TRACE_CPUPROFILER_EVENT_SCOPE(UWorldPartitionStreamingPolicy::SetCellsStateToUnloaded);
 
@@ -819,8 +981,8 @@ void UWorldPartitionStreamingPolicy::SetCellsStateToUnloaded(const TArray<const
 		{
 			UE_LOG(LogWorldPartition, Verbose, TEXT("UWorldPartitionStreamingPolicy::UnloadCells %s"), *Cell->GetName());
 			Cell->Unload();
-			ActivatedCells.Remove(Cell);
-			LoadedCells.Remove(Cell);
+			CurrentState.ActivatedCells.Remove(Cell);
+			CurrentState.LoadedCells.Remove(Cell);
 		}
 	}
 }
@@ -844,7 +1006,8 @@ bool UWorldPartitionStreamingPolicy::CanAddCellToWorld(const UWorldPartitionRunt
 		return true;
 	}
 
-	return !ShouldSkipCellForPerformance(InCell);
+	// When performance is degrading, start skipping non blocking cells
+	return InCell->GetBlockOnSlowLoading();
 }
 
 bool UWorldPartitionStreamingPolicy::IsStreamingCompleted(const TArray<FWorldPartitionStreamingSource>* InStreamingSources) const
@@ -881,7 +1044,7 @@ bool UWorldPartitionStreamingPolicy::IsStreamingCompleted(const TArray<FWorldPar
 	}
 
 	// Test spatially loaded cells using streaming sources (or provided streaming source)
-	TArrayView<const FWorldPartitionStreamingSource> QueriedStreamingSources = bTestProvidedStreamingSource ? *InStreamingSources : StreamingSources;
+	TArrayView<const FWorldPartitionStreamingSource> QueriedStreamingSources = bTestProvidedStreamingSource ? *InStreamingSources : CurrentState.StreamingSources;
 	for (const FWorldPartitionStreamingSource& StreamingSource : QueriedStreamingSources)
 	{
 		// Build a query source from a Streaming Source
@@ -958,7 +1121,7 @@ bool UWorldPartitionStreamingPolicy::IsStreamingCompleted(EWorldPartitionRuntime
 
 bool UWorldPartitionStreamingPolicy::DrawRuntimeHash2D(FWorldPartitionDraw2DContext& DrawContext)
 {
-	if (StreamingSources.Num() > 0 && WorldPartition->RuntimeHash)
+	if (CurrentState.StreamingSources.Num() > 0 && WorldPartition->RuntimeHash)
 	{
 		return WorldPartition->RuntimeHash->Draw2D(DrawContext);
 	}
@@ -969,18 +1132,18 @@ void UWorldPartitionStreamingPolicy::DrawRuntimeHash3D()
 {
 	if (WorldPartition->IsInitialized() && WorldPartition->RuntimeHash)
 	{
-		WorldPartition->RuntimeHash->Draw3D(StreamingSources);
+		WorldPartition->RuntimeHash->Draw3D(CurrentState.StreamingSources);
 	}
 }
 
 void UWorldPartitionStreamingPolicy::OnCellShown(const UWorldPartitionRuntimeCell* InCell)
 {
-	ActivatedCells.OnAddedToWorld(InCell);
+	CurrentState.ActivatedCells.OnAddedToWorld(InCell);
 }
 
 void UWorldPartitionStreamingPolicy::OnCellHidden(const UWorldPartitionRuntimeCell* InCell)
 {
-	ActivatedCells.OnRemovedFromWorld(InCell);
+	CurrentState.ActivatedCells.OnRemovedFromWorld(InCell);
 }
 
 void FActivatedCells::Add(const UWorldPartitionRuntimeCell* InCell)
@@ -998,6 +1161,12 @@ void FActivatedCells::Remove(const UWorldPartitionRuntimeCell* InCell)
 	PendingAddToWorldCells.Remove(InCell);
 }
 
+void FActivatedCells::Reset()
+{
+	Cells.Reset();
+	PendingAddToWorldCells.Reset();
+}
+
 void FActivatedCells::OnAddedToWorld(const UWorldPartitionRuntimeCell* InCell)
 {
 	PendingAddToWorldCells.Remove(InCell);
@@ -1014,4 +1183,68 @@ void FActivatedCells::OnRemovedFromWorld(const UWorldPartitionRuntimeCell* InCel
 	}
 }
 
+void FWorldPartitionUpdateStreamingCurrentState::Reset()
+{
+	StreamingSources.Reset();
+	LoadedCells.Reset();
+	ActivatedCells.Reset();
+}
+
+void FWorldPartitionUpdateStreamingCurrentState::CopyFrom(const FWorldPartitionUpdateStreamingCurrentState& InCurrentState)
+{
+	StreamingSources = InCurrentState.StreamingSources;
+	LoadedCells = InCurrentState.LoadedCells;
+	ActivatedCells = InCurrentState.ActivatedCells;
+}
+
+bool FWorldPartitionUpdateStreamingTargetState::IsEmpty() const
+{
+	return ToLoadCells.IsEmpty() && ToActivateCells.IsEmpty() && ToUnloadCells.IsEmpty() && ToDeactivateCells.IsEmpty();
+}
+
+void FWorldPartitionUpdateStreamingTargetState::Reset()
+{
+	ToLoadCells.Reset();
+	ToActivateCells.Reset();
+	ToDeactivateCells.Reset();
+	ToUnloadCells.Reset();
+	StreamingPerformance = EWorldPartitionStreamingPerformance::Good;
+}
+
+UWorldPartitionStreamingPolicy::FUpdateStreamingStateParams::FUpdateStreamingStateParams(UWorldPartitionStreamingPolicy* InPolicy, const FWorldPartitionUpdateStreamingCurrentState& InCurrentState)
+	: World(InPolicy->GetWorld())
+	, RuntimeHash(InPolicy->GetOuterUWorldPartition()->RuntimeHash)
+	, bCanStream(InPolicy->GetOuterUWorldPartition()->CanStream())
+	, bIsServer(InPolicy->GetOuterUWorldPartition()->IsServer())
+	, bIsStreamingInEnabled(InPolicy->GetOuterUWorldPartition()->IsStreamingInEnabled())
+	, bIsServerStreamingEnabled(InPolicy->GetOuterUWorldPartition()->IsServerStreamingEnabled())
+	, bIsServerStreamingOutEnabled(InPolicy->GetOuterUWorldPartition()->IsServerStreamingOutEnabled())
+	, bIsBlockingCausedByBadStreamingPerformance(InPolicy->IsInBlockTillLevelStreamingCompleted(true))
+	, bShouldMergeStreamingSourceInfo(InPolicy->bShouldMergeStreamingSourceInfo)
+	, PolicyUpdateStreamingStateEpoch(InPolicy->UpdateStreamingStateCounter)
+	, DataLayersLogicOperator(InPolicy->GetOuterUWorldPartition()->GetDataLayersLogicOperator())
+	, WorldPartitionInstanceTransform(InPolicy->GetOuterUWorldPartition()->GetInstanceTransform())
+	, CurrentState(InCurrentState)
+	, ServerDisallowedStreamingOutDataLayers(InPolicy->GetServerDisallowedStreamingOutDataLayers())
+	, FrameActivateCells(InPolicy->FrameActivateCells)
+	, FrameLoadCells(InPolicy->FrameLoadCells)
+	, WorldDataLayersEffectiveStatesRef(FWorldDataLayersEffectiveStatesAccessor::Get(InPolicy->GetOuterUWorldPartition()->GetTypedOuter<UWorld>()->GetWorldDataLayers()))
+{
+	check(IsInGameThread());
+}
+
+UWorldPartitionStreamingPolicy::FUpdateStreamingStateParams& UWorldPartitionStreamingPolicy::FUpdateStreamingStateParams::SetRequiredWorldDataLayersEffectiveStatesCopy(bool bInRequiredEffectiveStatesCopy)
+{
+	check(IsInGameThread());
+	if (bInRequiredEffectiveStatesCopy)
+	{
+		WorldDataLayersEffectiveStatesCopy = WorldDataLayersEffectiveStatesRef;
+	}
+	else
+	{
+		WorldDataLayersEffectiveStatesCopy.Reset();
+	}
+	return *this;
+}
+
 #undef LOCTEXT_NAMESPACE
diff --git a/Engine/Source/Runtime/Engine/Private/WorldPartition/WorldPartitionSubsystem.cpp b/Engine/Source/Runtime/Engine/Private/WorldPartition/WorldPartitionSubsystem.cpp
--- a/Engine/Source/Runtime/Engine/Private/WorldPartition/WorldPartitionSubsystem.cpp
+++ b/Engine/Source/Runtime/Engine/Private/WorldPartition/WorldPartitionSubsystem.cpp
@@ -954,7 +954,7 @@ void UWorldPartitionSubsystem::UpdateStreamingSources()
 
 	StreamingSources.Reset();
 
-	UWorld* World = GetWorld();
+	const UWorld* World = GetWorld();
 	bool bIsUsingReplayStreamingSources = false;
 	if (AWorldPartitionReplay::IsPlaybackEnabled(World))
 	{
@@ -1189,6 +1189,11 @@ void UWorldPartitionSubsystem::UpdateStreamingStateInternal(const UWorld* InWorl
 		return;
 	}
 
+	ON_SCOPE_EXIT
+	{
+		WorldPartitionSubsystem->OnStreamingStateUpdated().Broadcast();
+	};
+
 	TArray<TObjectPtr<UWorldPartition>> RegisteredWorldPartitionsCopy;
 	auto GetRegisteredWorldPartitionsCopy = [&RegisteredWorldPartitionsCopy, InWorldPartition, WorldPartitionSubsystem]()
 	{
diff --git a/Engine/Source/Runtime/Engine/Public/WorldPartition/DataLayer/DataLayerManager.h b/Engine/Source/Runtime/Engine/Public/WorldPartition/DataLayer/DataLayerManager.h
--- a/Engine/Source/Runtime/Engine/Public/WorldPartition/DataLayer/DataLayerManager.h
+++ b/Engine/Source/Runtime/Engine/Public/WorldPartition/DataLayer/DataLayerManager.h
@@ -122,8 +122,9 @@ private:
 
 	ENGINE_API AWorldDataLayers* GetWorldDataLayers() const;
 	ENGINE_API const UDataLayerInstance* GetDataLayerInstanceFromAssetName(const FName& InDataLayerAssetPathName) const;
-
 	ENGINE_API void BroadcastOnDataLayerInstanceRuntimeStateChanged(const UDataLayerInstance* InDataLayer, EDataLayerRuntimeState InState);
+	ENGINE_API static bool IsAnyDataLayerInEffectiveRuntimeState(TArrayView<const FName> InDataLayerNames, EDataLayerRuntimeState InState, const FWorldDataLayersEffectiveStates& InEffectiveStates);
+	ENGINE_API static bool IsAllDataLayerInEffectiveRuntimeState(TArrayView<const FName> InDataLayerNames, EDataLayerRuntimeState InState, const FWorldDataLayersEffectiveStates& InEffectiveStates);
 
 	/** Referenced objects (used by verse) */
 	UPROPERTY(Transient)
@@ -152,6 +153,7 @@ private:
 	friend class UWorldPartitionSubsystem;
 	friend class FWorldPartitionDebugHelper;
 	friend class UVerseDataLayerManagerBase;
+	friend struct FWorldPartitionStreamingContext;
 
 #if WITH_EDITOR
 private:
diff --git a/Engine/Source/Runtime/Engine/Public/WorldPartition/DataLayer/WorldDataLayers.h b/Engine/Source/Runtime/Engine/Public/WorldPartition/DataLayer/WorldDataLayers.h
--- a/Engine/Source/Runtime/Engine/Public/WorldPartition/DataLayer/WorldDataLayers.h
+++ b/Engine/Source/Runtime/Engine/Public/WorldPartition/DataLayer/WorldDataLayers.h
@@ -47,6 +47,36 @@ struct FActorPlacementDataLayers
 	}
 };
 
+struct FWorldDataLayersEffectiveStates
+{
+	FWorldDataLayersEffectiveStates();
+	const TSet<FName>& GetAllEffectiveActiveDataLayerNames() const;
+	const TSet<FName>& GetAllEffectiveLoadedDataLayerNames() const;
+	EDataLayerRuntimeState GetDataLayerEffectiveRuntimeStateByName(FName InDataLayerName) const;
+	int32 GetUpdateEpoch() const { return UpdateEpoch; }
+
+protected:
+	void Reset();
+	void SetReplicatedEffectiveActiveDataLayerNames(const TArray<FName>& InRepEffectiveActiveDataLayerNames);
+	void SetReplicatedEffectiveLoadedDataLayerNames(const TArray<FName>& InRepEffectiveLoadedDataLayerNames);
+	bool SetDataLayerEffectiveRuntimeState(FName InDataLayerName, bool bIsLocalDataLayer, EDataLayerRuntimeState NewEffectiveRuntimeState, EDataLayerRuntimeState& OutOldEffectiveRuntimeState);
+	const TSet<FName>& GetReplicatedEffectiveActiveDataLayerNames() const { return ReplicatedEffectiveActiveDataLayerNames; }
+	const TSet<FName>& GetReplicatedEffectiveLoadedDataLayerNames() const { return ReplicatedEffectiveLoadedDataLayerNames; }
+	friend class AWorldDataLayers;
+
+private:
+	TSet<FName> ReplicatedEffectiveActiveDataLayerNames;
+	TSet<FName> ReplicatedEffectiveLoadedDataLayerNames;
+	TSet<FName> LocalEffectiveActiveDataLayerNames;
+	TSet<FName> LocalEffectiveLoadedDataLayerNames;
+
+	int32 UpdateEpoch;
+	mutable int32 AllEffectiveActiveDataLayerNamesEpoch;
+	mutable int32 AllEffectiveLoadedDataLayerNamesEpoch;
+	mutable TSet<FName> AllEffectiveActiveDataLayerNames;
+	mutable TSet<FName> AllEffectiveLoadedDataLayerNames;
+};
+
 /**
  * Actor containing data layers instances within a world.
  */
@@ -215,6 +245,8 @@ protected:
 	ENGINE_API void OnRep_EffectiveLoadedDataLayerNames();
 
 private:
+	ENGINE_API const FWorldDataLayersEffectiveStates& GetEffectiveStates() const;
+
 	// External Data Layers
 	bool AddExternalDataLayerInstance(UExternalDataLayerInstance* ExternalDataLayerInstance);
 	bool RemoveExternalDataLayerInstance(UExternalDataLayerInstance* ExternalDataLayerInstance);
@@ -321,17 +353,7 @@ private:
 	UPROPERTY(Transient, Replicated, ReplicatedUsing=OnRep_EffectiveLoadedDataLayerNames)
 	TArray<FName> RepEffectiveLoadedDataLayerNames;
 
-	// TSet do not support replication so we replicate an array and update the set in the OnRep_EffectiveActiveDataLayerNames/OnRep_EffectiveLoadedDataLayerNames
-	TSet<FName> EffectiveActiveDataLayerNames;
-	TSet<FName> EffectiveLoadedDataLayerNames;
-
-	TSet<FName> LocalEffectiveActiveDataLayerNames;
-	TSet<FName> LocalEffectiveLoadedDataLayerNames;
-
-	mutable int32 AllEffectiveActiveDataLayerNamesEpoch;
-	mutable TSet<FName> AllEffectiveActiveDataLayerNames;
-	mutable int32 AllEffectiveLoadedDataLayerNamesEpoch;
-	mutable TSet<FName> AllEffectiveLoadedDataLayerNames;
+	FWorldDataLayersEffectiveStates EffectiveStates;
 
 	int32 DataLayersStateEpoch;
 
@@ -342,8 +364,23 @@ private:
 	friend class UDataLayerManager;
 	friend class UExternalDataLayerManager;
 	friend class UDataLayerEditorSubsystem;
+	friend struct FWorldDataLayersEffectiveStatesAccessor;
 };
 
+struct FWorldDataLayersEffectiveStatesAccessor
+{
+private:
+	static const FWorldDataLayersEffectiveStates& Get(const AWorldDataLayers* InWorldDataLayers)
+	{
+		check(InWorldDataLayers);
+		return InWorldDataLayers->GetEffectiveStates();
+	}
+	friend struct FWorldPartitionStreamingContext;
+	friend class UWorldPartitionStreamingPolicy;
+	friend class UDataLayerManager;
+};
+
+
 DEFINE_ACTORDESC_TYPE(AWorldDataLayers, FWorldDataLayersActorDesc);
 
 #if WITH_EDITOR
@@ -454,9 +491,9 @@ void AWorldDataLayers::OverwriteDataLayerRuntimeStates(const TArray<IdentifierTy
 		});
 
 		MARK_PROPERTY_DIRTY_FROM_NAME(AWorldDataLayers, RepEffectiveActiveDataLayerNames, this);
-		RepEffectiveActiveDataLayerNames = EffectiveActiveDataLayerNames.Array();
+		RepEffectiveActiveDataLayerNames = EffectiveStates.GetReplicatedEffectiveActiveDataLayerNames().Array();
 		MARK_PROPERTY_DIRTY_FROM_NAME(AWorldDataLayers, RepEffectiveLoadedDataLayerNames, this);
-		RepEffectiveLoadedDataLayerNames = EffectiveLoadedDataLayerNames.Array();
+		RepEffectiveLoadedDataLayerNames = EffectiveStates.GetReplicatedEffectiveLoadedDataLayerNames().Array();
 	}
 }
 
diff --git a/Engine/Source/Runtime/Engine/Public/WorldPartition/RuntimeHashSet/WorldPartitionRuntimeHashSet.h b/Engine/Source/Runtime/Engine/Public/WorldPartition/RuntimeHashSet/WorldPartitionRuntimeHashSet.h
--- a/Engine/Source/Runtime/Engine/Public/WorldPartition/RuntimeHashSet/WorldPartitionRuntimeHashSet.h
+++ b/Engine/Source/Runtime/Engine/Public/WorldPartition/RuntimeHashSet/WorldPartitionRuntimeHashSet.h
@@ -176,7 +176,7 @@ public:
 	// Streaming interface
 	ENGINE_API virtual void ForEachStreamingCells(TFunctionRef<bool(const UWorldPartitionRuntimeCell*)> Func) const override;
 	ENGINE_API virtual void ForEachStreamingCellsQuery(const FWorldPartitionStreamingQuerySource& QuerySource, TFunctionRef<bool(const UWorldPartitionRuntimeCell*)> Func, FWorldPartitionQueryCache* QueryCache) const override;
-	ENGINE_API virtual void ForEachStreamingCellsSources(const TArray<FWorldPartitionStreamingSource>& Sources, TFunctionRef<bool(const UWorldPartitionRuntimeCell*, EStreamingSourceTargetState)> Func) const override;
+	ENGINE_API virtual void ForEachStreamingCellsSources(const TArray<FWorldPartitionStreamingSource>& Sources, TFunctionRef<bool(const UWorldPartitionRuntimeCell*, EStreamingSourceTargetState)> Func, const FWorldPartitionStreamingContext& Context = FWorldPartitionStreamingContext()) const override;
 
 protected:
 	ENGINE_API virtual bool SupportsWorldAssetStreaming(const FName& InTargetGrid) override;
diff --git a/Engine/Source/Runtime/Engine/Public/WorldPartition/WorldPartition.h b/Engine/Source/Runtime/Engine/Public/WorldPartition/WorldPartition.h
--- a/Engine/Source/Runtime/Engine/Public/WorldPartition/WorldPartition.h
+++ b/Engine/Source/Runtime/Engine/Public/WorldPartition/WorldPartition.h
@@ -490,10 +490,6 @@ public:
 
 	ENGINE_API bool IsMainWorldPartition() const;
 
-private:
-	ENGINE_API void Tick(float DeltaSeconds);
-
-public:
 	ENGINE_API bool CanAddCellToWorld(const IWorldPartitionCell* InCell) const;
 	ENGINE_API bool IsStreamingCompleted(const TArray<FWorldPartitionStreamingSource>* InStreamingSources) const;
 	ENGINE_API bool IsStreamingCompleted(EWorldPartitionRuntimeCellState QueryState, const TArray<FWorldPartitionStreamingQuerySource>& QuerySources, bool bExactState) const;
@@ -708,20 +704,22 @@ private:
 	static ENGINE_API FAutoConsoleVariableRef CVarEnableServerStreamingOut;
 	static ENGINE_API FAutoConsoleVariableRef CVarUseMakingVisibleTransactionRequests;
 	static ENGINE_API FAutoConsoleVariableRef CVarUseMakingInvisibleTransactionRequests;
-
+	
 	ENGINE_API void OnWorldMatchStarting();
 	ENGINE_API void OnWorldPreBeginPlay();
-	
-#if WITH_EDITOR
-	ENGINE_API void OnLevelActorDeleted(AActor* Actor);
-	ENGINE_API void OnPostBugItGoCalled(const FVector& Loc, const FRotator& Rot);
-#endif
+	ENGINE_API void OnStreamingStateUpdated();
+	ENGINE_API void Tick(float DeltaSeconds);
+	void OnPreChangeStreamingContent();
+	int32 GetUpdateStreamingStateEpoch() const;
 
 	// Delegates registration
 	ENGINE_API void RegisterDelegates();
 	ENGINE_API void UnregisterDelegates();	
 
 #if WITH_EDITOR
+	ENGINE_API void OnLevelActorDeleted(AActor* Actor);
+	ENGINE_API void OnPostBugItGoCalled(const FVector& Loc, const FRotator& Rot);
+
 	void HashActorDescInstance(FWorldPartitionActorDescInstance* ActorDescInstance);
 	void UnhashActorDescInstance(FWorldPartitionActorDescInstance* ActorDescInstance);
 	void OnContentBundleRemovedContent(const FContentBundleEditor* ContentBundle);
@@ -807,6 +805,7 @@ private:
 #endif
 	class AWorldPartitionReplay* Replay;
 
+	friend struct FWorldPartitionStreamingContext;
 	friend class AWorldPartitionReplay;
 	friend class UWorldPartitionSubsystem;
 	friend class UExternalDataLayerManager;
diff --git a/Engine/Source/Runtime/Engine/Public/WorldPartition/WorldPartitionReplay.h b/Engine/Source/Runtime/Engine/Public/WorldPartition/WorldPartitionReplay.h
--- a/Engine/Source/Runtime/Engine/Public/WorldPartition/WorldPartitionReplay.h
+++ b/Engine/Source/Runtime/Engine/Public/WorldPartition/WorldPartitionReplay.h
@@ -52,8 +52,8 @@ class AWorldPartitionReplay : public AActor
 public:
 	static ENGINE_API void Initialize(UWorld* World);
 	static ENGINE_API void Uninitialize(UWorld* World);
-	static ENGINE_API bool IsPlaybackEnabled(UWorld* World);
-	static ENGINE_API bool IsRecordingEnabled(UWorld* World);
+	static ENGINE_API bool IsPlaybackEnabled(const UWorld* World);
+	static ENGINE_API bool IsRecordingEnabled(const UWorld* World);
 
 	ENGINE_API virtual void RewindForReplay() override;
 	ENGINE_API virtual void PreReplication(IRepChangedPropertyTracker& ChangedPropertyTracker) override;
diff --git a/Engine/Source/Runtime/Engine/Public/WorldPartition/WorldPartitionRuntimeCell.h b/Engine/Source/Runtime/Engine/Public/WorldPartition/WorldPartitionRuntimeCell.h
--- a/Engine/Source/Runtime/Engine/Public/WorldPartition/WorldPartitionRuntimeCell.h
+++ b/Engine/Source/Runtime/Engine/Public/WorldPartition/WorldPartitionRuntimeCell.h
@@ -25,6 +25,7 @@ class UDataLayerManager;
 class UExternalDataLayerAsset;
 class FStreamingGenerationActorDescView;
 struct FHierarchicalLogArchive;
+struct FWorldPartitionStreamingContext;
 
 enum class EWorldPartitionDataLayersLogicOperator : uint8;
 
@@ -248,10 +249,6 @@ class UWorldPartitionRuntimeCell : public UObject, public IWorldPartitionCell, p
 	ENGINE_API virtual bool CanUnload() const PURE_VIRTUAL(UWorldPartitionRuntimeCell::CanUnload, return true;);
 	ENGINE_API virtual void Activate() const PURE_VIRTUAL(UWorldPartitionRuntimeCell::Activate,);
 	ENGINE_API virtual void Deactivate() const PURE_VIRTUAL(UWorldPartitionRuntimeCell::Deactivate,);
-	UE_DEPRECATED(5.4, "IsAddedToWorld is deprecated.")
-	ENGINE_API virtual bool IsAddedToWorld() const { return false; }
-	UE_DEPRECATED(5.4, "CanAddToWorld is deprecated.")
-	ENGINE_API virtual bool CanAddToWorld() const { return false; }
 	ENGINE_API virtual ULevel* GetLevel() const PURE_VIRTUAL(UWorldPartitionRuntimeCell::GetLevel, return nullptr;);
 	ENGINE_API virtual EWorldPartitionRuntimeCellState GetCurrentState() const PURE_VIRTUAL(UWorldPartitionRuntimeCell::GetCurrentState, return EWorldPartitionRuntimeCellState::Unloaded;);
 	virtual FLinearColor GetDebugColor(EWorldPartitionRuntimeCellVisualizeMode VisualizeMode) const { static const FLinearColor DefaultColor = FLinearColor::Black.CopyWithNewOpacity(0.25f); return DefaultColor; }
@@ -259,8 +256,6 @@ class UWorldPartitionRuntimeCell : public UObject, public IWorldPartitionCell, p
 	virtual void SetIsAlwaysLoaded(bool bInIsAlwaysLoaded) { bIsAlwaysLoaded = bInIsAlwaysLoaded; }
 	ENGINE_API virtual void SetStreamingPriority(int32 InStreamingPriority) const PURE_VIRTUAL(UWorldPartitionRuntimeCell::SetStreamingPriority,);
 	virtual EStreamingStatus GetStreamingStatus() const { return LEVEL_Unloaded; }
-	UE_DEPRECATED(5.3, "IsLoading is deprecated.")
-	virtual bool IsLoading() const { return false; }
 	void SetClientOnlyVisible(bool bInClientOnlyVisible) { bClientOnlyVisible = bInClientOnlyVisible; }
 	bool GetClientOnlyVisible() const { return bClientOnlyVisible; }
 	virtual FGuid const& GetContentBundleID() const { return ContentBundleID; }
@@ -285,17 +280,22 @@ class UWorldPartitionRuntimeCell : public UObject, public IWorldPartitionCell, p
 	//~End IWorldPartitionCell Interface
 
 	ENGINE_API UDataLayerManager* GetDataLayerManager() const;
-	ENGINE_API EDataLayerRuntimeState GetCellEffectiveWantedState() const;
-	FORCEINLINE EDataLayerRuntimeState GetCellEffectiveWantedState(int32 InDataLayersStateEpoch) const
-	{
-		if (EffectiveWantedStateEpoch != InDataLayersStateEpoch)
-		{
-			return GetCellEffectiveWantedState();
-		}
-		return EffectiveWantedState;
-	}
+	ENGINE_API EDataLayerRuntimeState GetCellEffectiveWantedState(const FWorldPartitionStreamingContext& Context) const;
 	FORCEINLINE bool HasDataLayers() const { return !DataLayers.IsEmpty(); }
 
+	//~Begin Deprecation
+	UE_DEPRECATED(5.3, "IsLoading is deprecated.")
+	virtual bool IsLoading() const { return false; }
+	UE_DEPRECATED(5.4, "IsAddedToWorld is deprecated.")
+	ENGINE_API virtual bool IsAddedToWorld() const { return false; }
+	UE_DEPRECATED(5.4, "CanAddToWorld is deprecated.")
+	ENGINE_API virtual bool CanAddToWorld() const { return false; }
+	UE_DEPRECATED(5.5, "Use version that takes FWorldPartitionStreamingContext instead.")
+	EDataLayerRuntimeState GetCellEffectiveWantedState(int32 InDataLayersStateEpoch) const { return EDataLayerRuntimeState::Unloaded; }
+	UE_DEPRECATED(5.5, "Use version that takes FWorldPartitionStreamingContext instead.")
+	EDataLayerRuntimeState GetCellEffectiveWantedState() const { return EDataLayerRuntimeState::Unloaded; }
+	//~End Deprecation
+
 	void SetBlockOnSlowLoading(bool bInBlockOnSlowLoading) { bBlockOnSlowLoading = bInBlockOnSlowLoading; }
 	bool GetBlockOnSlowLoading() const { return bBlockOnSlowLoading; }
 #if WITH_EDITOR
@@ -373,6 +373,8 @@ protected:
 	//          The same should be done for injected external streaming objects.
 	ENGINE_API virtual bool ShouldServerWaitForClientLevelVisibility() const { return true; }
 
+	ENGINE_API EDataLayerRuntimeState GetCellEffectiveWantedStateRaw() const { return !HasDataLayers() ? EDataLayerRuntimeState::Activated : EffectiveWantedState; }
+
 	UPROPERTY()
 	bool bIsAlwaysLoaded;
 
@@ -422,7 +424,7 @@ protected:
 
 public:
 	//~Begin UWorldPartitionRuntimeCellData Proxy
-	inline void AppendStreamingSourceInfo(const FWorldPartitionStreamingSource& Source, const FSphericalSector& SourceShape) const { RuntimeCellData->AppendStreamingSourceInfo(Source, SourceShape); }
+	inline void AppendStreamingSourceInfo(const FWorldPartitionStreamingSource& Source, const FSphericalSector& SourceShape, const FWorldPartitionStreamingContext& Context) const { RuntimeCellData->AppendStreamingSourceInfo(Source, SourceShape, Context); }
 	inline void MergeStreamingSourceInfo() const { RuntimeCellData->MergeStreamingSourceInfo(); }
 	ENGINE_API int32 SortCompare(const UWorldPartitionRuntimeCell* Other) const;
 	
@@ -438,8 +440,14 @@ public:
 	ENGINE_API virtual bool IsDebugShown() const;
 	//~End UWorldPartitionRuntimeCellData Proxy
 
+	//~Begin Deprecation
+	UE_DEPRECATED(5.5, "Use version that takes FWorldPartitionStreamingContext instead.")
+	inline void AppendStreamingSourceInfo(const FWorldPartitionStreamingSource& Source, const FSphericalSector& SourceShape) {}
+	//~End Deprecation
+
 	UPROPERTY()
 	TObjectPtr<UWorldPartitionRuntimeCellData> RuntimeCellData;
 
 	friend class UWorldPartitionStreamingPolicy;
+	friend struct FSpatialHashStreamingGrid;
 };
diff --git a/Engine/Source/Runtime/Engine/Public/WorldPartition/WorldPartitionRuntimeCellData.h b/Engine/Source/Runtime/Engine/Public/WorldPartition/WorldPartitionRuntimeCellData.h
--- a/Engine/Source/Runtime/Engine/Public/WorldPartition/WorldPartitionRuntimeCellData.h
+++ b/Engine/Source/Runtime/Engine/Public/WorldPartition/WorldPartitionRuntimeCellData.h
@@ -8,6 +8,7 @@
 #include "WorldPartitionRuntimeCellData.generated.h"
 
 class UActorContainer;
+struct FWorldPartitionStreamingContext;
 
 /** Caches information on streaming source that will be used later on to sort cell. */
 UCLASS(MinimalAPI)
@@ -25,8 +26,8 @@ class UWorldPartitionRuntimeCellData : public UObject
 	ENGINE_API virtual void DumpStateLog(FHierarchicalLogArchive& Ar) const;
 #endif
 
-	ENGINE_API virtual void ResetStreamingSourceInfo() const;
-	ENGINE_API virtual void AppendStreamingSourceInfo(const FWorldPartitionStreamingSource& Source, const FSphericalSector& SourceShape) const;
+	ENGINE_API virtual void ResetStreamingSourceInfo(const FWorldPartitionStreamingContext& Context) const;
+	ENGINE_API virtual void AppendStreamingSourceInfo(const FWorldPartitionStreamingSource& Source, const FSphericalSector& SourceShape, const FWorldPartitionStreamingContext& Context) const;
 	ENGINE_API virtual void MergeStreamingSourceInfo() const {}
 	ENGINE_API virtual int32 SortCompare(const UWorldPartitionRuntimeCellData* InOther) const;
 
@@ -42,8 +43,12 @@ class UWorldPartitionRuntimeCellData : public UObject
 	virtual bool IsDebugShown() const { return true; }
 	ENGINE_API virtual FString GetDebugName() const;
 
-	static ENGINE_API int32 StreamingSourceCacheEpoch;
-	static inline void DirtyStreamingSourceCacheEpoch() { ++StreamingSourceCacheEpoch; }
+	//~Begin Deprecation
+	UE_DEPRECATED(5.5, "Use version that takes FWorldPartitionStreamingContext instead.")
+	ENGINE_API virtual void ResetStreamingSourceInfo() const {}
+	UE_DEPRECATED(5.5, "Use version that takes FWorldPartitionStreamingContext instead.")
+	ENGINE_API virtual void AppendStreamingSourceInfo(const FWorldPartitionStreamingSource& Source, const FSphericalSector& SourceShape) const {}
+	//~End Deprecation
 
 	// Minimum affecting source priority
 	mutable uint8 CachedMinSourcePriority;
diff --git a/Engine/Source/Runtime/Engine/Public/WorldPartition/WorldPartitionRuntimeCellDataSpatialHash.h b/Engine/Source/Runtime/Engine/Public/WorldPartition/WorldPartitionRuntimeCellDataSpatialHash.h
--- a/Engine/Source/Runtime/Engine/Public/WorldPartition/WorldPartitionRuntimeCellDataSpatialHash.h
+++ b/Engine/Source/Runtime/Engine/Public/WorldPartition/WorldPartitionRuntimeCellDataSpatialHash.h
@@ -14,8 +14,8 @@ class UWorldPartitionRuntimeCellDataSpatialHash : public UWorldPartitionRuntimeC
 	GENERATED_UCLASS_BODY()
 
 	//~Begin UWorldPartitionRuntimeCellData
-	ENGINE_API virtual void ResetStreamingSourceInfo() const override;
-	ENGINE_API virtual void AppendStreamingSourceInfo(const FWorldPartitionStreamingSource& Source, const FSphericalSector& SourceShape) const override;
+	ENGINE_API virtual void ResetStreamingSourceInfo(const FWorldPartitionStreamingContext& Context) const override;
+	ENGINE_API virtual void AppendStreamingSourceInfo(const FWorldPartitionStreamingSource& Source, const FSphericalSector& SourceShape, const FWorldPartitionStreamingContext& Context) const override;
 	ENGINE_API virtual void MergeStreamingSourceInfo() const override;
 	ENGINE_API virtual int32 SortCompare(const UWorldPartitionRuntimeCellData* InOther) const override;
 	ENGINE_API virtual FBox GetCellBounds() const override;
diff --git a/Engine/Source/Runtime/Engine/Public/WorldPartition/WorldPartitionRuntimeHash.h b/Engine/Source/Runtime/Engine/Public/WorldPartition/WorldPartitionRuntimeHash.h
--- a/Engine/Source/Runtime/Engine/Public/WorldPartition/WorldPartitionRuntimeHash.h
+++ b/Engine/Source/Runtime/Engine/Public/WorldPartition/WorldPartitionRuntimeHash.h
@@ -15,6 +15,7 @@
 #include "WorldPartition/WorldPartitionStreamingGenerationContext.h"
 #include "WorldPartition/WorldPartitionRuntimeContainerResolving.h"
 #include "WorldPartition/DataLayer/DataLayerInstanceProviderInterface.h"
+#include "WorldPartition/DataLayer/WorldDataLayers.h"
 #if WITH_EDITOR
 #include "CookPackageSplitter.h"
 #include "Misc/HierarchicalLogArchive.h"
@@ -22,6 +23,7 @@
 #include "WorldPartitionRuntimeHash.generated.h"
 
 struct FHierarchicalLogArchive;
+struct FDataLayerInstanceNames;
 class FWorldPartitionDraw2DContext;
 class UExternalDataLayerAsset;
 class UExternalDataLayerInstance;
@@ -148,6 +150,31 @@ private:
 	TMap<const UWorldPartitionRuntimeCell*, double> CellToSourceMinSqrDistances;
 };
 
+struct ENGINE_API FWorldPartitionStreamingContext
+{
+public:
+	static FWorldPartitionStreamingContext Create(const UWorld* InWorld);
+	FWorldPartitionStreamingContext();
+	bool IsValid() const { return bIsValid; }
+
+private:
+	FWorldPartitionStreamingContext(const UWorld* InWorld);
+	FWorldPartitionStreamingContext(EWorldPartitionDataLayersLogicOperator InDataLayersLogicOperator, const FWorldDataLayersEffectiveStates& InDataLayerEffectiveStates, int32 InUpdateStreamingStateEpoch);
+
+	int32 GetResolvingDataLayersRuntimeStateEpoch() const { check(IsValid()); return DataLayerEffectiveStates->GetUpdateEpoch(); }
+	int32 GetUpdateStreamingStateEpoch() const { check(IsValid()); return UpdateStreamingStateEpoch; }
+	EDataLayerRuntimeState ResolveDataLayerRuntimeState(const FDataLayerInstanceNames& InDataLayerNames) const;
+
+	bool bIsValid;
+	EWorldPartitionDataLayersLogicOperator DataLayersLogicOperator;
+	const FWorldDataLayersEffectiveStates* DataLayerEffectiveStates;
+	int32 UpdateStreamingStateEpoch;
+
+	friend class UWorldPartitionStreamingPolicy;
+	friend class UWorldPartitionRuntimeCell;
+	friend class UWorldPartitionRuntimeCellData;
+};
+
 UCLASS(Abstract, Config=Engine, AutoExpandCategories=(WorldPartition), Within=WorldPartition, MinimalAPI)
 class UWorldPartitionRuntimeHash : public UObject
 {
@@ -210,11 +237,16 @@ public:
 	class FStreamingSourceCells
 	{
 	public:
-		void AddCell(const UWorldPartitionRuntimeCell* Cell, const FWorldPartitionStreamingSource& Source, const FSphericalSector& SourceShape);
+		void AddCell(const UWorldPartitionRuntimeCell* Cell, const FWorldPartitionStreamingSource& Source, const FSphericalSector& SourceShape, const FWorldPartitionStreamingContext& Context);
 		void Reset() { Cells.Reset(); }
 		int32 Num() const { return Cells.Num(); }
 		TSet<const UWorldPartitionRuntimeCell*>& GetCells() { return Cells; }
 
+		//~Begin Deprecation
+		UE_DEPRECATED(5.5, "Use version that takes FWorldPartitionStreamingContext instead.")
+		void AddCell(const UWorldPartitionRuntimeCell* Cell, const FWorldPartitionStreamingSource& Source, const FSphericalSector& SourceShape) {}
+		//~End Deprecation
+
 	private:
 		TSet<const UWorldPartitionRuntimeCell*> Cells;
 	};
@@ -222,7 +254,7 @@ public:
 	// Streaming interface
 	virtual void ForEachStreamingCells(TFunctionRef<bool(const UWorldPartitionRuntimeCell*)> Func) const {}
 	virtual void ForEachStreamingCellsQuery(const FWorldPartitionStreamingQuerySource& QuerySource, TFunctionRef<bool(const UWorldPartitionRuntimeCell*)> Func, FWorldPartitionQueryCache* QueryCache = nullptr) const {}
-	virtual void ForEachStreamingCellsSources(const TArray<FWorldPartitionStreamingSource>& Sources, TFunctionRef<bool(const UWorldPartitionRuntimeCell*, EStreamingSourceTargetState)> Func) const {}
+	virtual void ForEachStreamingCellsSources(const TArray<FWorldPartitionStreamingSource>& Sources, TFunctionRef<bool(const UWorldPartitionRuntimeCell*, EStreamingSourceTargetState)> Func, const FWorldPartitionStreamingContext& Context = FWorldPartitionStreamingContext()) const {}
 	// Computes a hash value of all runtime hash specific dependencies that affects the update of the streaming
 	virtual uint32 ComputeUpdateStreamingHash() const { return 0; }
 
diff --git a/Engine/Source/Runtime/Engine/Public/WorldPartition/WorldPartitionRuntimeSpatialHash.h b/Engine/Source/Runtime/Engine/Public/WorldPartition/WorldPartitionRuntimeSpatialHash.h
--- a/Engine/Source/Runtime/Engine/Public/WorldPartition/WorldPartitionRuntimeSpatialHash.h
+++ b/Engine/Source/Runtime/Engine/Public/WorldPartition/WorldPartitionRuntimeSpatialHash.h
@@ -179,14 +179,21 @@ struct FSpatialHashStreamingGrid
 	// Used by PIE/Game
 	ENGINE_API int64 GetCellSize(int32 Level) const;
 	ENGINE_API void GetCells(const FWorldPartitionStreamingQuerySource& QuerySource, TSet<const UWorldPartitionRuntimeCell*>& OutCells, bool bEnableZCulling, FWorldPartitionQueryCache* QueryCache = nullptr) const;
-	ENGINE_API void GetCells(const TArray<FWorldPartitionStreamingSource>& Sources, UWorldPartitionRuntimeHash::FStreamingSourceCells& OutActivateCells, UWorldPartitionRuntimeHash::FStreamingSourceCells& OutLoadCells, bool bEnableZCulling) const;
-	ENGINE_API void GetNonSpatiallyLoadedCells(TSet<const UWorldPartitionRuntimeCell*>& OutActivateCells, TSet<const UWorldPartitionRuntimeCell*>& OutLoadCells) const;
+	ENGINE_API void GetCells(const TArray<FWorldPartitionStreamingSource>& Sources, UWorldPartitionRuntimeHash::FStreamingSourceCells& OutActivateCells, UWorldPartitionRuntimeHash::FStreamingSourceCells& OutLoadCells, bool bEnableZCulling, const FWorldPartitionStreamingContext& Context) const;
+	ENGINE_API void GetNonSpatiallyLoadedCells(TSet<const UWorldPartitionRuntimeCell*>& OutActivateCells, TSet<const UWorldPartitionRuntimeCell*>& OutLoadCells, const FWorldPartitionStreamingContext& Context) const;
 	ENGINE_API void Draw2D(const class UWorldPartitionRuntimeSpatialHash* Owner, const FBox2D& Region2D, const FBox2D& GridScreenBounds, TFunctionRef<FVector2D(const FVector2D&, bool)> WorldToScreen, FWorldPartitionDraw2DContext& DrawContext) const;
 	ENGINE_API void Draw3D(const class UWorldPartitionRuntimeSpatialHash* Owner, const TArray<FWorldPartitionStreamingSource>& Sources, const FTransform& Transform) const;
 	ENGINE_API void ForEachRuntimeCell(TFunctionRef<bool(const UWorldPartitionRuntimeCell*)> Func) const;
 	ENGINE_API const FSquare2DGridHelper& GetGridHelper() const;
 	ENGINE_API float GetLoadingRange() const;
 
+	//~Begin Deprecation
+	UE_DEPRECATED(5.5, "Use version that takes FWorldPartitionStreamingContext instead.")
+	ENGINE_API void GetCells(const TArray<FWorldPartitionStreamingSource>& Sources, UWorldPartitionRuntimeHash::FStreamingSourceCells& OutActivateCells, UWorldPartitionRuntimeHash::FStreamingSourceCells& OutLoadCells, bool bEnableZCulling) const {}
+	UE_DEPRECATED(5.5, "Use version that takes FWorldPartitionStreamingContext instead.")
+	ENGINE_API void GetNonSpatiallyLoadedCells(TSet<const UWorldPartitionRuntimeCell*>& OutActivateCells, TSet<const UWorldPartitionRuntimeCell*>& OutLoadCells) const {}
+	//~End Deprecation
+
 #if WITH_EDITOR
 	void DumpStateLog(FHierarchicalLogArchive& Ar) const;
 #endif
@@ -340,7 +347,7 @@ public:
 	// streaming interface
 	ENGINE_API virtual void ForEachStreamingCells(TFunctionRef<bool(const UWorldPartitionRuntimeCell*)> Func) const override;
 	ENGINE_API virtual void ForEachStreamingCellsQuery(const FWorldPartitionStreamingQuerySource& QuerySource, TFunctionRef<bool(const UWorldPartitionRuntimeCell*)> Func, FWorldPartitionQueryCache* QueryCache = nullptr) const override;
-	ENGINE_API virtual void ForEachStreamingCellsSources(const TArray<FWorldPartitionStreamingSource>& Sources, TFunctionRef<bool(const UWorldPartitionRuntimeCell*, EStreamingSourceTargetState)> Func) const override;
+	ENGINE_API virtual void ForEachStreamingCellsSources(const TArray<FWorldPartitionStreamingSource>& Sources, TFunctionRef<bool(const UWorldPartitionRuntimeCell*, EStreamingSourceTargetState)> Func, const FWorldPartitionStreamingContext& Context = FWorldPartitionStreamingContext()) const override;
 	ENGINE_API virtual uint32 ComputeUpdateStreamingHash() const override;
 
 	ENGINE_API virtual bool InjectExternalStreamingObject(URuntimeHashExternalStreamingObjectBase* ExternalStreamingObject) override;
diff --git a/Engine/Source/Runtime/Engine/Public/WorldPartition/WorldPartitionStreamingPolicy.h b/Engine/Source/Runtime/Engine/Public/WorldPartition/WorldPartitionStreamingPolicy.h
--- a/Engine/Source/Runtime/Engine/Public/WorldPartition/WorldPartitionStreamingPolicy.h
+++ b/Engine/Source/Runtime/Engine/Public/WorldPartition/WorldPartitionStreamingPolicy.h
@@ -9,12 +9,14 @@
 
 #pragma once
 
+#include "Async/TaskGraphFwd.h"
 #include "Containers/Set.h"
 #include "Misc/CoreDelegates.h"
 #include "WorldPartition/WorldPartitionRuntimeCell.h"
 #include "WorldPartition/WorldPartitionRuntimeHash.h"
 #include "WorldPartition/WorldPartitionStreamingSource.h"
 #include "WorldPartition/WorldPartitionRuntimeContainerResolving.h"
+#include "WorldPartition/DataLayer/WorldDataLayers.h"
 #include "WorldPartitionStreamingPolicy.generated.h"
 
 class UWorldPartition;
@@ -30,6 +32,7 @@ struct FActivatedCells
 	bool Contains(const UWorldPartitionRuntimeCell* InCell) const { return Cells.Contains(InCell); }
 	void OnAddedToWorld(const UWorldPartitionRuntimeCell* InCell);
 	void OnRemovedFromWorld(const UWorldPartitionRuntimeCell* InCell);
+	void Reset();
 
 	const TSet<TObjectPtr<const UWorldPartitionRuntimeCell>>& GetCells() const { return Cells; }
 	const TSet<const UWorldPartitionRuntimeCell*>& GetPendingAddToWorldCells() const { return PendingAddToWorldCells; }
@@ -42,6 +45,49 @@ private:
 	TSet<const UWorldPartitionRuntimeCell*> PendingAddToWorldCells;
 };
 
+USTRUCT()
+struct FWorldPartitionUpdateStreamingTargetState
+{
+	GENERATED_USTRUCT_BODY()
+
+	UPROPERTY(Transient)
+	TArray<TObjectPtr<const UWorldPartitionRuntimeCell>> ToLoadCells;
+
+	UPROPERTY(Transient)
+	TArray<TObjectPtr<const UWorldPartitionRuntimeCell>> ToActivateCells;
+
+	UPROPERTY(Transient)
+	TArray<TObjectPtr<const UWorldPartitionRuntimeCell>> ToDeactivateCells;
+
+	UPROPERTY(Transient)
+	TArray<TObjectPtr<const UWorldPartitionRuntimeCell>> ToUnloadCells;
+
+	EWorldPartitionStreamingPerformance StreamingPerformance = EWorldPartitionStreamingPerformance::Good;
+
+	bool bUpdateServerEpoch = false;
+
+	bool IsEmpty() const;
+	void Reset();
+};
+
+USTRUCT()
+struct FWorldPartitionUpdateStreamingCurrentState
+{
+	GENERATED_USTRUCT_BODY()
+
+	// Streaming Sources
+	TArray<FWorldPartitionStreamingSource> StreamingSources;
+
+	UPROPERTY(Transient)
+	TSet<TObjectPtr<const UWorldPartitionRuntimeCell>> LoadedCells;
+
+	UPROPERTY(Transient)
+	FActivatedCells ActivatedCells;
+
+	void Reset();
+	void CopyFrom(const FWorldPartitionUpdateStreamingCurrentState& InCurrentState);
+};
+
 UCLASS(Abstract, Within = WorldPartition)
 class UWorldPartitionStreamingPolicy : public UObject
 {
@@ -86,9 +132,9 @@ public:
 	virtual bool ConvertEditorPathToRuntimePath(const FSoftObjectPath& InPath, FSoftObjectPath& OutPath) const { return false; }
 	virtual UObject* GetSubObject(const TCHAR* SubObjectPath) { return nullptr; }
 
-	const TArray<FWorldPartitionStreamingSource>& GetStreamingSources() const { return StreamingSources; }
+	const TArray<FWorldPartitionStreamingSource>& GetStreamingSources() const { return CurrentState.StreamingSources; }
 
-	EWorldPartitionStreamingPerformance GetStreamingPerformance() const { return StreamingPerformance; }
+	EWorldPartitionStreamingPerformance GetStreamingPerformance() const { return CurrentStreamingPerformance; }
 
 	static bool IsUpdateStreamingOptimEnabled();
 
@@ -100,71 +146,122 @@ public:
 protected:
 	virtual void SetCellStateToLoaded(const UWorldPartitionRuntimeCell* InCell, int32& InOutMaxCellsToLoad);
 	virtual void SetCellStateToActivated(const UWorldPartitionRuntimeCell* InCell, int32& InOutMaxCellsToLoad);
-	virtual void SetCellsStateToUnloaded(const TArray<const UWorldPartitionRuntimeCell*>& ToUnloadCells);
+	virtual void SetCellsStateToUnloaded(const TArray<TObjectPtr<const UWorldPartitionRuntimeCell>>& ToUnloadCells);
 	virtual void GetCellsToUpdate(TArray<const UWorldPartitionRuntimeCell*>& OutToLoadCells, TArray<const UWorldPartitionRuntimeCell*>& OutToActivateCells);
 	virtual void GetCellsToReprioritize(TArray<const UWorldPartitionRuntimeCell*>& OutToLoadCells, TArray<const UWorldPartitionRuntimeCell*>& OutToActivateCells);
 	virtual void UpdateStreamingSources(bool bCanOptimizeUpdate);
-	void UpdateStreamingPerformance(const TSet<const UWorldPartitionRuntimeCell*>& InCells);
-	bool ShouldSkipCellForPerformance(const UWorldPartitionRuntimeCell* Cell) const;
+	void UpdateStreamingPerformance(EWorldPartitionStreamingPerformance NewStreamingPerformance);
 	bool IsInBlockTillLevelStreamingCompleted(bool bIsCausedByBadStreamingPerformance = false) const;
 
+	// Outer World Partition
 	const UWorldPartition* WorldPartition;
-	UPROPERTY(Transient)
-	TSet<TObjectPtr<const UWorldPartitionRuntimeCell>> LoadedCells;
 
-	UPROPERTY(Transient)
-	FActivatedCells ActivatedCells;
+private:
+	enum class EAsyncUpdateTaskState
+	{
+		None,
+		Pending,
+		Started
+	};
+
+	struct FUpdateStreamingStateParams
+	{
+		FUpdateStreamingStateParams(UWorldPartitionStreamingPolicy* InPolicy, const FWorldPartitionUpdateStreamingCurrentState& InCurrentState);
+		FUpdateStreamingStateParams& SetRequiredWorldDataLayersEffectiveStatesCopy(bool bInRequiredEffectiveStatesCopy);
+
+		const UWorld* World;
+		const UWorldPartitionRuntimeHash* RuntimeHash;
+		const bool bCanStream;
+		const bool bIsServer;
+		const bool bIsStreamingInEnabled;
+		const bool bIsServerStreamingEnabled;
+		const bool bIsServerStreamingOutEnabled;
+		const bool bIsBlockingCausedByBadStreamingPerformance;
+		const bool bShouldMergeStreamingSourceInfo;
+		const int32 PolicyUpdateStreamingStateEpoch;
+		const EWorldPartitionDataLayersLogicOperator DataLayersLogicOperator;
+		const FTransform WorldPartitionInstanceTransform;
+		const FWorldPartitionUpdateStreamingCurrentState& CurrentState;
+		const TSet<FName>& ServerDisallowedStreamingOutDataLayers;
+		TSet<const UWorldPartitionRuntimeCell*>& FrameActivateCells;
+		TSet<const UWorldPartitionRuntimeCell*>& FrameLoadCells;
+
+		const FWorldDataLayersEffectiveStates& GetWorldDataLayersEffectiveStates() const { return WorldDataLayersEffectiveStatesCopy.Get(WorldDataLayersEffectiveStatesRef); }
+	private:
+		const FWorldDataLayersEffectiveStates& WorldDataLayersEffectiveStatesRef;
+		TOptional<FWorldDataLayersEffectiveStates> WorldDataLayersEffectiveStatesCopy;
+	};
 
-	// Streaming Sources
-	TArray<FWorldPartitionStreamingSource> StreamingSources;
+	// Update optimization
+	uint32 ComputeUpdateStreamingHash(bool bCanOptimizeUpdate) const;
+	int32 ComputeServerStreamingEnabledEpoch() const;
 
-	TSet<const UWorldPartitionRuntimeCell*> FrameActivateCells;
-	TSet<const UWorldPartitionRuntimeCell*> FrameLoadCells;
+	void OnStreamingStateUpdated();
+	void OnPreChangeStreamingContent();
+	bool WaitForAsyncUpdateStreamingState();
+	void PostUpdateStreamingStateInternal_GameThread(FWorldPartitionUpdateStreamingTargetState& InOutTargetState);
+	const TSet<FName>& GetServerDisallowedStreamingOutDataLayers() const;
+	static void UpdateStreamingStateInternal(const FUpdateStreamingStateParams& InParams, FWorldPartitionUpdateStreamingTargetState& OutTargetState);
 
-	// Used by UWorldPartitionSubsystem
+	// Current streaming state
 	UPROPERTY(Transient)
-	TArray<TObjectPtr<const UWorldPartitionRuntimeCell>> ToActivateCells;
-	
+	FWorldPartitionUpdateStreamingCurrentState CurrentState;
+
+	// Current streaming performance
 	UPROPERTY(Transient)
-	TArray<TObjectPtr<const UWorldPartitionRuntimeCell>> ToLoadCells;
-	
-	int32 ProcessedToActivateCells;
-	int32 ProcessedToLoadCells;
+	EWorldPartitionStreamingPerformance CurrentStreamingPerformance;
 
-private:
-	// Update optimization
-	uint32 ComputeUpdateStreamingHash(bool bCanOptimizeUpdate) const;
-	int32 ComputeServerStreamingEnabledEpoch() const;
+	// Target state
+	UPROPERTY(Transient)
+	FWorldPartitionUpdateStreamingTargetState TargetState;
 
-	const TSet<FName>& GetServerDisallowedStreamingOutDataLayers();
+	// Asynchronous update task input payload
+	UPROPERTY(Transient)
+	FWorldPartitionUpdateStreamingCurrentState AsyncTaskCurrentState;
 
-	// CVars to control update optimization
-	static bool IsUpdateOptimEnabled;
-	static int32 ForceUpdateFrameCount;
-	static FAutoConsoleVariableRef CVarUpdateOptimEnabled;
-	static FAutoConsoleVariableRef CVarForceUpdateFrameCount;
+	// Asynchronous update task output payload
+	UPROPERTY(Transient)
+	FWorldPartitionUpdateStreamingTargetState AsyncTaskTargetState;
 
-	bool bCriticalPerformanceRequestedBlockTillOnWorld;
-	
 	UPROPERTY()
 	bool bShouldMergeStreamingSourceInfo;
 
+	bool bCriticalPerformanceRequestedBlockTillOnWorld;
 	int32 CriticalPerformanceBlockTillLevelStreamingCompletedEpoch;
+	int32 ProcessedToLoadCells;		// Used to know if last update fully processed ToLoadCells
+	int32 ProcessedToActivateCells; // Used to know if last update fully processed ToActivateCells
 	int32 ServerStreamingStateEpoch;
 	int32 ServerStreamingEnabledEpoch;
 	uint32 UpdateStreamingHash;
 	uint32 UpdateStreamingSourcesHash;
-	uint32 UpdateStreamingStateCalls;
+	uint32 UpdateStreamingStateCounter;
+
+	// Asynchronous update tasks
+	EAsyncUpdateTaskState AsyncUpdateTaskState;
+	FGraphEventRef AsyncUpdateStreamingStateTask;
+	FGraphEventRef AsyncPostUpdateStreamingStateTask;
+	int32 AsyncShouldSkipUpdateCounter;
 
-	TOptional<TSet<FName>> CachedServerDisallowStreamingOutDataLayers;
+	mutable TOptional<TSet<FName>> CachedServerDisallowStreamingOutDataLayers;
+	// Used internally by UpdateStreamingStateInternal (avoids re-allocations)
+	mutable TSet<const UWorldPartitionRuntimeCell*> FrameActivateCells;
+	mutable TSet<const UWorldPartitionRuntimeCell*> FrameLoadCells;
+
+	// CVars to control update optimization
+	static bool IsUpdateOptimEnabled;
+	static bool IsAsyncUpdateStreamingStateEnabled;
+	static int32 ForceUpdateFrameCount;
+	static FAutoConsoleVariableRef CVarUpdateOptimEnabled;
+	static FAutoConsoleVariableRef CVarAsyncUpdateStreamingStateEnabled;
+	static FAutoConsoleVariableRef CVarForceUpdateFrameCount;
 
-	EWorldPartitionStreamingPerformance StreamingPerformance;
 #if !UE_BUILD_SHIPPING
-	void UpdateDebugCellsStreamingPriority(const TSet<const UWorldPartitionRuntimeCell*>& ActivateStreamingCells, const TSet<const UWorldPartitionRuntimeCell*>& LoadStreamingCells);
+	static void UpdateDebugCellsStreamingPriority(const TSet<const UWorldPartitionRuntimeCell*>& InActivateStreamingCells, const TSet<const UWorldPartitionRuntimeCell*>& InLoadStreamingCells, bool bInShouldMergeStreamingSourceInfo);
 
 	double OnScreenMessageStartTime;
 	EWorldPartitionStreamingPerformance  OnScreenMessageStreamingPerformance;
 #endif
 
+	friend class UWorldPartition;
 	friend class UWorldPartitionSubsystem;
 };
\ No newline at end of file
diff --git a/Engine/Source/Runtime/Engine/Public/WorldPartition/WorldPartitionSubsystem.h b/Engine/Source/Runtime/Engine/Public/WorldPartition/WorldPartitionSubsystem.h
--- a/Engine/Source/Runtime/Engine/Public/WorldPartition/WorldPartitionSubsystem.h
+++ b/Engine/Source/Runtime/Engine/Public/WorldPartition/WorldPartitionSubsystem.h
@@ -96,6 +96,9 @@ public:
 	DECLARE_DELEGATE_RetVal_OneParam(bool, FWorldPartitionStreamingSourceProviderFilter, const IWorldPartitionStreamingSourceProvider*);
 	FWorldPartitionStreamingSourceProviderFilter& OnIsStreamingSourceProviderFiltered() { return IsStreamingSourceProviderFiltered; }
 
+	DECLARE_MULTICAST_DELEGATE(FWorldPartitionStreamingStateUpdated);
+	FWorldPartitionStreamingStateUpdated& OnStreamingStateUpdated() { return StreamingStateUpdated; }
+
 	static ENGINE_API TMulticastDelegate<void(UWorldPartitionSubsystem*, UWorld*)> OnWorldPartitionSubsystemInitialized;
 	static ENGINE_API TMulticastDelegate<void(UWorldPartitionSubsystem*, UWorld*)> OnWorldPartitionSubsystemDeinitialized;
 
@@ -179,6 +182,7 @@ private:
 	// Streaming Sources
 	TSet<IWorldPartitionStreamingSourceProvider*> StreamingSourceProviders;
 	FWorldPartitionStreamingSourceProviderFilter IsStreamingSourceProviderFiltered;
+	FWorldPartitionStreamingStateUpdated StreamingStateUpdated;
 	TArray<FWorldPartitionStreamingSource> StreamingSources;
 	TMap<FName, FStreamingSourceVelocity> StreamingSourcesVelocity;
 	uint32 StreamingSourcesHash;
