[NetPhysics] Async flow for NetworkPhysicsComponent (NPC) along with various fixes.

- Physics sleep state handled by FPhysicsReplicationAsync::ResimulationReplication to fix bug where client could sleep while server was awake, which would result in stuttering when a resimulation triggers.
- Don't set RepMovement quantization in NPC, this should be set manually on the actor blueprint.
- Don't require a StaticMeshActor to simulate to set physics replication mode.

#virtualized

diff --git a/Engine/Plugins/Experimental/ChaosModularVehicle/Source/ChaosModularVehicleEngine/Private/ChaosModularVehicle/ModularVehicleBaseComponent.cpp b/Engine/Plugins/Experimental/ChaosModularVehicle/Source/ChaosModularVehicleEngine/Private/ChaosModularVehicle/ModularVehicleBaseComponent.cpp
--- a/Engine/Plugins/Experimental/ChaosModularVehicle/Source/ChaosModularVehicleEngine/Private/ChaosModularVehicle/ModularVehicleBaseComponent.cpp
+++ b/Engine/Plugins/Experimental/ChaosModularVehicle/Source/ChaosModularVehicleEngine/Private/ChaosModularVehicle/ModularVehicleBaseComponent.cpp
@@ -774,7 +774,7 @@ void UModularVehicleBaseComponent::ActionTreeUpdates(Chaos::FSimTreeUpdates* Nex
 				// Network replication data needs to be updated, this is currently studily slow
 				if (NetworkPhysicsComponent)
 				{
-					TSharedPtr<Chaos::FBaseRewindHistory>& History = NetworkPhysicsComponent->GetStateHistory();
+					TSharedPtr<Chaos::FBaseRewindHistory>& History = NetworkPhysicsComponent->GetStateHistory_Internal();
 					Chaos::TDataRewindHistory<FNetworkModularVehicleStates>* StateHistory = static_cast<Chaos::TDataRewindHistory<FNetworkModularVehicleStates>*>(History.Get());
 					if (StateHistory)
 					{
diff --git a/Engine/Source/Runtime/Engine/Private/PhysicsEngine/Experimental/PhysScene_Chaos.cpp b/Engine/Source/Runtime/Engine/Private/PhysicsEngine/Experimental/PhysScene_Chaos.cpp
--- a/Engine/Source/Runtime/Engine/Private/PhysicsEngine/Experimental/PhysScene_Chaos.cpp
+++ b/Engine/Source/Runtime/Engine/Private/PhysicsEngine/Experimental/PhysScene_Chaos.cpp
@@ -139,9 +139,12 @@ public:
 	{
 		using namespace Chaos;
 
+		/* #TODO implement and re-enable resim commands. This callback must run on the main thread and resim currently does
+		 * not defer its callbacks to the main thread making its execution unsafe.
 		const UPhysicsSettings* PhysicsSettings = UPhysicsSettings::Get();
 		const bool bAllowResim = PhysicsSettings->PhysicsPrediction.bEnablePhysicsPrediction;
 		const int32 NumFrames = PhysicsSettings->GetPhysicsHistoryCount();
+		*/
 
 		TArray<int32> CommandIndicesToRemove;
 		CommandIndicesToRemove.Reserve(PendingCommands.Num());
@@ -1358,6 +1361,9 @@ const FRigidBodyState* FPhysScene_Chaos::GetStateFromReplicationCache(UPrimitive
 
 void FPhysScene_Chaos::RegisterForReplicationCache(UPrimitiveComponent* RootComponent)
 {
+	// ToDo, remove call to EnableAsyncPhysicsTickCallback when ReplicationCache is refactored to run with non-GT Frozen async callback
+	EnableAsyncPhysicsTickCallback();
+
 	const FObjectKey Key(RootComponent);
 	ReplicationCache.Map.Add(Key, FReplicationCacheData(RootComponent, GetSolver()->GetSolverTime()));
 }
diff --git a/Engine/Source/Runtime/Engine/Private/PhysicsEngine/NetworkPhysicsComponent.cpp b/Engine/Source/Runtime/Engine/Private/PhysicsEngine/NetworkPhysicsComponent.cpp
--- a/Engine/Source/Runtime/Engine/Private/PhysicsEngine/NetworkPhysicsComponent.cpp
+++ b/Engine/Source/Runtime/Engine/Private/PhysicsEngine/NetworkPhysicsComponent.cpp
@@ -1,7 +1,6 @@
 // Copyright Epic Games, Inc. All Rights Reserved.
 
 #include "Physics/NetworkPhysicsComponent.h"
-#include "Physics/NetworkPhysicsSettingsComponent.h"
 #include "Components/PrimitiveComponent.h"
 #include "EngineLogs.h"
 #include "EngineUtils.h"
@@ -23,6 +22,10 @@ namespace PhysicsReplicationCVars
 {
 	namespace ResimulationCVars
 	{
+		int32 InputRedundancy = 3;
+		static FAutoConsoleVariableRef CVarResimInputRedundancy(TEXT("np2.Resim.InputRedundancy"), InputRedundancy, TEXT("How many inputs to send with each unreliable network message to account for packetloss."));
+		int32 StateRedundancy = 1;
+		static FAutoConsoleVariableRef CVarResimStateRedundancy(TEXT("np2.Resim.StateRedundancy"), StateRedundancy, TEXT("How many states to send with each unreliable network message to account for packetloss."));
 		bool bAllowRewindToClosestState = true;
 		static FAutoConsoleVariableRef CVarResimAllowRewindToClosestState(TEXT("np2.Resim.AllowRewindToClosestState"), bAllowRewindToClosestState, TEXT("When rewinding to a specific frame, if the client doens't have state data for that frame, use closest data available. Only affects the first rewind frame, when FPBDRigidsEvolution is set to Reset."));
 		bool bCompareStateToTriggerRewind = false;
@@ -37,6 +40,8 @@ namespace PhysicsReplicationCVars
 		static FAutoConsoleVariableRef CVarResimApplyDataInsteadOfMergeData(TEXT("np2.Resim.ApplyDataInsteadOfMergeData"), bApplyDataInsteadOfMergeData, TEXT("When true, call ApplyData for each data instead of MergeData when having to use multiple data entries in one frame."));
 		bool bAllowInputExtrapolation = true;
 		static FAutoConsoleVariableRef CVarResimAllowInputExtrapolation(TEXT("np2.Resim.AllowInputExtrapolation"), bAllowInputExtrapolation, TEXT("When true and not locally controlled, allow inputs to be extrapolated from last known and if there is a gap allow interpolation between two known inputs."));
+		bool bValidateDataOnGameThread = false;
+		static FAutoConsoleVariableRef CVarResimValidateDataOnGameThread(TEXT("np2.Resim.ValidateDataOnGameThread"), bValidateDataOnGameThread, TEXT("When true, perform server-side input validation through FNetworkPhysicsData::ValidateData on the Game Thread. If false, perform the call on the Physics Thread."));
 	}
 }
 
@@ -168,7 +173,11 @@ bool FNetworkPhysicsRewindDataImportantStateProxy::NetSerialize(FArchive& Ar, cl
 	return NetSerializeBase(Ar, Map, bOutSuccess, [this]() { return Owner->ReplicatedImportantState.History->CreateNew(); });
 }
 
-// after presimulate internal (asyncinput internal simulation done and the output created)
+
+// --------------------------- Network Physics Callback ---------------------------
+
+/** After physics has solved and output is created
+* NOTE: GameThread is frozen during this callback */
 void FNetworkPhysicsCallback::ApplyCallbacks_Internal(int32 PhysicsStep, const TArray<Chaos::ISimCallbackObject*>& SimCallbackObjects)
 {
 	QUICK_SCOPE_CYCLE_COUNTER(NetworkPhysicsComponent_ApplyCallbacks_Internal);
@@ -183,7 +192,7 @@ void FNetworkPhysicsCallback::ApplyCallbacks_Internal(int32 PhysicsStep, const T
 	}
 }
 
-// before presimulate internal (asyncinput internal simulation is not done yet)
+// Before PreSimulate_Internal
 void FNetworkPhysicsCallback::ProcessInputs_Internal(int32 PhysicsStep, const TArray<Chaos::FSimCallbackInputAndObject>& SimCallbacks)
 {
 	PreProcessInputsInternal.Broadcast(PhysicsStep);
@@ -224,10 +233,6 @@ int32 FNetworkPhysicsCallback::TriggerRewindIfNeeded_Internal(int32 LatestStepCo
 		if (NetMode == NM_Client)
 		{
 			const int32 ReplicationFrame = RewindData->GetResimFrame();
-
-#if DEBUG_NETWORK_PHYSICS || DEBUG_REWIND_DATA
-			UE_LOG(LogChaos, Log, TEXT("CLIENT | PT | TriggerRewindIfNeeded_Internal | Replication Frame = %d"), ReplicationFrame);
-#endif
 			ResimFrame = (ResimFrame == INDEX_NONE) ? ReplicationFrame : (ReplicationFrame == INDEX_NONE) ? ResimFrame : FMath::Min(ReplicationFrame, ResimFrame);
 		}
 
@@ -235,7 +240,7 @@ int32 FNetworkPhysicsCallback::TriggerRewindIfNeeded_Internal(int32 LatestStepCo
 		{
 			const int32 ValidFrame = RewindData->FindValidResimFrame(ResimFrame);
 #if DEBUG_NETWORK_PHYSICS || DEBUG_REWIND_DATA
-			UE_LOG(LogChaos, Log, TEXT("CLIENT | PT | TriggerRewindIfNeeded_Internal | Resim Frame = %d | Valid Frame = %d"), ResimFrame, ValidFrame);
+			UE_LOG(LogChaos, Log, TEXT("CLIENT | PT | TriggerRewindIfNeeded_Internal | Requested Resim Frame = %d | Valid Resim Frame = %d"), ResimFrame, ValidFrame);
 #endif
 			ResimFrame = ValidFrame;
 		}
@@ -366,6 +371,9 @@ void FNetworkPhysicsCallback::ProcessInputs_External(int32 PhysicsStep, const TA
 	}
 }
 
+
+// --------------------------- Network Physics System ---------------------------
+
 UNetworkPhysicsSystem::UNetworkPhysicsSystem()
 {}
 
@@ -415,6 +423,9 @@ void UNetworkPhysicsSystem::OnWorldPostInit(UWorld* World, const UWorld::Initial
 	}
 }
 
+
+// --------------------------- GameThread Network Physics Component ---------------------------
+
 UNetworkPhysicsComponent::UNetworkPhysicsComponent(const FObjectInitializer& ObjectInitializer)
 	: Super(ObjectInitializer)
 {
@@ -428,63 +439,38 @@ UNetworkPhysicsComponent::UNetworkPhysicsComponent() : Super()
 
 void UNetworkPhysicsComponent::InitPhysics()
 {
-	if (const IConsoleVariable* CVarCompareInputToTriggerRewind = IConsoleManager::Get().FindConsoleVariable(TEXT("np2.Resim.CompareInputToTriggerRewind")))
+	if (const IConsoleVariable* CVarInputRedundancy = IConsoleManager::Get().FindConsoleVariable(TEXT("np2.Resim.InputRedundancy")))
 	{
-		bCompareInputToTriggerRewind = CVarCompareInputToTriggerRewind->GetBool();
+		InputRedundancy = CVarInputRedundancy->GetInt();
 	}
-	if (const IConsoleVariable* CVarCompareStateToTriggerRewind = IConsoleManager::Get().FindConsoleVariable(TEXT("np2.Resim.CompareStateToTriggerRewind")))
+	if (const IConsoleVariable* CVarStateRedundancy = IConsoleManager::Get().FindConsoleVariable(TEXT("np2.Resim.StateRedundancy")))
 	{
-		bCompareStateToTriggerRewind = CVarCompareStateToTriggerRewind->GetBool();
+		StateRedundancy = CVarStateRedundancy->GetInt();
 	}
 
-	AActor* Owner = GetOwner();
-	if (Owner)
+	if (AActor* Owner = GetOwner())
 	{
-		FRepMovement& RepMovement = Owner->GetReplicatedMovement_Mutable();
-		RepMovement.LocationQuantizationLevel = EVectorQuantization::RoundTwoDecimals;
-		RepMovement.RotationQuantizationLevel = ERotatorQuantization::ShortComponents;
-		RepMovement.VelocityQuantizationLevel = EVectorQuantization::RoundTwoDecimals;
-
 		if (UPrimitiveComponent* RootPrimComp = Cast<UPrimitiveComponent>(Owner->GetRootComponent()))
 		{
 			RootPhysicsObject = RootPrimComp->GetPhysicsObjectByName(NAME_None);
 		}
 	}
 
+	/** NOTE:
+	* If the NetworkPhysicsComponent is added as a SubObject after the actor has processed bAutoActivate and bWantsInitializeComponent
+	* SetActive(true) and InitializeComponent() needs to be called manually for the component to function properly. */
+	PrimaryComponentTick.bCanEverTick = true;
 	bAutoActivate = true;
 	bWantsInitializeComponent = true;
 	SetIsReplicatedByDefault(true);
-	SetAsyncPhysicsTickEnabled(true);
 }
 
 void UNetworkPhysicsComponent::BeginPlay()
 {
 	Super::BeginPlay();
-	UWorld* World = GetWorld();
-
-	if (World->WorldType == EWorldType::PIE || World->WorldType == EWorldType::Game)
-	{
-		if (FPhysScene* PhysScene = World->GetPhysicsScene())
-		{
-			if (Chaos::FPhysicsSolver* Solver = PhysScene->GetSolver())
-			{
-				if (UPhysicsSettings::Get()->PhysicsPrediction.bEnablePhysicsPrediction)
-				{
-					SetupRewindData();
 
-					if(FNetworkPhysicsCallback* SolverCallback = static_cast<FNetworkPhysicsCallback*>(Solver->GetRewindCallback()))
-					{
-						SolverCallback->PreProcessInputsInternal.AddUObject(this, &UNetworkPhysicsComponent::OnPreProcessInputsInternal);
-						SolverCallback->PostProcessInputsInternal.AddUObject(this, &UNetworkPhysicsComponent::OnPostProcessInputsInternal);
-					}
-				}
-				else
-				{
-					UE_LOG(LogChaos, Warning, TEXT("A NetworkPhysicsComponent is trying to set up but 'Project Settings -> Physics -> Physics Prediction' is not enabled. The component might not work as intended."));
-				}
-			}
-		}
-	}
+	// Update async component with current component properties
+	UpdateAsyncComponent();
 }
 
 void UNetworkPhysicsComponent::InitializeComponent()
@@ -494,22 +480,20 @@ void UNetworkPhysicsComponent::InitializeComponent()
 	// Cache CVar values
 	bEnableUnreliableFlow = PhysicsReplicationCVars::ResimulationCVars::bEnableUnreliableFlow;
 	bEnableReliableFlow = PhysicsReplicationCVars::ResimulationCVars::bEnableReliableFlow;
-	bApplyDataInsteadOfMergeData = PhysicsReplicationCVars::ResimulationCVars::bApplyDataInsteadOfMergeData;
-	bAllowInputExtrapolation = PhysicsReplicationCVars::ResimulationCVars::bAllowInputExtrapolation;
+	bValidateDataOnGameThread = PhysicsReplicationCVars::ResimulationCVars::bValidateDataOnGameThread;
 
+	UNetworkPhysicsSettingsComponent* SettingsComponent = nullptr;
 	if (AActor* Owner = GetOwner())
 	{
 		// Get settings from NetworkPhysicsSettingsComponent, if there is one
-		if (UNetworkPhysicsSettingsComponent* PhysicsSettings = Owner->FindComponentByClass<UNetworkPhysicsSettingsComponent>())
+		SettingsComponent = Owner->FindComponentByClass<UNetworkPhysicsSettingsComponent>();
+		if (SettingsComponent)
 		{
-			InputRedundancy = PhysicsSettings->NetworkPhysicsComponentSettings.GetRedundantInputs(InputRedundancy);
-			StateRedundancy = PhysicsSettings->NetworkPhysicsComponentSettings.GetRedundantStates(StateRedundancy);
-			bCompareInputToTriggerRewind = PhysicsSettings->NetworkPhysicsComponentSettings.GetCompareInputToTriggerRewind(bCompareInputToTriggerRewind);
-			bCompareStateToTriggerRewind = PhysicsSettings->NetworkPhysicsComponentSettings.GetCompareStateToTriggerRewind(bCompareStateToTriggerRewind);
-			bEnableUnreliableFlow = PhysicsSettings->NetworkPhysicsComponentSettings.GetEnableUnreliableFlow();
-			bEnableReliableFlow = PhysicsSettings->NetworkPhysicsComponentSettings.GetEnableReliableFlow();
-			bApplyDataInsteadOfMergeData = PhysicsSettings->NetworkPhysicsComponentSettings.GetApplyDataInsteadOfMergeData();
-			bAllowInputExtrapolation = PhysicsSettings->NetworkPhysicsComponentSettings.GetAllowInputExtrapolation();
+			InputRedundancy = SettingsComponent->NetworkPhysicsComponentSettings.GetRedundantInputs();
+			StateRedundancy = SettingsComponent->NetworkPhysicsComponentSettings.GetRedundantStates();
+			bEnableUnreliableFlow = SettingsComponent->NetworkPhysicsComponentSettings.GetEnableUnreliableFlow();
+			bEnableReliableFlow = SettingsComponent->NetworkPhysicsComponentSettings.GetEnableReliableFlow();
+			bValidateDataOnGameThread = SettingsComponent->NetworkPhysicsComponentSettings.GetValidateDataOnGameThread();
 
 			if (ReplicatedInputs.History)
 			{
@@ -522,18 +506,18 @@ void UNetworkPhysicsComponent::InitializeComponent()
 		}
 	}
 
-	// TEMP: Stops ensure that happens if first ever network broadcast for the actor is a multicast from PT (with GT frozen), this instead makes the first multicast happen in GT.
-	// ToDo, remove this when NetworkPhysicComponent has a proper Async ISimCallback flow.
-	if (bEnableReliableFlow && ReplicatedImportantInput.Owner != nullptr)
-	{
-		MulticastReceiveImportantInputData(ReplicatedImportantInput);
-	}
-
 	if (UWorld* World = GetWorld())
 	{
-		if (UNetworkPhysicsSystem* NetworkManager = World->GetSubsystem<UNetworkPhysicsSystem>())
+		if (FPhysScene* PhysScene = World->GetPhysicsScene())
 		{
-			NetworkManager->RegisterNetworkComponent(this);
+			if (Chaos::FPhysicsSolver* Solver = PhysScene->GetSolver())
+			{
+				// Create async component to run on Physics Thread
+				NetworkPhysicsComponent_Internal = Solver->CreateAndRegisterSimCallbackObject_External<FAsyncNetworkPhysicsComponent>();
+				NetworkPhysicsComponent_Internal->SettingsComponent = SettingsComponent ? SettingsComponent->GetNetworkPhysicsSettings_Internal() : nullptr;
+				NetworkPhysicsComponent_Internal->RootPhysicsObject = RootPhysicsObject;
+				CreateAsyncDataHistory();
+			}
 		}
 	}
 }
@@ -543,9 +527,13 @@ void UNetworkPhysicsComponent::UninitializeComponent()
 	Super::UninitializeComponent();
 	if (UWorld* World = GetWorld())
 	{
-		if (UNetworkPhysicsSystem* NetworkManager = World->GetSubsystem<UNetworkPhysicsSystem>())
+		if (FPhysScene* PhysScene = World->GetPhysicsScene())
 		{
-			NetworkManager->UnregisterNetworkComponent(this);
+			if (Chaos::FPhysicsSolver* Solver = PhysScene->GetSolver())
+			{
+				// Clear async component from Physics Thread and memory
+				Solver->UnregisterAndFreeSimCallbackObject_External(NetworkPhysicsComponent_Internal);
+			}
 		}
 	}
 }
@@ -563,163 +551,106 @@ void UNetworkPhysicsComponent::GetLifetimeReplicatedProps(TArray< FLifetimePrope
 	DOREPLIFETIME_WITH_PARAMS_FAST(UNetworkPhysicsComponent, ReplicatedStates, ReplicatedParams);
 }
 
-void UNetworkPhysicsComponent::AsyncPhysicsTickComponent(float DeltaTime, float SimTime)
+// Called every Game Thread frame
+void UNetworkPhysicsComponent::TickComponent(float DeltaTime, ELevelTick TickType, FActorComponentTickFunction* ThisTickFunction)
 {
-	QUICK_SCOPE_CYCLE_COUNTER(NetworkPhysicsComponent_AsyncPhysicsTick);
-
-	Super ::AsyncPhysicsTickComponent(DeltaTime, SimTime);
-
-	// Record the received states from the server into the history for future use
-	if (UWorld* World = GetWorld())
-	{
-		if (FPhysScene* PhysScene = World->GetPhysicsScene())
-		{
-			if(!PhysScene->GetSolver()->GetEvolution()->IsResimming())
-			{
-				// Replicate inputs across the network
-				SendInputData();
-
-				// Replicate states across the network
-				SendStateData();
-			}
-		}
-	}
+	Super::TickComponent(DeltaTime, TickType, ThisTickFunction);
+	UpdateAsyncComponent(false);
+	NetworkMarshaledData();
 }
 
-void UNetworkPhysicsComponent::SendInputData()
+void UNetworkPhysicsComponent::NetworkMarshaledData()
 {
-	const bool bIsServer = HasServerWorld();
-	// Inputs are sent from the server or locally controlled actors/pawns
-	if (InputHistory && (IsLocallyControlled() || bIsServer))
+	if (NetworkPhysicsComponent_Internal)
 	{
-		const APlayerController* PlayerController = GetPlayerController();
-		if (!PlayerController)
-		{
-			PlayerController = GetWorld()->GetFirstPlayerController();
-		}
-
-		// Check that client and server have synced the physics tick offset
-		if (PlayerController && PlayerController->GetNetworkPhysicsTickOffsetAssigned())
+		const bool bIsServer = HasServerWorld();
+		while (Chaos::TSimCallbackOutputHandle<FAsyncNetworkPhysicsComponentOutput> AsyncOutput = NetworkPhysicsComponent_Internal->PopFutureOutputData_External())
 		{
-			// Send latest N frames from history
-			const int32 ToFrame = InputHistory->GetLatestFrame();
-
-			// -- Default / Unreliable Flow --
-			if (bEnableUnreliableFlow)
+			// Unimportant / Unreliable
+			if (bEnableUnreliableFlow
+				&& AsyncOutput->InputData
+				&& AsyncOutput->InputData->HasDataInHistory()
+				&& AsyncOutput->InputData->CopyAllData(*ReplicatedInputs.History, /*bIncludeUnimportant*/ true, /*bIncludeImportant*/ bEnableReliableFlow == false))
 			{
-				const int32 FromFrame = ToFrame - InputRedundancy;
-				if (InputHistory->CopyData(*ReplicatedInputs.History, FromFrame, ToFrame, /*bIncludeUnimportant*/ true, /*bIncludeImportant*/ bEnableReliableFlow == false))
+				if (bIsServer)
 				{
-					if (bIsServer)
-					{
-						// Server sends inputs through property replication
-						MARK_PROPERTY_DIRTY_FROM_NAME(UNetworkPhysicsComponent, ReplicatedInputs, this);
-					}
-					else
-					{
-						// Clients send inputs through an RPC to the server
-						ServerReceiveInputData(ReplicatedInputs);
-					}
-
-#if DEBUG_NETWORK_PHYSICS
-					{
-						FAsyncPhysicsTimestamp Timestamp = const_cast<APlayerController*>(PlayerController)->GetPhysicsTimestamp();
-						FString NetRole = HasServerWorld() ? FString("SERVER    ") : (IsLocallyControlled() ? FString("AUTONOMOUS") : FString("PROXY     "));
-						ReplicatedInputs.History->DebugData(FString::Printf(TEXT("%s | PT | SendInputData | CurrentLocalFrame = %d | CurrentServerFrame = %d | Component = %s"), *NetRole, Timestamp.LocalFrame, Timestamp.ServerFrame, *GetFullName()));
-					}
-#endif
+					// Server sends inputs through property replication
+					MARK_PROPERTY_DIRTY_FROM_NAME(UNetworkPhysicsComponent, ReplicatedInputs, this);
+				}
+				else if (IsLocallyControlled())
+				{
+					// Clients send inputs through an RPC to the server
+					ServerReceiveInputData(ReplicatedInputs);
 				}
 			}
 
-			// -- Important / Reliable flow --
+			// Important / Reliable
 			if (bEnableReliableFlow)
 			{
-				/* Get the latest valid frame that can hold new important data:
-				* 1. Frame after last time we called SendInputData
-				* 2. Frame after last entry in replicated data
-				* 3. Earliest possible frame in history */
-				const int32 FromFrame = FMath::Max(FMath::Max(LastInputSendFrame + 1, ReplicatedImportantInput.History->GetLatestFrame() + 1), ToFrame - InputHistory->GetHistorySize());
-
-				for (int32 CopyFrame = FromFrame; CopyFrame <= ToFrame; CopyFrame++)
+				for (const TUniquePtr<Chaos::FBaseRewindHistory>& InputImportant : AsyncOutput->InputDataImportant)
 				{
-					if (InputHistory->CopyData(*ReplicatedImportantInput.History, CopyFrame, CopyFrame, /*bIncludeUnimportant*/ false, /*bIncludeImportant*/ true))
+					if (!InputImportant || !InputImportant->HasDataInHistory())
+					{
+						continue;
+					}
+
+					ReplicatedImportantInput.History->ResizeDataHistory(InputImportant->GetHistorySize());
+					if (InputImportant->CopyAllData(*ReplicatedImportantInput.History, /*bIncludeUnimportant*/ false, /*bIncludeImportant*/ true))
 					{
 						if (bIsServer)
 						{
 							MulticastReceiveImportantInputData(ReplicatedImportantInput);
 						}
-						else
+						else if (IsLocallyControlled())
 						{
 							ServerReceiveImportantInputData(ReplicatedImportantInput);
 						}
-#if DEBUG_NETWORK_PHYSICS
-						{
-							FAsyncPhysicsTimestamp Timestamp = const_cast<APlayerController*>(PlayerController)->GetPhysicsTimestamp();
-							FString NetRole = HasServerWorld() ? FString("SERVER    ") : (IsLocallyControlled() ? FString("AUTONOMOUS") : FString("PROXY     "));
-							ReplicatedImportantInput.History->DebugData(FString::Printf(TEXT("%s | PT | SendInputData | RELIABLE | CurrentLocalFrame = %d | CurrentServerFrame = %d | Component = %s"), *NetRole, Timestamp.LocalFrame, Timestamp.ServerFrame, *GetFullName()));
-						}
-#endif
 					}
 				}
 			}
-			LastInputSendFrame = InputHistory->GetLatestFrame();
-		}
-	}
 
-	if (bStopRelayingLocalInputsDeferred)
-	{
-		bIsRelayingLocalInputs = false;
-		bStopRelayingLocalInputsDeferred = false;
-	}
-}
-
-void UNetworkPhysicsComponent::SendStateData()
-{
-	if (HasServerWorld() && StateHistory)
-	{
-		// Send latest N frames from history
-		const int32 ToFrame = StateHistory->GetLatestFrame();
-
-		// -- Default / Unreliable Flow --
-		if (bEnableUnreliableFlow)
-		{
-			const int32 FromFrame = ToFrame - StateRedundancy;
-			if (StateHistory->CopyData(*ReplicatedStates.History, FromFrame, ToFrame, /*bIncludeUnimportant*/ true, /*bIncludeImportant*/ bEnableReliableFlow == false))
+			if (bIsServer)
 			{
-				// If on server we should send the states onto all the clients through repnotify
-				MARK_PROPERTY_DIRTY_FROM_NAME(UNetworkPhysicsComponent, ReplicatedStates, this);
+				// Unimportant / Unreliable
+				if (bEnableUnreliableFlow
+					&& AsyncOutput->StateData
+					&& AsyncOutput->StateData->HasDataInHistory()
+					&& AsyncOutput->StateData->CopyAllData(*ReplicatedStates.History, /*bIncludeUnimportant*/ true, /*bIncludeImportant*/ bEnableReliableFlow == false))
+				{
+					// If on server we should send the states onto all the clients through repnotify
+					MARK_PROPERTY_DIRTY_FROM_NAME(UNetworkPhysicsComponent, ReplicatedStates, this);
+				}
 
-#if DEBUG_NETWORK_PHYSICS
-				ReplicatedStates.History->DebugData(FString::Printf(TEXT("SERVER | PT | SendStateData | Component = %s"), *GetFullName()));
-#endif
+				// Important / Reliable
+				if (bEnableReliableFlow)
+				{
+					for (const TUniquePtr<Chaos::FBaseRewindHistory>& StateImportant : AsyncOutput->StateDataImportant)
+					{
+						if (!StateImportant || !StateImportant->HasDataInHistory())
+						{
+							continue;
+						}
+
+						ReplicatedImportantState.History->ResizeDataHistory(StateImportant->GetHistorySize());
+						if (StateImportant->CopyAllData(*ReplicatedImportantState.History, /*bIncludeUnimportant*/ false, /*bIncludeImportant*/ true))
+						{
+							MulticastReceiveImportantStateData(ReplicatedImportantState);
+						}
+					}
+				}
 			}
-		}
-		
-		// -- Important / Reliable flow --
-		if (bEnableReliableFlow)
-		{
-			/* Get the latest valid frame that can hold new important data: 
-			* 1. Frame after last time we called SendStateData
-			* 2. Frame after last entry in replicated data
-			* 3. Earliest possible frame in history */
-			const int32 FromFrame = FMath::Max(FMath::Max(LastStateSendFrame + 1, ReplicatedImportantState.History->GetLatestFrame() + 1), ToFrame - StateHistory->GetHistorySize());
 
-			for (int32 CopyFrame = FromFrame; CopyFrame <= ToFrame; CopyFrame++)
+			if (bStopRelayingLocalInputsDeferred)
 			{
-				if (StateHistory->CopyData(*ReplicatedImportantState.History, CopyFrame, CopyFrame, /*bIncludeUnimportant*/ false, /*bIncludeImportant*/ true))
-				{
-					MulticastReceiveImportantStateData(ReplicatedImportantState);
-#if DEBUG_NETWORK_PHYSICS
-					ReplicatedImportantState.History->DebugData(FString::Printf(TEXT("SERVER | PT | SendStateData | RELIABLE | Component = %s"), *GetFullName()));
-#endif
-				}
+				bIsRelayingLocalInputs = false;
+				bStopRelayingLocalInputsDeferred = false;
+				UpdateAsyncComponent();
 			}
 		}
-		LastStateSendFrame = StateHistory->GetLatestFrame();
 	}
 }
 
-/* Deprecated 5.4 */
+/** Deprecated 5.4 */
 void UNetworkPhysicsComponent::CorrectServerToLocalOffset(const int32 LocalToServerOffset)
 {
 	if (IsLocallyControlled() && !HasServerWorld() && StateHistory)
@@ -749,108 +680,55 @@ void UNetworkPhysicsComponent::CorrectServerToLocalOffset(const int32 LocalToSer
 
 void UNetworkPhysicsComponent::OnRep_SetReplicatedStates()
 {
-	APlayerController* PlayerController = GetPlayerController();
-	if (!PlayerController)
+	if (!NetworkPhysicsComponent_Internal || !StateHelper || !ReplicatedStates.History)
+	{
+		return;
+	}
+
+	if (FAsyncNetworkPhysicsComponentInput* AsyncInput = NetworkPhysicsComponent_Internal->GetProducerInputData_External())
 	{
-		PlayerController = GetWorld()->GetFirstPlayerController();
+		if (!AsyncInput->StateData)
+		{
+			AsyncInput->StateData = StateHelper->CreateUniqueRewindHistory(StateRedundancy);
+		}
+
+		ReplicatedStates.History->CopyAllData(*AsyncInput->StateData.Get(), /*bIncludeUnimportant*/ true, /*bIncludeImportant*/ true);
 	}
 
-	// The replicated states should only be used on the client since the server already have authoritative local ones
-	if (PlayerController && !HasServerWorld() && StateHistory)
+	/** Deprecated 5.4 */
+	if (InputCmdCVars::bCmdOffsetEnabled)
 	{
-		const int32 LocalOffset = PlayerController->GetNetworkPhysicsTickOffset();
+		APlayerController* PlayerController = GetPlayerController();
+		if (!PlayerController)
+		{
+			PlayerController = GetWorld()->GetFirstPlayerController();
+		}
 
-		/* Deprecated 5.4 */
-		if (InputCmdCVars::bCmdOffsetEnabled)
+		if (PlayerController && !HasServerWorld())
 		{
 			PRAGMA_DISABLE_DEPRECATION_WARNINGS
 			const int32 LocalToServerOffset = PlayerController->GetLocalToServerAsyncPhysicsTickOffset();
 			CorrectServerToLocalOffset(LocalToServerOffset);
 			PRAGMA_ENABLE_DEPRECATION_WARNINGS
 		}
-
-		// Record the received states from the server into the history for future use
-		if (UWorld* World = GetWorld())
-		{
-			if (FPhysScene* PhysScene = World->GetPhysicsScene())
-			{
-				TSharedPtr<Chaos::FBaseRewindHistory> ReceivedStates = MakeShareable(ReplicatedStates.History->Clone().Release());
-				PhysScene->EnqueueAsyncPhysicsCommand(0, this, [this, PhysScene, ReceivedStates, LocalOffset]()
-				{
-					if (bCompareStateToTriggerRewind && IsLocallyControlled())
-					{
-						// Record received data into history
-						const int32 ResimFrame = StateHistory->ReceiveNewData(*ReceivedStates, LocalOffset, /*CompareDataForRewind*/ bCompareStateToTriggerRewind);
-					
-						// Trigger resimulation if ReceiveNewData returned a valid frame to rewind to
-						TriggerResimulation(ResimFrame);
-					}
-					else 
-					{
-						// Record received data into history
-						PRAGMA_DISABLE_DEPRECATION_WARNINGS // TODO: Change to ReceiveNewData() in UE 5.6 and remove deprecation pragma
-						StateHistory->ReceiveNewDatas(*ReceivedStates, LocalOffset);
-						PRAGMA_ENABLE_DEPRECATION_WARNINGS
-					}
-
-#if DEBUG_NETWORK_PHYSICS
-					{
-						const int32 CurrentFrame = PhysScene->GetSolver()->GetCurrentFrame();
-						ReceivedStates->DebugData(FString::Printf(TEXT("CLIENT | PT | OnRep_SetReplicatedStates | CurrentLocalFrame = %d | LocalOffset = %d | Component = %s"), CurrentFrame, LocalOffset, *GetFullName()));
-					}
-#endif
-				}, false);
-			}
-		}
 	}
 }
 
 void UNetworkPhysicsComponent::OnRep_SetReplicatedInputs()
 {
-	APlayerController* PlayerController = GetPlayerController();
-	if (!PlayerController)
+	if (!NetworkPhysicsComponent_Internal || !InputHelper || !ReplicatedInputs.History)
 	{
-		PlayerController = GetWorld()->GetFirstPlayerController();
+		return ;
 	}
 
-	// Put replicated inputs into input history, also done for the local player since the server can alter inputs if invalid and the correct server-authoritative input should be used during a resimulation
-	if (PlayerController && !HasServerWorld() && InputHistory)
+	if (FAsyncNetworkPhysicsComponentInput* AsyncInput = NetworkPhysicsComponent_Internal->GetProducerInputData_External())
 	{
-		const int32 LocalOffset = PlayerController->GetNetworkPhysicsTickOffset();
-
-		// Record the received inputs from the server into the history for future use
-		if (UWorld* World = GetWorld())
+		if (!AsyncInput->InputData)
 		{
-			if (FPhysScene* PhysScene = World->GetPhysicsScene())
-			{
-				TSharedPtr<Chaos::FBaseRewindHistory> ReceivedInputs = MakeShareable(ReplicatedInputs.History->Clone().Release());
-				PhysScene->EnqueueAsyncPhysicsCommand(0, this, [this, PhysScene, ReceivedInputs, LocalOffset]()
-				{
-					if (bCompareInputToTriggerRewind && IsLocallyControlled())
-					{
-						// Record received data into history
-						const int32 ResimFrame = InputHistory->ReceiveNewData(*ReceivedInputs, LocalOffset, /*CompareDataForRewind*/ bCompareInputToTriggerRewind, /*Important*/false);
-						
-						// Trigger resimulation if ReceiveNewData returned a valid frame to rewind to
-						TriggerResimulation(ResimFrame);
-					}
-					else
-					{
-						// Record received data into history
-						PRAGMA_DISABLE_DEPRECATION_WARNINGS // TODO: Change to ReceiveNewData() in UE 5.6 and remove deprecation pragma
-						InputHistory->ReceiveNewDatas(*ReceivedInputs, LocalOffset);
-						PRAGMA_ENABLE_DEPRECATION_WARNINGS
-					}
-
-#if DEBUG_NETWORK_PHYSICS
-					{
-						const int32 CurrentFrame = PhysScene->GetSolver()->GetCurrentFrame();
-						ReceivedInputs->DebugData(FString::Printf(TEXT("CLIENT | PT | OnRep_SetReplicatedInputs | CurrentLocalFrame = %d | LocalOffset = %d | Component = %s"), CurrentFrame, LocalOffset, *GetFullName()));
-					}
-#endif
-				}, false);
-			}
+			AsyncInput->InputData = InputHelper->CreateUniqueRewindHistory(InputRedundancy);
 		}
+
+		ReplicatedInputs.History->CopyAllData(*AsyncInput->InputData.Get(), /*bIncludeUnimportant*/ true, /*bIncludeImportant*/ true);
 	}
 }
 
@@ -862,219 +740,548 @@ void UNetworkPhysicsComponent::ServerReceiveInputsDatas_Implementation(const FNe
 
 void UNetworkPhysicsComponent::ServerReceiveInputData_Implementation(const FNetworkPhysicsRewindDataInputProxy& ClientInputs)
 {
-	if (InputHistory)
-	{ 
-		if (UWorld* World = GetWorld())
-		{
-			if (FPhysScene* PhysScene = World->GetPhysicsScene())
-			{
-				// Make another copy of the client inputs for the physics thread to consume
-				TSharedPtr<Chaos::FBaseRewindHistory> ReceivedInputs = MakeShareable(ClientInputs.History->Clone().Release());
-				PhysScene->EnqueueAsyncPhysicsCommand(0, this, [this, ReceivedInputs, PhysScene]()
-				{
-					// Validate data in the received inputs
-					ReceivedInputs->ValidateDataInHistory(ActorComponent);
+	if (!NetworkPhysicsComponent_Internal || !InputHelper || !ClientInputs.History)
+	{
+		return;
+	}
 
-					// Record received data into history
-					PRAGMA_DISABLE_DEPRECATION_WARNINGS // TODO: Change to ReceiveNewData() in UE 5.6 and remove deprecation pragma
-					InputHistory->ReceiveNewDatas(*ReceivedInputs, 0);
-					PRAGMA_ENABLE_DEPRECATION_WARNINGS
+	if (FAsyncNetworkPhysicsComponentInput* AsyncInput = NetworkPhysicsComponent_Internal->GetProducerInputData_External())
+	{
+		if (!AsyncInput->InputData)
+		{
+			AsyncInput->InputData = InputHelper->CreateUniqueRewindHistory(InputRedundancy);
+		}
 
-#if DEBUG_NETWORK_PHYSICS
-					{
-						const int32 CurrentFrame = PhysScene->GetSolver()->GetCurrentFrame();
-						ReceivedInputs->DebugData(FString::Printf(TEXT("SERVER | PT | ServerReceiveInputData | CurrentLocalFrame = %d | Component = %s"), CurrentFrame, *GetFullName()));
-					}
-#endif
-				}, false);
-			}
+		// Validate data in the received inputs
+		if (bValidateDataOnGameThread && ActorComponent)
+		{
+			ClientInputs.History->ValidateDataInHistory(ActorComponent);
 		}
+
+		ClientInputs.History->CopyAllData(*AsyncInput->InputData.Get(), /*bIncludeUnimportant*/ true, /*bIncludeImportant*/ true);
 	}
 }
 
 void UNetworkPhysicsComponent::ServerReceiveImportantInputData_Implementation(const FNetworkPhysicsRewindDataImportantInputProxy& ClientInputs)
 {
-	if (InputHistory)
+	if (!NetworkPhysicsComponent_Internal || !ClientInputs.History)
 	{
-		if (UWorld* World = GetWorld())
-		{
-			if (FPhysScene* PhysScene = World->GetPhysicsScene())
-			{
-				// Make another copy of the client inputs for the physics thread to consume
-				TSharedPtr<Chaos::FBaseRewindHistory> ReceivedInputs = MakeShareable(ClientInputs.History->Clone().Release());
-				PhysScene->EnqueueAsyncPhysicsCommand(0, this, [this, ReceivedInputs, PhysScene]()
-					{
-						// Validate data in the received inputs
-						ReceivedInputs->ValidateDataInHistory(ActorComponent);
-
-						// Record received data into history
-						InputHistory->ReceiveNewData(*ReceivedInputs, 0, /*CompareDataForRewind*/false, /*Important*/true);
+		return;
+	}
 
-						// Cache earliest important data frame since last ApplyData was called
-						NewImportantInputFrame = FMath::Min(NewImportantInputFrame, ReceivedInputs->GetEarliestFrame());
+	if (FAsyncNetworkPhysicsComponentInput* AsyncInput = NetworkPhysicsComponent_Internal->GetProducerInputData_External())
+	{
+		// Initialize received data since not all data is networked and when we clone this we expect to have fully initialized data
+		ClientInputs.History->Initialize();
 
-#if DEBUG_NETWORK_PHYSICS
-						{
-							const int32 CurrentFrame = PhysScene->GetSolver()->GetCurrentFrame();
-							ReceivedInputs->DebugData(FString::Printf(TEXT("SERVER | PT | ServerReceiveImportantInputData | RELIABLE | CurrentLocalFrame = %d | NewImportantInputFrame = %d | Component = %s"), CurrentFrame, NewImportantInputFrame, *GetFullName()));
-						}
-#endif
-					}, false);
-			}
+		// Validate data in the received inputs
+		if (bValidateDataOnGameThread && ActorComponent)
+		{
+			ClientInputs.History->ValidateDataInHistory(ActorComponent);
 		}
+
+		// Create new data collection for marshaling
+		AsyncInput->InputDataImportant.Add(ClientInputs.History->Clone());
 	}
 }
 
 void UNetworkPhysicsComponent::MulticastReceiveImportantInputData_Implementation(const FNetworkPhysicsRewindDataImportantInputProxy& ServerInputs)
 {
-	// Multicast RPC, ignore on server
+	// Ignore Multicast on server
 	if (HasServerWorld())
 	{
 		return;
 	}
 
-	APlayerController* PlayerController = GetPlayerController();
-	if (!PlayerController)
+	if (!NetworkPhysicsComponent_Internal || !ServerInputs.History)
 	{
-		PlayerController = GetWorld()->GetFirstPlayerController();
+		return;
 	}
 
-	if (PlayerController && InputHistory)
+	if (FAsyncNetworkPhysicsComponentInput* AsyncInput = NetworkPhysicsComponent_Internal->GetProducerInputData_External())
 	{
-		const int32 LocalOffset = PlayerController->GetNetworkPhysicsTickOffset();
-		
-		if (UWorld* World = GetWorld())
-		{
-			if (FPhysScene* PhysScene = World->GetPhysicsScene())
-			{
-				// Make another copy of the client inputs for the physics thread to consume
-				TSharedPtr<Chaos::FBaseRewindHistory> ReceivedInputs = MakeShareable(ServerInputs.History->Clone().Release());
-				PhysScene->EnqueueAsyncPhysicsCommand(0, this, [this, ReceivedInputs, PhysScene, LocalOffset]()
-					{
-						// Record received data into history
-						const bool bCompareDataForRewind = (bCompareInputToTriggerRewind && IsLocallyControlled());
-						const int32 ResimFrame = InputHistory->ReceiveNewData(*ReceivedInputs, LocalOffset, /*CompareDataForRewind*/bCompareDataForRewind, /*Important*/true);
-						
-						// Trigger resimulation if ReceiveNewData returned a valid frame to rewind to
-						TriggerResimulation(ResimFrame);
-
-						// Cache earliest important data frame since last ApplyData was called
-						NewImportantInputFrame = FMath::Min(NewImportantInputFrame, ReceivedInputs->GetEarliestFrame());
-
-#if DEBUG_NETWORK_PHYSICS
-						{
-							const int32 CurrentFrame = PhysScene->GetSolver()->GetCurrentFrame();
-							FString NetRole = HasServerWorld() ? FString("SERVER    ") : (IsLocallyControlled() ? FString("AUTONOMOUS") : FString("PROXY     "));
-							ReceivedInputs->DebugData(FString::Printf(TEXT("%s | PT | MulticastReceiveImportantInputData | RELIABLE | CurrentLocalFrame = %d | NewImportantInputFrame = %d | Component = %s"), *NetRole, CurrentFrame, NewImportantInputFrame, *GetFullName()));
-						}
-#endif
-					}, false);
-			}
-		}
+		// Initialize received data since not all data is networked and when we clone this we expect to have fully initialized data
+		ServerInputs.History->Initialize();
 
+		// Create new data collection for marshaling
+		AsyncInput->InputDataImportant.Add(ServerInputs.History->Clone());
 	}
 }
 
 void UNetworkPhysicsComponent::MulticastReceiveImportantStateData_Implementation(const FNetworkPhysicsRewindDataImportantStateProxy& ServerStates)
 {
-	// Multicast RPC, ignore on server
+	// Ignore Multicast on server
 	if (HasServerWorld())
 	{
 		return;
 	}
 
-	APlayerController* PlayerController = GetPlayerController();
-	if (!PlayerController)
+	if (!NetworkPhysicsComponent_Internal || !ServerStates.History)
 	{
-		PlayerController = GetWorld()->GetFirstPlayerController();
+		return;
 	}
 
-	if (PlayerController && StateHistory)
+	if (FAsyncNetworkPhysicsComponentInput* AsyncInput = NetworkPhysicsComponent_Internal->GetProducerInputData_External())
 	{
-		const int32 LocalOffset = PlayerController->GetNetworkPhysicsTickOffset();
-
-		if (UWorld* World = GetWorld())
-		{
-			if (FPhysScene* PhysScene = World->GetPhysicsScene())
-			{
-				// Make another copy of the client inputs for the physics thread to consume
-				TSharedPtr<Chaos::FBaseRewindHistory> ReceivedStates = MakeShareable(ServerStates.History->Clone().Release());
-				PhysScene->EnqueueAsyncPhysicsCommand(0, this, [this, ReceivedStates, PhysScene, LocalOffset]()
-					{
-						const bool bCompareDataForRewind = (bCompareStateToTriggerRewind && IsLocallyControlled());
-						const int32 ResimFrame = StateHistory->ReceiveNewData(*ReceivedStates, LocalOffset, /*CompareDataForRewind*/bCompareDataForRewind, /*bImportant*/true);
-						TriggerResimulation(ResimFrame);
-
-#if DEBUG_NETWORK_PHYSICS
-						{
-							const int32 CurrentFrame = PhysScene->GetSolver()->GetCurrentFrame();
-							FString NetRole = HasServerWorld() ? FString("SERVER    ") : (IsLocallyControlled() ? FString("AUTONOMOUS") : FString("PROXY     "));
-							ReceivedStates->DebugData(FString::Printf(TEXT("%s | PT | MulticastReceiveImportantStateData | RELIABLE | CurrentLocalFrame = %d | Component = %s"), *NetRole, CurrentFrame, *GetFullName()));
-						}
-#endif
-					}, false);
-			}
-		}
+		// Initialize received data since not all data is networked and when we clone this we expect to have fully initialized data
+		ServerStates.History->Initialize();
 
+		// Create new data collection for marshaling
+		AsyncInput->StateDataImportant.Add(ServerStates.History->Clone());
 	}
 }
 
-void UNetworkPhysicsComponent::OnPreProcessInputsInternal(const int32 PhysicsStep)
+/** DEPRECATED UE 5.5 */
+const float UNetworkPhysicsComponent::GetCurrentInputDecay(FNetworkPhysicsData* PhysicsData)
 {
-	const bool bIsServer = HasServerWorld();
+	if (!PhysicsData)
+	{
+		return 0.0f;
+	}
 
-#if DEBUG_NETWORK_PHYSICS
-	if (bIsServer)
+	FPhysScene* PhysScene = GetWorld()->GetPhysicsScene();
+	if (!PhysScene)
 	{
-		UE_LOG(LogChaos, Log, TEXT("SERVER | PT | OnPreProcessInputsInternal | At Frame %d | Component = %s"), PhysicsStep, *GetFullName());
+		return 0.0f;
 	}
-	else
+
+	Chaos::FPhysicsSolver* Solver = PhysScene->GetSolver();
+	if (!Solver)
 	{
-		UE_LOG(LogChaos, Log, TEXT("CLIENT | PT | OnPreProcessInputsInternal | At Frame %d | Component = %s"), PhysicsStep, *GetFullName());
+		return 0.0f;
 	}
-#endif
 
-	if (ActorComponent)
-	{ 
-		bool bIsSolverReset = false;
-		bool bIsSolverResim = false;
+	Chaos::FRewindData* RewindData = Solver->GetRewindData();
+	if (!RewindData)
+	{
+		return 0.0f;
+	}
+	
+	const float NumPredictedInputs = RewindData->CurrentFrame() - PhysicsData->LocalFrame; // Number of frames we have used the same PhysicsData for during resim
+	const float MaxPredictedInputs = RewindData->GetLatestFrame() - 1 - PhysicsData->LocalFrame; // Max number of frames PhysicsData registered frame until end of resim
 
-		if (FPhysScene* PhysScene = GetWorld()->GetPhysicsScene())
-		{ 
-			if (Chaos::FPhysicsSolver* Solver = PhysScene->GetSolver())
-			{
-				bIsSolverResim = Solver->GetEvolution()->IsResimming();
-				bIsSolverReset = Solver->GetEvolution()->IsResetting();
-			}
-		}
+	// Linear decay
+	const float PredictionAlpha = MaxPredictedInputs > 0 ? (NumPredictedInputs / MaxPredictedInputs) : 0.0f;
 
-		// Apply replicated state on clients if we are resimulating
-		if (bIsSolverResim && StateHistory)
-		{
-			FNetworkPhysicsData* PhysicsData = StateData.Get();
-			PhysicsData->LocalFrame = PhysicsStep;
-			const bool bExactFrame = PhysicsReplicationCVars::ResimulationCVars::bAllowRewindToClosestState ? !bIsSolverReset : true;
-			PRAGMA_DISABLE_DEPRECATION_WARNINGS // TODO: Change to ExtractData() in UE 5.6 and remove deprecation pragma
-			if (StateHistory->ExtractDatas(PhysicsStep, bIsSolverReset, PhysicsData, bExactFrame) && PhysicsData->bReceivedData)
-			PRAGMA_ENABLE_DEPRECATION_WARNINGS
-			{
-				PhysicsData->ApplyData(ActorComponent);
-#if DEBUG_NETWORK_PHYSICS
-				UE_LOG(LogChaos, Log, TEXT("		Applying extracted state from history | bExactFrame = %d | LocalFrame = %d | ServerFrame = %d | InputFrame = %d | Data: %s")
-					, bExactFrame, PhysicsData->LocalFrame, PhysicsData->ServerFrame, PhysicsData->InputFrame, *PhysicsData->DebugData());
-#endif
-			}
-#if DEBUG_NETWORK_PHYSICS
-			else
-			{
-				UE_LOG(LogChaos, Log, TEXT("		FAILED to extract and apply state from history | bExactFrame = %d | -- Printing history --"), bExactFrame);
-				StateHistory->DebugData(FString::Printf(TEXT("StateHistory | Component = %s"), *GetFullName()));
+	return PredictionAlpha;
+}
+
+bool UNetworkPhysicsComponent::HasServerWorld() const
+{
+	return GetWorld()->IsNetMode(NM_DedicatedServer) || GetWorld()->IsNetMode(NM_ListenServer);
+}
+
+bool UNetworkPhysicsComponent::HasLocalController() const
+{
+	if (APlayerController* PlayerController = GetPlayerController())
+	{
+		return PlayerController->IsLocalController();
+	}
+	return false;
+}
+
+bool UNetworkPhysicsComponent::IsLocallyControlled() const
+{
+	if (bIsRelayingLocalInputs)
+	{
+		return true;
+	}
+
+	if (APlayerController* PlayerController = GetPlayerController())
+	{
+		return PlayerController->IsLocalController();
+	}
+	return false;
+}
+
+APlayerController* UNetworkPhysicsComponent::GetPlayerController() const
+{
+	if (APlayerController* PC = Cast<APlayerController>(GetOwner()))
+	{
+		return PC;
+	}
+
+	if (APawn* Pawn = Cast<APawn>(GetOwner()))
+	{
+		if (APlayerController * PC = Pawn->GetController<APlayerController>())
+		{
+			return PC;
+		}
+
+		// In this case the APlayerController can be found as the owner of the pawn
+		if (APlayerController* PC = Cast<APlayerController>(Pawn->GetOwner()))
+		{
+			return PC;
+		}
+	}
+
+	return nullptr;
+}
+
+void UNetworkPhysicsComponent::UpdateAsyncComponent(const bool bFullUpdate)
+{
+	// Marshal data from Game Thread to Physics Thread
+	if (NetworkPhysicsComponent_Internal)
+	{
+		if (FAsyncNetworkPhysicsComponentInput* AsyncInput = NetworkPhysicsComponent_Internal->GetProducerInputData_External())
+		{
+			if (!HasServerWorld())
+			{
+				if (UWorld* World = GetWorld())
+				{
+					if (APlayerController* PC = World->GetFirstPlayerController())
+					{
+						AsyncInput->NetworkPhysicsTickOffset = PC->GetNetworkPhysicsTickOffset();
+					}
+				}
+			}
+
+			if (bFullUpdate)
+			{
+				AsyncInput->bIsLocallyControlled = IsLocallyControlled();
+
+				if (UWorld* World = GetWorld())
+				{ 
+					AsyncInput->NetMode = World->GetNetMode();
+				}
+
+				if (AActor* Owner = GetOwner())
+				{ 
+					AsyncInput->NetRole = Owner->GetLocalRole();
+					AsyncInput->PhysicsReplicationMode = Owner->GetPhysicsReplicationMode();
+					AsyncInput->ActorName = AActor::GetDebugName(Owner);
+				}
+				
+				if (ActorComponent)
+				{
+					AsyncInput->ActorComponent = ActorComponent;
+				}
+
+			}
+		}
+	}
+}
+
+void UNetworkPhysicsComponent::CreateAsyncDataHistory()
+{
+	if (NetworkPhysicsComponent_Internal)
+	{
+		if (FAsyncNetworkPhysicsComponentInput* AsyncInput = NetworkPhysicsComponent_Internal->GetProducerInputData_External())
+		{
+			AsyncInput->ActorComponent = ActorComponent;
+
+			if (InputHelper)
+			{
+				// Marshal the input helper to create both input data and input history on the physics thread
+				AsyncInput->InputHelper = InputHelper->Clone();
+			}
+
+			if (StateHelper)
+			{
+				// Marshal the state helper to create both state data and state history on the physics thread
+				AsyncInput->StateHelper = StateHelper->Clone();
+			}
+		}
+	}
+}
+
+void UNetworkPhysicsComponent::RemoveDataHistory()
+{
+	// Tell the async network physics component to unregister from RewindData
+	if (NetworkPhysicsComponent_Internal)
+	{
+		if (FAsyncNetworkPhysicsComponentInput* AsyncInput = NetworkPhysicsComponent_Internal->GetProducerInputData_External())
+		{
+			AsyncInput->bUnregisterDataHistoryFromRewindData = true;
+		}
+	}
+}
+
+void UNetworkPhysicsComponent::AddDataHistory()
+{
+	// Tell the async network physics component to register in RewindData
+	if (NetworkPhysicsComponent_Internal)
+	{
+		if (FAsyncNetworkPhysicsComponentInput* AsyncInput = NetworkPhysicsComponent_Internal->GetProducerInputData_External())
+		{
+			AsyncInput->bRegisterDataHistoryInRewindData = true;
+		}
+	}
+}
+
+TSharedPtr<Chaos::FBaseRewindHistory>& UNetworkPhysicsComponent::GetStateHistory_Internal()
+{
+	if (NetworkPhysicsComponent_Internal)
+	{
+		return NetworkPhysicsComponent_Internal->StateHistory;
+	}
+	return StateHistory;
+}
+
+TSharedPtr<Chaos::FBaseRewindHistory>& UNetworkPhysicsComponent::GetInputHistory_Internal()
+{
+	if (NetworkPhysicsComponent_Internal)
+	{
+		return NetworkPhysicsComponent_Internal->InputHistory;
+	}
+	return InputHistory;
+}
+
+
+// --------------------------- Async Network Physics Component ---------------------------
+
+// Initialize static
+const FNetworkPhysicsSettingsNetworkPhysicsComponent FAsyncNetworkPhysicsComponent::SettingsNetworkPhysicsComponent_Default = FNetworkPhysicsSettingsNetworkPhysicsComponent();
+
+FAsyncNetworkPhysicsComponent::FAsyncNetworkPhysicsComponent() : TSimCallbackObject()
+	, bIsLocallyControlled(true)
+	, NetMode(ENetMode::NM_Standalone)
+	, NetRole(ENetRole::ROLE_Authority)
+	, NetworkPhysicsTickOffset(0)
+	, PhysicsReplicationMode(EPhysicsReplicationMode::Default)
+{
+}
+
+FAsyncNetworkPhysicsComponent::~FAsyncNetworkPhysicsComponent()
+{
+	if (Chaos::FPhysicsSolverBase* BaseSolver = GetSolver())
+	{
+		// Unregister for Pre- and Post- ProcessInputs_Internal callbacks
+		if (FNetworkPhysicsCallback* SolverCallback = static_cast<FNetworkPhysicsCallback*>(BaseSolver->GetRewindCallback()))
+		{
+			SolverCallback->PreProcessInputsInternal.Remove(DelegateOnPreProcessInputs_Internal);
+			DelegateOnPreProcessInputs_Internal.Reset();
+
+			SolverCallback->PostProcessInputsInternal.Remove(DelegateOnPostProcessInputs_Internal);
+			DelegateOnPostProcessInputs_Internal.Reset();
+		}
+	}
+
+	UnregisterDataHistoryFromRewindData();
+}
+
+void FAsyncNetworkPhysicsComponent::OnPostInitialize_Internal()
+{
+	if (Chaos::FPhysicsSolverBase* BaseSolver = GetSolver())
+	{
+		if (BaseSolver->IsNetworkPhysicsPredictionEnabled())
+		{
+			// Register for Pre- and Post- ProcessInputs_Internal callbacks
+			if (FNetworkPhysicsCallback* SolverCallback = static_cast<FNetworkPhysicsCallback*>(BaseSolver->GetRewindCallback()))
+			{
+				DelegateOnPreProcessInputs_Internal = SolverCallback->PreProcessInputsInternal.AddRaw(this, &FAsyncNetworkPhysicsComponent::OnPreProcessInputs_Internal);
+				DelegateOnPostProcessInputs_Internal = SolverCallback->PostProcessInputsInternal.AddRaw(this, &FAsyncNetworkPhysicsComponent::OnPostProcessInputs_Internal);
+			}
+		}
+		else
+		{
+			UE_LOG(LogChaos, Warning, TEXT("A NetworkPhysicsComponent is trying to set up but 'Project Settings -> Physics -> Physics Prediction' is not enabled. The component might not work as intended."));
+		}
+	}
+}
+
+const FNetworkPhysicsSettingsNetworkPhysicsComponent& FAsyncNetworkPhysicsComponent::GetComponentSettings()
+{
+	return SettingsComponent ? SettingsComponent->Settings.NetworkPhysicsComponentSettings : SettingsNetworkPhysicsComponent_Default;
+};
+
+void FAsyncNetworkPhysicsComponent::ConsumeAsyncInput()
+{
+	if (const FAsyncNetworkPhysicsComponentInput* AsyncInput = GetConsumerInput_Internal())
+	{
+		const FNetworkPhysicsSettingsNetworkPhysicsComponent& ComponentSettings = GetComponentSettings();
+
+		/** Onetime setup marshaled data */
+		{
+			if (AsyncInput->bIsLocallyControlled.IsSet())
+			{
+				bIsLocallyControlled = *AsyncInput->bIsLocallyControlled;
+			}
+			if (AsyncInput->NetMode.IsSet())
+			{
+				NetMode = *AsyncInput->NetMode;
+			}
+			if (AsyncInput->NetRole.IsSet())
+			{
+				NetRole = *AsyncInput->NetRole;
+			}
+			if (AsyncInput->NetworkPhysicsTickOffset.IsSet())
+			{
+				NetworkPhysicsTickOffset = *AsyncInput->NetworkPhysicsTickOffset;
+			}
+			if (AsyncInput->PhysicsReplicationMode.IsSet())
+			{
+				PhysicsReplicationMode = *AsyncInput->PhysicsReplicationMode;
+			}
+			if (AsyncInput->ActorComponent.IsSet())
+			{
+				ActorComponent = *AsyncInput->ActorComponent;
+			}
+			if (AsyncInput->ActorName.IsSet())
+			{
+				ActorName = *AsyncInput->ActorName;
+			}
+			if (AsyncInput->bRegisterDataHistoryInRewindData.IsSet())
+			{
+				RegisterDataHistoryInRewindData();
+			}
+			if (AsyncInput->bUnregisterDataHistoryFromRewindData.IsSet())
+			{
+				UnregisterDataHistoryFromRewindData();
+			}
+			if (AsyncInput->InputHelper.IsSet())
+			{
+				// Setup rewind data if not already done, and get history size
+				const int32 NumFrames = SetupRewindData();
+
+				// Create input data and history
+				InputData = (*AsyncInput->InputHelper)->CreateUniqueData();
+				InputHistory = MakeShareable((*AsyncInput->InputHelper)->CreateUniqueRewindHistory(NumFrames).Release());
+				RegisterDataHistoryInRewindData();
+			}
+			if (AsyncInput->StateHelper.IsSet())
+			{
+				// Setup rewind data if not already done, and get history size
+				const int32 NumFrames = SetupRewindData();
+
+				// Create state data and history
+				StateData = (*AsyncInput->StateHelper)->CreateUniqueData();
+				StateHistory = MakeShareable((*AsyncInput->StateHelper)->CreateUniqueRewindHistory(NumFrames).Release());
+				RegisterDataHistoryInRewindData();
+			}
+		}
+
+		/** Continuously marshaled data */
+		{
+			/** Receive data helper */
+			auto ReceiveHelper = [&](Chaos::FBaseRewindHistory* History, Chaos::FBaseRewindHistory* ReceiveData, const bool bImportant, const bool bCompareData)
+			{
+				const bool bCompareDataForRewind = bCompareData && IsLocallyControlled() && !IsServer();
+				const int32 ResimFrame = History->ReceiveNewData(*ReceiveData, (IsServer() ? 0 : NetworkPhysicsTickOffset), bCompareDataForRewind, bImportant);
+				if (bCompareDataForRewind)
+				{
+					TriggerResimulation(ResimFrame);
+				}
+#if DEBUG_NETWORK_PHYSICS
+				{
+					FString NetRoleString = IsServer() ? FString("SERVER") : (IsLocallyControlled() ? FString("AUTONO") : FString("PROXY "));
+					ReceiveData->DebugData(FString::Printf(TEXT("%s | PT | RECEIVE DATA | bImportant: %d | Name: %s"), *NetRoleString, bImportant, *GetActorName()));
+				}
+#endif
+			};
+
+			// Receive Inputs
+			if (AsyncInput->InputData && AsyncInput->InputData->HasDataInHistory())
+			{
+				// Validate data in the received inputs on the server
+				if (!ComponentSettings.GetValidateDataOnGameThread() && IsServer() && ActorComponent)
+				{
+					AsyncInput->InputData->ValidateDataInHistory(ActorComponent);
+				}
+				ReceiveHelper(InputHistory.Get(), AsyncInput->InputData.Get(), /*bImportant*/false, ComponentSettings.GetCompareInputToTriggerRewind());
+			}
+
+			// Receive States
+			if (AsyncInput->StateData && AsyncInput->StateData->HasDataInHistory())
+			{
+				ReceiveHelper(StateHistory.Get(), AsyncInput->StateData.Get(), /*bImportant*/false, ComponentSettings.GetCompareStateToTriggerRewind());
+			}
+
+			// Receive Important Inputs
+			for (const TUniquePtr<Chaos::FBaseRewindHistory>& InputImportant : AsyncInput->InputDataImportant)
+			{
+				if (!InputImportant || !InputImportant->HasDataInHistory())
+				{
+					continue;
+				}
+				ReceiveHelper(InputHistory.Get(), InputImportant.Get(), /*bImportant*/true, ComponentSettings.GetCompareInputToTriggerRewind());
+			}
+
+			// Receive Important States
+			for (const TUniquePtr<Chaos::FBaseRewindHistory>& StateImportant : AsyncInput->StateDataImportant)
+			{
+				if (!StateImportant || !StateImportant->HasDataInHistory())
+				{
+					continue;
+				}
+				ReceiveHelper(StateHistory.Get(), StateImportant.Get(), /*bImportant*/true, ComponentSettings.GetCompareStateToTriggerRewind());
+			}
+		}
+	}
+}
+
+FAsyncNetworkPhysicsComponentOutput& FAsyncNetworkPhysicsComponent::GetAsyncOutput_Internal()
+{
+	FAsyncNetworkPhysicsComponentOutput& AsyncOutput = GetProducerOutputData_Internal();
+	const FNetworkPhysicsSettingsNetworkPhysicsComponent& ComponentSettings = GetComponentSettings();
+
+	// InputData marshal from PT to GT is needed for: LocallyControlled and Server
+	if ((IsLocallyControlled() || IsServer()) && AsyncOutput.InputData == nullptr && InputHistory != nullptr)
+	{
+		AsyncOutput.InputData = InputHistory->CreateNew();
+		AsyncOutput.InputData->ResizeDataHistory(ComponentSettings.GetRedundantInputs());
+	}
+
+	// StateData marshal from PT to GT is needed for: Server
+	if (IsServer() && AsyncOutput.StateData == nullptr && StateHistory != nullptr)
+	{
+		AsyncOutput.StateData = StateHistory->CreateNew();
+		AsyncOutput.StateData->ResizeDataHistory(ComponentSettings.GetRedundantStates());
+	}
+
+	return AsyncOutput;
+}
+
+void FAsyncNetworkPhysicsComponent::OnPreProcessInputs_Internal(const int32 PhysicsStep)
+{
+	ConsumeAsyncInput();
+
+	const FNetworkPhysicsSettingsNetworkPhysicsComponent& ComponentSettings = GetComponentSettings();
+	const bool bIsServer = IsServer();
+
+	bool bIsSolverReset = false;
+	bool bIsSolverResim = false;
+	if (Chaos::FPBDRigidsEvolution* Evolution = GetEvolution())
+	{
+		bIsSolverResim = Evolution->IsResimming();
+		bIsSolverReset = Evolution->IsResetting();
+	}
+
+#if DEBUG_NETWORK_PHYSICS
+	{
+		const int32 InputBufferSize = (IsServer() && InputHistory) ? (InputHistory->GetLatestFrame() - PhysicsStep) : 0;
+		const FString NetRoleString = IsServer() ? FString("SERVER") : (IsLocallyControlled() ? FString("AUTONO") : FString("PROXY "));
+		UE_LOG(LogChaos, Log, TEXT("%s | PT | OnPreProcessInputsInternal | At Frame %d | IsResim: %d | FirstResimFrame: %d | InputBuffer: %d | Name = %s"), *NetRoleString, PhysicsStep, bIsSolverResim, bIsSolverReset, InputBufferSize, *GetActorName());
+	}
+#endif
+
+	if (ActorComponent)
+	{
+		// Apply replicated state on clients if we are resimulating
+		if (bIsSolverResim && StateHistory && StateData)
+		{
+			FNetworkPhysicsData* PhysicsData = StateData.Get();
+			PhysicsData->LocalFrame = PhysicsStep;
+			const bool bExactFrame = PhysicsReplicationCVars::ResimulationCVars::bAllowRewindToClosestState ? !bIsSolverReset : true;
+			if (StateHistory->ExtractData(PhysicsStep, bIsSolverReset, PhysicsData, bExactFrame) && PhysicsData->bReceivedData)
+			{
+				PhysicsData->ApplyData(ActorComponent);
+#if DEBUG_NETWORK_PHYSICS
+				UE_LOG(LogChaos, Log, TEXT("		Applying extracted state from history | bExactFrame = %d | LocalFrame = %d | ServerFrame = %d | InputFrame = %d | Data: %s")
+					, bExactFrame, PhysicsData->LocalFrame, PhysicsData->ServerFrame, PhysicsData->InputFrame, *PhysicsData->DebugData());
+#endif
+			}
+#if DEBUG_NETWORK_PHYSICS
+			else if (PhysicsStep <= StateHistory->GetLatestFrame())
+			{
+				UE_LOG(LogChaos, Log, TEXT("		FAILED to extract and apply state from history | bExactFrame = %d | -- Printing history --"), bExactFrame);
+				StateHistory->DebugData(FString::Printf(TEXT("StateHistory | Component = %s"), *GetActorName()));
 			}
 #endif
 		}
 
 		// Apply replicated inputs on server and simulated proxies (and on local player if we are resimulating)
-		if (InputHistory && (!IsLocallyControlled() || bIsSolverResim))
+		if (InputHistory && InputData && (!IsLocallyControlled() || bIsSolverResim))
 		{
 			FNetworkPhysicsData* PhysicsData = InputData.Get();
 			int32 NextExpectedLocalFrame = PhysicsData->LocalFrame + 1;
@@ -1083,7 +1290,7 @@ void UNetworkPhysicsComponent::OnPreProcessInputsInternal(const int32 PhysicsSte
 			// There are important inputs earlier than upcoming input to apply
 			if (NewImportantInputFrame < NextExpectedLocalFrame && !bIsSolverResim)
 			{
-				if (bApplyDataInsteadOfMergeData)
+				if (ComponentSettings.GetApplyDataInsteadOfMergeData())
 				{
 #if DEBUG_NETWORK_PHYSICS
 					UE_LOG(LogChaos, Log, TEXT("	Reapplying multiple data due to receiving an important data that was previously missed. FromFrame: %d | ToFrame: %d | IsLocallyControlled = %d"), NewImportantInputFrame, (NextExpectedLocalFrame - 1), IsLocallyControlled());
@@ -1101,9 +1308,7 @@ void UNetworkPhysicsComponent::OnPreProcessInputsInternal(const int32 PhysicsSte
 				}
 			}
 
-			PRAGMA_DISABLE_DEPRECATION_WARNINGS // TODO: Change to ExtractData() in UE 5.6 and remove deprecation pragma
-			if (InputHistory->ExtractDatas(PhysicsStep, bIsSolverReset, PhysicsData, /*bExactFrame*/(bAllowInputExtrapolation == false)))
-			PRAGMA_ENABLE_DEPRECATION_WARNINGS
+			if (InputHistory->ExtractData(PhysicsStep, bIsSolverReset, PhysicsData, /*bExactFrame*/(ComponentSettings.GetAllowInputExtrapolation() == false)))
 			{
 				// Calculate input decay if we are resimulating and we don't have up to date inputs
 				if (bIsSolverResim)
@@ -1114,59 +1319,63 @@ void UNetworkPhysicsComponent::OnPreProcessInputsInternal(const int32 PhysicsSte
 						PhysicsData->DecayData(InputDecay);
 					}
 				}
-				// Apply input if data is unimportant or if data is important and unaltered (extrapolated, interpolated)
-				else if (!PhysicsData->bImportant || PhysicsData->InputFrame != INDEX_NONE)
+				// Check if we have a gap between last used input and current input
+				else if (PhysicsData->LocalFrame > NextExpectedLocalFrame)
 				{
-					// Check if we have a gap between last used input and current input
-					if (PhysicsData->LocalFrame > NextExpectedLocalFrame)
+					if (ComponentSettings.GetApplyDataInsteadOfMergeData())
 					{
-						if (bApplyDataInsteadOfMergeData)
-						{
 #if DEBUG_NETWORK_PHYSICS
-							UE_LOG(LogChaos, Log, TEXT("	Applying multiple data instead of merging, from LocalFrame %d into LocalFrame %d | IsLocallyControlled = %d"), NextExpectedLocalFrame, PhysicsData->LocalFrame, IsLocallyControlled());
+						UE_LOG(LogChaos, Log, TEXT("	Applying multiple data instead of merging, from LocalFrame %d into LocalFrame %d | IsLocallyControlled = %d"), NextExpectedLocalFrame, PhysicsData->LocalFrame, IsLocallyControlled());
 #endif
-							// Iterate over each input and call ApplyData, except on the last, it will get handled by the normal ApplyData call further down
-							const int32 LastFrame = PhysicsData->LocalFrame;
-							for (; NextExpectedLocalFrame <= LastFrame; NextExpectedLocalFrame++)
+						// Iterate over each input and call ApplyData, except on the last, it will get handled by the normal ApplyData call further down
+						const int32 LastFrame = PhysicsData->LocalFrame;
+						for (; NextExpectedLocalFrame <= LastFrame; NextExpectedLocalFrame++)
+						{
+							if (InputHistory->ExtractData(NextExpectedLocalFrame, bIsSolverReset, PhysicsData, true) && NextExpectedLocalFrame < LastFrame)
 							{
-								if (InputHistory->ExtractData(NextExpectedLocalFrame, bIsSolverReset, PhysicsData, true) && NextExpectedLocalFrame < LastFrame)
-								{
-									PhysicsData->ApplyData(ActorComponent);
-								}
+								PhysicsData->ApplyData(ActorComponent);
 							}
 						}
-						else
-						{
+					}
+					else
+					{
 #if DEBUG_NETWORK_PHYSICS
-							UE_LOG(LogChaos, Log, TEXT("	Merging inputs from LocalFrame %d into LocalFrame %d | IsLocallyControlled = %d"), NextExpectedLocalFrame, PhysicsData->LocalFrame, IsLocallyControlled());
+						UE_LOG(LogChaos, Log, TEXT("	Merging inputs from LocalFrame %d into LocalFrame %d | IsLocallyControlled = %d"), NextExpectedLocalFrame, PhysicsData->LocalFrame, IsLocallyControlled());
 #endif
-							// Merge all inputs since last used input
-							InputHistory->MergeData(NextExpectedLocalFrame, PhysicsData);
-						}
+						// Merge all inputs since last used input
+						InputHistory->MergeData(NextExpectedLocalFrame, PhysicsData);
 					}
+				}
 
-					// If the extracted input data was altered (extrapolated, interpolated) on the server, record it into the history for it to get replicated to clients
-					if (bIsServer && !bIsSolverResim && PhysicsData->InputFrame == INDEX_NONE)
+				// If the extracted input data was altered (extrapolated, interpolated) on the server, record it into the history for it to get replicated to clients
+				if (bIsServer && PhysicsData->InputFrame == INDEX_NONE)
+				{
+#if DEBUG_NETWORK_PHYSICS
+					if (PhysicsStep > InputHistory->GetLatestFrame())
 					{
-						PhysicsData->bReceivedData = true; // Mark the input data as received so that it doesn't get overwritten by incoming client inputs
-						PhysicsData->bImportant = false;
-						PhysicsData->LocalFrame = PhysicsStep;
-						InputHistory->RecordData(PhysicsStep, PhysicsData);
+						UE_LOG(LogChaos, Log, TEXT("	Input buffer is empty, input for frame %d was extrapolated from frame: %d"), PhysicsStep, PhysicsData->LocalFrame);
 					}
+#endif
+					PhysicsData->bReceivedData = true; // Mark the input data as received so that it doesn't get overwritten by incoming client inputs
+					PhysicsData->bImportant = false;
+					PhysicsData->LocalFrame = PhysicsStep;
+					InputHistory->RecordData(PhysicsStep, PhysicsData);
 				}
 
 				PhysicsData->ApplyData(ActorComponent);
 
 #if DEBUG_NETWORK_PHYSICS
-				UE_LOG(LogChaos, Log, TEXT("		Applying extracted input from history | LocalFrame = %d | ServerFrame = %d | InputFrame = %d | IsResim = %d | IsLocallyControlled = %d | InputDecay = %f | Data: %s")
-					, PhysicsData->LocalFrame, PhysicsData->ServerFrame, PhysicsData->InputFrame, bIsSolverResim, IsLocallyControlled(), GetCurrentInputDecay(PhysicsData), *PhysicsData->DebugData());
+				{
+					UE_LOG(LogChaos, Log, TEXT("		Applying extracted input from history | LocalFrame = %d | ServerFrame = %d | InputFrame = %d | IsResim = %d | IsLocallyControlled = %d | InputDecay = %f | Data: %s")
+						, PhysicsData->LocalFrame, PhysicsData->ServerFrame, PhysicsData->InputFrame, bIsSolverResim, IsLocallyControlled(), GetCurrentInputDecay(PhysicsData), *PhysicsData->DebugData());
+				}
 #endif
 			}
 #if DEBUG_NETWORK_PHYSICS
-			else
+			else if (PhysicsStep <= InputHistory->GetLatestFrame())
 			{
 				UE_LOG(LogChaos, Log, TEXT("		FAILED to extract and apply input from history | IsResim = %d | IsLocallyControlled = %d | -- Printing history --"), bIsSolverResim, IsLocallyControlled());
-				InputHistory->DebugData(FString::Printf(TEXT("InputHistory | Component = %s"), *GetFullName()));
+				InputHistory->DebugData(FString::Printf(TEXT("InputHistory | Name = %s"), *GetActorName()));
 			}
 #endif
 		}
@@ -1174,52 +1383,41 @@ void UNetworkPhysicsComponent::OnPreProcessInputsInternal(const int32 PhysicsSte
 	NewImportantInputFrame = INT_MAX;
 }
 
-void UNetworkPhysicsComponent::OnPostProcessInputsInternal(const int32 PhysicsStep)
+void FAsyncNetworkPhysicsComponent::OnPostProcessInputs_Internal(const int32 PhysicsStep)
 {
-	const bool bIsServer = HasServerWorld();
+	const FNetworkPhysicsSettingsNetworkPhysicsComponent& ComponentSettings = GetComponentSettings();
+	const bool bIsServer = IsServer();
 
 	bool bIsSolverReset = false;
 	bool bIsSolverResim = false;
-#if DEBUG_NETWORK_PHYSICS
-	if (bIsServer)
+	if (Chaos::FPBDRigidsEvolution* Evolution = GetEvolution())
 	{
-		UE_LOG(LogChaos, Log, TEXT("SERVER | PT | OnPostProcessInputsInternal | At Frame %d | Component = %s"), PhysicsStep, *GetFullName());
+		bIsSolverResim = Evolution->IsResimming();
+		bIsSolverReset = Evolution->IsResetting();
 	}
-	else
+
+#if DEBUG_NETWORK_PHYSICS
 	{
-		UE_LOG(LogChaos, Log, TEXT("CLIENT | PT | OnPostProcessInputsInternal | At Frame %d | Component = %s"), PhysicsStep, *GetFullName());
+		FString NetRoleString = IsServer() ? FString("SERVER") : (IsLocallyControlled() ? FString("AUTONO") : FString("PROXY "));
+		UE_LOG(LogChaos, Log, TEXT("%s | PT | OnPostProcessInputsInternal | At Frame %d | IsResim: %d | FirstResimFrame: %d | Name = %s"), *NetRoleString, PhysicsStep, bIsSolverResim, bIsSolverReset, *GetActorName());
 	}
 #endif
 
 	if (ActorComponent)
 	{
-		if (FPhysScene* PhysScene = GetWorld()->GetPhysicsScene())
-		{
-			if (Chaos::FPhysicsSolver* Solver = PhysScene->GetSolver())
-			{
-				bIsSolverResim = Solver->GetEvolution()->IsResimming();
-				bIsSolverReset = Solver->GetEvolution()->IsResetting();
-			}
-		}
-
-		APlayerController* PlayerController = GetPlayerController();
-		if (!PlayerController)
-		{
-			PlayerController = GetWorld()->GetFirstPlayerController();
-		}
-
 		const bool bShouldCacheInputHistory = IsLocallyControlled() && !bIsSolverResim;
-		// for the inputs client local ones are ground truth otherwise use the replicated ones coming from the server
-		if (PlayerController && bShouldCacheInputHistory && (InputData != nullptr))
+		// For the inputs client local ones are ground truth otherwise use the replicated ones coming from the server
+		if (bShouldCacheInputHistory && (InputData != nullptr))
 		{
+			// Prepare to gather input data
 			FNetworkPhysicsData* PhysicsData = InputData.Get();
-			PhysicsData->PrepareFrame(PhysicsStep, bIsServer, PlayerController->GetNetworkPhysicsTickOffset());
+			PhysicsData->PrepareFrame(PhysicsStep, bIsServer, GetNetworkPhysicsTickOffset());
 
+			// Gather input data from implementation
 			PhysicsData->BuildData(ActorComponent);
 
-			PRAGMA_DISABLE_DEPRECATION_WARNINGS // TODO: Change to RecordData() in UE 5.6 and remove deprecation pragma
-			InputHistory->RecordDatas(PhysicsStep, PhysicsData);
-			PRAGMA_ENABLE_DEPRECATION_WARNINGS
+			// Record input in history
+			InputHistory->RecordData(PhysicsStep, PhysicsData);
 
 #if DEBUG_NETWORK_PHYSICS
 			{
@@ -1229,30 +1427,30 @@ void UNetworkPhysicsComponent::OnPostProcessInputsInternal(const int32 PhysicsSt
 #endif
 		}
 
-		const bool bShouldCacheStateHistory = bIsServer || (bCompareStateToTriggerRewind && bShouldCacheInputHistory);
-		if (PlayerController && StateHistory && bShouldCacheStateHistory)
+		const bool bShouldCacheStateHistory = bIsServer || (ComponentSettings.GetCompareStateToTriggerRewind() && bShouldCacheInputHistory);
+		if (StateHistory && StateData && bShouldCacheStateHistory)
 		{
 			// Compute of the local frame coming from the client that was used to generate this state
 			int32 InputFrame = INDEX_NONE;
+			if (InputData)
 			{
 				FNetworkPhysicsData* PhysicsData = InputData.Get();
-				PRAGMA_DISABLE_DEPRECATION_WARNINGS // TODO: Change to ExtractData() in UE 5.6 and remove deprecation pragma
-				if (InputHistory && InputHistory->ExtractDatas(PhysicsStep, false, PhysicsData, true))
-				PRAGMA_ENABLE_DEPRECATION_WARNINGS
+				if (InputHistory && InputHistory->ExtractData(PhysicsStep, false, PhysicsData, true))
 				{
 					InputFrame = PhysicsData->InputFrame;
 				}
 			}
 
+			// Prepare to gather state data
 			FNetworkPhysicsData* PhysicsData = StateData.Get();
-			PhysicsData->PrepareFrame(PhysicsStep, bIsServer, PlayerController->GetNetworkPhysicsTickOffset());
+			PhysicsData->PrepareFrame(PhysicsStep, bIsServer, GetNetworkPhysicsTickOffset());
 			PhysicsData->InputFrame = InputFrame;
 
+			// Gather input data from implementation
 			PhysicsData->BuildData(ActorComponent);
 
-			PRAGMA_DISABLE_DEPRECATION_WARNINGS // TODO: Change to RecordData() in UE 5.6 and remove deprecation pragma
-			StateHistory->RecordDatas(PhysicsStep, PhysicsData);
-			PRAGMA_ENABLE_DEPRECATION_WARNINGS
+			// Record input in history
+			StateHistory->RecordData(PhysicsStep, PhysicsData);
 
 #if DEBUG_NETWORK_PHYSICS
 			{
@@ -1262,188 +1460,263 @@ void UNetworkPhysicsComponent::OnPostProcessInputsInternal(const int32 PhysicsSt
 #endif
 		}
 	}
+
+	// Marshal inputs and states from PT to GT for networking
+	FAsyncNetworkPhysicsComponentOutput& AsyncOutput = GetAsyncOutput_Internal();
+	SendInputData_Internal(AsyncOutput);
+	SendStateData_Internal(AsyncOutput);
+	FinalizeOutputData_Internal();
 }
 
-const float UNetworkPhysicsComponent::GetCurrentInputDecay(FNetworkPhysicsData* PhysicsData)
+void FAsyncNetworkPhysicsComponent::SendInputData_Internal(FAsyncNetworkPhysicsComponentOutput& AsyncOutput)
 {
-	if (!PhysicsData)
-	{
-		return 0.0f;
-	}
+	const bool bIsServer = IsServer();
 
-	FPhysScene* PhysScene = GetWorld()->GetPhysicsScene();
-	if (!PhysScene)
+	// Inputs are sent from the server or locally controlled actors/pawns
+	if (AsyncOutput.InputData && InputHistory && (IsLocallyControlled() || bIsServer))
 	{
-		return 0.0f;
-	}
+		const FNetworkPhysicsSettingsNetworkPhysicsComponent& ComponentSettings = GetComponentSettings();
 
-	Chaos::FPhysicsSolver* Solver = PhysScene->GetSolver();
-	if (!Solver)
-	{
-		return 0.0f;
-	}
+		// Send latest N frames from history
+		const int32 ToFrame = FMath::Max(0, InputHistory->GetLatestFrame());
 
-	Chaos::FRewindData* RewindData = Solver->GetRewindData();
-	if (!RewindData)
-	{
-		return 0.0f;
-	}
-	
-	const float NumPredictedInputs = RewindData->CurrentFrame() - PhysicsData->LocalFrame; // Number of frames we have used the same PhysicsData for during resim
-	const float MaxPredictedInputs = RewindData->GetLatestFrame() - 1 - PhysicsData->LocalFrame; // Max number of frames PhysicsData registered frame until end of resim
+		// -- Default / Unreliable Flow --
+		if (ComponentSettings.GetEnableUnreliableFlow())
+		{
+			const uint8 InputsToSend = ComponentSettings.GetRedundantInputs();
+			const int32 FromFrame = FMath::Max(0, ToFrame - InputsToSend - 1); // Remove 1 since both ToFrame and FromFrame are inclusive
 
-	// Linear decay
-	const float PredictionAlpha = MaxPredictedInputs > 0 ? (NumPredictedInputs / MaxPredictedInputs) : 0.0f;
+			// Resize marshaling history if needed
+			AsyncOutput.InputData->ResizeDataHistory(InputsToSend);
+			
+			if (InputHistory->CopyData(*AsyncOutput.InputData, FromFrame, ToFrame, /*bIncludeUnimportant*/ true, /*bIncludeImportant*/ ComponentSettings.GetEnableReliableFlow() == false))
+			{
+#if DEBUG_NETWORK_PHYSICS
+				{
+					const int32 LocalFrame = GetRigidSolver()->GetCurrentFrame();
+					const int32 ServerFrame = IsServer() ? LocalFrame : LocalFrame + GetNetworkPhysicsTickOffset();
+					FString NetRoleString = IsServer() ? FString("SERVER") : (IsLocallyControlled() ? FString("AUTONO") : FString("PROXY "));
+					AsyncOutput.InputData->DebugData(FString::Printf(TEXT("%s | PT | SendInputData_Internal | UNRELIABLE | CurrentLocalFrame = %d | CurrentServerFrame = %d | Name: %s"), *NetRoleString, LocalFrame, ServerFrame, *GetActorName()));
+				}
+#endif
+			}
+		}
 
-	return PredictionAlpha;
-}
+		// -- Important / Reliable flow --
+		if (ComponentSettings.GetEnableReliableFlow())
+		{
+			/* Get the latest valid frame that can hold new important data:
+			* 1. Frame after last time we called SendInputData
+			* 2. Earliest possible frame in history */
+			const int32 FromFrame = FMath::Max(LastInputSendFrame + 1, ToFrame - InputHistory->GetHistorySize());
 
-bool UNetworkPhysicsComponent::HasServerWorld() const
-{
-	return GetWorld()->IsNetMode(NM_DedicatedServer) || GetWorld()->IsNetMode(NM_ListenServer);
-}
+			// Check if we have important data to marshal
+			if (const int32 Count = InputHistory->CountValidData(FromFrame, ToFrame, /*bIncludeUnimportant*/ false, /*bIncludeImportant*/ true) > 0)
+			{
+				// Create new data collection for marshaling
+				const int32 Idx = AsyncOutput.InputDataImportant.Add(InputHistory->CreateNew());
+				AsyncOutput.InputDataImportant[Idx]->ResizeDataHistory(Count);
 
-bool UNetworkPhysicsComponent::HasLocalController() const
-{
-	if (APlayerController* PlayerController = GetPlayerController())
-	{
-		return PlayerController->IsLocalController();
+				// Copy over data
+				if (InputHistory->CopyData(*AsyncOutput.InputDataImportant[Idx], FromFrame, ToFrame, /*bIncludeUnimportant*/ false, /*bIncludeImportant*/ true))
+				{
+#if DEBUG_NETWORK_PHYSICS
+					{
+						const int32 LocalFrame = GetRigidSolver()->GetCurrentFrame();
+						const int32 ServerFrame = IsServer() ? LocalFrame : LocalFrame + GetNetworkPhysicsTickOffset();
+						FString NetRoleString = IsServer() ? FString("SERVER") : (IsLocallyControlled() ? FString("AUTONO") : FString("PROXY "));
+						AsyncOutput.InputDataImportant[Idx]->DebugData(FString::Printf(TEXT("%s | PT | SendInputData_Internal | RELIABLE | CurrentLocalFrame = %d | CurrentServerFrame = %d | Name: %s"), *NetRoleString, LocalFrame, ServerFrame, *GetActorName()));
+					}
+#endif
+				}
+				
+			}
+		}
+		LastInputSendFrame = InputHistory->GetLatestFrame();
 	}
-	return false;
 }
 
-bool UNetworkPhysicsComponent::IsLocallyControlled() const
+void FAsyncNetworkPhysicsComponent::SendStateData_Internal(FAsyncNetworkPhysicsComponentOutput& AsyncOutput)
 {
-	if (bIsRelayingLocalInputs)
+	if (IsServer() && StateHistory && AsyncOutput.StateData)
 	{
-		return true;
-	}
+		const FNetworkPhysicsSettingsNetworkPhysicsComponent& ComponentSettings = GetComponentSettings();
 
-	if (APlayerController* PlayerController = GetPlayerController())
-	{
-		return PlayerController->IsLocalController();
-	}
-	return false;
-}
-
-APlayerController* UNetworkPhysicsComponent::GetPlayerController() const
-{
-	if (APlayerController* PC = Cast<APlayerController>(GetOwner()))
-	{
-		return PC;
-	}
+		// Send latest N frames from history
+		const int32 ToFrame = FMath::Max(0, StateHistory->GetLatestFrame());
 
-	if (APawn* Pawn = Cast<APawn>(GetOwner()))
-	{
-		if (APlayerController * PC = Pawn->GetController<APlayerController>())
+		// -- Default / Unreliable Flow --
+		if (ComponentSettings.GetEnableUnreliableFlow())
 		{
-			return PC;
+			const uint8 StatesToSend = ComponentSettings.GetRedundantStates();
+			const int32 FromFrame = FMath::Max(0, ToFrame - StatesToSend - 1); // Remove 1 since both ToFrame and FromFrame are inclusive
+
+			// Resize marshaling history if needed
+			AsyncOutput.StateData->ResizeDataHistory(StatesToSend);
+
+			if (StateHistory->CopyData(*AsyncOutput.StateData, FromFrame, ToFrame, /*bIncludeUnimportant*/ true, /*bIncludeImportant*/ ComponentSettings.GetEnableReliableFlow() == false))
+			{
+#if DEBUG_NETWORK_PHYSICS
+				{
+					const int32 LocalFrame = GetRigidSolver()->GetCurrentFrame();
+					const int32 ServerFrame = IsServer() ? LocalFrame : LocalFrame + GetNetworkPhysicsTickOffset();
+					AsyncOutput.StateData->DebugData(FString::Printf(TEXT("SERVER | PT | SendStateData_Internal | UNRELIABLE | CurrentLocalFrame = %d | CurrentServerFrame = %d | Name: %s"), LocalFrame, ServerFrame, *GetActorName()));
+				}
+#endif
+			}
 		}
 
-		// In this case the APlayerController can be found as the owner of the pawn
-		if (APlayerController* PC = Cast<APlayerController>(Pawn->GetOwner()))
+		// -- Important / Reliable flow --
+		if (ComponentSettings.GetEnableReliableFlow())
 		{
-			return PC;
-		}
+			/* Get the latest valid frame that can hold new important data:
+			* 1. Frame after last time we called SendStateData
+			* 2. Earliest possible frame in history */
+			const int32 FromFrame = FMath::Max(LastStateSendFrame + 1, ToFrame - StateHistory->GetHistorySize());
 
+			// Check if we have important data to marshal
+			if (const int32 Count = StateHistory->CountValidData(FromFrame, ToFrame, /*bIncludeUnimportant*/ false, /*bIncludeImportant*/ true) > 0)
+			{
+				// Create new data collection for marshaling
+				const int32 Idx = AsyncOutput.StateDataImportant.Add(StateHistory->CreateNew());
+				AsyncOutput.StateDataImportant[Idx]->ResizeDataHistory(Count);
+
+				// Copy over data
+				if (StateHistory->CopyData(*AsyncOutput.StateDataImportant[Idx], FromFrame, ToFrame, /*bIncludeUnimportant*/ false, /*bIncludeImportant*/ true))
+				{
+#if DEBUG_NETWORK_PHYSICS
+					{
+						const int32 LocalFrame = GetRigidSolver()->GetCurrentFrame();
+						const int32 ServerFrame = IsServer() ? LocalFrame : LocalFrame + GetNetworkPhysicsTickOffset();
+						AsyncOutput.StateDataImportant[Idx]->DebugData(FString::Printf(TEXT("SERVER | PT | SendStateData_Internal | RELIABLE | CurrentLocalFrame = %d | CurrentServerFrame = %d | Name: %s"), LocalFrame, ServerFrame, *GetActorName()));
+					}
+#endif
+				}
+			}
+		}
+		LastStateSendFrame = StateHistory->GetLatestFrame();
 	}
+}
 
-	return nullptr;
+Chaos::FPBDRigidsSolver* FAsyncNetworkPhysicsComponent::GetRigidSolver()
+{
+	return static_cast<Chaos::FPBDRigidsSolver*>(GetSolver());
 }
 
-void UNetworkPhysicsComponent::RemoveDataHistory()
+Chaos::FPBDRigidsEvolution* FAsyncNetworkPhysicsComponent::GetEvolution()
 {
-	if (GetWorld())
+	if (Chaos::FPBDRigidsSolver* RigidSolver = GetRigidSolver())
 	{
-		if (FPhysScene* PhysScene = GetWorld()->GetPhysicsScene())
-		{
-			if (Chaos::FPhysicsSolver* Solver = PhysScene->GetSolver())
-			{
-				if (Chaos::FRewindData* RewindData = Solver->GetRewindData())
-				{
-					RewindData->RemoveInputHistory(InputHistory);
-					RewindData->RemoveStateHistory(StateHistory);
-				}
-			}
-		}
+		return RigidSolver->GetEvolution();
 	}
+	return nullptr;
 }
-void UNetworkPhysicsComponent::AddDataHistory()
+
+void FAsyncNetworkPhysicsComponent::TriggerResimulation(int32 ResimFrame)
 {
-	if (FPhysScene* PhysScene = GetWorld()->GetPhysicsScene())
+	if (ResimFrame != INDEX_NONE)
 	{
-		if (Chaos::FPhysicsSolver* Solver = PhysScene->GetSolver())
+		if (Chaos::FPBDRigidsSolver* RigidSolver = GetRigidSolver())
 		{
-			if (Chaos::FRewindData* RewindData = Solver->GetRewindData())
+			if (Chaos::FRewindData* RewindData = RigidSolver->GetRewindData())
 			{
-				RewindData->AddInputHistory(InputHistory);
-				if (StateHistory)
+				// Mark particle/island as resim
+				Chaos::FReadPhysicsObjectInterface_Internal Interface = Chaos::FPhysicsObjectInternalInterface::GetRead();
+				if (Chaos::FPBDRigidParticleHandle* POHandle = Interface.GetRigidParticle(RootPhysicsObject))
 				{
-					RewindData->AddStateHistory(StateHistory);
+					if (Chaos::FPBDRigidsEvolution* Evolution = GetEvolution())
+					{
+						Evolution->GetIslandManager().SetParticleResimFrame(POHandle, ResimFrame);
+					}
 				}
+
+				// Set resim frame in rewind data
+				ResimFrame = (RewindData->GetResimFrame() == INDEX_NONE) ? ResimFrame : FMath::Min(ResimFrame, RewindData->GetResimFrame());
+				RewindData->SetResimFrame(ResimFrame);
 			}
 		}
 	}
 }
 
-int32 UNetworkPhysicsComponent::SetupRewindData()
+const float FAsyncNetworkPhysicsComponent::GetCurrentInputDecay(const FNetworkPhysicsData* PhysicsData)
 {
-	int32 NumFrames = UPhysicsSettings::Get()->GetPhysicsHistoryCount();
+	if (!PhysicsData)
+	{
+		return 0.0f;
+	}
+	
+	Chaos::FPhysicsSolverBase* BaseSolver = GetSolver();
+	if (!BaseSolver)
+	{
+		return 0.0f;
+	}
 
-	AActor* Owner = GetOwner();
-	// Don't let this actor initialize RewindData if the actor is not using resimulation
-	if (Owner && Owner->GetPhysicsReplicationMode() != EPhysicsReplicationMode::Resimulation)
+	Chaos::FRewindData* RewindData = BaseSolver->GetRewindData();
+	if (!RewindData)
 	{
-		return NumFrames;
+		return 0.0f;
 	}
 
-	if (FPhysScene* PhysScene = GetWorld()->GetPhysicsScene())
+	const float NumPredictedInputs = RewindData->CurrentFrame() - PhysicsData->LocalFrame; // Number of frames we have used the same PhysicsData for during resim
+	const float MaxPredictedInputs = RewindData->GetLatestFrame() - 1 - PhysicsData->LocalFrame; // Max number of frames PhysicsData registered frame until end of resim
+
+	// Linear decay
+	const float PredictionAlpha = MaxPredictedInputs > 0 ? (NumPredictedInputs / MaxPredictedInputs) : 0.0f;
+
+	return PredictionAlpha;
+}
+
+void FAsyncNetworkPhysicsComponent::RegisterDataHistoryInRewindData()
+{
+	if (Chaos::FPhysicsSolverBase* BaseSolver = GetSolver())
 	{
-		if (Chaos::FPhysicsSolver* Solver = PhysScene->GetSolver())
+		if (Chaos::FRewindData* RewindData = BaseSolver->GetRewindData())
 		{
-			if (UPhysicsSettings::Get()->PhysicsPrediction.bEnablePhysicsPrediction && Solver->GetRewindData() == nullptr)
-			{
-				Solver->EnableRewindCapture(NumFrames, true);
-			}
+			UnregisterDataHistoryFromRewindData();
 
-			if (Chaos::FRewindData* RewindData = Solver->GetRewindData())
+			RewindData->AddInputHistory(InputHistory);
+			if (StateHistory)
 			{
-				NumFrames = RewindData->Capacity();
+				RewindData->AddStateHistory(StateHistory);
 			}
 		}
 	}
-
-	return NumFrames;
 }
 
-void UNetworkPhysicsComponent::TriggerResimulation(int32 ResimFrame)
+void FAsyncNetworkPhysicsComponent::UnregisterDataHistoryFromRewindData()
 {
-	if (ResimFrame != INDEX_NONE)
+	if (Chaos::FPhysicsSolverBase* BaseSolver = GetSolver())
 	{
-		if (FPhysScene* PhysScene = GetWorld()->GetPhysicsScene())
+		if (Chaos::FRewindData* RewindData = BaseSolver->GetRewindData())
 		{
-			if (Chaos::FPhysicsSolver* Solver = PhysScene->GetSolver())
-			{
-				if (Chaos::FRewindData* RewindData = Solver->GetRewindData())
-				{
-					// Mark particle/island as resim
-					Chaos::FReadPhysicsObjectInterface_Internal Interface = Chaos::FPhysicsObjectInternalInterface::GetRead();
-					if (Chaos::FPBDRigidParticleHandle* POHandle = Interface.GetRigidParticle(RootPhysicsObject))
-					{
-						Solver->GetEvolution()->GetIslandManager().SetParticleResimFrame(POHandle, ResimFrame);
-					}
-
-					// Set resim frame in rewind data
-					ResimFrame = (RewindData->GetResimFrame() == INDEX_NONE) ? ResimFrame : FMath::Min(ResimFrame, RewindData->GetResimFrame());
-					RewindData->SetResimFrame(ResimFrame);
-				}
-			}
+			RewindData->RemoveInputHistory(InputHistory);
+			RewindData->RemoveStateHistory(StateHistory);
 		}
 	}
 }
 
+const int32 FAsyncNetworkPhysicsComponent::SetupRewindData()
+{
+	int32 NumFrames = 0;
 
+	if (Chaos::FPBDRigidsSolver* RigidSolver = GetRigidSolver())
+	{
+		NumFrames = RigidSolver->GetPhysicsHistoryCount();
 
+		// Don't let this actor initialize RewindData if not using resimulation
+		if (GetPhysicsReplicationMode() == EPhysicsReplicationMode::Resimulation)
+		{
+			if (RigidSolver->IsNetworkPhysicsPredictionEnabled() && RigidSolver->GetRewindData() == nullptr)
+			{
+				RigidSolver->EnableRewindCapture(NumFrames, true);
+			}
+		}
 
+		if (Chaos::FRewindData* RewindData = RigidSolver->GetRewindData())
+		{
+			NumFrames = RewindData->Capacity();
+		}
+	}
 
-
+	return NumFrames;
+}
diff --git a/Engine/Source/Runtime/Engine/Private/PhysicsEngine/NetworkPhysicsSettingsComponent.cpp b/Engine/Source/Runtime/Engine/Private/PhysicsEngine/NetworkPhysicsSettingsComponent.cpp
--- a/Engine/Source/Runtime/Engine/Private/PhysicsEngine/NetworkPhysicsSettingsComponent.cpp
+++ b/Engine/Source/Runtime/Engine/Private/PhysicsEngine/NetworkPhysicsSettingsComponent.cpp
@@ -36,17 +36,17 @@ void UNetworkPhysicsSettingsComponent::InitializeComponent()
 	Super::InitializeComponent();
 
 	using namespace Chaos;
-	NetworkPhysicsSettingsAsync = nullptr;
+	NetworkPhysicsSettings_Internal = nullptr;
 	if (UWorld* World = GetWorld())
 	{
 		if (FPhysScene* PhysScene = World->GetPhysicsScene())
 		{
 			if (Chaos::FPhysicsSolver* Solver = PhysScene->GetSolver())
 			{
-				NetworkPhysicsSettingsAsync = Solver->CreateAndRegisterSimCallbackObject_External<FNetworkPhysicsSettingsComponentAsync>();
+				NetworkPhysicsSettings_Internal = Solver->CreateAndRegisterSimCallbackObject_External<FNetworkPhysicsSettingsComponentAsync>();
 				
 				// Marshal settings data from GT to PT
-				if (NetworkPhysicsSettingsAsync)
+				if (NetworkPhysicsSettings_Internal)
 				{
 					if (AActor* Owner = GetOwner())
 					{
@@ -54,7 +54,7 @@ void UNetworkPhysicsSettingsComponent::InitializeComponent()
 						{
 							if (Chaos::FConstPhysicsObjectHandle PhysicsObject = RootPrimComp->GetPhysicsObjectByName(NAME_None))
 							{
-								FNetworkPhysicsSettingsAsyncInput* AsyncInput = NetworkPhysicsSettingsAsync->GetProducerInputData_External();
+								FNetworkPhysicsSettingsAsyncInput* AsyncInput = NetworkPhysicsSettings_Internal->GetProducerInputData_External();
 								AsyncInput->PhysicsObject = PhysicsObject;
 								AsyncInput->Settings.GeneralSettings = GeneralSettings;
 								AsyncInput->Settings.DefaultReplicationSettings = DefaultReplicationSettings;
@@ -86,11 +86,11 @@ void UNetworkPhysicsSettingsComponent::UninitializeComponent()
 		{
 			if (Chaos::FPhysicsSolver* Solver = PhysScene->GetSolver())
 			{
-				Solver->UnregisterAndFreeSimCallbackObject_External(NetworkPhysicsSettingsAsync);
+				Solver->UnregisterAndFreeSimCallbackObject_External(NetworkPhysicsSettings_Internal);
 			}
 		}
 	}
-	NetworkPhysicsSettingsAsync = nullptr;
+	NetworkPhysicsSettings_Internal = nullptr;
 
 	if (AActor* Owner = GetOwner())
 	{
diff --git a/Engine/Source/Runtime/Engine/Private/PhysicsEngine/PhysicsReplication.cpp b/Engine/Source/Runtime/Engine/Private/PhysicsEngine/PhysicsReplication.cpp
--- a/Engine/Source/Runtime/Engine/Private/PhysicsEngine/PhysicsReplication.cpp
+++ b/Engine/Source/Runtime/Engine/Private/PhysicsEngine/PhysicsReplication.cpp
@@ -1995,6 +1995,7 @@ bool FPhysicsReplicationAsync::ResimulationReplication(Chaos::FPBDRigidParticleH
 		return true;
 	}
 
+	const bool bShouldSleep = (Target.TargetState.Flags & ERigidBodyFlags::Sleeping) != 0;
 	bool bClearTarget = true;
 
 	static constexpr Chaos::FFrameAndPhase::EParticleHistoryPhase RewindPhase = Chaos::FFrameAndPhase::EParticleHistoryPhase::PostPushData;
@@ -2031,67 +2032,81 @@ bool FPhysicsReplicationAsync::ResimulationReplication(Chaos::FPBDRigidParticleH
 	}
 #endif
 
-	if (LocalFrame > RewindData->GetBlockedResimFrame())
+	if (ShouldTriggerResim && Target.TickCount == 0 && LocalFrame > RewindData->GetBlockedResimFrame())
 	{
-		if (ShouldTriggerResim && Target.TickCount == 0)
-		{
-			// Trigger resimulation
-			RigidsSolver->GetEvolution()->GetIslandManager().SetParticleResimFrame(Handle, LocalFrame);
+		// Trigger resimulation
+		RigidsSolver->GetEvolution()->GetIslandManager().SetParticleResimFrame(Handle, LocalFrame);
 
-			int32 ResimFrame = RewindData->GetResimFrame();
-			ResimFrame = (ResimFrame == INDEX_NONE) ? LocalFrame : FMath::Min(ResimFrame, LocalFrame);
-			RewindData->SetResimFrame(ResimFrame);
-		}
-		else if (SettingsCurrent.ResimulationSettings.GetRuntimeCorrectionEnabled())
+		int32 ResimFrame = RewindData->GetResimFrame();
+		ResimFrame = (ResimFrame == INDEX_NONE) ? LocalFrame : FMath::Min(ResimFrame, LocalFrame);
+		RewindData->SetResimFrame(ResimFrame);
+	}
+	else if (SettingsCurrent.ResimulationSettings.GetRuntimeCorrectionEnabled())
+	{
+		// Wake up if is sleeping and should not sleep
+		if (Handle->IsSleeping() && !bShouldSleep)
 		{
-			const int32 NumPredictedFrames = RigidsSolver->GetCurrentFrame() - LocalFrame - Target.TickCount;
+			RigidsSolver->GetEvolution()->SetParticleObjectState(Handle, Chaos::EObjectStateType::Dynamic);
+		}
 
-			if (Target.TickCount <= NumPredictedFrames && NumPredictedFrames > 0)
+		const int32 NumPredictedFrames = RigidsSolver->GetCurrentFrame() - LocalFrame - Target.TickCount;
+
+		if (Target.TickCount <= NumPredictedFrames && NumPredictedFrames > 0)
+		{
+			// Positional Correction
+			const float CorrectionAmountX = SettingsCurrent.ResimulationSettings.GetPosStabilityMultiplier() / NumPredictedFrames;
+			const FVector PosDiffCorrection = ErrorOffset * CorrectionAmountX; // Same result as (ErrorOffset / NumPredictedFrames) * PosStabilityMultiplier
+			const FVector CorrectedX = Handle->GetX() + PosDiffCorrection;
+
+			// Rotational Correction
+			const float CorrectionAmountR = SettingsCurrent.ResimulationSettings.GetRotStabilityMultiplier() / NumPredictedFrames;
+			const FQuat DeltaQuat = PastState.GetR().Inverse() * Target.TargetState.Quaternion;
+			const FQuat TargetCorrectionR = Handle->GetR() * DeltaQuat;
+			const FQuat CorrectedR = FQuat::Slerp(Handle->GetR(), TargetCorrectionR, CorrectionAmountR);
+
+			if (SettingsCurrent.ResimulationSettings.GetRuntimeVelocityCorrectionEnabled())
 			{
-				// Positional Correction
-				const float CorrectionAmountX = SettingsCurrent.ResimulationSettings.GetPosStabilityMultiplier() / NumPredictedFrames;
-				const FVector PosDiffCorrection = ErrorOffset * CorrectionAmountX; // Same result as (ErrorOffset / NumPredictedFrames) * PosStabilityMultiplier
-				const FVector CorrectedX = Handle->GetX() + PosDiffCorrection;
-
-				// Rotational Correction
-				const float CorrectionAmountR = SettingsCurrent.ResimulationSettings.GetRotStabilityMultiplier() / NumPredictedFrames;
-				const FQuat DeltaQuat = PastState.GetR().Inverse() * Target.TargetState.Quaternion;
-				const FQuat TargetCorrectionR = Handle->GetR() * DeltaQuat;
-				const FQuat CorrectedR = FQuat::Slerp(Handle->GetR(), TargetCorrectionR, CorrectionAmountR);
-
-				if (SettingsCurrent.ResimulationSettings.GetRuntimeVelocityCorrectionEnabled())
-				{
-					// Linear Velocity Correction
-					const FVector LinVelDiff = Target.TargetState.LinVel - PastState.GetV(); // Velocity vector that the server covers but the client doesn't
-					const float CorrectionAmountV = SettingsCurrent.ResimulationSettings.GetVelStabilityMultiplier() / NumPredictedFrames;
-					const FVector VelCorrection = LinVelDiff * CorrectionAmountV; // Same result as (LinVelDiff / NumPredictedFrames) * VelStabilityMultiplier
-					const FVector CorrectedV = Handle->GetV() + VelCorrection;
-
-					// Angular Velocity Correction
-					const FVector AngVelDiff = FMath::DegreesToRadians(Target.TargetState.AngVel) - PastState.GetW(); // Angular velocity vector that the server covers but the client doesn't
-					const float CorrectionAmountW = SettingsCurrent.ResimulationSettings.GetAngVelStabilityMultiplier() / NumPredictedFrames;
-					const FVector AngVelCorrection = AngVelDiff * CorrectionAmountW; // Same result as (AngVelDiff / NumPredictedFrames) * VelStabilityMultiplier
-					const FVector CorrectedW = Handle->GetW() + AngVelCorrection;
+				// Linear Velocity Correction
+				const FVector LinVelDiff = Target.TargetState.LinVel - PastState.GetV(); // Velocity vector that the server covers but the client doesn't
+				const float CorrectionAmountV = SettingsCurrent.ResimulationSettings.GetVelStabilityMultiplier() / NumPredictedFrames;
+				const FVector VelCorrection = LinVelDiff * CorrectionAmountV; // Same result as (LinVelDiff / NumPredictedFrames) * VelStabilityMultiplier
+				const FVector CorrectedV = Handle->GetV() + VelCorrection;
+
+				// Angular Velocity Correction
+				const FVector AngVelDiff = FMath::DegreesToRadians(Target.TargetState.AngVel) - PastState.GetW(); // Angular velocity vector that the server covers but the client doesn't
+				const float CorrectionAmountW = SettingsCurrent.ResimulationSettings.GetAngVelStabilityMultiplier() / NumPredictedFrames;
+				const FVector AngVelCorrection = AngVelDiff * CorrectionAmountW; // Same result as (AngVelDiff / NumPredictedFrames) * VelStabilityMultiplier
+				const FVector CorrectedW = Handle->GetW() + AngVelCorrection;
 					
-					// Apply correction to velocities
-					Handle->SetV(CorrectedV);
-					Handle->SetW(CorrectedW);
-				}
+				// Apply correction to velocities
+				Handle->SetV(CorrectedV);
+				Handle->SetW(CorrectedW);
+			}
 
 #if !(UE_BUILD_SHIPPING || UE_BUILD_TEST)
-				if (PhysicsReplicationCVars::ResimulationCVars::bDrawDebug)
-				{
-					Chaos::FDebugDrawQueue::GetInstance().DrawDebugDirectionalArrow(Handle->GetX(), CorrectedX, 5.0f, FColor::MakeRandomSeededColor(LocalFrame), true, CharacterMovementCVars::NetCorrectionLifetime, 0, 0.5f);
-				}
-#endif
-				// Apply correction to position and rotation
-				RigidsSolver->GetEvolution()->ApplyParticleTransformCorrection(Handle, CorrectedX, CorrectedR, SettingsCurrent.ResimulationSettings.GetRuntimeCorrectConnectedBodies(), /*bInRecalculateFrictionOnConnectedBodies*/true, ReplicatedParticleIDs);
+			if (PhysicsReplicationCVars::ResimulationCVars::bDrawDebug)
+			{
+				Chaos::FDebugDrawQueue::GetInstance().DrawDebugDirectionalArrow(Handle->GetX(), CorrectedX, 5.0f, FColor::MakeRandomSeededColor(LocalFrame), true, CharacterMovementCVars::NetCorrectionLifetime, 0, 0.5f);
 			}
+#endif
+			// Apply correction to position and rotation
+			RigidsSolver->GetEvolution()->ApplyParticleTransformCorrection(Handle, CorrectedX, CorrectedR, SettingsCurrent.ResimulationSettings.GetRuntimeCorrectConnectedBodies(), /*bInRecalculateFrictionOnConnectedBodies*/true, ReplicatedParticleIDs);
+		}
 
-			// Keep target for NumPredictedFrames time to perform runtime corrections with until a new target is received
-			bClearTarget = Target.TickCount >= NumPredictedFrames;
+		// Keep target for NumPredictedFrames time to perform runtime corrections with until a new target is received
+		bClearTarget = Target.TickCount >= NumPredictedFrames;
+	}
+
+	// Set sleep state if we are about to clear the target from memory and the target is set to sleep
+	if (bClearTarget && bShouldSleep)
+	{
+		RigidsSolver->GetEvolution()->SetParticleObjectState(Handle, Chaos::EObjectStateType::Sleeping);
+		if (PhysicsReplicationCVars::PredictiveInterpolationCVars::bSleepConnectedBodies)
+		{
+			RigidsSolver->GetEvolution()->ApplySleepOnConnectedParticles(Handle);
 		}
 	}
+
 	return bClearTarget;
 }
 
diff --git a/Engine/Source/Runtime/Engine/Private/StaticMeshActor.cpp b/Engine/Source/Runtime/Engine/Private/StaticMeshActor.cpp
--- a/Engine/Source/Runtime/Engine/Private/StaticMeshActor.cpp
+++ b/Engine/Source/Runtime/Engine/Private/StaticMeshActor.cpp
@@ -55,10 +55,7 @@ void AStaticMeshActor::BeginPlay()
 		SetReplicates(true);
 	}	
 
-	if (StaticMeshComponent && StaticMeshComponent->BodyInstance.bSimulatePhysics)
-	{
-		SetPhysicsReplicationMode(StaticMeshPhysicsReplicationMode);
-	}
+	SetPhysicsReplicationMode(StaticMeshPhysicsReplicationMode);
 
 	Super::BeginPlay();
 }
diff --git a/Engine/Source/Runtime/Engine/Public/Physics/NetworkPhysicsComponent.h b/Engine/Source/Runtime/Engine/Public/Physics/NetworkPhysicsComponent.h
--- a/Engine/Source/Runtime/Engine/Public/Physics/NetworkPhysicsComponent.h
+++ b/Engine/Source/Runtime/Engine/Public/Physics/NetworkPhysicsComponent.h
@@ -11,9 +11,12 @@
 #include "Subsystems/WorldSubsystem.h"
 #include "PhysicsEngine/PhysicsSettings.h"
 #include "Chaos/PhysicsObject.h"
+#include "Physics/NetworkPhysicsSettingsComponent.h"
 
 #include "NetworkPhysicsComponent.generated.h"
 
+class FAsyncNetworkPhysicsComponent;
+
 DECLARE_MULTICAST_DELEGATE_OneParam(FOnPreProcessInputsInternal, const int32);
 DECLARE_MULTICAST_DELEGATE_OneParam(FOnPostProcessInputsInternal, const int32);
 DECLARE_MULTICAST_DELEGATE_TwoParams(FOnInjectInputsExternal, const int32 /* PhysicsStep */, const int32 /* NumSteps */);
@@ -58,6 +61,49 @@ struct TNetRewindHistory : public Chaos::TDataRewindHistory<DataType>
 		}
 	}
 
+	virtual int32 CountValidData(const uint32 StartFrame, const uint32 EndFrame, const bool bIncludeUnimportant = true, const bool bIncludeImportant = false)
+	{
+		// Find how many entries are valid in frame range
+		DataType FrameData;
+		int32 Count = 0;
+		for (uint32 CopyFrame = StartFrame; CopyFrame <= EndFrame; ++CopyFrame)
+		{
+			const int32 CopyIndex = Super::GetFrameIndex(CopyFrame);
+			if (CopyFrame == Super::DataHistory[CopyIndex].LocalFrame)
+			{
+				FrameData = Super::DataHistory[CopyIndex];
+
+				// Check if we should include unimportant and/or important data
+				if ((!FrameData.bImportant && bIncludeUnimportant) || (FrameData.bImportant && bIncludeImportant))
+				{
+					Count++;
+				}
+			}
+		}
+		return Count;
+	}
+
+	virtual void SetImportant(const bool bImportant, const int32 Frame = INDEX_NONE) override
+	{
+		if (Frame > INDEX_NONE)
+		{
+			if (Super::EvalData(Frame))
+			{
+				// Set importance on specified frame
+				Super::DataHistory[Super::CurrentIndex].bImportant = bImportant;
+			}
+		}
+		else
+		{
+			// Set importance on all frames
+			for (int32 CopyIndex = 0; CopyIndex < Super::NumFrames; ++CopyIndex)
+			{
+				Super::DataHistory[CopyIndex].bImportant = bImportant;
+			}
+		}
+
+	}
+
 	virtual void ApplyDataRange(const int32 FromFrame, const int32 ToFrame, void* ActorComponent, const bool bOnlyImportant = false) override
 	{
 		UActorComponent* NetworkComponent = static_cast<UActorComponent*>(ActorComponent);
@@ -73,6 +119,27 @@ struct TNetRewindHistory : public Chaos::TDataRewindHistory<DataType>
 		}
 	}
 
+	virtual bool CopyAllData(Chaos::FBaseRewindHistory& OutHistory, bool bIncludeUnimportant = true, bool bIncludeImportant = false) override
+	{
+		TNetRewindHistory& OutNetHistory = static_cast<TNetRewindHistory&>(OutHistory);
+		bool bHasCopiedData = false;
+
+		DataType FrameData;
+		for (int32 CopyIndex = 0; CopyIndex < Super::NumFrames; ++CopyIndex)
+		{
+			FrameData = Super::DataHistory[CopyIndex];
+
+			// Check if we should include unimportant and/or important data
+			if ((!FrameData.bImportant && bIncludeUnimportant) || (FrameData.bImportant && bIncludeImportant))
+			{
+				OutNetHistory.RecordData(FrameData.LocalFrame, &FrameData);
+				bHasCopiedData = true;
+			}
+		}
+		return bHasCopiedData;
+
+	}
+
 	virtual bool CopyData(Chaos::FBaseRewindHistory& OutHistory, const uint32 StartFrame, const uint32 EndFrame, bool bIncludeUnimportant = true, bool bIncludeImportant = false) override
 	{
 		TNetRewindHistory& OutNetHistory = static_cast<TNetRewindHistory&>(OutHistory);
@@ -389,6 +456,7 @@ struct TStructOpsTypeTraits<FNetworkPhysicsRewindDataImportantStateProxy> : publ
 	};
 };
 
+
 /**
  * Network physics rewind callback to manage all the sim callbacks rewind functionalities
  */
@@ -399,7 +467,7 @@ struct FNetworkPhysicsCallback : public Chaos::IRewindCallback
 		UpdateNetMode();
 	}
 
-	// Delegate on the internal inputs process 
+	// Delegate on the internal inputs process
 	FOnPreProcessInputsInternal PreProcessInputsInternal;
 	FOnPostProcessInputsInternal PostProcessInputsInternal;
 	// Bind to this for additional processing on the GT during InjectInputs_External()
@@ -426,12 +494,6 @@ struct FNetworkPhysicsCallback : public Chaos::IRewindCallback
 		RewindableCallbackObjects.Remove(SimCallbackObject);
 	}
 
-	// Updates the TMap on PhysScene that stores (non interpolated) physics data for replication.
-	// 
-	// Needs to be called from PT context to access fixed tick handle
-	// but also needs to be able to access GT data (actor iterator, actor state)
-	void UpdateReplicationMap_Internal(int32 PhysicsStep);
-
 	// Update client player on GT
 	UE_DEPRECATED(5.4, "Physics frame offset is handled by the PlayerController automatically, it's recommended to use APlayerController::GetAsyncPhysicsTimestamp() to get the ServerFrame and LocalFrame on both client and server. Also disable the deprecated flow by setting p.net.CmdOffsetEnabled = 0")
 	void UpdateClientPlayer_External(int32 PhysicsStep);
@@ -456,6 +518,7 @@ struct FNetworkPhysicsCallback : public Chaos::IRewindCallback
 	TArray<Chaos::ISimCallbackObject*> RewindableCallbackObjects;
 };
 
+
 /**
  * Network physics manager to initialize data required for rewind/resim
  */
@@ -475,19 +538,9 @@ public:
 
 	// Delegate at world init 
 	ENGINE_API void OnWorldPostInit(UWorld* World, const UWorld::InitializationValues);
-
-	// Register a component to be used by the callback 
-	void RegisterNetworkComponent(class UNetworkPhysicsComponent* NetworkComponent) {NetworkComponents.Add(NetworkComponent); }
-
-	// Remove a network component from registered list
-	void UnregisterNetworkComponent(class UNetworkPhysicsComponent* NetworkComponent) { NetworkComponents.Remove(NetworkComponent); }
-
-private:
-
-	// List of physics network components that will be used by the rewind callback
-	TArray<class UNetworkPhysicsComponent*> NetworkComponents;
 };
 
+
 /**
  * Base network physics data that will be used by physics
  */
@@ -590,7 +643,6 @@ struct FNetworkPhysicsData
 	*/
 	virtual bool CompareData(const FNetworkPhysicsData& PredictedData) { return true; }
 
-	/** DEPRECATED */
 	UE_DEPRECATED(5.4, "Deprecated, use ApplyData instead")
 	virtual void ApplyDatas(UActorComponent* NetworkComponent) const { }
 	UE_DEPRECATED(5.4, "Deprecated, use BuildData instead")
@@ -601,6 +653,11 @@ struct FNetworkPhysicsData
 	/** Return string with custom debug data */
 	virtual const FString DebugData() { return FString(" - DebugData() not implemented - "); }
 
+	bool operator==(const FNetworkPhysicsData& Other) const
+	{
+		return ServerFrame == Other.ServerFrame && LocalFrame == Other.LocalFrame && InputFrame == Other.InputFrame;
+	}
+
 	friend UNetworkPhysicsComponent;
 };
 
@@ -620,9 +677,30 @@ struct UE_DEPRECATED(5.4, "Deprecated, use FNetworkPhysicsData instead") FNetwor
 	}
 };
 
+/** Base for helper, to create data and data history */
+struct FNetworkPhysicsDataHelper
+{
+	virtual ~FNetworkPhysicsDataHelper() = default;
+	
+	virtual TUniquePtr<FNetworkPhysicsDataHelper> Clone() const = 0;
+	virtual TUniquePtr<FNetworkPhysicsData> CreateUniqueData() const = 0;
+	virtual TUniquePtr<Chaos::FBaseRewindHistory> CreateUniqueRewindHistory(const int32 Size) const = 0;
+};
+
+/** Helper for the creation of state / input data and history with correct derived type */
+template<typename DataType>
+struct TNetworkPhysicsDataHelper : FNetworkPhysicsDataHelper
+{
+	virtual ~TNetworkPhysicsDataHelper() = default;
+
+	virtual TUniquePtr<FNetworkPhysicsDataHelper> Clone() const { return MakeUnique<TNetworkPhysicsDataHelper>(*this); }
+	virtual TUniquePtr<FNetworkPhysicsData> CreateUniqueData() const { return MakeUnique<DataType>(); }
+	virtual TUniquePtr<Chaos::FBaseRewindHistory> CreateUniqueRewindHistory(const int32 Size) const { return MakeUnique<TNetRewindHistory<DataType>>(Size); }
+};
 
 /**
- * Network physics component that will be attached to any player controller
+ * Network physics component to add to actors or pawns that control their physic simulation through applying inputs,
+ * and should support networking through physics resimulation.
  */
 UCLASS(BlueprintType, MinimalAPI)
 class UNetworkPhysicsComponent : public UActorComponent
@@ -631,43 +709,18 @@ class UNetworkPhysicsComponent : public UActorComponent
 public:
 	ENGINE_API UNetworkPhysicsComponent();
 
-	// Get the player controller on which the component is attached
+	// Get the player controller
 	ENGINE_API virtual APlayerController* GetPlayerController() const;
 
 	// Init the network physics component 
 	ENGINE_API void InitPhysics();
 
-	UE_DEPRECATED(5.4, "Deprecated, use SendInputData() instead")
-	UFUNCTION(Server, unreliable, meta = (DeprecatedFunction, DeprecationMessage = "ServerReceiveInputsDatas has been deprecated. Use ServerReceiveInputData instead."))
-	ENGINE_API void ServerReceiveInputsDatas(const FNetworkPhysicsRewindDataInputProxy& ClientInputs);
-	// Async physics tick component function per frame from the solver
-	ENGINE_API virtual void AsyncPhysicsTickComponent(float DeltaTime, float SimTime) override;
+	// Called every frame
+	ENGINE_API virtual void TickComponent(float DeltaTime, ELevelTick TickType, FActorComponentTickFunction* ThisTickFunction) override;
 
 	// Function to init the replicated properties
 	ENGINE_API virtual void GetLifetimeReplicatedProps(TArray<FLifetimeProperty>& OutLifeTimeProps) const override;
 
-	// Replicate input data 
-	ENGINE_API void SendInputData();
-
-	UE_DEPRECATED(5.4, "Deprecated, use SendInputData() instead")
-	ENGINE_API void SendLocalInputsDatas() { SendInputData(); }
-
-	// Replicate state data
-	ENGINE_API void SendStateData();
-
-	UE_DEPRECATED(5.4, "Deprecated, use SendStateData() instead")
-	ENGINE_API void SendLocalStatesDatas() { SendStateData(); }
-
-	// Delegate linked to the physics rewind callback to send record local inputs/states
-	ENGINE_API void OnPreProcessInputsInternal(const int32 PhysicsStep);
-
-	// Delegate linked to the physics rewind callback to send record local inputs/states
-	ENGINE_API void OnPostProcessInputsInternal(const int32 PhysicsStep);
-
-	// Correct the player controller Server to local offset based on the received replicated states
-	UE_DEPRECATED(5.4, "Physics frame offset is handled by the PlayerController automatically, it's recommended to use APlayerController::GetAsyncPhysicsTimestamp() to get the ServerFrame and LocalFrame on both client and server. Also disable the deprecated flow by setting p.net.CmdOffsetEnabled = 0")
-	ENGINE_API void CorrectServerToLocalOffset(const int32 LocalToServerOffset);
-
 	// Used to create any physics engine information for this component 
 	ENGINE_API virtual void BeginPlay() override;
 
@@ -677,50 +730,26 @@ public:
 	// Unregister the component from the network manager
 	ENGINE_API virtual void UninitializeComponent() override;
 
-	/** Register and create both state and input to be both networked and cached in history */
-	template<typename PhysicsTraits>
-	void CreateDataHistory(UActorComponent* HistoryComponent);
-	
-	template<typename PhysicsTraits>
-	UE_DEPRECATED(5.4, "Deprecated, use CreateDataHistory() instead")
-	void CreateDatasHistory(UActorComponent* HistoryComponent);
-	
-	/**  Register and create input history
-	* Please use CreateDataHistory() if both input and custom state are supposed to be networked and cached in history.
-	* NOTE: Registering Input without State requires networking push-model to be enabled to take advantage of all the CPU and Network Bandwidth savings, CVar: Net.IsPushModelEnabled 1 */
-	template<class InputsType>
-	void CreateInputHistory(UActorComponent* HistoryComponent);
-
-
 	// Remove state/input history from rewind data
 	ENGINE_API void RemoveDataHistory();
-	UE_DEPRECATED(5.4, "Deprecated, use RemoveDataHistory() instead")
-	ENGINE_API void RemoveDatasHistory() { RemoveDataHistory(); }
 
 	// Add state/input history to rewind data
 	ENGINE_API void AddDataHistory();
-	UE_DEPRECATED(5.4, "Deprecated, use AddDataHistory() instead")
-	ENGINE_API void AddDatasHistory() { AddDataHistory(); }
-
-	// Enable RewindData history caching and return the history size
-	ENGINE_API int32 SetupRewindData();
-
-	// Get the data factory that will be used for net serialization
-	TSharedPtr<Chaos::FBaseRewindHistory>& GetStateHistory() { return StateHistory; }
-	UE_DEPRECATED(5.4, "Deprecated, use GetStateHistory() instead")
-	TSharedPtr<Chaos::FBaseRewindHistory>& GetStatesHistory() { return GetStateHistory(); }
 
-	// Get the data factory that will be used for net serialization
-	TSharedPtr<Chaos::FBaseRewindHistory>& GetInputHistory() { return InputHistory; }
-	UE_DEPRECATED(5.4, "Deprecated, use GetInputHistory() instead")
-	TSharedPtr<Chaos::FBaseRewindHistory>& GetInputsHistory() { return GetInputHistory(); }
+	// Get the GameThread state history (not guaranteed to be the exact data used in physics, for that use GetStateHistory_Internal on PhysicsThread)
+	ENGINE_API TSharedPtr<Chaos::FBaseRewindHistory>& GetStateHistory_External() { return StateHistory; }
+	
+	// Get the PhysicsThread state history (if there is none it returns the GameThread history)
+	ENGINE_API TSharedPtr<Chaos::FBaseRewindHistory>& GetStateHistory_Internal();
+	
+	// Get the GameThread input history (not guaranteed to be the exact data used in physics, for that use GetInputHistory_Internal on PhysicsThread)
+	ENGINE_API TSharedPtr<Chaos::FBaseRewindHistory>& GetInputHistory_External() { return InputHistory; }
+	
+	// Get the PhysicsThread input history (if there is none it returns the GameThread history)
+	ENGINE_API TSharedPtr<Chaos::FBaseRewindHistory>& GetInputHistory_Internal();
 
 	// Check if the world is on server
 	ENGINE_API bool HasServerWorld() const;
-
-	// Check if the player controller exists and is local
-	UE_DEPRECATED(5.4, "Deprecated, use IsLocallyControlled() which takes both local player controlled and local relayed inputs into account.")
-	ENGINE_API bool HasLocalController() const;
 	
 	// Check if this is controlled locally through relayed inputs or an existing local player controller
 	ENGINE_API bool IsLocallyControlled() const;
@@ -731,6 +760,7 @@ public:
 	ENGINE_API void SetIsRelayingLocalInputs(bool bInRelayingLocalInputs)
 	{
 		bIsRelayingLocalInputs = bInRelayingLocalInputs;
+		UpdateAsyncComponent();
 	}
 
 	/** Stop relaying local inputs after next network send.
@@ -747,8 +777,8 @@ public:
 	/** Check if this is controlled locally through relayed inputs from autonomous proxy. It's recommended to use IsLocallyControlled() when checking if this is locally controlled. */
 	ENGINE_API const bool GetIsRelayingLocalInputs() const { return bIsRelayingLocalInputs; }
 
-	/** Returns the current amount of input decay during resimulation as a magnitude from 0.0 to 1.0. Returns 0 if not currently resimulating. */
-	ENGINE_API const float GetCurrentInputDecay(FNetworkPhysicsData* PhysicsData);
+	/** Get FAsyncNetworkPhysicsComponent on the Physics Thread */
+	FAsyncNetworkPhysicsComponent* GetNetworkPhysicsComponent_Internal() { NetworkPhysicsComponent_Internal; }
 
 protected : 
 
@@ -793,8 +823,25 @@ protected :
 	FNetworkPhysicsRewindDataStateProxy ReplicatedStates;
 
 private:
-	/** Trigger a resimulation on frame */
-	void TriggerResimulation(int32 ResimFrame);
+	/** Send network data from marshaled data */
+	void NetworkMarshaledData();
+
+	/** Update the async component with GT properties */
+	void UpdateAsyncComponent(const bool bFullUpdate = true);
+
+	/** Setup data type and history in async input and trigger the creation of data history in the async component */
+	ENGINE_API void CreateAsyncDataHistory();
+
+public:
+	/** Register and create both state and input to be both networked and cached in history */
+	template<typename PhysicsTraits>
+	void CreateDataHistory(UActorComponent* HistoryComponent);
+
+	/**  Register and create input history
+	* Please use CreateDataHistory() if both input and custom state are supposed to be networked and cached in history.
+	* NOTE: Registering Input without State requires networking push-model to be enabled to take advantage of all the CPU and Network Bandwidth savings, CVar: Net.IsPushModelEnabled 1 */
+	template<class InputsType>
+	void CreateInputHistory(UActorComponent* HistoryComponent);
 
 private:
 
@@ -804,17 +851,20 @@ private:
 	friend struct FNetworkPhysicsRewindDataImportantInputProxy;
 	friend struct FNetworkPhysicsRewindDataImportantStateProxy;
 
-	// States history uses to rewind simulation 
+	// Network Physics Component data internal to the physics thread
+	FAsyncNetworkPhysicsComponent* NetworkPhysicsComponent_Internal;
+
+	// States history on GameThread
 	TSharedPtr<Chaos::FBaseRewindHistory> StateHistory;
 
-	// Inputs history used during simulation
+	// Inputs history on GameThread
 	TSharedPtr<Chaos::FBaseRewindHistory> InputHistory;
 
-	// Local temporary inputs data used by pre/post process inputs functions
-	TUniquePtr<FNetworkPhysicsData> InputData;
+	// Helper for the creation of input data and history with correct derived type
+	TUniquePtr<FNetworkPhysicsDataHelper> InputHelper;
 
-	// Local temporary states data used by pre/post process inputs functions
-	TUniquePtr<FNetworkPhysicsData> StateData;
+	// Helper for the creation of state data and history with correct derived type
+	TUniquePtr<FNetworkPhysicsDataHelper> StateHelper;
 
 	// Send last N number of inputs each replication call to patch up holes due to packet loss
 	int8 InputRedundancy = 3;
@@ -834,17 +884,97 @@ private:
 	// If we are currently relaying inputs and will stop after next network send.
 	bool bStopRelayingLocalInputsDeferred = false;
 
-	int32 LastInputSendFrame = INDEX_NONE;
-	int32 LastStateSendFrame = INDEX_NONE;
-	int32 NewImportantInputFrame = INT_MAX;
-
-	// Temporary settings variables, ToDo, Remove when async ISimCallback flow is implemented and these can be retrieved from NetworkPhysicsSettingsComponent on PT
-	bool bCompareStateToTriggerRewind = false;
-	bool bCompareInputToTriggerRewind = false;
+	// ToDo, retrieve from NetworkPhysicsSettingsComponent so changes at runtime gets picked up
 	bool bEnableUnreliableFlow = true;
 	bool bEnableReliableFlow = false;
-	bool bApplyDataInsteadOfMergeData = false;
-	bool bAllowInputExtrapolation = true;
+	bool bValidateDataOnGameThread = false;
+
+
+	/** ----- Deprecated API ----- */
+public:
+	UE_DEPRECATED(5.4, "Deprecated, use ServerReceiveInputData instead")
+		UFUNCTION(Server, unreliable, meta = (DeprecatedFunction, DeprecationMessage = "ServerReceiveInputsDatas has been deprecated. Use ServerReceiveInputData instead."))
+		ENGINE_API void ServerReceiveInputsDatas(const FNetworkPhysicsRewindDataInputProxy& ClientInputs);
+
+	// Replicate input data 
+	UE_DEPRECATED(5.5, "Deprecated, inputs are sent from PT marshaled data and can't be manually triggered on GT.")
+		ENGINE_API void SendInputData() {}
+
+	UE_DEPRECATED(5.4, "Deprecated, use SendInputData() instead")
+		ENGINE_API void SendLocalInputsDatas()
+	{
+		PRAGMA_DISABLE_DEPRECATION_WARNINGS
+			SendInputData();
+		PRAGMA_ENABLE_DEPRECATION_WARNINGS
+	}
+
+	// Replicate state data
+	UE_DEPRECATED(5.5, "Deprecated, states are sent from PT marshaled data and can't be manually triggered on GT.")
+		ENGINE_API void SendStateData() {}
+
+	UE_DEPRECATED(5.4, "Deprecated, use SendStateData() instead")
+		ENGINE_API void SendLocalStatesDatas()
+	{
+		PRAGMA_DISABLE_DEPRECATION_WARNINGS
+			SendStateData();
+		PRAGMA_ENABLE_DEPRECATION_WARNINGS
+	}
+
+	// Delegate linked to the physics rewind callback to send record local inputs/states
+	UE_DEPRECATED(5.5, "Deprecated, Logic moved to async flow through FAsyncNetworkPhysicsComponent::OnPreProcessInputs_Internal instead")
+		ENGINE_API void OnPreProcessInputsInternal(const int32 PhysicsStep) {};
+
+	// Delegate linked to the physics rewind callback to send record local inputs/states
+	UE_DEPRECATED(5.5, "Deprecated, Logic moved to async flow through FAsyncNetworkPhysicsComponent::OnPostProcessInputs_Internal instead")
+		ENGINE_API void OnPostProcessInputsInternal(const int32 PhysicsStep) {};
+
+	// Correct the player controller Server to local offset based on the received replicated states
+	UE_DEPRECATED(5.4, "Physics frame offset is handled by the PlayerController automatically, it's recommended to use APlayerController::GetAsyncPhysicsTimestamp() to get the ServerFrame and LocalFrame on both client and server. Also disable the deprecated flow by setting p.net.CmdOffsetEnabled = 0")
+		ENGINE_API void CorrectServerToLocalOffset(const int32 LocalToServerOffset);
+
+	UE_DEPRECATED(5.4, "Deprecated, use RemoveDataHistory() instead")
+		ENGINE_API void RemoveDatasHistory() { RemoveDataHistory(); }
+
+	UE_DEPRECATED(5.4, "Deprecated, use AddDataHistory() instead")
+		ENGINE_API void AddDatasHistory() { AddDataHistory(); }
+
+	// Enable RewindData history caching and return the history size
+	UE_DEPRECATED(5.5, "Deprecated, don't setup rewind data from GameThread")
+		ENGINE_API int32 SetupRewindData() { return 0; }
+
+	UE_DEPRECATED(5.5, "Deprecated, use GetStateHistory_Internal() or GetStateHistory_External() instead")
+		TSharedPtr<Chaos::FBaseRewindHistory>& GetStateHistory() { return GetStateHistory_Internal(); }
+
+	UE_DEPRECATED(5.4, "Deprecated, use GetStateHistory() instead")
+		TSharedPtr<Chaos::FBaseRewindHistory>& GetStatesHistory()
+	{
+		PRAGMA_DISABLE_DEPRECATION_WARNINGS
+		return GetStateHistory();
+		PRAGMA_ENABLE_DEPRECATION_WARNINGS
+	}
+
+	UE_DEPRECATED(5.5, "Deprecated, use GetInputHistory_Internal() or GetInputHistory_External() instead")
+		TSharedPtr<Chaos::FBaseRewindHistory>& GetInputHistory() { return GetInputHistory_Internal(); }
+	
+	UE_DEPRECATED(5.4, "Deprecated, use GetInputHistory() instead")
+		TSharedPtr<Chaos::FBaseRewindHistory>& GetInputsHistory()
+	{
+		PRAGMA_DISABLE_DEPRECATION_WARNINGS
+		return GetInputHistory();
+		PRAGMA_ENABLE_DEPRECATION_WARNINGS
+	}
+
+	// Check if the player controller exists and is local
+	UE_DEPRECATED(5.4, "Deprecated, use IsLocallyControlled() which takes both local player controlled and local relayed inputs into account.")
+		ENGINE_API bool HasLocalController() const;
+
+	template<typename PhysicsTraits>
+	UE_DEPRECATED(5.4, "Deprecated, use CreateDataHistory() instead")
+		void CreateDatasHistory(UActorComponent* HistoryComponent);
+
+	/** Returns the current amount of input decay during resimulation as a magnitude from 0.0 to 1.0. Returns 0 if not currently resimulating. */
+	UE_DEPRECATED(5.5, "Deprecated, use FAsyncNetworkPhysicsComponent::GetCurrentInputDecay instead.")
+		ENGINE_API const float GetCurrentInputDecay(FNetworkPhysicsData* PhysicsData);
 };
 
 /** DEPRECATED UE 5.4 */
@@ -857,13 +987,8 @@ FORCEINLINE void UNetworkPhysicsComponent::CreateDatasHistory(UActorComponent* H
 template<typename PhysicsTraits>
 FORCEINLINE void UNetworkPhysicsComponent::CreateDataHistory(UActorComponent* HistoryComponent)
 {
-	const int32 NumFrames = SetupRewindData();
-
-	InputHistory = MakeShared<TNetRewindHistory<typename PhysicsTraits::InputsType>>(NumFrames);
-	StateHistory = MakeShared<TNetRewindHistory<typename PhysicsTraits::StatesType>>(NumFrames);
-
-	InputData = MakeUnique<typename PhysicsTraits::InputsType>();
-	StateData = MakeUnique<typename PhysicsTraits::StatesType>();
+	InputHelper = MakeUnique<TNetworkPhysicsDataHelper<typename PhysicsTraits::InputsType>>();
+	StateHelper = MakeUnique<TNetworkPhysicsDataHelper<typename PhysicsTraits::StatesType>>();
 
 	ReplicatedInputs.History = MakeUnique<TNetRewindHistory<typename PhysicsTraits::InputsType>>(InputRedundancy);
 	ReplicatedInputs.Owner = this;
@@ -878,18 +1003,14 @@ FORCEINLINE void UNetworkPhysicsComponent::CreateDataHistory(UActorComponent* Hi
 	ReplicatedImportantState.Owner = this;
 
 	ActorComponent = HistoryComponent;
-	
-	AddDataHistory();
+
+	CreateAsyncDataHistory();
 }
 
 template<class InputsType>
 FORCEINLINE void UNetworkPhysicsComponent::CreateInputHistory(UActorComponent* HistoryComponent)
 {
-	const int32 NumFrames = SetupRewindData();
-
-	InputHistory = MakeShared<TNetRewindHistory<InputsType>>(NumFrames);
-
-	InputData = MakeUnique<InputsType>();
+	InputHelper = MakeUnique<TNetworkPhysicsDataHelper<InputsType>>();
 
 	ReplicatedInputs.History = MakeUnique<TNetRewindHistory<InputsType>>(InputRedundancy);
 	ReplicatedInputs.Owner = this;
@@ -899,6 +1020,196 @@ FORCEINLINE void UNetworkPhysicsComponent::CreateInputHistory(UActorComponent* H
 
 	ActorComponent = HistoryComponent;
 
-	AddDataHistory();
+	CreateAsyncDataHistory();
 }
 
+
+// --------------------------- PhysicsThread Network Physics Component ---------------------------
+
+struct FAsyncNetworkPhysicsComponentInput : public Chaos::FSimCallbackInput
+{
+	TOptional<bool> bIsLocallyControlled;
+	TOptional<ENetMode> NetMode;
+	TOptional<ENetRole> NetRole;
+	TOptional<int32> NetworkPhysicsTickOffset;
+	TOptional<EPhysicsReplicationMode> PhysicsReplicationMode;
+	TOptional<TObjectPtr<UActorComponent>> ActorComponent;
+	TOptional<FString> ActorName;
+	TOptional<TUniquePtr<FNetworkPhysicsDataHelper>> InputHelper;
+	TOptional<TUniquePtr<FNetworkPhysicsDataHelper>> StateHelper;
+	TOptional<bool> bRegisterDataHistoryInRewindData;
+	TOptional<bool> bUnregisterDataHistoryFromRewindData;
+
+	TUniquePtr<Chaos::FBaseRewindHistory> InputData;
+	TUniquePtr<Chaos::FBaseRewindHistory> StateData;
+
+	TArray<TUniquePtr<Chaos::FBaseRewindHistory>> InputDataImportant;
+	TArray<TUniquePtr<Chaos::FBaseRewindHistory>> StateDataImportant;
+
+	void Reset()
+	{
+		bIsLocallyControlled.Reset();
+		NetMode.Reset();
+		NetRole.Reset();
+		NetworkPhysicsTickOffset.Reset();
+		PhysicsReplicationMode.Reset();
+		ActorComponent.Reset();
+		ActorName.Reset();
+		InputHelper.Reset();
+		StateHelper.Reset();
+		bRegisterDataHistoryInRewindData.Reset();
+		bUnregisterDataHistoryFromRewindData.Reset();
+
+		if (InputData)
+		{
+			InputData->ResetFast();
+		}
+		if (StateData)
+		{
+			StateData->ResetFast();
+		}
+
+		// Todo, optimize
+		InputDataImportant.Reset();
+		StateDataImportant.Reset();
+	}
+};
+
+struct FAsyncNetworkPhysicsComponentOutput : public Chaos::FSimCallbackOutput
+{
+	TUniquePtr<Chaos::FBaseRewindHistory> InputData;
+	TUniquePtr<Chaos::FBaseRewindHistory> StateData;
+
+	TArray<TUniquePtr<Chaos::FBaseRewindHistory>> InputDataImportant;
+	TArray<TUniquePtr<Chaos::FBaseRewindHistory>> StateDataImportant;
+
+	void Reset()
+	{
+		if (InputData)
+		{
+			InputData->ResetFast();
+		}
+		if (StateData)
+		{
+			StateData->ResetFast();
+		}
+
+		// Todo, optimize
+		InputDataImportant.Reset();
+		StateDataImportant.Reset();
+	}
+};
+
+class FAsyncNetworkPhysicsComponent : public Chaos::TSimCallbackObject<
+	FAsyncNetworkPhysicsComponentInput,
+	FAsyncNetworkPhysicsComponentOutput,
+	Chaos::ESimCallbackOptions::None>
+{
+
+public:
+	FAsyncNetworkPhysicsComponent();
+	~FAsyncNetworkPhysicsComponent();
+
+	// Get reference to async output for current internal frame and initialize it if not already done
+	FAsyncNetworkPhysicsComponentOutput& GetAsyncOutput_Internal();
+
+	// If this network physics component is locally controlled, can be either server or autonomous proxy
+	const bool IsLocallyControlled() const { return bIsLocallyControlled; }
+
+	// If we are on the server
+	const bool IsServer() const { return (NetMode == ENetMode::NM_DedicatedServer || NetMode == ENetMode::NM_ListenServer); }
+
+	// Get the ENetRole
+	const ENetRole GetNetRole() const { return NetRole; }
+
+	// Get actor name
+	const FString GetActorName() const { return ActorName; }
+
+	// Get the physics tick offset (add to clients physics tick to get the servers corresponding physics tick)
+	const int32 GetNetworkPhysicsTickOffset() const { return NetworkPhysicsTickOffset; }
+
+	// Get the physics replication mode used
+	const EPhysicsReplicationMode GetPhysicsReplicationMode() const { return PhysicsReplicationMode; }
+
+	// Add state/input history to rewind data
+	void RegisterDataHistoryInRewindData();
+
+	// Remove state/input history from rewind data
+	void UnregisterDataHistoryFromRewindData();
+
+	// Enable RewindData history caching and return the history size
+	const int32 SetupRewindData();
+
+private:
+	// Initialize, bind delegates etc.
+	virtual void OnPostInitialize_Internal() override;
+
+	// Delegate ran at the start of FNetworkPhysicsCallback::ProcessInputs_Internal, used to receive and apply inputs and states if needed.
+	void OnPreProcessInputs_Internal(const int32 PhysicsStep);
+
+	// Delegate ran at the end of FNetworkPhysicsCallback::ProcessInputs_Internal, used to record and send inputs and states if needed.
+	void OnPostProcessInputs_Internal(const int32 PhysicsStep);
+
+	// Consume data from async input
+	void ConsumeAsyncInput();
+
+	/** Get the rigid solver */
+	Chaos::FPBDRigidsSolver* GetRigidSolver();
+
+	/** Get the rigid solver evolution */
+	Chaos::FPBDRigidsEvolution* GetEvolution();
+
+	/** Get the settings for this NetworkPhysicsComponent */
+	const FNetworkPhysicsSettingsNetworkPhysicsComponent& GetComponentSettings();
+
+	/** Trigger a resimulation on frame */
+	void TriggerResimulation(int32 ResimFrame);
+
+	/** Returns the current amount of input decay during resimulation as a magnitude from 0.0 to 1.0. Returns 0 if not currently resimulating. */
+	const float GetCurrentInputDecay(const FNetworkPhysicsData* PhysicsData);
+
+	/** Populate input data in AsyncOutput to send over the network */
+	void SendInputData_Internal(FAsyncNetworkPhysicsComponentOutput& AsyncOutput);
+
+	/** Populate state data in AsyncOutput to send over the network */
+	void SendStateData_Internal(FAsyncNetworkPhysicsComponentOutput& AsyncOutput);
+
+private:
+	friend UNetworkPhysicsComponent;
+
+	bool bIsLocallyControlled;
+	ENetMode NetMode;
+	ENetRole NetRole;
+	int32 NetworkPhysicsTickOffset;
+	EPhysicsReplicationMode PhysicsReplicationMode;
+	FString ActorName;
+
+	int32 LastInputSendFrame = INDEX_NONE;
+	int32 LastStateSendFrame = INDEX_NONE;
+	int32 NewImportantInputFrame = INT_MAX;
+
+	// Component settings
+	FNetworkPhysicsSettingsComponentAsync* SettingsComponent;
+	static const FNetworkPhysicsSettingsNetworkPhysicsComponent SettingsNetworkPhysicsComponent_Default;
+
+	// Actor component that will be used to fill the histories
+	TObjectPtr<UActorComponent> ActorComponent;
+
+	// Root components physics object
+	Chaos::FPhysicsObjectHandle RootPhysicsObject;
+
+	// States history uses to rewind simulation 
+	TSharedPtr<Chaos::FBaseRewindHistory> StateHistory;
+
+	// Inputs history used during simulation
+	TSharedPtr<Chaos::FBaseRewindHistory> InputHistory;
+
+	// Local temporary inputs data used by pre/post process inputs functions
+	TUniquePtr<FNetworkPhysicsData> InputData;
+
+	// Local temporary states data used by pre/post process inputs functions
+	TUniquePtr<FNetworkPhysicsData> StateData;
+
+	FDelegateHandle DelegateOnPreProcessInputs_Internal;
+	FDelegateHandle DelegateOnPostProcessInputs_Internal;
+};
diff --git a/Engine/Source/Runtime/Engine/Public/Physics/NetworkPhysicsSettingsComponent.h b/Engine/Source/Runtime/Engine/Public/Physics/NetworkPhysicsSettingsComponent.h
--- a/Engine/Source/Runtime/Engine/Public/Physics/NetworkPhysicsSettingsComponent.h
+++ b/Engine/Source/Runtime/Engine/Public/Physics/NetworkPhysicsSettingsComponent.h
@@ -66,6 +66,11 @@ namespace PhysicsReplicationCVars
 		extern bool bEnableReliableFlow;
 		extern bool bApplyDataInsteadOfMergeData;
 		extern bool bAllowInputExtrapolation;
+		extern bool bValidateDataOnGameThread;
+		extern int32 InputRedundancy;
+		extern int32 StateRedundancy;
+		extern bool bCompareStateToTriggerRewind;
+		extern bool bCompareInputToTriggerRewind;
 	}
 }
 
@@ -304,59 +309,66 @@ struct FNetworkPhysicsSettingsNetworkPhysicsComponent
 
 	UPROPERTY(config, EditDefaultsOnly, Category = "Overrides", Meta = (InlineEditConditionToggle))
 	uint32 bOverrideRedundantInputs : 1;
-	// Override how many inputs to synchronize each sync to cover packet loss.
+	// Overrides CVar: np2.Resim.InputRedundancy -- How many inputs to send with each unreliable network message to account for packetloss.
 	UPROPERTY(config, EditDefaultsOnly, Category = "Overrides", Meta = (EditCondition = "bOverrideRedundantInputs"))
-	uint8 RedundantInputs = 3;
-	uint8 GetRedundantInputs(uint8 DefaultValue) { return bOverrideRedundantInputs ? RedundantInputs : DefaultValue; }
+	uint8 RedundantInputs = PhysicsReplicationCVars::ResimulationCVars::InputRedundancy;
+	const uint8 GetRedundantInputs() const { return bOverrideRedundantInputs ? RedundantInputs : PhysicsReplicationCVars::ResimulationCVars::InputRedundancy; }
 
 	UPROPERTY(config, EditDefaultsOnly, Category = "Overrides", Meta = (InlineEditConditionToggle))
 	uint32 bOverrideRedundantStates : 1;
-	// Override how many states to synchronize each sync to cover packet loss.
+	// Overrides CVar: np2.Resim.StateRedundancy -- How many states to send with each unreliable network message to account for packetloss.
 	UPROPERTY(config, EditDefaultsOnly, Category = "Overrides", Meta = (EditCondition = "bOverrideRedundantStates"))
-	uint8 RedundantStates = 1;
-	uint8 GetRedundantStates(uint8 DefaultValue) { return bOverrideRedundantStates ? RedundantStates : DefaultValue; }
+	uint8 RedundantStates = PhysicsReplicationCVars::ResimulationCVars::StateRedundancy;
+	const uint8 GetRedundantStates() const { return bOverrideRedundantStates ? RedundantStates : PhysicsReplicationCVars::ResimulationCVars::StateRedundancy; }
 
 	UPROPERTY(config, EditDefaultsOnly, Category = "Overrides", Meta = (InlineEditConditionToggle))
 	uint32 bOverrideCompareStateToTriggerRewind : 1;
 	// Overrides CVar: np2.Resim.CompareStateToTriggerRewind -- When true, cache local players custom state struct in rewind history and compare the predicted state with incoming server state to trigger resimulations if they differ, comparison done through FNetworkPhysicsData::CompareData.
 	UPROPERTY(config, EditDefaultsOnly, Category = "Overrides", Meta = (EditCondition = "bOverrideCompareStateToTriggerRewind"))
-	bool bCompareStateToTriggerRewind = false;
-	bool GetCompareStateToTriggerRewind(bool DefaultValue) { return bOverrideCompareStateToTriggerRewind ? bCompareStateToTriggerRewind : DefaultValue; }
+	bool bCompareStateToTriggerRewind = PhysicsReplicationCVars::ResimulationCVars::bCompareStateToTriggerRewind;
+	const bool GetCompareStateToTriggerRewind() const { return bOverrideCompareStateToTriggerRewind ? bCompareStateToTriggerRewind : PhysicsReplicationCVars::ResimulationCVars::bCompareStateToTriggerRewind; }
 
 	UPROPERTY(config, EditDefaultsOnly, Category = "Overrides", Meta = (InlineEditConditionToggle))
 	uint32 bOverrideCompareInputToTriggerRewind : 1;
 	// Overrides CVar: np2.Resim.CompareInputToTriggerRewind -- When true, compare local players predicted inputs with incoming server inputs to trigger resimulations if they differ, comparison done through FNetworkPhysicsData::CompareData.
 	UPROPERTY(config, EditDefaultsOnly, Category = "Overrides", Meta = (EditCondition = "bOverrideCompareInputToTriggerRewind"))
-	bool bCompareInputToTriggerRewind = false;
-	bool GetCompareInputToTriggerRewind(bool DefaultValue) { return bOverrideCompareInputToTriggerRewind ? bCompareInputToTriggerRewind : DefaultValue; }
+	bool bCompareInputToTriggerRewind = PhysicsReplicationCVars::ResimulationCVars::bCompareInputToTriggerRewind;
+	const bool GetCompareInputToTriggerRewind() const { return bOverrideCompareInputToTriggerRewind ? bCompareInputToTriggerRewind : PhysicsReplicationCVars::ResimulationCVars::bCompareInputToTriggerRewind; }
 
 	UPROPERTY(config, EditDefaultsOnly, Category = "Overrides", Meta = (InlineEditConditionToggle))
 	uint32 bOverrideEnableUnreliableFlow : 1;
 	// Overrides CVar: np2.Resim.EnableUnreliableFlow -- When true, allow data to be sent unreliably. Also sends FNetworkPhysicsData not marked with FNetworkPhysicsData::bimportant unreliably over the network.
 	UPROPERTY(config, EditDefaultsOnly, Category = "Overrides", Meta = (EditCondition = "bOverrideEnableUnreliableFlow"))
 	bool bEnableUnreliableFlow = PhysicsReplicationCVars::ResimulationCVars::bEnableUnreliableFlow;
-	bool GetEnableUnreliableFlow() { return bOverrideEnableUnreliableFlow ? bEnableUnreliableFlow : PhysicsReplicationCVars::ResimulationCVars::bEnableUnreliableFlow; }
+	const bool GetEnableUnreliableFlow() const { return bOverrideEnableUnreliableFlow ? bEnableUnreliableFlow : PhysicsReplicationCVars::ResimulationCVars::bEnableUnreliableFlow; }
 
 	UPROPERTY(config, EditDefaultsOnly, Category = "Overrides", Meta = (InlineEditConditionToggle))
 	uint32 bOverrideEnableReliableFlow : 1;
 	// Overrides CVar: np2.Resim.EnableReliableFlow -- EXPERIMENTAL -- When true, allow data to be sent reliably. Also send FNetworkPhysicsData marked with FNetworkPhysicsData::bimportant reliably over the network.
 	UPROPERTY(config, EditDefaultsOnly, Category = "Overrides", Meta = (EditCondition = "bOverrideEnableReliableFlow"))
 	bool bEnableReliableFlow = PhysicsReplicationCVars::ResimulationCVars::bEnableReliableFlow;
-	bool GetEnableReliableFlow() { return bOverrideEnableReliableFlow ? bEnableReliableFlow : PhysicsReplicationCVars::ResimulationCVars::bEnableReliableFlow; }
+	const bool GetEnableReliableFlow() const { return bOverrideEnableReliableFlow ? bEnableReliableFlow : PhysicsReplicationCVars::ResimulationCVars::bEnableReliableFlow; }
 
 	UPROPERTY(config, EditDefaultsOnly, Category = "Overrides", Meta = (InlineEditConditionToggle))
 	uint32 bOverrideApplyDataInsteadOfMergeData : 1;
 	// Overrides CVar: np2.Resim.ApplyDataInsteadOfMergeData -- When true, call ApplyData for each data instead of MergeData when having to use multiple data entries in one frame.
 	UPROPERTY(config, EditDefaultsOnly, Category = "Overrides", Meta = (EditCondition = "bOverrideApplyDataInsteadOfMergeData"))
 	bool bApplyDataInsteadOfMergeData = PhysicsReplicationCVars::ResimulationCVars::bApplyDataInsteadOfMergeData;
-	bool GetApplyDataInsteadOfMergeData() { return bOverrideApplyDataInsteadOfMergeData ? bApplyDataInsteadOfMergeData : PhysicsReplicationCVars::ResimulationCVars::bApplyDataInsteadOfMergeData; }
+	const bool GetApplyDataInsteadOfMergeData() const { return bOverrideApplyDataInsteadOfMergeData ? bApplyDataInsteadOfMergeData : PhysicsReplicationCVars::ResimulationCVars::bApplyDataInsteadOfMergeData; }
 
 	UPROPERTY(config, EditDefaultsOnly, Category = "Overrides", Meta = (InlineEditConditionToggle))
 	uint32 bOverrideAllowInputExtrapolation : 1;
 	// Overrides CVar: np2.Resim.AllowInputExtrapolation -- When true and not locally controlled, allow inputs to be extrapolated from last known and if there is a gap allow interpolation between two known inputs.
 	UPROPERTY(config, EditDefaultsOnly, Category = "Overrides", Meta = (EditCondition = "bOverrideAllowInputExtrapolation"))
 	bool bAllowInputExtrapolation = PhysicsReplicationCVars::ResimulationCVars::bAllowInputExtrapolation;
-	bool GetAllowInputExtrapolation() { return bOverrideAllowInputExtrapolation ? bAllowInputExtrapolation : PhysicsReplicationCVars::ResimulationCVars::bAllowInputExtrapolation; }
+	const bool GetAllowInputExtrapolation() const { return bOverrideAllowInputExtrapolation ? bAllowInputExtrapolation : PhysicsReplicationCVars::ResimulationCVars::bAllowInputExtrapolation; }
+
+	UPROPERTY(config, EditDefaultsOnly, Category = "Overrides", Meta = (InlineEditConditionToggle))
+	uint32 bOverrideValidateDataOnGameThread : 1;
+	// Overrides CVar: np2.Resim.ValidateDataOnGameThread -- When true, perform server-side input validation through FNetworkPhysicsData::ValidateData on the Game Thread. If false, perform the call on the Physics Thread.
+	UPROPERTY(config, EditDefaultsOnly, Category = "Overrides", Meta = (EditCondition = "bOverrideValidateDataOnGameThread"))
+	bool bValidateDataOnGameThread = PhysicsReplicationCVars::ResimulationCVars::bValidateDataOnGameThread;
+	const bool GetValidateDataOnGameThread() const { return bOverrideValidateDataOnGameThread ? bValidateDataOnGameThread : PhysicsReplicationCVars::ResimulationCVars::bValidateDataOnGameThread; }
 };
 
 /*
@@ -388,6 +400,9 @@ public:
 
 	virtual void BeginPlay() override;
 
+	/** Get the settings internal to the PhysicsThread (Only access construct on the Physics Thread) */
+	FNetworkPhysicsSettingsComponentAsync* GetNetworkPhysicsSettings_Internal() const { return NetworkPhysicsSettings_Internal; };
+
 public:
 	UPROPERTY(EditDefaultsOnly, Category = "Networked Physics Settings")
 	FNetworkPhysicsSettings GeneralSettings;
@@ -413,12 +428,15 @@ public:
 	*/
 
 private:
-	FNetworkPhysicsSettingsComponentAsync* NetworkPhysicsSettingsAsync;
+	FNetworkPhysicsSettingsComponentAsync* NetworkPhysicsSettings_Internal;
 
 	// Game Thread map of settings component per actor
 	static TMap<AActor*, UNetworkPhysicsSettingsComponent*> ObjectToSettings_External;
 
+
+	/* --- Static API --- */
 public:
+	/** Get the settings component for a specified actor  */
 	static UNetworkPhysicsSettingsComponent* GetSettingsForActor(AActor* Owner);
 };
 
diff --git a/Engine/Source/Runtime/Experimental/Chaos/Private/PBDRigidsSolver.cpp b/Engine/Source/Runtime/Experimental/Chaos/Private/PBDRigidsSolver.cpp
--- a/Engine/Source/Runtime/Experimental/Chaos/Private/PBDRigidsSolver.cpp
+++ b/Engine/Source/Runtime/Experimental/Chaos/Private/PBDRigidsSolver.cpp
@@ -1935,7 +1935,7 @@ namespace Chaos
 				if (MRewindData->RewindToFrame(ResimStep))
 				{
 #if DEBUG_REWIND_DATA
-					UE_LOG(LogTemp, Warning, TEXT("COMMON | PT | ConditionalApplyRewind_Internal | PERFORMING RESIMULATION | Resim From Frame = %d | Num Steps = %d | To Current Frame: %d"), ResimStep, NumResimSteps, CurrentFrame);
+					UE_LOG(LogChaos, Warning, TEXT("CLIENT | PT | ConditionalApplyRewind_Internal | PERFORMING RESIMULATION | Resim From Frame = %d | Num Steps = %d | To Current Frame: %d"), ResimStep, NumResimSteps, CurrentFrame);
 #endif
 
 					SetIsResimming(true);
@@ -2012,14 +2012,14 @@ namespace Chaos
 #if DEBUG_REWIND_DATA
 				else
 				{
-					UE_LOG(LogTemp, Log, TEXT("COMMON | PT | ConditionalApplyRewind_Internal | Resimulation failed, FRewindData::RewindToFrame returned false | Current Frame = %d | Num Steps = %d | Resim Frame = %d | Last Frame = %d | Rewind History Size = %d"), CurrentFrame, NumResimSteps, ResimStep, LastStep, MarshallingManager.GetNumHistory_Internal());
+					UE_LOG(LogChaos, Log, TEXT("CLIENT | PT | ConditionalApplyRewind_Internal | Resimulation failed, FRewindData::RewindToFrame returned false | Current Frame = %d | Num Steps = %d | Resim Frame = %d | Last Frame = %d | Rewind History Size = %d"), CurrentFrame, NumResimSteps, ResimStep, LastStep, MarshallingManager.GetNumHistory_Internal());
 				}
 #endif
 			}
 #if DEBUG_REWIND_DATA
 			else
 			{
-				UE_LOG(LogTemp, Log, TEXT("COMMON | PT | ConditionalApplyRewind_Internal | Resimulation failed, invalid rewind frame data | Current Frame = %d | Num Steps = %d | Resim Frame = %d | Last Frame = %d | Rewind History Size = %d"), CurrentFrame, NumResimSteps, ResimStep, LastStep, MarshallingManager.GetNumHistory_Internal());
+				UE_LOG(LogChaos, Log, TEXT("CLIENT | PT | ConditionalApplyRewind_Internal | Resimulation failed, invalid rewind frame data | Current Frame = %d | Num Steps = %d | Resim Frame = %d | Last Frame = %d | Rewind History Size = %d"), CurrentFrame, NumResimSteps, ResimStep, LastStep, MarshallingManager.GetNumHistory_Internal());
 			}
 #endif
 			// Clear the ResimFrame no matter if resimulation succeeded or failed (if it failed it's not going to succeed next frame either based on the same ResimFrame)
diff --git a/Engine/Source/Runtime/Experimental/Chaos/Private/RewindData.cpp b/Engine/Source/Runtime/Experimental/Chaos/Private/RewindData.cpp
--- a/Engine/Source/Runtime/Experimental/Chaos/Private/RewindData.cpp
+++ b/Engine/Source/Runtime/Experimental/Chaos/Private/RewindData.cpp
@@ -270,7 +270,7 @@ bool FRewindData::RewindToFrame(int32 Frame)
 	if (Frame < EarliestFrame)
 	{
 #if DEBUG_REWIND_DATA
-		UE_LOG(LogTemp, Log, TEXT("COMMON | PT | RewindToFrame | Failed due to rewind frame earlier than available history | Rewind Frame: %d | Earliest Frame: %d"), Frame, EarliestFrame);
+		UE_LOG(LogChaos, Log, TEXT("CLIENT | PT | RewindToFrame | Failed due to rewind frame earlier than available history | Rewind Frame: %d | Earliest Frame: %d"), Frame, EarliestFrame);
 #endif
 		return false;
 	}
@@ -279,7 +279,7 @@ bool FRewindData::RewindToFrame(int32 Frame)
 	if (Frame == EarliestFrame && bNeedsSave && FramesSaved == Managers.Capacity())
 	{
 #if DEBUG_REWIND_DATA
-		UE_LOG(LogTemp, Log, TEXT("COMMON | PT | RewindToFrame | Failed due to rewinding to last available frame and bNeedsSave is set to true"));
+		UE_LOG(LogChaos, Log, TEXT("CLIENT | PT | RewindToFrame | Failed due to rewinding to last available frame and bNeedsSave is set to true"));
 #endif
 		return false;
 	}
@@ -902,7 +902,7 @@ int32 FRewindData::FindValidResimFrame(const int32 RequestedFrame)
 	if (RequestedFrame <= BlockResimFrame)
 	{
 #if DEBUG_REWIND_DATA
-		UE_LOG(LogTemp, Log, TEXT("COMMON | PT | FindValidResimFrame | Resim is blocked | BlockResimFrame: %d | RequestedFrame: %d"), BlockResimFrame, RequestedFrame);
+		UE_LOG(LogChaos, Log, TEXT("CLIENT | PT | FindValidResimFrame | Resim is blocked | BlockResimFrame: %d | RequestedFrame: %d"), BlockResimFrame, RequestedFrame);
 #endif
 
 		return ValidFrame;
@@ -958,7 +958,7 @@ int32 FRewindData::FindValidResimFrame(const int32 RequestedFrame)
 		bHasTargetHistory = true;
 		
 #if DEBUG_REWIND_DATA
-		UE_LOG(LogTemp, Log, TEXT("COMMON | PT | FindValidResimFrame | Processing resim particles | History Frame: %d | Total Particle Count: %d | ResimIslands Particle Count: %d | ResimFrameValidation: %d"), ValidFrame, DirtyParticles.Num(), ResimIslandParticles.Num(), ResimFrameValidation);
+		UE_LOG(LogChaos, Log, TEXT("CLIENT | PT | FindValidResimFrame | Processing resim particles | History Frame: %d | Total Particle Count: %d | ResimIslands Particle Count: %d | ResimFrameValidation: %d"), ValidFrame, DirtyParticles.Num(), ResimIslandParticles.Num(), ResimFrameValidation);
 #endif
 
 		if ((EResimFrameValidation)ResimFrameValidation == EResimFrameValidation::IslandValidation)
@@ -1038,7 +1038,7 @@ int32 FRewindData::FindValidResimFrame(const int32 RequestedFrame)
 		ValidFrame = bResimIncompleteHistory ? RequestedFrame : INDEX_NONE;
 
 #if DEBUG_REWIND_DATA
-		UE_LOG(LogTemp, Warning, TEXT("COMMON | PT | FindValidResimFrame | No valid resim frame found | RequestedFrame: %d | ValidFrame: %d | EarliestFrame: %d | HasTargetHistory: %d | EarliestHistoryFrame: %d | CurrentFrame: %d | FramesSaved: %d | ResimFrameValidation: %d"), RequestedFrame, ValidFrame, EarliestFrame, bHasTargetHistory, GetEarliestFrame_Internal(), CurrentFrame(), FramesSaved, ResimFrameValidation);
+		UE_LOG(LogChaos, Warning, TEXT("CLIENT | PT | FindValidResimFrame | No valid resim frame found | RequestedFrame: %d | ValidFrame: %d | EarliestFrame: %d | HasTargetHistory: %d | EarliestHistoryFrame: %d | CurrentFrame: %d | FramesSaved: %d | ResimFrameValidation: %d"), RequestedFrame, ValidFrame, EarliestFrame, bHasTargetHistory, GetEarliestFrame_Internal(), CurrentFrame(), FramesSaved, ResimFrameValidation);
 #endif
 	}
 
diff --git a/Engine/Source/Runtime/Experimental/Chaos/Public/RewindData.h b/Engine/Source/Runtime/Experimental/Chaos/Public/RewindData.h
--- a/Engine/Source/Runtime/Experimental/Chaos/Public/RewindData.h
+++ b/Engine/Source/Runtime/Experimental/Chaos/Public/RewindData.h
@@ -34,14 +34,27 @@ struct FBaseRewindHistory
 	/** Create a polymorphic copy of the history */
 	virtual TUniquePtr<FBaseRewindHistory> Clone() const = 0;
 
+	/** Initialize history */
+	virtual void Initialize() { }
+
 	/** Set the package map for serialization */
 	FORCEINLINE virtual void SetPackageMap(class UPackageMap* InPackageMap) {}
 
 	/** Check if the history buffer contains an entry for the given frame*/
 	FORCEINLINE virtual bool HasValidData(const int32 ValidFrame) const { return false; }
 	UE_DEPRECATED(5.4, "Deprecated, use HasValidData() instead")
-	FORCEINLINE virtual bool HasValidDatas(const int32 ValidFrame) const { return HasValidData(ValidFrame);
-	}
+	FORCEINLINE virtual bool HasValidDatas(const int32 ValidFrame) const { return HasValidData(ValidFrame);	}
+
+	/** Find how many entries are valid in frame range
+	* @param StartFrame = Included
+	* @param EndFrame = Included
+	* @param bIncludeUnimportant = If to include unimportant data entries
+	* @param bIncludeImportant = If to include important data entries */
+	FORCEINLINE virtual int32 CountValidData(const uint32 StartFrame, const uint32 EndFrame, const bool bIncludeUnimportant = true, const bool bIncludeImportant = false) { return 0; }
+
+	/** Mark frame important or unimportant
+	* @param Frame = use INDEX_NONE to set importance on all entries */
+	FORCEINLINE virtual void SetImportant(const bool bImportant, const int32 Frame = INDEX_NONE) {}
 
 	/** Extract data from the history buffer at a given time */
 	FORCEINLINE virtual bool ExtractData(const int32 ExtractFrame, const bool bResetSolver, void* HistoryData, const bool bExactFrame = false) { return true; }
@@ -59,11 +72,16 @@ struct FBaseRewindHistory
 	UE_DEPRECATED(5.4, "Deprecated, use RecordData() instead")
 	FORCEINLINE virtual bool RecordDatas(const int32 RecordFrame, const void* HistoryDatas) { return RecordData(RecordFrame, HistoryDatas); }
 	
+	/** Copy all data from local history into into @param OutHistory
+	* @param bIncludeUnimportant = If to copy unimportant data entries
+	* @param bIncludeImportant = If to copy important data entries */
+	virtual bool CopyAllData(Chaos::FBaseRewindHistory& OutHistory, bool bIncludeUnimportant = true, bool bIncludeImportant = false) { return false; }
+
 	/** Copy data from local history into @param OutHistory
 	* @param StartFrame = Included
 	* @param EndFrame = Included 
 	* @param bIncludeUnimportant = If to copy unimportant data entries
-	* @param bIncludeImportant = If to copy important data entries*/
+	* @param bIncludeImportant = If to copy important data entries */
 	virtual bool CopyData(Chaos::FBaseRewindHistory& OutHistory, const uint32 StartFrame, const uint32 EndFrame, bool bIncludeUnimportant = true, bool bIncludeImportant = false) { return false; }
 
 	/** Create a polymorphic copy of only a range of frames, applying the frame offset to the copies
@@ -105,8 +123,14 @@ struct FBaseRewindHistory
 	/** Return the max size of the history */
 	virtual const int32 GetHistorySize() const { return 0; }
 
+	/** Return if history has valid data */
+	virtual const bool HasDataInHistory() const { return false; }
+
 	/** Resize the history */
 	virtual void ResizeDataHistory(const int32 FrameCount, const EAllowShrinking AllowShrinking = EAllowShrinking::Yes) {}
+
+	/** Perform a fast reset, marking the data history as reset but not clearing the data or resetting collections */
+	virtual void ResetFast() {}
 };
 
 /** Templated data history holding a data buffer */
@@ -114,18 +138,28 @@ template<typename DataType>
 struct TDataRewindHistory : public FBaseRewindHistory
 {
 	FORCEINLINE TDataRewindHistory(const int32 FrameCount, const bool bIsHistoryLocal) :
-		bIsLocalHistory(bIsHistoryLocal), DataHistory(), LatestFrame(0), CurrentFrame(0), CurrentIndex(0), NumFrames(FrameCount)
+		bIsLocalHistory(bIsHistoryLocal), DataHistory(), LatestFrame(INDEX_NONE), CurrentFrame(0), CurrentIndex(0), NumFrames(FrameCount)
 	{
 		DataHistory.SetNum(NumFrames);
 	}
 
 	FORCEINLINE TDataRewindHistory(const int32 FrameCount) :
-		DataHistory(), LatestFrame(0), CurrentFrame(0), CurrentIndex(0), NumFrames(FrameCount)
+		DataHistory(), LatestFrame(INDEX_NONE), CurrentFrame(0), CurrentIndex(0), NumFrames(FrameCount)
 	{
 		DataHistory.SetNum(NumFrames);
 	}
 	FORCEINLINE virtual ~TDataRewindHistory() {}
 
+	/** Initialize history */
+	virtual void Initialize()
+	{
+		// Iterate over current data to find latest frame
+		for (int32 FrameIndex = 0; FrameIndex < NumFrames; ++FrameIndex)
+		{
+			LatestFrame = FMath::Max(LatestFrame, DataHistory[FrameIndex].LocalFrame);
+		}
+	}
+
 protected:
 
 	/** Get the closest (min/max) valid data from the data frame */
@@ -160,7 +194,7 @@ public :
 	/** Extract states at a given time */
 	FORCEINLINE virtual bool ExtractData(const int32 ExtractFrame, const bool bResetSolver, void* HistoryData, const bool bExactFrame = false) override
 	{
-		// Early out if we are trying to extract data but the latest data is more than the whole buffer size old
+		// Early out if we are trying to extract data later than latest frame but the latest data is more than the whole buffer size old, don't extrapolate that far.
 		if (ExtractFrame - NumFrames > GetLatestFrame())
 		{
 			return false;
@@ -171,19 +205,17 @@ public :
 		{
 			CurrentFrame = ExtractFrame;
 			CurrentIndex = ExtractIndex;
-
-#if DEBUG_NETWORK_PHYSICS
-			UE_LOG(LogTemp, Log, TEXT("		Found matching data into history at frame %d"), ExtractFrame);
-#endif
 			*static_cast<DataType*>(HistoryData) = DataHistory[CurrentIndex];
 			return true;
 		}
 		else if(!bExactFrame)
 		{
+#if DEBUG_REWIND_DATA
 			if (bResetSolver)
 			{
 				UE_LOG(LogChaos, Warning, TEXT("		Unable to extract data at frame %d while rewinding the simulation"), ExtractFrame);
 			}
+#endif
 			PRAGMA_DISABLE_DEPRECATION_WARNINGS // TODO: Change to ClosestData() in UE 5.6 and remove deprecation pragma
 			const int32 MinFrameIndex = ClosestDatas(ExtractFrame, true);
 			const int32 MaxFrameIndex = ClosestDatas(ExtractFrame, false);
@@ -207,7 +239,7 @@ public :
 				ExtractedData.InputFrame = INDEX_NONE; // Clear InputFrame since this history entry is now altered and doesn't correspond to the source entry with the same InputFrame
 
 #if DEBUG_NETWORK_PHYSICS
-				UE_LOG(LogTemp, Log, TEXT("		Smoothing data between frame %d and %d - > [%d %d]"), DataHistory[MinFrameIndex].LocalFrame, DataHistory[MaxFrameIndex].LocalFrame, static_cast<DataType*>(HistoryData)->InputFrame, static_cast<DataType*>(HistoryData)->ServerFrame);
+				UE_LOG(LogChaos, Log, TEXT("		Smoothing data between frame %d and %d - > [%d %d]"), DataHistory[MinFrameIndex].LocalFrame, DataHistory[MaxFrameIndex].LocalFrame, static_cast<DataType*>(HistoryData)->InputFrame, static_cast<DataType*>(HistoryData)->ServerFrame);
 #endif
 				return true;
 			}
@@ -223,14 +255,14 @@ public :
 				ExtractedData.InputFrame = INDEX_NONE; // Clear InputFrame since this history entry is now altered and doesn't correspond to the source entry with the same InputFrame
 
 #if DEBUG_NETWORK_PHYSICS
-				UE_LOG(LogTemp, Log, TEXT("		Setting data to frame %d"), DataHistory[MinFrameIndex].LocalFrame);
+				UE_LOG(LogChaos, Log, TEXT("		Setting data to frame %d"), DataHistory[MinFrameIndex].LocalFrame);
 #endif
 				return true;
 			}
 			else
 			{
 #if DEBUG_NETWORK_PHYSICS
-				UE_LOG(LogTemp, Log, TEXT("		Failed to find data bounds : Min = %d | Max = %d"), MinFrameIndex, MaxFrameIndex);
+				UE_LOG(LogChaos, Log, TEXT("		Failed to find data bounds : Min = %d | Max = %d"), MinFrameIndex, MaxFrameIndex);
 #endif
 				return false;
 			}
@@ -348,6 +380,11 @@ public :
 		return NumFrames;
 	}
 
+	virtual const bool HasDataInHistory() const
+	{
+		return LatestFrame > INDEX_NONE;
+	}
+
 	/** Resize the history */
 	FORCEINLINE void ResizeDataHistory(const int32 FrameCount, const EAllowShrinking AllowShrinking = EAllowShrinking::Yes) override
 	{
@@ -364,6 +401,13 @@ public :
 		return FMath::Abs(Frame % NumFrames);
 	}
 
+	FORCEINLINE virtual void ResetFast()
+	{
+		LatestFrame = INDEX_NONE;
+		CurrentFrame = 0;
+		CurrentIndex = 0;
+	}
+
 protected : 
 
 	/** Check if the history is on the local/remote client*/
@@ -1620,7 +1664,7 @@ public:
 	/** Add input history to the rewind data for future use while resimulating */
 	void AddInputHistory(const TSharedPtr<FBaseRewindHistory>& InputHistory)
 	{
-		InputHistories.Add(InputHistory.ToWeakPtr());
+		InputHistories.AddUnique(InputHistory.ToWeakPtr());
 	}
 	UE_DEPRECATED(5.4, "Deprecated, use AddInputHistory() instead")
 	void AddInputsHistory(const TSharedPtr<FBaseRewindHistory>& InputsHistory) { AddInputHistory(InputsHistory); }
@@ -1636,7 +1680,7 @@ public:
 	/** Add state history to the rewind data for future use while rewinding */
 	void AddStateHistory(const TSharedPtr<FBaseRewindHistory>& StateHistory)
 	{
-		StateHistories.Add(StateHistory.ToWeakPtr());
+		StateHistories.AddUnique(StateHistory.ToWeakPtr());
 	}
 	UE_DEPRECATED(5.4, "Deprecated, use AddStateHistory() instead")
 	void AddStatesHistory(const TSharedPtr<FBaseRewindHistory>& StatesHistory) { AddStateHistory(StatesHistory); }
