[NetPhysics] Fix missing return, fix compilation error.

diff --git a/Engine/Source/Runtime/Engine/Public/Physics/NetworkPhysicsComponent.h b/Engine/Source/Runtime/Engine/Public/Physics/NetworkPhysicsComponent.h
--- a/Engine/Source/Runtime/Engine/Public/Physics/NetworkPhysicsComponent.h
+++ b/Engine/Source/Runtime/Engine/Public/Physics/NetworkPhysicsComponent.h
@@ -778,7 +778,7 @@ public:
 	ENGINE_API const bool GetIsRelayingLocalInputs() const { return bIsRelayingLocalInputs; }
 
 	/** Get FAsyncNetworkPhysicsComponent on the Physics Thread */
-	FAsyncNetworkPhysicsComponent* GetNetworkPhysicsComponent_Internal() { NetworkPhysicsComponent_Internal; }
+	FAsyncNetworkPhysicsComponent* GetNetworkPhysicsComponent_Internal() { return NetworkPhysicsComponent_Internal; }
 
 protected : 
 
