Add some TRACE_SCOPE's to input code which might be hot perf paths for easier testing in Unreal Insights.

#rnx

diff --git a/Engine/Source/Runtime/Engine/Private/Components/InputComponent.cpp b/Engine/Source/Runtime/Engine/Private/Components/InputComponent.cpp
--- a/Engine/Source/Runtime/Engine/Private/Components/InputComponent.cpp
+++ b/Engine/Source/Runtime/Engine/Private/Components/InputComponent.cpp
@@ -25,6 +25,8 @@ UInputComponent::UInputComponent( const FObjectInitializer& ObjectInitializer )
 
 void UInputComponent::ConditionalBuildKeyMap(UPlayerInput* PlayerInput)
 {
+	TRACE_CPUPROFILER_EVENT_SCOPE(UInputComponent::ConditionalBuildKeyMap);
+	
 	if (!ensure(PlayerInput))
 	{
 		return;
@@ -94,6 +96,8 @@ void UInputComponent::ConditionalBuildKeyMap(UPlayerInput* PlayerInput)
 
 void UInputComponent::OnInputOwnerEndPlayed(AActor* InOwner, EEndPlayReason::Type EndPlayReason)
 {
+	TRACE_CPUPROFILER_EVENT_SCOPE(UInputComponent::OnInputOwnerEndPlayed);
+	
 	for (int32 Index = CachedKeyToActionInfo.Num() - 1; Index >= 0; --Index)
 	{
 		FCachedKeyToActionInfo& CachedInfo = CachedKeyToActionInfo[Index];
@@ -107,6 +111,8 @@ void UInputComponent::OnInputOwnerEndPlayed(AActor* InOwner, EEndPlayReason::Typ
 
 void UInputComponent::ClearBindingsForObject(UObject* InOwner)
 {
+	TRACE_CPUPROFILER_EVENT_SCOPE(UInputComponent::ClearBindingsForObject);
+	
 	for (int32 Index = CachedKeyToActionInfo.Num() - 1; Index >= 0; --Index)
 	{
 		FCachedKeyToActionInfo& CachedInfo = CachedKeyToActionInfo[Index];
diff --git a/Engine/Source/Runtime/Engine/Private/UserInterface/PlayerInput.cpp b/Engine/Source/Runtime/Engine/Private/UserInterface/PlayerInput.cpp
--- a/Engine/Source/Runtime/Engine/Private/UserInterface/PlayerInput.cpp
+++ b/Engine/Source/Runtime/Engine/Private/UserInterface/PlayerInput.cpp
@@ -133,6 +133,8 @@ void UPlayerInput::PostInitProperties()
 
 void UPlayerInput::FlushPressedKeys()
 {
+	TRACE_CPUPROFILER_EVENT_SCOPE(UPlayerInput::FlushPressedKeys);
+	
 	APlayerController* PlayerController = GetOuterAPlayerController();
 	ULocalPlayer* LocalPlayer = PlayerController ? Cast<ULocalPlayer>(PlayerController->Player) : nullptr;
 	if (LocalPlayer != nullptr)
@@ -1168,6 +1170,8 @@ void UPlayerInput::ProcessNonAxesKeys(FKey InKey, FKeyState* KeyState)
 
 void UPlayerInput::ProcessInputStack(const TArray<UInputComponent*>& InputComponentStack, const float DeltaTime, const bool bGamePaused)
 {
+	TRACE_CPUPROFILER_EVENT_SCOPE(UPlayerInput::ProcessInputStack);
+	
 	ConditionalBuildKeyMappings();
 
 	static TArray<TPair<FKey, FKeyState*>> KeysWithEvents;
@@ -1200,6 +1204,8 @@ void UPlayerInput::ProcessInputStack(const TArray<UInputComponent*>& InputCompon
 
 void UPlayerInput::EvaluateKeyMapState(const float DeltaTime, const bool bGamePaused, OUT TArray<TPair<FKey, FKeyState*>>& KeysWithEvents)
 {
+	TRACE_CPUPROFILER_EVENT_SCOPE(UPlayerInput::EvaluateKeyMapState);
+	
 	// Must be called non-recursively on the game thread
 	check(IsInGameThread() && !KeysWithEvents.Num());
 	
@@ -1280,6 +1286,8 @@ void UPlayerInput::EvaluateKeyMapState(const float DeltaTime, const bool bGamePa
 
 void UPlayerInput::EvaluateInputDelegates(const TArray<UInputComponent*>& InputComponentStack, const float DeltaTime, const bool bGamePaused, const TArray<TPair<FKey, FKeyState*>>& KeysWithEvents)
 {
+	TRACE_CPUPROFILER_EVENT_SCOPE(UPlayerInput::EvaluateInputDelegates);
+	
 	// We collect axis contributions by delegate, so we can sum up
 	// contributions from multiple bindings.
 	struct FAxisDelegateDetails
