[NetPhysics] Fix compiler warnings.

diff --git a/Engine/Source/Runtime/Engine/Private/PhysicsEngine/NetworkPhysicsComponent.cpp b/Engine/Source/Runtime/Engine/Private/PhysicsEngine/NetworkPhysicsComponent.cpp
--- a/Engine/Source/Runtime/Engine/Private/PhysicsEngine/NetworkPhysicsComponent.cpp
+++ b/Engine/Source/Runtime/Engine/Private/PhysicsEngine/NetworkPhysicsComponent.cpp
@@ -894,6 +894,12 @@ bool UNetworkPhysicsComponent::IsLocallyControlled() const
 	return false;
 }
 
+void UNetworkPhysicsComponent::SetIsRelayingLocalInputs(bool bInRelayingLocalInputs)
+{
+	bIsRelayingLocalInputs = bInRelayingLocalInputs;
+	UpdateAsyncComponent();
+}
+
 APlayerController* UNetworkPhysicsComponent::GetPlayerController() const
 {
 	if (APlayerController* PC = Cast<APlayerController>(GetOwner()))
@@ -1511,7 +1517,8 @@ void FAsyncNetworkPhysicsComponent::SendInputData_Internal(FAsyncNetworkPhysicsC
 			const int32 FromFrame = FMath::Max(LastInputSendFrame + 1, ToFrame - InputHistory->GetHistorySize());
 
 			// Check if we have important data to marshal
-			if (const int32 Count = InputHistory->CountValidData(FromFrame, ToFrame, /*bIncludeUnimportant*/ false, /*bIncludeImportant*/ true) > 0)
+			const int32 Count = InputHistory->CountValidData(FromFrame, ToFrame, /*bIncludeUnimportant*/ false, /*bIncludeImportant*/ true);
+			if (Count > 0)
 			{
 				// Create new data collection for marshaling
 				const int32 Idx = AsyncOutput.InputDataImportant.Add(InputHistory->CreateNew());
@@ -1575,7 +1582,8 @@ void FAsyncNetworkPhysicsComponent::SendStateData_Internal(FAsyncNetworkPhysicsC
 			const int32 FromFrame = FMath::Max(LastStateSendFrame + 1, ToFrame - StateHistory->GetHistorySize());
 
 			// Check if we have important data to marshal
-			if (const int32 Count = StateHistory->CountValidData(FromFrame, ToFrame, /*bIncludeUnimportant*/ false, /*bIncludeImportant*/ true) > 0)
+			const int32 Count = StateHistory->CountValidData(FromFrame, ToFrame, /*bIncludeUnimportant*/ false, /*bIncludeImportant*/ true);
+			if (Count > 0)
 			{
 				// Create new data collection for marshaling
 				const int32 Idx = AsyncOutput.StateDataImportant.Add(StateHistory->CreateNew());
diff --git a/Engine/Source/Runtime/Engine/Public/Physics/NetworkPhysicsComponent.h b/Engine/Source/Runtime/Engine/Public/Physics/NetworkPhysicsComponent.h
--- a/Engine/Source/Runtime/Engine/Public/Physics/NetworkPhysicsComponent.h
+++ b/Engine/Source/Runtime/Engine/Public/Physics/NetworkPhysicsComponent.h
@@ -757,11 +757,7 @@ public:
 	/** Mark this as controlled through locally relayed inputs rather than controlled as a pawn through a player controller.
 	* Set if NetworkPhysicsComponent is implemented on an AActor instead of APawn and it's currently being fed inputs, or if this is controlled by the server. 
 	* NOTE: The actor for this NetworkPhysicsComponent also needs to be owned by the local client if this is used client-side. */
-	ENGINE_API void SetIsRelayingLocalInputs(bool bInRelayingLocalInputs)
-	{
-		bIsRelayingLocalInputs = bInRelayingLocalInputs;
-		UpdateAsyncComponent();
-	}
+	ENGINE_API void SetIsRelayingLocalInputs(bool bInRelayingLocalInputs);
 
 	/** Stop relaying local inputs after next network send.
 	* Deferred version of SetIsRelayingLocalInputs(false) to ensure that the last replicated data gets sent.
