Fixed height fog compilation issue.

diff --git a/Engine/Source/Runtime/Engine/Private/Components/HeightFogComponent.cpp b/Engine/Source/Runtime/Engine/Private/Components/HeightFogComponent.cpp
--- a/Engine/Source/Runtime/Engine/Private/Components/HeightFogComponent.cpp
+++ b/Engine/Source/Runtime/Engine/Private/Components/HeightFogComponent.cpp
@@ -93,7 +93,8 @@ bool UExponentialHeightFogComponent::CanEditChange(const FProperty* InProperty)
 	{
 		FString PropertyName = InProperty->GetName();
 
-		if (DoesProjectSupportExpFogMatchesVolumetricFog())
+		static const auto CVarFog = IConsoleManager::Get().FindTConsoleVariableDataInt(TEXT("r.SupportExpFogMatchesVolumetricFog"));
+		if (CVarFog && CVarFog->GetValueOnAnyThread() > 0)
 		{
 			if (PropertyName == GET_MEMBER_NAME_STRING_CHECKED(UExponentialHeightFogComponent, DirectionalInscatteringExponent) ||
 				PropertyName == GET_MEMBER_NAME_STRING_CHECKED(UExponentialHeightFogComponent, DirectionalInscatteringStartDistance) ||
