Added ForEachSubsystemWithInterface to FObjectSubsystemCollection to allow iterating on subsystems implementing a given interface.
Allowed caching in FSubsystemCollectionBase's SubsystemArrayMap the subsystems that implement a given interface.

diff --git a/Engine/Source/Runtime/Engine/Private/Subsystems/SubsystemCollection.cpp b/Engine/Source/Runtime/Engine/Private/Subsystems/SubsystemCollection.cpp
--- a/Engine/Source/Runtime/Engine/Private/Subsystems/SubsystemCollection.cpp
+++ b/Engine/Source/Runtime/Engine/Private/Subsystems/SubsystemCollection.cpp
@@ -3,6 +3,7 @@
 #include "Subsystems/SubsystemCollection.h"
 
 #include "Subsystems/Subsystem.h"
+#include "UObject/Interface.h"
 #include "UObject/UObjectHash.h"
 #include "UObject/Package.h"
 #include "Modules/ModuleManager.h"
@@ -97,13 +98,15 @@ void FSubsystemCollectionBase::ForEachSubsystem(TFunctionRef<void(USubsystem*)>
 
 FSubsystemCollectionBase::FSubsystemArray& FSubsystemCollectionBase::FindAndPopulateSubsystemArrayInternal(UClass* SubsystemClass) const
 {
+	const bool bIsInterface = SubsystemClass->IsChildOf<UInterface>();
 	if (!SubsystemArrayMap.Contains(SubsystemClass))
 	{
 		TUniquePtr<FSubsystemArray>& NewList = SubsystemArrayMap.Emplace(SubsystemClass, MakeUnique<FSubsystemArray>());
 		for (auto Iter = SubsystemMap.CreateConstIterator(); Iter; ++Iter)
 		{
 			UClass* KeyClass = Iter.Key();
-			if (KeyClass->IsChildOf(SubsystemClass))
+			if ((!bIsInterface && KeyClass->IsChildOf(SubsystemClass)) || 
+				(bIsInterface && KeyClass->ImplementsInterface(SubsystemClass)))
 			{
 				NewList->Subsystems.Add(Iter.Value());
 			}
@@ -297,7 +300,9 @@ USubsystem* FSubsystemCollectionBase::AddAndInitializeSubsystem(UClass* Subsyste
 				// Not calling FatalErrorIfIteratingSubsystems because adding to the end of the array is safe for index-based iteration
 				for (TPair<UClass*, TUniquePtr<FSubsystemArray>>& Pair : SubsystemArrayMap)
 				{
-					if (SubsystemClass->IsChildOf(Pair.Key))
+					const bool bIsInterface = Pair.Key->IsChildOf<UInterface>();
+					if ((!bIsInterface && SubsystemClass->IsChildOf(Pair.Key)) || 
+						(bIsInterface && SubsystemClass->ImplementsInterface(Pair.Key)))
 					{
 						Pair.Value->Subsystems.Add(Subsystem);
 					}
@@ -326,7 +331,9 @@ void FSubsystemCollectionBase::RemoveAndDeinitializeSubsystem(USubsystem* Subsys
 
 	for (TPair<UClass*, TUniquePtr<FSubsystemArray>>& Pair : SubsystemArrayMap)
 	{
-		if (SubsystemClass->IsChildOf(Pair.Key))
+		const bool bIsInterface = Pair.Key->IsChildOf<UInterface>();
+		if ((!bIsInterface && SubsystemClass->IsChildOf(Pair.Key)) || 
+			(bIsInterface && SubsystemClass->ImplementsInterface(Pair.Key)))
 		{
 			UE_CLOG(Pair.Value->bIsIterating, LogSubsystemCollection, Fatal, TEXT("Attempted to deinitialize subsystem %s while iterating subsystems of type %s"),
 				*Subsystem->GetPathName(),
diff --git a/Engine/Source/Runtime/Engine/Public/Subsystems/SubsystemCollection.h b/Engine/Source/Runtime/Engine/Public/Subsystems/SubsystemCollection.h
--- a/Engine/Source/Runtime/Engine/Public/Subsystems/SubsystemCollection.h
+++ b/Engine/Source/Runtime/Engine/Public/Subsystems/SubsystemCollection.h
@@ -232,6 +232,16 @@ public:
 		});
 	}
 
+	template <typename TSubsystemInterface>
+	void ForEachSubsystemWithInterface(TFunctionRef<void(TBaseType*)> Operation) const
+	{
+		UClass* SubsystemInterfaceClass = TSubsystemInterface::StaticClass();
+		ForEachSubsystemOfClass(SubsystemInterfaceClass, [Operation = MoveTemp(Operation)](USubsystem* Subsystem)
+		{
+			Operation(CastChecked<TBaseType>(Subsystem));
+		});
+	}
+
 public:
 
 	/** Construct a FSubsystemCollection, pass in the owning object almost certainly (this). */
