PostProcessMaterial UserSceneTextures: CookDeterminism: Fix FMaterialShaderMapContent uninitialized values UserSceneTexturesX and UserSceneTexturesY when the UMaterial doesn't set them.
#rnx

diff --git a/Engine/Source/Runtime/Engine/Public/MaterialShared.h b/Engine/Source/Runtime/Engine/Public/MaterialShared.h
--- a/Engine/Source/Runtime/Engine/Public/MaterialShared.h
+++ b/Engine/Source/Runtime/Engine/Public/MaterialShared.h
@@ -1382,8 +1382,8 @@ private:
 	LAYOUT_FIELD(FSHAHash, ShaderContentHash);
 
 	LAYOUT_FIELD(FScriptName, UserSceneTextureOutput);
-	LAYOUT_FIELD(int32, UserTextureDivisorX);
-	LAYOUT_FIELD(int32, UserTextureDivisorY);
+	LAYOUT_FIELD_INITIALIZED(int32, UserTextureDivisorX, 0);
+	LAYOUT_FIELD_INITIALIZED(int32, UserTextureDivisorY, 0);
 	LAYOUT_FIELD(FScriptName, ResolutionRelativeToInput);
 
 	LAYOUT_FIELD_EDITORONLY(TMemoryImageArray<FMaterialProcessedSource>, ShaderProcessedSource);
