Added csv stat Ticks/Total which includes Basic/TicksQueued + batched ticks from various systems.
This replaces Basic/TicksQueued as the total for the Ticks Breakdown graph.

Any new managed ticks systems that provide a Ticks/ csv stat need to add their tick count to Ticks/Total.

diff --git a/Engine/Source/Runtime/Engine/Private/LevelTick.cpp b/Engine/Source/Runtime/Engine/Private/LevelTick.cpp
--- a/Engine/Source/Runtime/Engine/Private/LevelTick.cpp
+++ b/Engine/Source/Runtime/Engine/Private/LevelTick.cpp
@@ -1202,13 +1202,17 @@ static void RecordWorldCountsToCSV(UWorld* World, bool bDoingActorTicks)
 			bool bDetailed = (CVarDetailedTickContextForCSV.GetValueOnGameThread() != 0);
 
 			TSortedMap<FName, int32, FDefaultAllocator, FNameFastLess> TickContextToCountMap;
-			int32 EnabledCount;
+			int32 EnabledCount = 0;
 			FTickTaskManagerInterface::Get().GetEnabledTickFunctionCounts(World, TickContextToCountMap, EnabledCount, bDetailed, true);
 
 			for (auto It = TickContextToCountMap.CreateConstIterator(); It; ++It)
 			{
 				FCsvProfiler::Get()->RecordCustomStat(It->Key, CSV_CATEGORY_INDEX(Ticks), It->Value, ECsvCustomStatOp::Accumulate); // use accumulate in case we have more than one world ticking
 			}
+
+			// By default, Ticks/Total equals Basic/TicksQueued.
+			// Ticks/Total can be used to accumulate tick counts from game systems that don't run in the task graph.
+			CSV_CUSTOM_STAT(Ticks, Total, EnabledCount, ECsvCustomStatOp::Accumulate);
 		}
 
 		if (CVarRecordActorCountsToCSV.GetValueOnAnyThread())
