Fix missing transitions on the skin cache tangent accumulation buffers.

FGPUSkinCache::MakeBufferTransitions was skipping transitions if a resource was already in the target state, but that's incorrect for UAV -> UAV. These transitions are necessary in order to serialize dispatches which write to the same UAV. Also, the tangent accumulation buffers can appear multiple times in the array passed to MakeBufferTransitions, so we need to make sure we issue a single transition per resource (which requires sorting the array).

diff --git a/Engine/Source/Runtime/Engine/Private/GPUSkinCache.cpp b/Engine/Source/Runtime/Engine/Private/GPUSkinCache.cpp
--- a/Engine/Source/Runtime/Engine/Private/GPUSkinCache.cpp
+++ b/Engine/Source/Runtime/Engine/Private/GPUSkinCache.cpp
@@ -24,6 +24,7 @@ GPUSkinCache.cpp: Performs skinning on a compute shader into a buffer to avoid v
 #include "RenderingThread.h"
 #include "Stats/StatsTrace.h"
 #include "UObject/UObjectIterator.h"
+#include "Algo/Sort.h"
 #include "ComponentRecreateRenderStateContext.h"
 #include "Algo/Sort.h"
 
