Fix typo in comment.

diff --git a/Engine/Source/Runtime/Engine/Private/Actor.cpp b/Engine/Source/Runtime/Engine/Private/Actor.cpp
--- a/Engine/Source/Runtime/Engine/Private/Actor.cpp
+++ b/Engine/Source/Runtime/Engine/Private/Actor.cpp
@@ -1153,14 +1153,14 @@ static bool ForEachStreamingRelevantComponent(const AActor* InActor, bool bForEd
 			return false;
 		}
 
-		// Transient components shouldn't be part of the streeaming bounds, unless the actor itself is transient.
+		// Transient components shouldn't be part of the streaming bounds, unless the actor itself is transient.
 		// This is to allow transient actors to be loaded in PIE.
 		if (!InActor->HasAnyFlags(RF_Transient) && InComponent->HasAnyFlags(RF_Transient))
 		{
 			return false;
 		}
 
-		// Editor-only components shouldn't be part of the streeaming bounds, unless the actor itself is editor-only.
+		// Editor-only components shouldn't be part of the streaming bounds, unless the actor itself is editor-only.
 		// This is to allow editor-only actors to be loaded in PIE.
 		if (!bForEditor && !InActor->IsEditorOnly() && InComponent->IsEditorOnly())
 		{
