Protect against FRepChangelistState double-free

diff --git a/Engine/Source/Runtime/Engine/Private/RepLayout.cpp b/Engine/Source/Runtime/Engine/Private/RepLayout.cpp
--- a/Engine/Source/Runtime/Engine/Private/RepLayout.cpp
+++ b/Engine/Source/Runtime/Engine/Private/RepLayout.cpp
@@ -1162,6 +1162,10 @@ FRepChangelistState::~FRepChangelistState()
 #if WITH_PUSH_MODEL
 	UE_RepLayout_Private::ConditionallyRemovePushModelObject(PushModelObjectHandle);
 #endif // WITH_PUSH_MODEL
+
+	// Explicitly reset members to improve resilience to double-destruction
+	CustomDeltaChangelistState = nullptr;
+	StaticBuffer.Empty();
 }
 
 #if WITH_PUSH_MODEL
@@ -8381,6 +8385,8 @@ FRepStateStaticBuffer::~FRepStateStaticBuffer()
 	{
 		RepLayout->DestructProperties(*this);
 	}
+
+	Buffer = {};
 }
 
 const TCHAR* LexToString(ERepLayoutFlags Flag)
diff --git a/Engine/Source/Runtime/Engine/Public/Net/RepLayout.h b/Engine/Source/Runtime/Engine/Public/Net/RepLayout.h
--- a/Engine/Source/Runtime/Engine/Public/Net/RepLayout.h
+++ b/Engine/Source/Runtime/Engine/Public/Net/RepLayout.h
@@ -242,6 +242,14 @@ struct FRepSerializationSharedInfo
 		bIsValid(false)
 	{}
 
+	~FRepSerializationSharedInfo()
+	{
+		// Explicitly reset members to improve resilience to double-destruction.
+		SharedPropertyInfo = {};
+		SerializedProperties = nullptr;
+		bIsValid = false;
+	}
+
 	void SetValid()
 	{
 		bIsValid = true;
@@ -302,7 +310,7 @@ struct FRepSerializationSharedInfo
 private:
 
 	/** Whether or not shared serialization data has been successfully built. */
-	bool bIsValid;
+	bool bIsValid = false;
 };
 
 /**
@@ -347,7 +355,7 @@ public:
 	TArray<uint16> Changed;
 
 	/** Whether or not this Changelist should be resent due to a Nak. */
-	bool Resend;
+	bool Resend = false;
 };
 
 /**
@@ -398,6 +406,11 @@ public:
 
 	void CountBytes(FArchive& Ar) const;
 
+	void Empty()
+	{
+		Buffer.Empty();
+	}
+
 private:
 
 	// Properties will be copied in here so memory needs aligned to largest type
@@ -438,13 +451,13 @@ public:
 	TUniquePtr<struct FCustomDeltaChangelistState> CustomDeltaChangelistState;
 
 	/** Index in the buffer where changelist history starts (i.e., the Oldest changelist). */
-	int32 HistoryStart;
+	int32 HistoryStart = 0;
 
 	/** Index in the buffer where changelist history ends (i.e., the Newest changelist). */
-	int32 HistoryEnd;
+	int32 HistoryEnd = 0;
 
 	/** Number of times that properties have been compared */
-	int32 CompareIndex;
+	int32 CompareIndex = 0;
 
 	/** Tracking custom delta sends, for comparison against sending rep state. */
 	uint32 CustomDeltaChangeIndex = 0;
@@ -496,17 +509,22 @@ public:
 
 	~FReplicationChangelistMgr();
 
-	FRepChangelistState* GetRepChangelistState() const
+	FRepChangelistState* GetRepChangelistState()
+	{
+		return &RepChangelistState;
+	}
+
+	const FRepChangelistState* GetRepChangelistState() const
 	{
-		return const_cast<FRepChangelistState*>(&RepChangelistState);
+		return &RepChangelistState;
 	}
 
 	void CountBytes(FArchive& Ar) const;
 
 private:
 
-	uint32 LastReplicationFrame;
-	uint32 LastInitialReplicationFrame;
+	uint32 LastReplicationFrame = 0;
+	uint32 LastInitialReplicationFrame = 0;
 
 	FRepChangelistState RepChangelistState;
 };
