Client optimizations when waiting for the shadermaps to come back from the ODSC server. The hashing done in FODSCThread::AddShaderPipelineRequest is expensive, so we need to avoid doing it if we can.
* FShaderId already takes into account ShaderPipelineName/VFType->GetHashedName/ShaderTypeNames->GetHashedName. The material path is the only criterion which is not part of FShaderId, so use it directly during the map lookup. Only do the hashing when we actually need to send the permutation for compilation on the server.
* The assumption is that MaterialHash+VFHash+ShaderHash is good enough to avoid collision
* Replace the Key from FShaderId to FODSCShaderId to rely on hashed names instead of pointers. GetTypeHash in FShaderId is relying on dereferencing the pointers which might be unsafe if the shader/VF gets deleted behind our back
* Change CriticalSection to RWLock to allow multiple readers to test if an ODSC request has already been submitted
* Make a first read-only pass in TryGetShaders so that we call AddShaderPipelineRequest only once. That will also prevent calling ShouldCompilePermutation every frame until we receive the shadermap from ODSC

#rnx

diff --git a/Engine/Source/Runtime/Engine/Private/Materials/MaterialShared.cpp b/Engine/Source/Runtime/Engine/Private/Materials/MaterialShared.cpp
--- a/Engine/Source/Runtime/Engine/Private/Materials/MaterialShared.cpp
+++ b/Engine/Source/Runtime/Engine/Private/Materials/MaterialShared.cpp
@@ -3772,6 +3772,30 @@ bool FMaterial::TryGetShaders(const FMaterialShaderTypes& InTypes, const FVertex
 				}
 			}
 		}
+
+#if WITH_ODSC
+		TArray<FShaderId> RequestShaderIds;
+		TArray<FString> ShaderStageNamesToCompile;
+		FString MaterialName;
+        bool bODSCRequestAlreadySent = false;
+		if (bShouldForceRecompile)
+		{
+			MaterialName = GetFullPath();
+			for (auto* ShaderType : InTypes.PipelineType->GetStages())
+			{
+				ShaderStageNamesToCompile.Add(ShaderType->GetName());
+				RequestShaderIds.Add(FShaderId(ShaderType, ShaderMap->GetShaderMapId().CookedShaderMapIdHash, InTypes.PipelineType->GetHashedName(), InVertexFactoryType, kUniqueShaderPermutationId, ShaderPlatform));
+			}
+
+			bODSCRequestAlreadySent = GODSCManager->CheckIfRequestAlreadySent(RequestShaderIds, MaterialName);
+			if (bODSCRequestAlreadySent)
+			{
+				bShouldForceRecompile = false;
+			}
+
+		}
+#endif
+
         // we don't do 'else' here because when bShouldForceRecompile is true, we still want to use the current 
         // pipeline until we have a new one ready. The ODSC server might fail to find the right shader, and this might
         // skew results when doing some A/B comparisons		
@@ -3794,21 +3818,12 @@ bool FMaterial::TryGetShaders(const FMaterialShaderTypes& InTypes, const FVertex
 #endif // WITH_EDITOR || WITH_ODSC
 
 #if WITH_ODSC
-				if (FPlatformProperties::RequiresCookedData())
+				if (FPlatformProperties::RequiresCookedData() && !bODSCRequestAlreadySent)
 				{
 					if (GODSCManager->IsHandlingRequests())
 					{
-						const FString MaterialName = GetFullPath();
 						const FString VFTypeName(InVertexFactoryType ? InVertexFactoryType->GetName() : TEXT(""));
 						const FString PipelineName(InTypes.PipelineType->GetName());
-						TArray<FString> ShaderStageNamesToCompile;
-						TArray<FShaderId> RequestShaderIds;
-						for (auto* ShaderType : InTypes.PipelineType->GetStages())
-						{
-							ShaderStageNamesToCompile.Add(ShaderType->GetName());
-							RequestShaderIds.Add(FShaderId(ShaderType, ShaderMap->GetShaderMapId().CookedShaderMapIdHash, InTypes.PipelineType->GetHashedName(), InVertexFactoryType, kUniqueShaderPermutationId, ShaderPlatform));
-						}
-
 						GODSCManager->AddThreadedShaderPipelineRequest(ShaderPlatform, GetFeatureLevel(), GetQualityLevel(), MaterialName, VFTypeName, PipelineName, 
 						                                               ShaderStageNamesToCompile, kUniqueShaderPermutationId, RequestShaderIds);
 					}
@@ -3879,6 +3894,24 @@ bool FMaterial::TryGetShaders(const FMaterialShaderTypes& InTypes, const FVertex
 		        // we don't do 'else' here because when bShouldForceRecompile is true, we still want to use the current 
 		        // shader until we have a new one ready. The ODSC server might fail to find the right shader, and this might
 		        // skew results when doing some A/B comparisons		
+#if WITH_ODSC
+
+				TArray<FShaderId> RequestShaderIds;
+				FString MaterialName;
+                bool bODSCRequestAlreadySent = false;
+				if (bShouldForceRecompile)
+				{
+					MaterialName = GetFullPath();
+					RequestShaderIds.Add(FShaderId(ShaderType, ShaderMap->GetShaderMapId().CookedShaderMapIdHash, FHashedName(), InVertexFactoryType, PermutationId, ShaderPlatform));
+
+					bODSCRequestAlreadySent = GODSCManager->CheckIfRequestAlreadySent(RequestShaderIds, MaterialName);
+					if (bODSCRequestAlreadySent)
+					{
+						bShouldForceRecompile = false;
+					}
+				}
+#endif
+				
 				if (Shader == nullptr || bShouldForceRecompile)
 				{
 					bMissingShader = (Shader == nullptr);
@@ -3891,19 +3924,15 @@ bool FMaterial::TryGetShaders(const FMaterialShaderTypes& InTypes, const FVertex
 #endif // WITH_EDITOR || WITH_ODSC
 
 #if WITH_ODSC
-					if (FPlatformProperties::RequiresCookedData())
+					if (FPlatformProperties::RequiresCookedData() && !bODSCRequestAlreadySent)
 					{
 						if (GODSCManager->IsHandlingRequests())
 						{
-							const FString MaterialName = GetFullPath();
 							const FString VFTypeName(InVertexFactoryType ? InVertexFactoryType->GetName() : TEXT(""));
 							const FString PipelineName;
 							TArray<FString> ShaderStageNamesToCompile;
 							ShaderStageNamesToCompile.Add(ShaderType->GetName());
 
-							TArray<FShaderId> RequestShaderIds;
-							RequestShaderIds.Add(FShaderId(ShaderType, ShaderMap->GetShaderMapId().CookedShaderMapIdHash, FHashedName(), InVertexFactoryType, PermutationId, ShaderPlatform));
-
 							GODSCManager->AddThreadedShaderPipelineRequest(ShaderPlatform, GetFeatureLevel(), GetQualityLevel(), MaterialName, VFTypeName, PipelineName, 
 							                                               ShaderStageNamesToCompile, PermutationId, RequestShaderIds);
 						}
diff --git a/Engine/Source/Runtime/Engine/Private/ODSC/ODSCManager.cpp b/Engine/Source/Runtime/Engine/Private/ODSC/ODSCManager.cpp
--- a/Engine/Source/Runtime/Engine/Private/ODSC/ODSCManager.cpp
+++ b/Engine/Source/Runtime/Engine/Private/ODSC/ODSCManager.cpp
@@ -312,6 +312,15 @@ void FODSCManager::RegisterMaterialShaderMap(const FMaterialShaderMap& MaterialS
 	}
 }
  
+bool FODSCManager::CheckIfRequestAlreadySent(const TArray<FShaderId>& RequestShaderIds, const FString& MaterialName) const
+{
+	if (Thread)
+	{
+		return Thread->CheckIfRequestAlreadySent(RequestShaderIds, MaterialName);
+	}
+	return false;
+}
+
 void FODSCManager::TryLoadGlobalShaders(EShaderPlatform ShaderPlatform)
 {
 	check(IsODSCActive());
diff --git a/Engine/Source/Runtime/Engine/Private/ODSC/ODSCThread.cpp b/Engine/Source/Runtime/Engine/Private/ODSC/ODSCThread.cpp
--- a/Engine/Source/Runtime/Engine/Private/ODSC/ODSCThread.cpp
+++ b/Engine/Source/Runtime/Engine/Private/ODSC/ODSCThread.cpp
@@ -168,7 +168,7 @@ void FODSCThread::ResetMaterialsODSCData(ERHIFeatureLevel::Type FeatureLevel)
 	FlushRenderingCommands();
 
 	{
-		FScopeLock Lock(&RequestHashCriticalSection);
+		FWriteScopeLock WriteLock(RequestHashesRWLock);
 
 		// this will stop the rendering thread, and reattach components, in the destructor
 		FMaterialUpdateContext UpdateContext(FMaterialUpdateContext::EOptions::Default);
@@ -194,6 +194,35 @@ void FODSCThread::ResetMaterialsODSCData(ERHIFeatureLevel::Type FeatureLevel)
 #endif
 }
 
+FODSCThread::FODSCShaderId::FODSCShaderId(const FShaderId& ShaderId)
+: MaterialShaderMapHash(ShaderId.MaterialShaderMapHash)
+, ShaderTypeHashedName(ShaderId.Type ? ShaderId.Type->GetHashedName() : 0)
+, VFTypeHashedName(ShaderId.VFType ? ShaderId.VFType->GetHashedName() : 0)
+, ShaderPipelineName(ShaderId.ShaderPipelineName)
+, PermutationId(ShaderId.PermutationId)
+, Platform(ShaderId.Platform)
+{}
+
+bool FODSCThread::CheckIfRequestAlreadySent(const TArray<FShaderId>& RequestShaderIds, const FString& MaterialName) const
+{
+	FReadScopeLock ReadLock(RequestHashesRWLock);
+	for (const FShaderId& ShaderId : RequestShaderIds)
+	{
+		const FMaterialRequestsHashes* MaterialRequestHashes = RequestHashes.Find(ShaderId);
+		if (MaterialRequestHashes == nullptr)
+		{
+			return false;
+		}
+			
+		if (!MaterialRequestHashes->RequestStrings.Contains(MaterialName))
+		{
+			return false;
+		}
+	}
+
+	return true;
+}
+
 void FODSCThread::AddRequest(const TArray<FString>& MaterialsToCompile, const FString& ShaderTypesToLoad, EShaderPlatform ShaderPlatform, ERHIFeatureLevel::Type FeatureLevel, EMaterialQualityLevel::Type QualityLevel, ODSCRecompileCommand RecompileCommandType)
 {
 	PendingMaterialThreadedRequests.Enqueue(new FODSCMessageHandler(MaterialsToCompile, ShaderTypesToLoad, ShaderPlatform, FeatureLevel, QualityLevel, RecompileCommandType));
@@ -211,41 +240,39 @@ void FODSCThread::AddShaderPipelineRequest(
 	const TArray<FShaderId>& RequestShaderIds
 )
 {
-	// TODO: Requests for individual permutations come in here, but a single coalesced payload is submitted to the server since 
-	// we compile all material shader permutations encountered for the moment. Consider batching up requested permutations and 
-	// have the server skip compiling those not in the list. Ensure that DDC key and shader map assumptions are correct!
-
-	FString RequestString = (MaterialName + VertexFactoryName + PipelineName);
-	for (const auto& ShaderTypeName : ShaderTypeNames)
-	{
-		RequestString += ShaderTypeName;
-	}
-	const FString RequestHash = FMD5::HashAnsiString(*RequestString);
-
-	FScopeLock Lock(&RequestHashCriticalSection);
-
 	bool bShouldAddRequest = false;
-
-	for (const FShaderId& ShaderId : RequestShaderIds)
 	{
-		FMaterialRequestsHashes& MaterialRequestHashes = RequestHashes.FindOrAdd(ShaderId);
-		bool bAlreadyInSet = false;
-		MaterialRequestHashes.RequestStrings.Add(RequestString, &bAlreadyInSet);
-		if (!bAlreadyInSet)
+		FWriteScopeLock WriteLock(RequestHashesRWLock);
+
+		for (const FShaderId& ShaderId : RequestShaderIds)
 		{
-			bShouldAddRequest = true;
+			FMaterialRequestsHashes& MaterialRequestHashes = RequestHashes.FindOrAdd(ShaderId);
+			bool bAlreadyInSet = false;
+			MaterialRequestHashes.RequestStrings.Add(MaterialName, &bAlreadyInSet);
+			if (!bAlreadyInSet)
+			{
+				bShouldAddRequest = true;
+			}
 		}
 	}
 
 	if (bShouldAddRequest)
 	{
+		SCOPED_NAMED_EVENT(AddShaderPipelineRequest_AddRequest, FColor::Emerald);
+
+		FString RequestString = (MaterialName + VertexFactoryName + PipelineName);
+		for (const auto& ShaderTypeName : ShaderTypeNames)
+		{
+			RequestString += ShaderTypeName;
+		}
+		const FString RequestHash = FMD5::HashAnsiString(*RequestString);
 		PendingMeshMaterialThreadedRequests.Enqueue(FODSCRequestPayload(ShaderPlatform, FeatureLevel, QualityLevel, MaterialName, VertexFactoryName, PipelineName, ShaderTypeNames, PermutationId, RequestHash));
 	}
 }
 
 void FODSCThread::RegisterMaterialShaderMap(const FMaterialShaderMap& MaterialShaderMap)
 {
-	FScopeLock Lock(&RequestHashCriticalSection);
+	FWriteScopeLock WriteLock(RequestHashesRWLock);
 
 	TMap<FShaderId, TShaderRef<FShader>> ShadersInMap;
 	MaterialShaderMap.GetShaderList(ShadersInMap);
@@ -319,7 +346,6 @@ void FODSCThread::Process()
 	FODSCRequestPayload Payload;
 	TArray<FODSCRequestPayload> PayloadsToAggregate;
 	{
-		FScopeLock Lock(&RequestHashCriticalSection);
 		while (PendingMeshMaterialThreadedRequests.Dequeue(Payload))
 		{
 			PayloadsToAggregate.Add(Payload);
diff --git a/Engine/Source/Runtime/Engine/Private/ODSC/ODSCThread.h b/Engine/Source/Runtime/Engine/Private/ODSC/ODSCThread.h
--- a/Engine/Source/Runtime/Engine/Private/ODSC/ODSCThread.h
+++ b/Engine/Source/Runtime/Engine/Private/ODSC/ODSCThread.h
@@ -168,6 +168,9 @@ public:
 	void RegisterMaterialShaderMap(const FMaterialShaderMap& MaterialShaderMap);
 
 	void ResetMaterialsODSCData(ERHIFeatureLevel::Type FeatureLevel);
+
+	bool CheckIfRequestAlreadySent(const TArray<FShaderId>& RequestShaderIds, const FString& MaterialName) const;
+
 protected:
 
 	//~ Begin FRunnable Interface
@@ -209,6 +212,7 @@ private:
 	TQueue<FODSCMessageHandler*, EQueueMode::Spsc> CompletedThreadedRequests;
 
 	/** Lock to access the RequestHashes TMap */
+	mutable FRWLock RequestHashesRWLock;
 	FCriticalSection RequestHashCriticalSection;
 
 
@@ -217,8 +221,48 @@ private:
 		TSet<FString> RequestStrings;
 	};
 
+	struct FODSCShaderId
+	{
+	public:
+		inline FODSCShaderId() {}
+		FODSCShaderId(const FShaderId& ShaderId);
+
+		FSHAHash MaterialShaderMapHash;
+		FHashedName ShaderTypeHashedName = 0;
+		FHashedName VFTypeHashedName = 0;
+		FHashedName ShaderPipelineName = 0;
+		int32 PermutationId = 0;
+		uint32 Platform : SP_NumBits = SP_NumPlatforms;
+
+		friend inline uint32 GetTypeHash( const FODSCShaderId& Id )
+		{
+			return HashCombine(
+				GetTypeHash(Id.ShaderTypeHashedName),
+				HashCombine(GetTypeHash(Id.VFTypeHashedName),
+							HashCombine(GetTypeHash(Id.ShaderPipelineName),
+										HashCombine(GetTypeHash(Id.MaterialShaderMapHash),
+													HashCombine(GetTypeHash(Id.PermutationId), GetTypeHash(Id.Platform))))));
+		}
+
+		friend bool operator==(const FODSCShaderId& X, const FODSCShaderId& Y)
+		{
+			return X.ShaderTypeHashedName == Y.ShaderTypeHashedName
+			&& X.MaterialShaderMapHash == Y.MaterialShaderMapHash
+			&& X.ShaderPipelineName == Y.ShaderPipelineName
+			&& X.VFTypeHashedName == Y.VFTypeHashedName
+			&& X.PermutationId == Y.PermutationId 
+			&& X.Platform == Y.Platform;
+		}
+
+		friend bool operator!=(const FODSCShaderId& X, const FODSCShaderId& Y)
+		{
+			return !(X == Y);
+		}
+	};
+
+
 	/** Hashes for all Pending or Completed requests.  This is so we avoid making the same request multiple times. */
-	TMap<FShaderId, FMaterialRequestsHashes> RequestHashes;
+	TMap<FODSCShaderId, FMaterialRequestsHashes> RequestHashes;
 
 	/** Pointer to Runnable Thread */
 	FRunnableThread* Thread = nullptr;
diff --git a/Engine/Source/Runtime/Engine/Public/ODSC/ODSCManager.h b/Engine/Source/Runtime/Engine/Public/ODSC/ODSCManager.h
--- a/Engine/Source/Runtime/Engine/Public/ODSC/ODSCManager.h
+++ b/Engine/Source/Runtime/Engine/Public/ODSC/ODSCManager.h
@@ -109,6 +109,8 @@ public:
 
 	static void ReportODSCError(const FString& InErrorMessage);
 
+	bool CheckIfRequestAlreadySent(const TArray<FShaderId>& RequestShaderIds, const FString& MaterialName) const;
+
 private:
 
 	ENGINE_API void OnEnginePreExit();
