Deprecated UWorldSubsystem::UpdateStreamingState in favor of using IStreamingSubsystemInterface.
This allows to iterate only on WorldSubsystems that implement IStreamingSubsystemInterface (reduces the overhead of visiting all world subsystems and calling the default empty implementation, which can represent a considerable amount of wasted time).

diff --git a/Engine/Source/Runtime/Engine/Private/LevelInstance/LevelInstanceSubsystem.cpp b/Engine/Source/Runtime/Engine/Private/LevelInstance/LevelInstanceSubsystem.cpp
--- a/Engine/Source/Runtime/Engine/Private/LevelInstance/LevelInstanceSubsystem.cpp
+++ b/Engine/Source/Runtime/Engine/Private/LevelInstance/LevelInstanceSubsystem.cpp
@@ -255,7 +255,7 @@ bool ULevelInstanceSubsystem::IsLoading(const ILevelInstanceInterface* LevelInst
 	return LevelInstance->HasValidLevelInstanceID() && LoadingLevelInstances.Contains(LevelInstance->GetLevelInstanceID());
 }
 
-void ULevelInstanceSubsystem::UpdateStreamingState()
+void ULevelInstanceSubsystem::OnUpdateStreamingState()
 {
 	TRACE_CPUPROFILER_EVENT_SCOPE(ULevelInstanceSubsystem::UpdateStreamingState);
 
@@ -918,7 +918,7 @@ void ULevelInstanceSubsystem::Tick()
 	// For non-game world, Tick is responsible of processing LevelInstances to update/load/unload
 	if (!GetWorld()->IsGameWorld())
 	{
-		UpdateStreamingState();
+		OnUpdateStreamingState();
 
 		// Update Editor Mode if we are the GEditor world's Subsystem
 		if (GEditor->GetEditorWorldContext().World() == GetWorld())
@@ -2777,7 +2777,7 @@ void ULevelInstanceSubsystem::BlockOnLoading()
 	// Blocking until LevelInstance is loaded and all its child LevelInstances
 	while (LevelInstancesToLoadOrUpdate.Num())
 	{
-		UpdateStreamingState();
+		OnUpdateStreamingState();
 	}
 }
 
diff --git a/Engine/Source/Runtime/Engine/Private/Subsystems/WorldSubsystem.cpp b/Engine/Source/Runtime/Engine/Private/Subsystems/WorldSubsystem.cpp
--- a/Engine/Source/Runtime/Engine/Private/Subsystems/WorldSubsystem.cpp
+++ b/Engine/Source/Runtime/Engine/Private/Subsystems/WorldSubsystem.cpp
@@ -3,6 +3,7 @@
 #include "Subsystems/WorldSubsystem.h"
 #include "Engine/World.h"
 #include "Subsystems/Subsystem.h"
+#include "Streaming/StreamingWorldSubsystemInterface.h"
 
 #include UE_INLINE_GENERATED_CPP_BY_NAME(WorldSubsystem)
 
@@ -41,6 +42,15 @@ bool UWorldSubsystem::DoesSupportWorldType(const EWorldType::Type WorldType) con
 	return WorldType == EWorldType::Game || WorldType == EWorldType::Editor || WorldType == EWorldType::PIE;
 }
 
+PRAGMA_DISABLE_DEPRECATION_WARNINGS
+void UWorldSubsystem::UpdateStreamingState()
+{
+	if (IStreamingWorldSubsystemInterface* StreamingWorldSubsystem = Cast<IStreamingWorldSubsystemInterface>(this))
+	{
+		StreamingWorldSubsystem->OnUpdateStreamingState();
+	}
+}
+PRAGMA_ENABLE_DEPRECATION_WARNINGS
 
 // ----------------------------------------------------------------------------------
 
diff --git a/Engine/Source/Runtime/Engine/Private/World.cpp b/Engine/Source/Runtime/Engine/Private/World.cpp
--- a/Engine/Source/Runtime/Engine/Private/World.cpp
+++ b/Engine/Source/Runtime/Engine/Private/World.cpp
@@ -79,6 +79,7 @@
 #include "ProfilingDebugging/LoadTimeTracker.h"
 #include "Streaming/ServerStreamingLevelsVisibility.h"
 #include "Streaming/LevelStreamingDelegates.h"
+#include "Streaming/StreamingWorldSubsystemInterface.h"
 
 #if WITH_EDITOR
 	#include "DerivedDataCacheInterface.h"
@@ -4271,10 +4272,11 @@ void UWorld::InternalUpdateStreamingState()
 	{
 		WorldComposition->UpdateStreamingState();
 	}
-
+	
 	// Update World Subsystems required streaming levels
-	SubsystemCollection.ForEachSubsystem([](UWorldSubsystem* WorldSubsystem){
-	    WorldSubsystem->UpdateStreamingState();
+	SubsystemCollection.ForEachSubsystemWithInterface<UStreamingWorldSubsystemInterface>([](UWorldSubsystem* WorldSubsystem)
+	{
+		CastChecked<IStreamingWorldSubsystemInterface>(WorldSubsystem)->OnUpdateStreamingState();
 	});
 }
 
diff --git a/Engine/Source/Runtime/Engine/Private/WorldPartition/WorldPartitionSubsystem.cpp b/Engine/Source/Runtime/Engine/Private/WorldPartition/WorldPartitionSubsystem.cpp
--- a/Engine/Source/Runtime/Engine/Private/WorldPartition/WorldPartitionSubsystem.cpp
+++ b/Engine/Source/Runtime/Engine/Private/WorldPartition/WorldPartitionSubsystem.cpp
@@ -1101,7 +1101,7 @@ void UWorldPartitionSubsystem::OnWorldBeginPlay(UWorld& InWorld)
 	}
 }
 
-void UWorldPartitionSubsystem::UpdateStreamingState()
+void UWorldPartitionSubsystem::OnUpdateStreamingState()
 {
 	SCOPE_CYCLE_COUNTER(STAT_WorldPartitionUpdateStreamingState);
 
diff --git a/Engine/Source/Runtime/Engine/Public/LevelInstance/LevelInstanceSubsystem.h b/Engine/Source/Runtime/Engine/Public/LevelInstance/LevelInstanceSubsystem.h
--- a/Engine/Source/Runtime/Engine/Public/LevelInstance/LevelInstanceSubsystem.h
+++ b/Engine/Source/Runtime/Engine/Public/LevelInstance/LevelInstanceSubsystem.h
@@ -12,6 +12,7 @@
 #include "WorldPartition/WorldPartitionHandle.h"
 #include "WorldPartition/WorldPartitionActorContainerID.h"
 #include "LevelInstance/LevelInstancePropertyOverrideAsset.h"
+#include "Streaming/StreamingWorldSubsystemInterface.h"
 
 #if WITH_EDITOR
 #include "EditorLevelUtils.h"
@@ -50,7 +51,7 @@ ENUM_CLASS_FLAGS(ELevelInstanceBreakFlags);
  * ULevelInstanceSubsystem
  */
 UCLASS(MinimalAPI)
-class ULevelInstanceSubsystem : public UWorldSubsystem
+class ULevelInstanceSubsystem : public UWorldSubsystem, public IStreamingWorldSubsystemInterface
 {
 	GENERATED_BODY()
 
@@ -68,9 +69,9 @@ public:
 	ENGINE_API virtual bool DoesSupportWorldType(EWorldType::Type WorldType) const override;
 	//~ End USubsystem Interface.
 	
-	//~ Begin UWorldSubsystem Interface.
-	ENGINE_API virtual void UpdateStreamingState() override;
-	//~ End UWorldSubsystem Interface.
+	//~Begin IStreamingWorldSubsystemInterface
+	ENGINE_API virtual void OnUpdateStreamingState() override;
+	//~End IStreamingWorldSubsystemInterface
 
 	ENGINE_API ILevelInstanceInterface* GetLevelInstance(const FLevelInstanceID& LevelInstanceID) const;
 	ENGINE_API ILevelInstanceInterface* GetOwningLevelInstance(const ULevel* Level) const;
diff --git a/Engine/Source/Runtime/Engine/Public/Streaming/StreamingWorldSubsystemInterface.h b/Engine/Source/Runtime/Engine/Public/Streaming/StreamingWorldSubsystemInterface.h
new file mode 100644
--- /dev/null
+++ b/Engine/Source/Runtime/Engine/Public/Streaming/StreamingWorldSubsystemInterface.h
@@ -0,0 +1,22 @@
+// Copyright Epic Games, Inc. All Rights Reserved.
+
+#pragma once
+
+#include "UObject/Interface.h"
+#include "StreamingWorldSubsystemInterface.generated.h"
+
+/** Interface for world subsystems that require an update for streaming (called by UWorld::InternalUpdateStreamingState) */
+UINTERFACE(meta = (CannotImplementInterfaceInBlueprint))
+class UStreamingWorldSubsystemInterface : public UInterface
+{
+	GENERATED_BODY()
+};
+
+class IStreamingWorldSubsystemInterface
+{
+	GENERATED_BODY()
+
+public:
+
+	virtual void OnUpdateStreamingState() {}
+};
diff --git a/Engine/Source/Runtime/Engine/Public/Subsystems/WorldSubsystem.h b/Engine/Source/Runtime/Engine/Public/Subsystems/WorldSubsystem.h
--- a/Engine/Source/Runtime/Engine/Public/Subsystems/WorldSubsystem.h
+++ b/Engine/Source/Runtime/Engine/Public/Subsystems/WorldSubsystem.h
@@ -38,8 +38,8 @@ public:
 	/** Called after world components (e.g. line batcher and all level components) have been updated */
 	virtual void OnWorldComponentsUpdated(UWorld& World) {}
 
-	/** Updates sub-system required streaming levels (called by world's UpdateStreamingState function) */
-	virtual void UpdateStreamingState() {}
+	UE_DEPRECATED(5.5, "Implement IStreamingWorldSubsystemInterface instead.")
+	ENGINE_API virtual void UpdateStreamingState();
 
 protected:
 	ENGINE_API virtual bool DoesSupportWorldType(const EWorldType::Type WorldType) const;
diff --git a/Engine/Source/Runtime/Engine/Public/WorldPartition/WorldPartitionSubsystem.h b/Engine/Source/Runtime/Engine/Public/WorldPartition/WorldPartitionSubsystem.h
--- a/Engine/Source/Runtime/Engine/Public/WorldPartition/WorldPartitionSubsystem.h
+++ b/Engine/Source/Runtime/Engine/Public/WorldPartition/WorldPartitionSubsystem.h
@@ -7,6 +7,7 @@
 #include "WorldPartition/Filter/WorldPartitionActorFilter.h"
 #include "WorldPartition/WorldPartitionActorContainerID.h"
 #include "WorldPartition/WorldPartitionHandle.h"
+#include "Streaming/StreamingWorldSubsystemInterface.h"
 #include "Containers/Map.h"
 #include "Containers/Set.h"
 #include "Misc/Guid.h"
@@ -47,7 +48,7 @@ private:
  */
 
 UCLASS(MinimalAPI)
-class UWorldPartitionSubsystem : public UTickableWorldSubsystem
+class UWorldPartitionSubsystem : public UTickableWorldSubsystem, public IStreamingWorldSubsystemInterface
 {
 	GENERATED_BODY()
 
@@ -61,10 +62,13 @@ public:
 
 	//~ Begin UWorldSubsystem Interface.
 	ENGINE_API virtual void OnWorldBeginPlay(UWorld& InWorld) override;
-	ENGINE_API virtual void UpdateStreamingState() override;
 	ENGINE_API virtual void OnWorldComponentsUpdated(UWorld& World) override;
 	//~ End UWorldSubsystem Interface.
 
+	//~Begin IStreamingWorldSubsystemInterface
+	ENGINE_API virtual void OnUpdateStreamingState() override;
+	//~End IStreamingWorldSubsystemInterface
+
 	//~ Begin FTickableGameObject
 	ENGINE_API virtual void Tick(float DeltaSeconds) override;
 	virtual bool IsTickableInEditor() const override { return true; }
diff --git a/Engine/Source/Runtime/MovieSceneTracks/Private/Systems/MovieSceneDataLayerSystem.cpp b/Engine/Source/Runtime/MovieSceneTracks/Private/Systems/MovieSceneDataLayerSystem.cpp
--- a/Engine/Source/Runtime/MovieSceneTracks/Private/Systems/MovieSceneDataLayerSystem.cpp
+++ b/Engine/Source/Runtime/MovieSceneTracks/Private/Systems/MovieSceneDataLayerSystem.cpp
@@ -560,7 +560,7 @@ void UMovieSceneDataLayerSystem::OnRun(FSystemTaskPrerequisites& InPrerequisites
 				FlushTypeString = TEXT("FlushStreamingVisibility");
 
 				// Make sure any DataLayer state change is processed before flushing visibility					
-				WorldPartitionSubsystem->UpdateStreamingState();
+				WorldPartitionSubsystem->OnUpdateStreamingState();
 				World->FlushLevelStreaming(EFlushLevelStreamingType::Visibility);
 			}
 
