- Allow editor to load non-confidential platform extension DP.ini files
- Always load the GlobalDefaults DP, not just when ALLOW_OTHER_PLATFORM_CONFIG is defined
- Allow a couple of dp exec commands when ALLOW_OTHER_PLATFORM_CONFIG is 0

diff --git a/Engine/Source/Runtime/Engine/Private/DeviceProfiles/DeviceProfileManager.cpp b/Engine/Source/Runtime/Engine/Private/DeviceProfiles/DeviceProfileManager.cpp
--- a/Engine/Source/Runtime/Engine/Private/DeviceProfiles/DeviceProfileManager.cpp
+++ b/Engine/Source/Runtime/Engine/Private/DeviceProfiles/DeviceProfileManager.cpp
@@ -69,9 +69,7 @@ UDeviceProfileManager& UDeviceProfileManager::Get(bool bFromPostCDOContruct)
 		DeviceProfileManagerSingleton = NewObject<UDeviceProfileManager>();
 
 		DeviceProfileManagerSingleton->AddToRoot();
-#if ALLOW_OTHER_PLATFORM_CONFIG
 		DeviceProfileManagerSingleton->LoadProfiles();
-#endif
 
 		// always start with an active profile, even if we create it on the spot
 		UDeviceProfile* ActiveProfile = DeviceProfileManagerSingleton->FindProfile(GetPlatformDeviceProfileName());
@@ -719,6 +717,9 @@ UDeviceProfile* UDeviceProfileManager::CreateProfile(const FString& ProfileName,
 	UDeviceProfile* DeviceProfile = FindObject<UDeviceProfile>( GetTransientPackage(), *ProfileName );
 	if (DeviceProfile == NULL)
 	{
+		// set the config platform if it's confidential - so that we will save to platform specific DP.ini, instead of DefaultDP.ini
+		bool bNeedSaveToPlatformConfig = ConfigPlatform != nullptr && FDataDrivenPlatformInfoRegistry::GetPlatformInfo(FName(ConfigPlatform)).bIsConfidential;
+
 		// use ConfigPlatform ini hierarchy to look in for the parent profile
 		// @todo config: we could likely cache local ini files to speed this up,
 		// along with the ones we load in LoadConfig
@@ -768,6 +769,15 @@ UDeviceProfile* UDeviceProfileManager::CreateProfile(const FString& ProfileName,
 			// if the config needs to come from a platform, set it now, then reload the config
 			DeviceProfile->ConfigPlatform = ConfigPlatform;
 			DeviceProfile->LoadConfig();
+			
+			// when we don't need to save to PlatformDP.ini, we want to save to Project/Config/DefaultDP.ini, otherwise
+			// we want to save to Project/Platforms/Plat/Config/PlatDP.ini, so we keep the COnfigPlatform
+			// around for saving, so it will route to the correct .ini file to save (once it's loaded, the function
+			// UDeviceProfile::GetConfigOverridePlatform() is only needed when saving)
+			if (!bNeedSaveToPlatformConfig)
+			{
+				DeviceProfile->ConfigPlatform = TEXT("");
+			}
 		}
 
 		// make sure the DP has all the LODGroups it needs
@@ -921,19 +931,15 @@ void UDeviceProfileManager::LoadProfiles()
 	{
 		TMap<FString, FString> DeviceProfileToPlatformConfigMap;
 		DeviceProfileToPlatformConfigMap.Add(TEXT("GlobalDefaults,None"), FPlatformProperties::IniPlatformName());
-		TArray<FName> ConfidentialPlatforms = FDataDrivenPlatformInfoRegistry::GetConfidentialPlatforms();
 		
-#if !ALLOW_OTHER_PLATFORM_CONFIG
-		checkf(ConfidentialPlatforms.Contains(FPlatformProperties::IniPlatformName()) == false,
-			TEXT("UDeviceProfileManager::LoadProfiles is called from a confidential platform (%s). Confidential platforms are not expected to be editor/non-cooked builds."), 
-			ANSI_TO_TCHAR(FPlatformProperties::IniPlatformName()));
-#endif
+#if ALLOW_OTHER_PLATFORM_CONFIG
+		TArray<FName> PlatformsToLoad = FDataDrivenPlatformInfoRegistry::GetSortedPlatformNames(EPlatformInfoType::TruePlatformsOnly);
 
 		// go over all the platforms we find, starting with the current platform
-		for (int32 PlatformIndex = 0; PlatformIndex <= ConfidentialPlatforms.Num(); PlatformIndex++)
+		for (int32 PlatformIndex = 0; PlatformIndex <= PlatformsToLoad.Num(); PlatformIndex++)
 		{
 			// which platform's set of ini files should we load from?
-			FString ConfigLoadPlatform = PlatformIndex == 0 ? FString(FPlatformProperties::IniPlatformName()) : ConfidentialPlatforms[PlatformIndex - 1].ToString();
+			FString ConfigLoadPlatform = PlatformIndex == 0 ? FString(FPlatformProperties::IniPlatformName()) : PlatformsToLoad[PlatformIndex - 1].ToString();
 
 			// load the DP.ini files (from current platform and then by the extra confidential platforms)
 			FConfigFile LocalPlatformConfigFile;
@@ -952,7 +958,8 @@ void UDeviceProfileManager::LoadProfiles()
 				}
 			}
 		}
-
+#endif
+		
 		// now that we have gathered all the unique DPs, load them from the proper platform hierarchy
 		for (auto It = DeviceProfileToPlatformConfigMap.CreateIterator(); It; ++It)
 		{
@@ -962,15 +969,7 @@ void UDeviceProfileManager::LoadProfiles()
 
 			if (FindObject<UDeviceProfile>(GetTransientPackage(), *Name) == NULL)
 			{
-				// set the config platform if it's not the current platform
-				if (It.Value() != FPlatformProperties::IniPlatformName())
-				{
-					CreateProfile(Name, DeviceType, TEXT(""), *It.Value());
-				}
-				else
-				{
-					CreateProfile(Name, DeviceType);
-				}
+				CreateProfile(Name, DeviceType, TEXT(""), *It.Value());
 			}
 		}
 
@@ -1442,6 +1441,8 @@ static bool GetCVarForDeviceProfile( FOutputDevice& Ar, FString DPName, FString
 	return true;
 }
 
+#endif
+
 class FPlatformCVarExec : public FSelfRegisteringExec
 {
 protected:
@@ -1449,6 +1450,7 @@ protected:
 	// FSelfRegisteringExec interface
 	virtual bool Exec_Runtime(UWorld* InWorld, const TCHAR* Cmd, FOutputDevice& Ar) override
 	{
+#if ALLOW_OTHER_PLATFORM_CONFIG
 		if (FParse::Command(&Cmd, TEXT("dpcvar")))
 		{
 			FString DPName, CVarName;
@@ -1526,6 +1528,27 @@ protected:
 
 			FConfigCacheIni::ClearOtherPlatformConfigs();
 		}
+#else
+		if (FParse::Command(&Cmd, TEXT("dpdump")))
+		{
+			UDeviceProfile* DeviceProfile = UDeviceProfileManager::Get().FindProfile(Cmd, false);
+			if (DeviceProfile)
+			{
+				Ar.Logf(TEXT("All cvars found for deviceprofile %s"), Cmd);
+				Ar.Logf(TEXT("  CVARS CURRENTLY DISABLED DUE TO ALLOW_OTHER_PLATFORM_CONFIG"), Cmd);
+
+				// log out the LODGroups fully
+				FArrayProperty* LODGroupsProperty = FindFProperty<FArrayProperty>(UDeviceProfile::StaticClass(), GET_MEMBER_NAME_CHECKED(UTextureLODSettings, TextureLODGroups));
+				FScriptArrayHelper_InContainer ArrayHelper(LODGroupsProperty, DeviceProfile);
+				for (int32 Index = 0; Index < ArrayHelper.Num(); Index++)
+				{
+					FString	Buffer;
+					LODGroupsProperty->Inner->ExportTextItem_Direct(Buffer, ArrayHelper.GetRawPtr(Index), ArrayHelper.GetRawPtr(Index), DeviceProfile, 0);
+					Ar.Logf(TEXT("LODGroup[%d]=%s"), Index, *Buffer);
+				}
+			}
+		}
+#endif
 		else if (FParse::Command(&Cmd, TEXT("dpreapply")))
 		{
 			UDeviceProfileManager::Get().ReapplyDeviceProfile();
@@ -1537,5 +1560,3 @@ protected:
 
 } GPlatformCVarExec;
 
-
-#endif
