Avoid direct use of VTData in FVirtualTexture2DResource member functions.
This is to prepare for derived classes that have null VTData.

diff --git a/Engine/Source/Runtime/Engine/Private/Texture2D.cpp b/Engine/Source/Runtime/Engine/Private/Texture2D.cpp
--- a/Engine/Source/Runtime/Engine/Private/Texture2D.cpp
+++ b/Engine/Source/Runtime/Engine/Private/Texture2D.cpp
@@ -1809,9 +1809,9 @@ class IAllocatedVirtualTexture* FVirtualTexture2DResource::AcquireAllocatedVT()
 	{
 		FAllocatedVTDescription VTDesc;
 		VTDesc.Dimensions = 2;
-		VTDesc.TileSize = VTData->TileSize;
-		VTDesc.TileBorderSize = VTData->TileBorderSize;
-		VTDesc.NumTextureLayers = VTData->GetNumLayers();
+		VTDesc.TileSize = GetTileSize();
+		VTDesc.TileBorderSize = GetBorderSize();
+		VTDesc.NumTextureLayers = GetNumLayers();
 		VTDesc.bShareDuplicateLayers = bSinglePhysicalSpace;
 
 		for (uint32 LayerIndex = 0u; LayerIndex < VTDesc.NumTextureLayers; ++LayerIndex)
