-Fix access to the UMapBuildDataRegistry when using LightingScenarios and VLM with multiple levels.

diff --git a/Engine/Source/Runtime/Engine/Private/StaticLightingBuildContext.cpp b/Engine/Source/Runtime/Engine/Private/StaticLightingBuildContext.cpp
--- a/Engine/Source/Runtime/Engine/Private/StaticLightingBuildContext.cpp
+++ b/Engine/Source/Runtime/Engine/Private/StaticLightingBuildContext.cpp
@@ -179,7 +179,16 @@ FPrecomputedVolumetricLightmapData& FStaticLightingBuildContext::GetOrCreateLeve
 		}
 	}
 
-	UMapBuildDataRegistry* Registry = GetOrCreateRegistryForLevelGuid(LevelId);
+	UMapBuildDataRegistry* Registry = nullptr;
+	if (LightingScenario)
+	{
+		Registry = LightingScenario->GetOrCreateMapBuildData();
+	}
+	else
+	{
+		Registry = GetOrCreateRegistryForLevelGuid(LevelId);
+	}
+	
 	FGuid BuildDataID = GetLevelBuildDataID(LevelId);
 	if (FPrecomputedVolumetricLightmapData* Data = Registry->GetLevelPrecomputedVolumetricLightmapBuildData(BuildDataID))
 	{
