#ue5 - Fix crash when PackageMap is null in UCharacterMovementComponent::MoveResponsePacked_ClientReceive(). This mirrors a similar fix that was applied to UCharacterMovementComponent::ServerMovePacked_ServerReceive() in CL 26510683.

#tests PIE with multiple clients and bad network profile, using ToggleDebugCamera (because that is found in many of the logs)

diff --git a/Engine/Source/Runtime/Engine/Private/Components/CharacterMovementComponent.cpp b/Engine/Source/Runtime/Engine/Private/Components/CharacterMovementComponent.cpp
--- a/Engine/Source/Runtime/Engine/Private/Components/CharacterMovementComponent.cpp
+++ b/Engine/Source/Runtime/Engine/Private/Components/CharacterMovementComponent.cpp
@@ -10430,6 +10430,12 @@ void UCharacterMovementComponent::MoveResponsePacked_ClientReceive(const FCharac
 		MoveResponseBitReader.PackageMap = PackedBits.GetPackageMap();
 	}
 
+	if (MoveResponseBitReader.PackageMap == nullptr)
+	{
+		devCode(UE_LOG(LogNetPlayerMovement, Error, TEXT("MoveResponsePacked_ClientReceive: Failed to find PackageMap for data serialization!")));
+		return;
+	}
+
 	// Deserialize bits to response data struct.
 	// We had to wait until now and use the temp bit stream because the RPC doesn't know about the virtual overrides on the possibly custom struct that is our data container.
 	FCharacterMoveResponseDataContainer& ResponseDataContainer = GetMoveResponseDataContainer();
