ForEachStreamingRelevantComponent  now takes into account all UActorComponent when being invoked for editor.

diff --git a/Engine/Source/Runtime/Engine/Private/Actor.cpp b/Engine/Source/Runtime/Engine/Private/Actor.cpp
--- a/Engine/Source/Runtime/Engine/Private/Actor.cpp
+++ b/Engine/Source/Runtime/Engine/Private/Actor.cpp
@@ -1189,11 +1189,15 @@ static bool ForEachStreamingRelevantComponent(const AActor* InActor, bool bForEd
 		bHasStreamingRelevantComponents |= HandleComponent(Component);
 	});
 
-	if (!bHasStreamingRelevantComponents)
+	if (bForEditor || !bHasStreamingRelevantComponents)
 	{
 		InActor->ForEachComponent<UActorComponent>(false, [&bHasStreamingRelevantComponents, &HandleComponent](UActorComponent* Component)
 		{
-			bHasStreamingRelevantComponents |= HandleComponent(Component);
+			// We already handled UPrimitive components.
+			if (!Cast<UPrimitiveComponent>(Component))
+			{
+				bHasStreamingRelevantComponents |= HandleComponent(Component);
+			}
 		});
 	}
 
