[Materials] Fix optimization of dynamic branch translation in HLSLMaterialTranslator (Follow up on CL 34065622).

- The then/else branches in the lerp() intrinsic must be reversed: lerp(A, B, false) selects A, not B.
- Only use ternary operator ?: on scalar and vector types; Generated erroneous HLSL code with LWC vectors.
- Also generate only a single branch if the condition is "true" or "false"; Adds minor constant folding optiomziation.

#rnx

diff --git a/Engine/Source/Runtime/Engine/Private/Materials/HLSLMaterialTranslator.cpp b/Engine/Source/Runtime/Engine/Private/Materials/HLSLMaterialTranslator.cpp
--- a/Engine/Source/Runtime/Engine/Private/Materials/HLSLMaterialTranslator.cpp
+++ b/Engine/Source/Runtime/Engine/Private/Materials/HLSLMaterialTranslator.cpp
@@ -6464,22 +6464,40 @@ int32 FHLSLMaterialTranslator::DynamicBranch(int32 Condition, int32 A, int32 B)
 		AddCodeChunk(MCT_VoidStatement, TEXT("//%s"), *StaticBoolParameter->GetParameterName().ToString());
 	}
 
-	if ((ResultType & MCT_Float) != 0)
+	const FString ConditionCode = GetParameterCode(Condition);
+	if (ConditionCode == TEXT("true"))
 	{
-		// Use lerp() intrinsic for floating-point values to avoid dynamic branching
-		return AddCodeChunk(ResultType, TEXT("lerp(%s, %s, %s)"), *GetParameterCode(A), *GetParameterCode(B), *GetParameterCode(Condition));
+		const FString ThenBranchCode = GetParameterCode(A);
+		return AddCodeChunk(ResultType, TEXT("%s"), *ThenBranchCode);
 	}
-	else if ((ResultType & (MCT_Numeric | MCT_UInt)) != 0)
+	else if (ConditionCode == TEXT("false"))
+	{
+		// Minor constant folding optimization: Only generate 
+		const FString ElseBranchCode = GetParameterCode(B);
+		return AddCodeChunk(ResultType, TEXT("%s"), *ElseBranchCode);
+	}
+	else if ((ResultType & MCT_Float) != 0)
+	{
+		// Use lerp() intrinsic for floating-point values to avoid dynamic branching; Here the then/else branches are reversed!
+		const FString ThenBranchCode = GetParameterCode(A);
+		const FString ElseBranchCode = GetParameterCode(B);
+		return AddCodeChunk(ResultType, TEXT("lerp(%s, %s, %s)"), *ElseBranchCode, *ThenBranchCode, *ConditionCode);
+	}
+	else if ((ResultType & (MCT_UInt | MCT_Bool)) != 0)
 	{
 		// Use ternary-operator to simplify output for dynamic branch of two numerical values
-		return AddCodeChunk(ResultType, TEXT("%s ? %s : %s"), *GetParameterCode(Condition), *GetParameterCode(A), *GetParameterCode(B));
+		const FString ThenBranchCode = GetParameterCode(A);
+		const FString ElseBranchCode = GetParameterCode(B);
+		return AddCodeChunk(ResultType, TEXT("%s ? %s : %s"), *ConditionCode, *ThenBranchCode, *ElseBranchCode);
 	}
 	else
 	{
 		// Fallback to switch-case statement for then/else branches for all other types
-		FString SymbolName = CreateSymbolName(TEXT("Static"));
+		const FString SymbolName = CreateSymbolName(TEXT("Static"));
+		const FString ThenBranchCode = GetParameterCode(A);
+		const FString ElseBranchCode = GetParameterCode(B);
 		AddCodeChunk(MCT_VoidStatement, TEXT("%s %s;"), HLSLTypeString(ResultType), *SymbolName);
-		AddCodeChunk(MCT_VoidStatement, TEXT("[branch] switch (int(%s)){ default: %s = %s; break; case 0: %s = %s; break;}"), *GetParameterCode(Condition), *SymbolName, *GetParameterCode(A), *SymbolName, *GetParameterCode(B));
+		AddCodeChunk(MCT_VoidStatement, TEXT("[branch] switch (int(%s)){ default: %s = %s; break; case 0: %s = %s; break;}"), *ConditionCode, *SymbolName, *ThenBranchCode, *SymbolName, *ElseBranchCode);
 		return AddCodeChunk(ResultType, *SymbolName);
 	}
 }
