Enable EWorldType::Game worlds to be dedicated servers if desired in the editor.
PR #11965

diff --git a/Engine/Source/Runtime/Engine/Private/NetDriver.cpp b/Engine/Source/Runtime/Engine/Private/NetDriver.cpp
--- a/Engine/Source/Runtime/Engine/Private/NetDriver.cpp
+++ b/Engine/Source/Runtime/Engine/Private/NetDriver.cpp
@@ -2011,9 +2011,9 @@ void UNetDriver::FlushHandler()
 
 ENetMode UNetDriver::GetNetMode() const
 {
-	// Special case for PIE - forcing dedicated server behavior
+	// Special case for PIE and game worlds - forcing dedicated server behavior if requested
 #if WITH_EDITOR
-	if (World && World->WorldType == EWorldType::PIE && IsServer())
+	if (World && (World->WorldType == EWorldType::PIE || World->WorldType == EWorldType::Game) && IsServer())
 	{
 		//@todo: world context won't be valid during seamless travel CopyWorldData
 		FWorldContext* WorldContext = GEngine->GetWorldContextFromWorld(World);
