#ue5 - Check for null capsule component in UCharacterMovementComponent::AdjustProxyCapsuleSize()

#tests PIE

diff --git a/Engine/Source/Runtime/Engine/Private/Components/CharacterMovementComponent.cpp b/Engine/Source/Runtime/Engine/Private/Components/CharacterMovementComponent.cpp
--- a/Engine/Source/Runtime/Engine/Private/Components/CharacterMovementComponent.cpp
+++ b/Engine/Source/Runtime/Engine/Private/Components/CharacterMovementComponent.cpp
@@ -1718,7 +1718,7 @@ void UCharacterMovementComponent::PostPhysicsTickComponent(float DeltaTime, FCha
 
 void UCharacterMovementComponent::AdjustProxyCapsuleSize()
 {
-	if (bShrinkProxyCapsule && CharacterOwner && CharacterOwner->GetLocalRole() == ROLE_SimulatedProxy)
+	if (bShrinkProxyCapsule && CharacterOwner && CharacterOwner->GetLocalRole() == ROLE_SimulatedProxy && CharacterOwner->GetCapsuleComponent() != nullptr)
 	{
 		bShrinkProxyCapsule = false;
 
