UserFlags accessible in View Uniform Buffer added to Post Process Settings, effectively allowing per-view overrides on materials.  A new TestPostVolumeUserFlag expression node was added, returning a boolean result of a bit test with the UserFlags.  The UserFlags are also technically accessible in Custom HLSL for more arbitrary use.  Tested for both legacy and new translate.

diff --git a/Engine/Shaders/Private/MaterialTemplate.ush b/Engine/Shaders/Private/MaterialTemplate.ush
--- a/Engine/Shaders/Private/MaterialTemplate.ush
+++ b/Engine/Shaders/Private/MaterialTemplate.ush
@@ -1115,6 +1115,11 @@ float GetWaterWaveParamIndex(FMaterialVertexParameters Parameters)
 #endif
 }
 
+float PostVolumeUserFlagTest(float CompareValue)
+{
+	return (((int)1 << (int)CompareValue) & View.PostVolumeUserFlags) ? 1.0 : 0.0;
+}
+
 // Returns whether a scene texture id is a for a post process input or not.
 bool IsPostProcessInputSceneTexture(const uint SceneTextureId)
 {
diff --git a/Engine/Source/Runtime/Engine/Classes/Engine/Scene.h b/Engine/Source/Runtime/Engine/Classes/Engine/Scene.h
--- a/Engine/Source/Runtime/Engine/Classes/Engine/Scene.h
+++ b/Engine/Source/Runtime/Engine/Classes/Engine/Scene.h
@@ -1178,6 +1178,9 @@ struct FPostProcessSettings
 	UPROPERTY(EditAnywhere, BlueprintReadWrite, Category=Overrides, meta=(PinHiddenByDefault, InlineEditConditionToggle))
 	uint8 bOverride_ScreenSpaceReflectionRoughnessScale:1; // TODO: look useless...
 
+	UPROPERTY(EditAnywhere, BlueprintReadWrite, Category=Overrides, meta=(PinHiddenByDefault, InlineEditConditionToggle))
+	uint8 bOverride_UserFlags:1;
+
 	// -----------------------------------------------------------------------
 
 	// Ray Tracing
@@ -2281,6 +2284,10 @@ struct FPostProcessSettings
 	UPROPERTY()
 	float ScreenPercentage_DEPRECATED;
 
+	/** Per-view user flags accessible in materials via TestPostVolumeUserFlag node, allowing per-view overrides of material behavior. */
+	UPROPERTY(EditAnywhere, BlueprintReadWrite, Category = "Rendering Features", meta = (editcondition = "bOverride_UserFlags", DisplayName = "User Flags"))
+	int32 UserFlags;
+
 
 
 	// Note: Adding properties before this line require also changes to the OverridePostProcessSettings() function and 
diff --git a/Engine/Source/Runtime/Engine/Private/Materials/HLSLMaterialTranslator.cpp b/Engine/Source/Runtime/Engine/Private/Materials/HLSLMaterialTranslator.cpp
--- a/Engine/Source/Runtime/Engine/Private/Materials/HLSLMaterialTranslator.cpp
+++ b/Engine/Source/Runtime/Engine/Private/Materials/HLSLMaterialTranslator.cpp
@@ -11336,6 +11336,18 @@ int32 FHLSLMaterialTranslator::DepthOfFieldFunction(int32 Depth, int32 FunctionV
 		*GetParameterCode(Depth), FunctionValueIndex);
 }
 
+int32 FHLSLMaterialTranslator::PostVolumeUserFlagTestFunction(int32 Input)
+{
+	if (Input == INDEX_NONE)
+	{
+		return INDEX_NONE;
+	}
+
+	return AddCodeChunk(MCT_Float,
+		TEXT("PostVolumeUserFlagTest(%s)"),
+		*GetParameterCode(Input));
+}
+
 int32 FHLSLMaterialTranslator::Sobol(int32 Cell, int32 Index, int32 Seed)
 {
 	AddEstimatedTextureSample(2);
diff --git a/Engine/Source/Runtime/Engine/Private/Materials/HLSLMaterialTranslator.h b/Engine/Source/Runtime/Engine/Private/Materials/HLSLMaterialTranslator.h
--- a/Engine/Source/Runtime/Engine/Private/Materials/HLSLMaterialTranslator.h
+++ b/Engine/Source/Runtime/Engine/Private/Materials/HLSLMaterialTranslator.h
@@ -1192,6 +1192,7 @@ protected:
 
 	virtual int32 AntialiasedTextureMask(int32 Tex, int32 UV, float Threshold, uint8 Channel) override;
 	virtual int32 DepthOfFieldFunction(int32 Depth, int32 FunctionValueIndex) override;
+	virtual int32 PostVolumeUserFlagTestFunction(int32 Input) override;
 	virtual int32 Sobol(int32 Cell, int32 Index, int32 Seed) override;
 	virtual int32 TemporalSobol(int32 Index, int32 Seed) override;
 	virtual int32 Noise(int32 Position, EPositionOrigin PositionOrigin, float Scale, int32 Quality, uint8 NoiseFunction, bool bTurbulence, int32 Levels, float OutputMin, float OutputMax, float LevelScale, int32 FilterWidth, bool bTiling, uint32 RepeatSize) override;
diff --git a/Engine/Source/Runtime/Engine/Private/Materials/MaterialExpressionHLSL.cpp b/Engine/Source/Runtime/Engine/Private/Materials/MaterialExpressionHLSL.cpp
--- a/Engine/Source/Runtime/Engine/Private/Materials/MaterialExpressionHLSL.cpp
+++ b/Engine/Source/Runtime/Engine/Private/Materials/MaterialExpressionHLSL.cpp
@@ -153,6 +153,7 @@
 #include "Materials/MaterialExpressionPerInstanceRandom.h"
 #include "Materials/MaterialExpressionPixelDepth.h"
 #include "Materials/MaterialExpressionPixelNormalWS.h"
+#include "Materials/MaterialExpressionPostVolumeUserFlagTest.h"
 #include "Materials/MaterialExpressionPower.h"
 #include "Materials/MaterialExpressionPrecomputedAOMask.h"
 #include "Materials/MaterialExpressionPreSkinnedLocalBounds.h"
@@ -1471,6 +1472,19 @@ bool UMaterialExpressionParticleMotionBlurFade::GenerateHLSLExpression(FMaterial
 	return true;
 }
 
+bool UMaterialExpressionPostVolumeUserFlagTest::GenerateHLSLExpression(FMaterialHLSLGenerator& Generator, UE::HLSLTree::FScope& Scope, int32 OutputIndex, UE::HLSLTree::FExpression const*& OutExpression) const
+{
+	using namespace UE::HLSLTree;
+	const FExpression* BitIndexExpression = BitIndex.AcquireHLSLExpressionOrConstant(Generator, Scope, (float)ConstBitIndex);
+	if (!BitIndexExpression)
+	{
+		return false;
+	}
+
+	OutExpression = Generator.GetTree().NewExpression<Material::FExpressionPostVolumeUserFlagTest>(BitIndexExpression);
+	return true;
+}
+
 bool UMaterialExpressionTextureObject::GenerateHLSLExpression(FMaterialHLSLGenerator& Generator, UE::HLSLTree::FScope& Scope, int32 OutputIndex, UE::HLSLTree::FExpression const*& OutExpression) const
 {
 	using namespace UE::HLSLTree;
diff --git a/Engine/Source/Runtime/Engine/Private/Materials/MaterialExpressions.cpp b/Engine/Source/Runtime/Engine/Private/Materials/MaterialExpressions.cpp
--- a/Engine/Source/Runtime/Engine/Private/Materials/MaterialExpressions.cpp
+++ b/Engine/Source/Runtime/Engine/Private/Materials/MaterialExpressions.cpp
@@ -189,6 +189,7 @@
 #include "Materials/MaterialExpressionPerInstanceCustomData.h"
 #include "Materials/MaterialExpressionPixelDepth.h"
 #include "Materials/MaterialExpressionPixelNormalWS.h"
+#include "Materials/MaterialExpressionPostVolumeUserFlagTest.h"
 #include "Materials/MaterialExpressionPower.h"
 #include "Materials/MaterialExpressionPreSkinnedNormal.h"
 #include "Materials/MaterialExpressionPreSkinnedPosition.h"
@@ -22452,6 +22453,53 @@ void UMaterialExpressionDepthOfFieldFunction::GetCaption(TArray<FString>& OutCap
 }
 #endif // WITH_EDITOR
 
+///////////////////////////////////////////////////////////////////////////////
+// UMaterialExpressionPostVolumeUserFlagTest
+///////////////////////////////////////////////////////////////////////////////
+UMaterialExpressionPostVolumeUserFlagTest::UMaterialExpressionPostVolumeUserFlagTest(const FObjectInitializer& ObjectInitializer)
+	: Super(ObjectInitializer)
+{
+#if WITH_EDITORONLY_DATA
+	// Structure to hold one-time initialization
+	struct FConstructorStatics
+	{
+		FText NAME_GetPostProcessSetting;
+		FConstructorStatics()
+			: NAME_GetPostProcessSetting(LOCTEXT("Get Post Process Setting", "Get Post Process Setting"))
+		{
+		}
+	};
+	static FConstructorStatics ConstructorStatics;
+
+	MenuCategories.Add(ConstructorStatics.NAME_GetPostProcessSetting);
+
+	bCollapsed = true;
+#endif
+}
+
+#if WITH_EDITOR
+int32 UMaterialExpressionPostVolumeUserFlagTest::Compile(class FMaterialCompiler* Compiler, int32 OutputIndex)
+{
+	int32 BitIndexCompiled;
+
+	if (BitIndex.GetTracedInput().Expression)
+	{
+		BitIndexCompiled = BitIndex.Compile(Compiler);
+	}
+	else
+	{
+		BitIndexCompiled = Compiler->Constant((float)ConstBitIndex);
+	}
+
+	if (BitIndexCompiled == INDEX_NONE)
+	{
+		return INDEX_NONE;
+	}
+
+	return Compiler->PostVolumeUserFlagTestFunction(BitIndexCompiled);
+}
+#endif // WITH_EDITOR
+
 ///////////////////////////////////////////////////////////////////////////////
 // UMaterialExpressionDDX
 ///////////////////////////////////////////////////////////////////////////////
diff --git a/Engine/Source/Runtime/Engine/Private/Materials/MaterialHLSLTree.cpp b/Engine/Source/Runtime/Engine/Private/Materials/MaterialHLSLTree.cpp
--- a/Engine/Source/Runtime/Engine/Private/Materials/MaterialHLSLTree.cpp
+++ b/Engine/Source/Runtime/Engine/Private/Materials/MaterialHLSLTree.cpp
@@ -3188,6 +3188,26 @@ void FExpressionDepthOfFieldFunction::EmitValueShader(FEmitContext& Context, FEm
 		FunctionValue);
 }
 
+bool FExpressionPostVolumeUserFlagTest::PrepareValue(FEmitContext& Context, FEmitScope& Scope, const FRequestedType& RequestedType, FPrepareValueResult& OutResult) const
+{
+	const FPreparedType& InputType = Context.PrepareExpression(InputExpression, Scope, Shader::EValueType::Float1);
+	if (InputType.IsVoid())
+	{
+		return false;
+	}
+
+	return OutResult.SetType(Context, RequestedType, EExpressionEvaluation::Shader, Shader::EValueType::Float1);
+}
+
+void FExpressionPostVolumeUserFlagTest::EmitValueShader(FEmitContext& Context, FEmitScope& Scope, const FRequestedType& RequestedType, FEmitValueShaderResult& OutResult) const
+{
+	FEmitShaderExpression* EmitInput = InputExpression->GetValueShader(Context, Scope, Shader::EValueType::Float1);
+
+	OutResult.Code = Context.EmitInlineExpression(Scope, Shader::EValueType::Float1,
+		TEXT("PostVolumeUserFlagTest(%)"),
+		EmitInput);
+}
+
 bool FExpressionSobolFunction::PrepareValue(FEmitContext& Context, FEmitScope& Scope, const FRequestedType& RequestedType, FPrepareValueResult& OutResult) const
 {
 	if ((CellExpression && Context.PrepareExpression(CellExpression, Scope, Shader::EValueType::Float2).IsVoid()) ||
diff --git a/Engine/Source/Runtime/Engine/Private/Scene.cpp b/Engine/Source/Runtime/Engine/Private/Scene.cpp
--- a/Engine/Source/Runtime/Engine/Private/Scene.cpp
+++ b/Engine/Source/Runtime/Engine/Private/Scene.cpp
@@ -649,6 +649,8 @@ FPostProcessSettings::FPostProcessSettings()
 
 	bMobileHQGaussian = false;
 
+	UserFlags = 0;
+
 #if DO_CHECK && WITH_EDITOR
 	static bool bCheckedMembers = false;
 	if (!bCheckedMembers)
diff --git a/Engine/Source/Runtime/Engine/Private/SceneView.cpp b/Engine/Source/Runtime/Engine/Private/SceneView.cpp
--- a/Engine/Source/Runtime/Engine/Private/SceneView.cpp
+++ b/Engine/Source/Runtime/Engine/Private/SceneView.cpp
@@ -1854,6 +1854,8 @@ void FSceneView::OverridePostProcessSettings(const FPostProcessSettings& Src, fl
 		SET_PP(MotionBlurTargetFPS);
 
 		SET_PP(AutoExposureApplyPhysicalCameraExposure);
+
+		SET_PP(UserFlags);
 	}
 
 	// Blendable objects
diff --git a/Engine/Source/Runtime/Engine/Public/MaterialCompiler.h b/Engine/Source/Runtime/Engine/Public/MaterialCompiler.h
--- a/Engine/Source/Runtime/Engine/Public/MaterialCompiler.h
+++ b/Engine/Source/Runtime/Engine/Public/MaterialCompiler.h
@@ -569,6 +569,7 @@ public:
 		return SamplePhysicsField(PositionArg, EPositionOrigin::Absolute, OutputType, TargetIndex);
 	}
 	virtual int32 DepthOfFieldFunction(int32 Depth, int32 FunctionValueIndex) = 0;
+	virtual int32 PostVolumeUserFlagTestFunction(int32 Input) = 0;
 	virtual int32 AtmosphericFogColor(int32 WorldPosition, EPositionOrigin PositionOrigin) = 0;
 	UE_DEPRECATED(5.4, "Use AtmosphericFogColor(int32, EPositionOrigin) instead")
 	int32 AtmosphericFogColor(int32 WorldPosition)
@@ -1103,6 +1104,7 @@ public:
 		return Compiler->DepthOfFieldFunction(Depth, FunctionValueIndex);
 	}
 
+	virtual int32 PostVolumeUserFlagTestFunction(int32 Input) override { return Compiler->PostVolumeUserFlagTestFunction(Input); }
 	virtual int32 GetHairUV() override { return Compiler->GetHairUV(); }
 	virtual int32 GetHairDimensions() override { return Compiler->GetHairDimensions(); }
 	virtual int32 GetHairSeed() override { return Compiler->GetHairSeed(); }
diff --git a/Engine/Source/Runtime/Engine/Public/MaterialHLSLTree.h b/Engine/Source/Runtime/Engine/Public/MaterialHLSLTree.h
--- a/Engine/Source/Runtime/Engine/Public/MaterialHLSLTree.h
+++ b/Engine/Source/Runtime/Engine/Public/MaterialHLSLTree.h
@@ -498,6 +498,19 @@ public:
 	virtual void EmitValueShader(FEmitContext& Context, FEmitScope& Scope, const FRequestedType& RequestedType, FEmitValueShaderResult& OutResult) const override;
 };
 
+class FExpressionPostVolumeUserFlagTest : public FExpression
+{
+public:
+	const FExpression* InputExpression;
+
+	FExpressionPostVolumeUserFlagTest(const FExpression* InInputExpression)
+		: InputExpression(InInputExpression)
+	{}
+
+	virtual bool PrepareValue(FEmitContext& Context, FEmitScope& Scope, const FRequestedType& RequestedType, FPrepareValueResult& OutResult) const override;
+	virtual void EmitValueShader(FEmitContext& Context, FEmitScope& Scope, const FRequestedType& RequestedType, FEmitValueShaderResult& OutResult) const override;
+};
+
 struct DataDrivenShaderPlatformData
 {
 	FName PlatformName;
diff --git a/Engine/Source/Runtime/Engine/Public/Materials/MaterialExpressionPostVolumeUserFlagTest.h b/Engine/Source/Runtime/Engine/Public/Materials/MaterialExpressionPostVolumeUserFlagTest.h
new file mode 100644
--- /dev/null
+++ b/Engine/Source/Runtime/Engine/Public/Materials/MaterialExpressionPostVolumeUserFlagTest.h
@@ -0,0 +1,33 @@
+// Copyright Epic Games, Inc. All Rights Reserved.
+
+
+#pragma once
+
+#include "CoreMinimal.h"
+#include "UObject/ObjectMacros.h"
+#include "MaterialExpressionIO.h"
+#include "Materials/MaterialExpression.h"
+#include "MaterialExpressionPostVolumeUserFlagTest.generated.h"
+
+UCLASS(MinimalAPI)
+class UMaterialExpressionPostVolumeUserFlagTest : public UMaterialExpression
+{
+	GENERATED_UCLASS_BODY()
+
+	UPROPERTY(meta = (RequiredInput = "false", ToolTip = "Defaults to 'ConstBitIndex' if not specified"))
+	FExpressionInput BitIndex;
+
+	/** only used if Input is not hooked up */
+	UPROPERTY(EditAnywhere, Category=MaterialExpressionPostVolumeUserFlagTest, meta=(OverridingInputProperty = "BitIndex"))
+	int32 ConstBitIndex;
+
+
+	//~ Begin UMaterialExpression Interface
+#if WITH_EDITOR
+	virtual FText GetCreationName() const override { return FText::FromString(TEXT("TestPostVolumeUserFlag")); }
+	virtual int32 Compile(class FMaterialCompiler* Compiler, int32 OutputIndex) override;
+	virtual void GetCaption(TArray<FString>& OutCaptions) const override { OutCaptions.Add(FString(TEXT("TestPostVolumeUserFlag"))); }
+	virtual bool GenerateHLSLExpression(FMaterialHLSLGenerator& Generator, UE::HLSLTree::FScope& Scope, int32 OutputIndex, UE::HLSLTree::FExpression const*& OutExpression) const override;
+#endif // WITH_EDITOR
+	//~ End UMaterialExpression Interface
+};
diff --git a/Engine/Source/Runtime/Engine/Public/SceneInterface.h b/Engine/Source/Runtime/Engine/Public/SceneInterface.h
--- a/Engine/Source/Runtime/Engine/Public/SceneInterface.h
+++ b/Engine/Source/Runtime/Engine/Public/SceneInterface.h
@@ -704,8 +704,10 @@ public:
 		FEngineShowFlags EngineShowFlags;
 		const AActor* ViewActor = nullptr;
 		class FSceneViewStateInterface* ViewStateInterface = nullptr;
+		int32 PostVolumeUserFlags = 0;						// UserFlags from Post Process Volume, used if bOverridesPostVolumeUserFlags set, allows per-view material overrides
 		bool bIsSceneCapture = false;
 		bool bUseMainViewFamilyShowFlags = false;			// Custom render pass should use flags from main view family, rather than ones in this structure
+		bool bOverridesPostVolumeUserFlags = false;			// Use PostVolumeUserFlags, instead of flags from main view family
 
 		class FCustomRenderPassBase* CustomRenderPass = nullptr;
 	};
diff --git a/Engine/Source/Runtime/Engine/Public/SceneView.h b/Engine/Source/Runtime/Engine/Public/SceneView.h
--- a/Engine/Source/Runtime/Engine/Public/SceneView.h
+++ b/Engine/Source/Runtime/Engine/Public/SceneView.h
@@ -992,6 +992,7 @@ enum ETranslucencyVolumeCascade
 	VIEW_UNIFORM_BUFFER_MEMBER_PER_VIEW(FVector4f, GlintLUTParameters0) \
 	VIEW_UNIFORM_BUFFER_MEMBER_PER_VIEW(FVector4f, GlintLUTParameters1) \
 	VIEW_UNIFORM_BUFFER_MEMBER(float, MaterialMaxEmissiveValue) \
+	VIEW_UNIFORM_BUFFER_MEMBER(int32, PostVolumeUserFlags) \
 	VIEW_UNIFORM_BUFFER_MEMBER(FIntVector4, EnvironmentComponentsFlags) \
 
 /** The uniform shader parameters associated with a view. */
diff --git a/Engine/Source/Runtime/Renderer/Private/SceneCaptureRendering.cpp b/Engine/Source/Runtime/Renderer/Private/SceneCaptureRendering.cpp
--- a/Engine/Source/Runtime/Renderer/Private/SceneCaptureRendering.cpp
+++ b/Engine/Source/Runtime/Renderer/Private/SceneCaptureRendering.cpp
@@ -1137,6 +1137,12 @@ void FScene::UpdateSceneCaptureContents(USceneCaptureComponent2D* CaptureCompone
 
 			PassInput.EngineShowFlags = CaptureComponent->ShowFlags;
 
+			if (CaptureComponent->PostProcessBlendWeight > 0.0f && CaptureComponent->PostProcessSettings.bOverride_UserFlags)
+			{
+				PassInput.PostVolumeUserFlags = CaptureComponent->PostProcessSettings.UserFlags;
+				PassInput.bOverridesPostVolumeUserFlags = true;
+			}
+
 			// Caching scene capture info to be passed to the scene renderer.
 			// #todo: We cannot (yet) guarantee for which ViewFamily this CRP will eventually be rendered since it will just execute the next time the scene is rendered by any FSceneRenderer. This seems quite problematic and could easily lead to unexpected behavior...
 			AddCustomRenderPass(nullptr, PassInput);
diff --git a/Engine/Source/Runtime/Renderer/Private/SceneRendering.cpp b/Engine/Source/Runtime/Renderer/Private/SceneRendering.cpp
--- a/Engine/Source/Runtime/Renderer/Private/SceneRendering.cpp
+++ b/Engine/Source/Runtime/Renderer/Private/SceneRendering.cpp
@@ -1592,6 +1592,7 @@ void FViewInfo::SetupUniformBufferParameters(
 	}
 
 	ViewUniformShaderParameters.MaterialMaxEmissiveValue = MaterialMaxEmissiveValue;
+	ViewUniformShaderParameters.PostVolumeUserFlags = FinalPostProcessSettings.UserFlags;
 
 	// This should probably be in SetupCommonViewUniformBufferParameters, but drags in too many dependencies
 	UpdateNoiseTextureParameters(ViewUniformShaderParameters);
@@ -2953,6 +2954,16 @@ FSceneRenderer::FSceneRenderer(const FSceneViewFamily* InViewFamily, FHitProxyCo
 		FViewInfo* ViewInfo = &CustomRenderPassInfo.Views.Emplace_GetRef(&NewView);
 		CustomRenderPassInfo.ViewFamily.Views.Add(ViewInfo);
 
+		if (PassInput.bOverridesPostVolumeUserFlags)
+		{
+			ViewInfo->FinalPostProcessSettings.UserFlags = PassInput.PostVolumeUserFlags;
+		}
+		else
+		{
+			// Arbitrarily use the post process UserFlags from the first view.
+			ViewInfo->FinalPostProcessSettings.UserFlags = Views[0].FinalPostProcessSettings.UserFlags;
+		}
+
 		// Must initialize to have a GPUScene connected to be able to collect dynamic primitives.
 		ViewInfo->DynamicPrimitiveCollector = FGPUScenePrimitiveCollector(&GPUSceneDynamicContext);
 		ViewInfo->bDisableQuerySubmissions = true;
