FTimerUnifiedDelegate use a TVariant
Saves 32 bytes (128->92)
#rnx

diff --git a/Engine/Source/Runtime/Engine/Private/TimerManager.cpp b/Engine/Source/Runtime/Engine/Private/TimerManager.cpp
--- a/Engine/Source/Runtime/Engine/Private/TimerManager.cpp
+++ b/Engine/Source/Runtime/Engine/Private/TimerManager.cpp
@@ -112,66 +112,74 @@ struct FTimerSourceList
 	// This is similar to FTimerUnifiedDelegate::ToString() but it tries to find the base class / exclude vptr printout info so that timers are collapsed/aggregated better
 	static FString GetPartialDeduplicateDelegateToString(const FTimerUnifiedDelegate& Delegate)
 	{
-		FString FunctionNameStr;
-		FString ObjectNameStr;
+		FString FunctionNameStr = TEXT("NotBound!");
+		FString ObjectNameStr = TEXT("NotBound!");
 
-		if (Delegate.FuncDelegate.IsBound())
+		if (const FTimerDelegate* FuncDelegate = Delegate.VariantDelegate.TryGet<FTimerDelegate>())
 		{
-			ObjectNameStr = TEXT("NonDynamicDelegate");
-			FName FunctionName;
+			if (FuncDelegate->IsBound())
+			{
+				ObjectNameStr = TEXT("NonDynamicDelegate");
+				FName FunctionName;
 #if USE_DELEGATE_TRYGETBOUNDFUNCTIONNAME
-			FunctionName = Delegate.FuncDelegate.TryGetBoundFunctionName();
+				FunctionName = FuncDelegate->TryGetBoundFunctionName();
 #endif
-			if (FunctionName.IsNone())
-			{
-				uint64 ProgramCounter = Delegate.FuncDelegate.GetBoundProgramCounterForTimerManager();
-				if (ProgramCounter != 0)
+				if (FunctionName.IsNone())
 				{
-					// Add the function address
-					if (DumpTimerLogSymbolNames)
+					uint64 ProgramCounter = FuncDelegate->GetBoundProgramCounterForTimerManager();
+					if (ProgramCounter != 0)
 					{
-						// Try to resolve the function address to a symbol
-						FProgramCounterSymbolInfo SymbolInfo;
-						FPlatformStackWalk::ProgramCounterToSymbolInfo(ProgramCounter, /*out*/ SymbolInfo);
-						FunctionNameStr = FString::Printf(TEXT("%s [%s:%d]"), ANSI_TO_TCHAR(SymbolInfo.FunctionName), ANSI_TO_TCHAR(SymbolInfo.Filename), SymbolInfo.LineNumber);
+						// Add the function address
+						if (DumpTimerLogSymbolNames)
+						{
+							// Try to resolve the function address to a symbol
+							FProgramCounterSymbolInfo SymbolInfo;
+							FPlatformStackWalk::ProgramCounterToSymbolInfo(ProgramCounter, /*out*/ SymbolInfo);
+							FunctionNameStr = FString::Printf(TEXT("%s [%s:%d]"), ANSI_TO_TCHAR(SymbolInfo.FunctionName), ANSI_TO_TCHAR(SymbolInfo.Filename), SymbolInfo.LineNumber);
+						}
+						else
+						{
+							FunctionNameStr = FString::Printf(TEXT("func: 0x%llx"), ProgramCounter);
+						}
 					}
 					else
 					{
-						FunctionNameStr = FString::Printf(TEXT("func: 0x%llx"), ProgramCounter);
+						FunctionNameStr = TEXT("0x0");
 					}
 				}
 				else
 				{
-					FunctionNameStr = TEXT("0x0");
+					FunctionNameStr = FunctionName.ToString();
 				}
 			}
-			else
-			{
-				FunctionNameStr = FunctionName.ToString();
-			}
 		}
-		else if (Delegate.FuncDynDelegate.IsBound())
+		else if (const FTimerDynamicDelegate* FuncDynDelegate = Delegate.VariantDelegate.TryGet<FTimerDynamicDelegate>())
 		{
-			const FName FuncFName = Delegate.FuncDynDelegate.GetFunctionName();
-			FunctionNameStr = FuncFName.ToString();
-
-			UClass* SourceClass = nullptr;
-			if (const UObject* Object = Delegate.FuncDynDelegate.GetUObject())
+			if (FuncDynDelegate->IsBound())
 			{
-				SourceClass = Object->GetClass();
-				if (UFunction* Func = SourceClass->FindFunctionByName(FuncFName))
+				const FName FuncFName = FuncDynDelegate->GetFunctionName();
+				FunctionNameStr = FuncFName.ToString();
+
+				UClass* SourceClass = nullptr;
+				if (const UObject* Object = FuncDynDelegate->GetUObject())
 				{
-					SourceClass = Func->GetOwnerClass();
+					SourceClass = Object->GetClass();
+					if (UFunction* Func = SourceClass->FindFunctionByName(FuncFName))
+					{
+						SourceClass = Func->GetOwnerClass();
+					}
 				}
-			}
 
-
-			ObjectNameStr = GetPathNameSafe(SourceClass);
+				ObjectNameStr = GetPathNameSafe(SourceClass);
+			}
 		}
-		else
+		else if (const FTimerFunction* TimerFunction = Delegate.VariantDelegate.TryGet<FTimerFunction>())
 		{
-			ObjectNameStr = TEXT("NotBound!");
-			FunctionNameStr = TEXT("NotBound!");
+			if (*TimerFunction)
+			{
+				FunctionNameStr = TEXT("TFunction");
+				ObjectNameStr = TEXT("");	
+			}
 		}
 
 		return FString::Printf(TEXT("%s,\"%s\""), *ObjectNameStr, *FunctionNameStr);
@@ -327,82 +335,173 @@ void FTimerManager::OnCrash()
 	UE_LOG(LogEngine, Warning, TEXT("TimerManager %p dump ended"), this);
 }
 
+void FTimerUnifiedDelegate::Execute() const
+{
+	switch (VariantDelegate.GetIndex())
+	{
+	case FTimerDelegateVariant::IndexOfType<FTimerDelegate>():
+		{
+			const FTimerDelegate& FuncDelegate = VariantDelegate.Get<FTimerDelegate>();
+			if (FuncDelegate.IsBound())
+			{
+				FScopeCycleCounterUObject Context(FuncDelegate.GetUObject());
+				FuncDelegate.Execute();
+			}
+			break;
+		}
+	case FTimerDelegateVariant::IndexOfType<FTimerDynamicDelegate>():
+		{
+			const FTimerDynamicDelegate& FuncDynDelegate = VariantDelegate.Get<FTimerDynamicDelegate>();
+			if (FuncDynDelegate.IsBound())
+			{
+				// stat scope is handled by UObject::ProcessEvent for the UFunction.
+				FuncDynDelegate.ProcessDelegate<UObject>(nullptr);
+			}
+			break;
+		}
+	case FTimerDelegateVariant::IndexOfType<FTimerFunction>():
+		{
+			if (const FTimerFunction& TimerFunction = VariantDelegate.Get<FTimerFunction>())
+			{
+				QUICK_SCOPE_CYCLE_COUNTER(STAT_FTimerUnifiedDelegate_Execute);
+				TimerFunction();
+			}
+			break;
+		}
+	default:
+		break;
+	}
+}
 
-FString FTimerUnifiedDelegate::ToString() const
+bool FTimerUnifiedDelegate::IsBound() const
 {
-	const UObject* Object = nullptr;
-	FString FunctionNameStr;
-	bool bDynDelegate = false;
+	switch (VariantDelegate.GetIndex())
+	{
+	case FTimerDelegateVariant::IndexOfType<FTimerDelegate>():
+		return VariantDelegate.Get<FTimerDelegate>().IsBound();
+	case FTimerDelegateVariant::IndexOfType<FTimerDynamicDelegate>():
+		return VariantDelegate.Get<FTimerDynamicDelegate>().IsBound();
+	case FTimerDelegateVariant::IndexOfType<FTimerFunction>():
+		return VariantDelegate.Get<FTimerFunction>()!=nullptr;
+	default:
+		return false;
+	}
+}
 
-	if (FuncDelegate.IsBound())
+const void* FTimerUnifiedDelegate::GetBoundObject() const
+{
+	switch (VariantDelegate.GetIndex())
 	{
-		FName FunctionName;
-#if USE_DELEGATE_TRYGETBOUNDFUNCTIONNAME
-		FunctionName = FuncDelegate.TryGetBoundFunctionName();
-#endif
-		if (FunctionName.IsNone())
+	case FTimerDelegateVariant::IndexOfType<FTimerDelegate>():
 		{
-			void** VtableAddr = nullptr;
-#if PLATFORM_COMPILER_CLANG || defined(_MSC_VER)
-			// Add the vtable address
-			const void* UserObject = FuncDelegate.GetObjectForTimerManager();
-			if (UserObject)
+			const FTimerDelegate& FuncDelegate = VariantDelegate.Get<FTimerDelegate>();
+			if (FuncDelegate.IsBound())
 			{
-				VtableAddr = *(void***)UserObject;
-				FunctionNameStr = FString::Printf(TEXT("vtbl: %p"), VtableAddr);
+				return FuncDelegate.GetObjectForTimerManager();
 			}
-#endif // PLATFORM_COMPILER_CLANG
-
-			uint64 ProgramCounter = FuncDelegate.GetBoundProgramCounterForTimerManager();
-			if (ProgramCounter != 0)
+			break;
+		}
+	case FTimerDelegateVariant::IndexOfType<FTimerDynamicDelegate>():
+		{
+			const FTimerDynamicDelegate& FuncDynDelegate = VariantDelegate.Get<FTimerDynamicDelegate>();
+			if (FuncDynDelegate.IsBound())
 			{
-				// Add the function address
+				return FuncDynDelegate.GetUObject();
+			}
+			break;
+		}
+	default:
+		break;
+	}
+	return nullptr;
+}
 
-#if PLATFORM_COMPILER_CLANG
-				// See if this is a virtual function. Heuristic is that real function addresses are higher than some value, and vtable offsets are lower
-				const uint64 MaxVTableAddressOffset = 32768;
-				if (DumpTimerLogResolveVirtualFunctions && VtableAddr && ProgramCounter > 0 && ProgramCounter < MaxVTableAddressOffset)
+FString FTimerUnifiedDelegate::ToString() const
+{
+	switch (VariantDelegate.GetIndex())
+	{
+	case FTimerDelegateVariant::IndexOfType<FTimerDelegate>():
+		{
+			const FTimerDelegate& FuncDelegate = VariantDelegate.Get<FTimerDelegate>();
+			if (FuncDelegate.IsBound())
+			{
+				FString FunctionNameStr;
+				FName FunctionName;
+#if USE_DELEGATE_TRYGETBOUNDFUNCTIONNAME
+				FunctionName = FuncDelegate.TryGetBoundFunctionName();
+#endif
+				if (FunctionName.IsNone())
 				{
-					// If the ProgramCounter is just an offset to the vtable (virtual member function) then resolve the actual ProgramCounter here.
-					ProgramCounter = (uint64)VtableAddr[ProgramCounter / sizeof(void*)];
-				}
+					void** VtableAddr = nullptr;
+#if PLATFORM_COMPILER_CLANG || defined(_MSC_VER)
+					// Add the vtable address
+					const void* UserObject = FuncDelegate.GetObjectForTimerManager();
+					if (UserObject)
+					{
+						VtableAddr = *(void***)UserObject;
+						FunctionNameStr = FString::Printf(TEXT("vtbl: %p"), VtableAddr);
+					}
 #endif // PLATFORM_COMPILER_CLANG
 
-				FunctionNameStr += FString::Printf(TEXT(" func: 0x%llx"), ProgramCounter);
+					uint64 ProgramCounter = FuncDelegate.GetBoundProgramCounterForTimerManager();
+					if (ProgramCounter != 0)
+					{
+						// Add the function address
+
+#if PLATFORM_COMPILER_CLANG
+						// See if this is a virtual function. Heuristic is that real function addresses are higher than some value, and vtable offsets are lower
+						const uint64 MaxVTableAddressOffset = 32768;
+						if (DumpTimerLogResolveVirtualFunctions && VtableAddr && ProgramCounter > 0 && ProgramCounter < MaxVTableAddressOffset)
+						{
+							// If the ProgramCounter is just an offset to the vtable (virtual member function) then resolve the actual ProgramCounter here.
+							ProgramCounter = (uint64)VtableAddr[ProgramCounter / sizeof(void*)];
+						}
+#endif // PLATFORM_COMPILER_CLANG
 
-				if (DumpTimerLogSymbolNames)
-				{
-					// Try to resolve the function address to a symbol
-					FProgramCounterSymbolInfo SymbolInfo;
-					SymbolInfo.FunctionName[0] = 0;
-					SymbolInfo.Filename[0] = 0;
-					SymbolInfo.LineNumber = 0;
-					FPlatformStackWalk::ProgramCounterToSymbolInfo(ProgramCounter, SymbolInfo);
-					FunctionNameStr += FString::Printf(TEXT(" %s [%s:%d]"), ANSI_TO_TCHAR(SymbolInfo.FunctionName), ANSI_TO_TCHAR(SymbolInfo.Filename), SymbolInfo.LineNumber);
+						FunctionNameStr += FString::Printf(TEXT(" func: 0x%llx"), ProgramCounter);
+
+						if (DumpTimerLogSymbolNames)
+						{
+							// Try to resolve the function address to a symbol
+							FProgramCounterSymbolInfo SymbolInfo;
+							SymbolInfo.FunctionName[0] = 0;
+							SymbolInfo.Filename[0] = 0;
+							SymbolInfo.LineNumber = 0;
+							FPlatformStackWalk::ProgramCounterToSymbolInfo(ProgramCounter, SymbolInfo);
+							FunctionNameStr += FString::Printf(TEXT(" %s [%s:%d]"), ANSI_TO_TCHAR(SymbolInfo.FunctionName), ANSI_TO_TCHAR(SymbolInfo.Filename), SymbolInfo.LineNumber);
+						}
+					}
+					else
+					{
+						FunctionNameStr = TEXT(" 0x0");
+					}
 				}
+				return FString::Printf(TEXT("DELEGATE,%s,%s"), TEXT("NO OBJ"), *FunctionNameStr);
 			}
-			else
+			return TEXT("UNBOUND DELEGATE");
+		}
+	case FTimerDelegateVariant::IndexOfType<FTimerDynamicDelegate>():
+		{
+			const FTimerDynamicDelegate& FuncDynDelegate = VariantDelegate.Get<FTimerDynamicDelegate>();
+			if (FuncDynDelegate.IsBound())
 			{
-				FunctionNameStr = TEXT(" 0x0");
+				const UObject* Object = FuncDynDelegate.GetUObject();
+				return FString::Printf(TEXT("DYN DELEGATE,%s,%s"), Object == nullptr ? TEXT("NO OBJ") : *Object->GetPathName(), *FuncDynDelegate.GetFunctionName().ToString());
 			}
+			return TEXT("UNBOUND DYN DELEGATE");
 		}
-		else
+	case FTimerDelegateVariant::IndexOfType<FTimerFunction>():
 		{
-			FunctionNameStr = FunctionName.ToString();
+			if (const FTimerFunction& TimerFunction = VariantDelegate.Get<FTimerFunction>())
+			{
+				QUICK_SCOPE_CYCLE_COUNTER(STAT_FTimerUnifiedDelegate_Execute);
+				return TEXT("TFUNCTION");
+			}
+			return TEXT("UNBOUND TFUNCTION");
 		}
+	default:
+		return TEXT("UNSET");
 	}
-	else if (FuncDynDelegate.IsBound())
-	{
-		Object = FuncDynDelegate.GetUObject();
-		FunctionNameStr = FuncDynDelegate.GetFunctionName().ToString();
-		bDynDelegate = true;
-	}
-	else
-	{
-		FunctionNameStr = TEXT("NotBound!");
-	}
-
-	return FString::Printf(TEXT("%s,%s,%s"), bDynDelegate ? TEXT("DYN DELEGATE") : TEXT("DELEGATE"), Object == nullptr ? TEXT("NO OBJ") : *Object->GetPathName(), *FunctionNameStr);
 }
 
 // ---------------------------------
@@ -455,7 +554,12 @@ FTimerHandle FTimerManager::K2_FindDynamicTimerHandle(FTimerDynamicDelegate InDy
 			for (FTimerHandle Handle : *TimersForObject)
 			{
 				const FTimerData& Data = GetTimer(Handle);
-				if (Data.Status != ETimerStatus::ActivePendingRemoval && Data.TimerDelegate.FuncDynDelegate == InDynamicDelegate)
+				if (Data.Status == ETimerStatus::ActivePendingRemoval)
+				{
+					continue;
+				}
+				const FTimerDynamicDelegate* DynDelegate = Data.TimerDelegate.VariantDelegate.TryGet<FTimerDynamicDelegate>();
+				if (DynDelegate && *DynDelegate == InDynamicDelegate)
 				{
 					Result = Handle;
 					break;
diff --git a/Engine/Source/Runtime/Engine/Public/TimerManager.h b/Engine/Source/Runtime/Engine/Public/TimerManager.h
--- a/Engine/Source/Runtime/Engine/Public/TimerManager.h
+++ b/Engine/Source/Runtime/Engine/Public/TimerManager.h
@@ -10,6 +10,7 @@
 #include "Stats/Stats.h"
 #include "UObject/Object.h"
 #include "Engine/TimerHandle.h"
+#include "Misc/TVariant.h"
 #include "Templates/Function.h"
 
 class UGameInstance;
@@ -19,6 +20,8 @@ struct FTimerSourceList;
 // using "not checked" user policy (means race detection is disabled) because this delegate is stored in a TSparseArray and causes its reallocation
 // from inside delegate's execution. This is incompatible with race detection that needs to access the delegate instance after its execution
 using FTimerDelegate = TDelegate<void(), FNotThreadSafeNotCheckedDelegateUserPolicy>;
+using FTimerFunction = TFunction<void(void)>;
+using FTimerDelegateVariant = TVariant<TYPE_OF_NULLPTR, FTimerDelegate, FTimerDynamicDelegate, FTimerFunction>;
 
 #ifndef UE_ENABLE_TRACKING_TIMER_SOURCES
 #define UE_ENABLE_TRACKING_TIMER_SOURCES !UE_BUILD_SHIPPING
@@ -28,60 +31,23 @@ using FTimerDelegate = TDelegate<void(), FNotThreadSafeNotCheckedDelegateUserPol
 struct FTimerUnifiedDelegate
 {
 	/** Holds the delegate to call. */
-	FTimerDelegate FuncDelegate;
-	/** Holds the dynamic delegate to call. */
-	FTimerDynamicDelegate FuncDynDelegate;
-	/** Holds the TFunction callback to call. */
-	TFunction<void(void)> FuncCallback;
+	FTimerDelegateVariant VariantDelegate;
 
 	FTimerUnifiedDelegate() {};
-	FTimerUnifiedDelegate(FTimerDelegate const& D) : FuncDelegate(D) {};
-	FTimerUnifiedDelegate(FTimerDynamicDelegate const& D) : FuncDynDelegate(D) {};
-	FTimerUnifiedDelegate(TFunction<void(void)>&& Callback) : FuncCallback(MoveTemp(Callback)) {}
-	
-	inline void Execute()
-	{
-		if (FuncDelegate.IsBound())
-		{
-			FScopeCycleCounterUObject Context(FuncDelegate.GetUObject());
-			FuncDelegate.Execute();
-		}
-		else if (FuncDynDelegate.IsBound())
-		{
-			// stat scope is handled by UObject::ProcessEvent for the UFunction.
-			FuncDynDelegate.ProcessDelegate<UObject>(nullptr);
-		}
-		else if ( FuncCallback )
-		{
-			QUICK_SCOPE_CYCLE_COUNTER(STAT_FTimerUnifiedDelegate_Execute);
-			FuncCallback();
-		}
-	}
 
-	inline bool IsBound() const
-	{
-		return ( FuncDelegate.IsBound() || FuncDynDelegate.IsBound() || FuncCallback );
-	}
+	FTimerUnifiedDelegate(FTimerDelegate const& D) : VariantDelegate(TInPlaceType<FTimerDelegate>(), D) {}
+	FTimerUnifiedDelegate(FTimerDynamicDelegate const& D) : VariantDelegate(TInPlaceType<FTimerDynamicDelegate>(), D) {}
+	FTimerUnifiedDelegate(FTimerFunction&& Callback) : VariantDelegate(TInPlaceType<FTimerFunction>(), MoveTemp(Callback)) {}
+	
+	void Execute() const;
 
-	inline const void* GetBoundObject() const
-	{
-		if (FuncDelegate.IsBound())
-		{
-			return FuncDelegate.GetObjectForTimerManager();
-		}
-		else if (FuncDynDelegate.IsBound())
-		{
-			return FuncDynDelegate.GetUObject();
-		}
+	bool IsBound() const;
 
-		return nullptr;
-	}
+	const void* GetBoundObject() const;
 
-	inline void Unbind()
+	void Unbind()
 	{
-		FuncDelegate.Unbind();
-		FuncDynDelegate.Unbind();
-		FuncCallback = nullptr;
+		VariantDelegate.Set<TYPE_OF_NULLPTR>(nullptr);
 	}
 
 	/** Utility to output info about delegate as a string. */
