Skeletal Mesh: Automatically update alt skin profiles and vertex attributes on mesh commit

#rnx

diff --git a/Engine/Source/Developer/SkeletalMeshUtilitiesCommon/Private/LODUtilities.cpp b/Engine/Source/Developer/SkeletalMeshUtilitiesCommon/Private/LODUtilities.cpp
--- a/Engine/Source/Developer/SkeletalMeshUtilitiesCommon/Private/LODUtilities.cpp
+++ b/Engine/Source/Developer/SkeletalMeshUtilitiesCommon/Private/LODUtilities.cpp
@@ -3109,13 +3109,13 @@ bool FLODUtilities::UpdateLODInfoVertexAttributes(
 
 	FSkeletalMeshLODModel& TargetLODModel = InSkeletalMesh->GetImportedModel()->LODModels[InTargetLODIndex];
 	
-	TArray<FSkeletalMeshVertexAttributeInfo>& SkelMeshAttributeInfos = InSkeletalMesh->GetLODInfo(InTargetLODIndex)->VertexAttributes; 
+	const TArray<FSkeletalMeshVertexAttributeInfo>& SkelMeshAttributeInfos = InSkeletalMesh->GetLODInfo(InTargetLODIndex)->VertexAttributes; 
 
 	// Retain any existing attribute infos and match based on names.
-	TMap<FName, FSkeletalMeshVertexAttributeInfo> ExistingAttributeInfos;
-	for (FSkeletalMeshVertexAttributeInfo& AttributeInfo: SkelMeshAttributeInfos)
+	TMap<FName, const FSkeletalMeshVertexAttributeInfo*> ExistingAttributeInfos;
+	for (const FSkeletalMeshVertexAttributeInfo& AttributeInfo: SkelMeshAttributeInfos)
 	{
-		ExistingAttributeInfos.Add(AttributeInfo.Name, MoveTemp(AttributeInfo));
+		ExistingAttributeInfos.Add(AttributeInfo.Name, &AttributeInfo);
 	}
 	
 	const FMeshDescription* MeshDescription = InSkeletalMesh->GetMeshDescription(InSourceLODIndex);
@@ -3130,8 +3130,6 @@ bool FLODUtilities::UpdateLODInfoVertexAttributes(
 		}
 	});
  
-	SkelMeshAttributeInfos.Reset(SourceAttributes.Num());
-
 	// If we're not copying the values, leave the existing data in place.
 	if (bInCopyAttributeValues)
 	{
@@ -3146,24 +3144,18 @@ bool FLODUtilities::UpdateLODInfoVertexAttributes(
 		const FName AttributeName(SourceAttributeInfo.Key);
 
 		// Did this definition already exist? Try to retain as much of the existing information as possible.
-		FSkeletalMeshVertexAttributeInfo Info;
+		const FSkeletalMeshVertexAttributeInfo* AttributeInfo = nullptr;
 		
 		if (ExistingAttributeInfos.Contains(AttributeName))
 		{
-			Info = MoveTemp(ExistingAttributeInfos[AttributeName]);
+			AttributeInfo = ExistingAttributeInfos[AttributeName];
 		}
-		else
-		{
-			Info.Name = AttributeName;
-		}
-
-		SkelMeshAttributeInfos.Add(Info);
 
-		if (Info.IsEnabledForRender() && bInCopyAttributeValues)
+		if (AttributeInfo && AttributeInfo->IsEnabledForRender() && bInCopyAttributeValues)
 		{
 			FSkeletalMeshModelVertexAttribute& ModelAttribute = TargetLODModel.VertexAttributes.FindOrAdd(AttributeName);
 
-			ModelAttribute.DataType = Info.DataType;
+			ModelAttribute.DataType = AttributeInfo->DataType;
 			ModelAttribute.ComponentCount = 1;
 
 			if (InTargetLODIndex == InSourceLODIndex)
@@ -3192,19 +3184,6 @@ bool FLODUtilities::UpdateLODInfoVertexAttributes(
 	// Wait for all the attribute conversion tasks to complete.
 	UE::Tasks::Wait(ConversionTasks);
 
-#if WITH_EDITOR
-	// Notify UI and other systems of the change
-	// Dispatch it on the game thread for thread safty as this can be called during cook on a worker thread
-	FFunctionGraphTask::CreateAndDispatchWhenReady(
-		[WeakSkelMesh = TWeakObjectPtr<USkeletalMesh>(InSkeletalMesh)]()
-		{
-			USkeletalMesh* InSkeletalMesh = WeakSkelMesh.Get();
-			if (InSkeletalMesh)
-			{
-				InSkeletalMesh->GetOnVertexAttributesArrayChanged().Broadcast();
-			}
-		}, TStatId(), NULL, ENamedThreads::GameThread);	
-#endif
 	
 	return true;
 }
diff --git a/Engine/Source/Runtime/Engine/Classes/Engine/SkeletalMesh.h b/Engine/Source/Runtime/Engine/Classes/Engine/SkeletalMesh.h
--- a/Engine/Source/Runtime/Engine/Classes/Engine/SkeletalMesh.h
+++ b/Engine/Source/Runtime/Engine/Classes/Engine/SkeletalMesh.h
@@ -512,10 +512,19 @@ public:
 
 #if WITH_EDITORONLY_DATA
 	/** Returns the number of source models.
-	 *  \note This value is the same as the return value of GetLODNum, as these two are kept in sync.
 	 */
 	ENGINE_API int32 GetNumSourceModels() const;
 
+	/** Returns all source models as immutable. See GetSourceModel for description of what each source model contains.
+	 *  \return The list of all source models. 
+	 */
+	ENGINE_API TConstArrayView<FSkeletalMeshSourceModel> GetAllSourceModels() const;
+
+	/** Returns all source models. See GetSourceModel for description of what each source model contains.
+	 *  \return The list of all source models. 
+	 */
+	ENGINE_API TArrayView<FSkeletalMeshSourceModel> GetAllSourceModels();
+	
 	/** Returns the source model object for the skeletal mesh of a given LOD. This source models stores
 	 *  an optional mesh description that the renderable data is generated from. If there is no mesh description
 	 *  object, then the geometry of this LOD is automatically generated from an earlier LOD.
@@ -574,23 +583,38 @@ public:
 	struct FCommitMeshDescriptionParams
 	{
 		FCommitMeshDescriptionParams() {}
+
+		/** Update the alternate skin weight profile list based on presence of non-default skin weight attributes in the  
+		 *  committed mesh. If a skin weight attribute exists on the mesh but not in the list of profiles, then the profile
+		 *  is added. Likewise, if a profile exists but there's no corresponding skin weight attribute then that profile
+		 *  is removed.
+		 */
+		bool bUpdateSkinWeightProfiles = true;
+		
+		/** Update the LOD's vertex attribute list based on presence of vertex attributes in the committed mesh. 
+		 *  If a vertex attribute is defined on the committed mesh but doesn't exist in the list of vertex attributes 
+		 *  on the LOD, then that vertex attribute is added to the LOD. Likewise, vertex attribute exists on the LOD, 
+		 *  but not on the committed mesh, then the attribute is removed from the LOD. 
+		 */
+		bool bUpdateVertexAttributes = true;
 		
 		/** Mark the package as dirty. If calling CommitMeshDescription from a non-game thread,
 		 *  this value should be set to \c false.
 		 */
 		bool bMarkPackageDirty = true;
 
-		/** Force the render data to update. By default the render data uses the hash of the mesh to check if
+		/** Force the render data to update. By default, the render data uses the hash of the mesh to check if
 		 *  an update is required, this forces this hash to be unique, causing the render data to update whether
 		 *  the mesh has changed or not.
 		 */
 		bool bForceUpdate = false;
 	};
 
-	/* Commits the stored mesh description object to bulk storage. This also forces the imported bounds to update
-	 * if the mesh committed is on LOD 0. If there is no stored mesh description, the bulk storage will be emptied.
+	/** Commits the stored mesh description object to bulk storage. This also forces the imported bounds to update
+	 *  if the mesh committed is on LOD 0. If there is no stored mesh description, the bulk storage will be emptied.
+	 *  \note It is thread-safe to commit multiple meshes simultaneously, as long as they're all on different LODs. 
 	 *  \param InLODIndex The LOD index at which to commit the mesh description to bulk storage.
-	 *  \param InParams An optional object to control how the commit is done.
+	 *  \param InParams An optional object to control what happens during the commit.
 	 *  \return \c true if the commit was successful. Even if there was no mesh description to store, the commit
 	 *    is still successful.
 	 */
@@ -631,52 +655,36 @@ public:
 	// Raw mesh data DDC string ID, there is no API to retrieve it, since only the LODModels need this value
 	
 
-	/* Fill the OutMesh with the imported data */
 	UE_DEPRECATED(5.4, "Use GetMeshDescription instead.")
 	ENGINE_API void LoadLODImportedData(const int32 LODIndex, FSkeletalMeshImportData& OutMesh) const;
 
-	/* Fill the asset LOD entry with the InMesh. */
 	UE_DEPRECATED(5.4, "Use CommitMeshDescription instead.")
 	ENGINE_API void SaveLODImportedData(const int32 LODIndex, const FSkeletalMeshImportData& InMesh);
 	
-	/* Return true if the imported data has all the necessary data to use the skeletalmesh builder. Return False otherwise.
-	 * Old asset before the refactor will not be able to be build until it get fully re-import.
-	 * This value is cache in the LODModel and update when we call SaveLODImportedData.
-	 */
 	UE_DEPRECATED(5.4, "Use HasMeshDescription instead.")
 	ENGINE_API bool IsLODImportedDataBuildAvailable(const int32 LODIndex) const;
 	
-	/* Return true if the imported data is present. Return false otherwise.
-	 * Old asset before the split workflow will not have this data and will not support import geo only or skinning only.
-	 * This value is cache in the LODModel and update when we call SaveLODImportedData.
-	 */
 	UE_DEPRECATED(5.4, "Use HasMeshDescription instead.")
 	ENGINE_API bool IsLODImportedDataEmpty(const int32 LODIndex) const;
 
-	/* Get the Versions of the geo and skinning data. We use those versions to answer to IsLODImportedDataBuildAvailable function. */
 	UE_DEPRECATED(5.4, "No equivalent provided since versioning is not surfaced for mesh description bulk data.")
 	ENGINE_API void GetLODImportedDataVersions(const int32 LODIndex, ESkeletalMeshGeoImportVersions& OutGeoImportVersion, ESkeletalMeshSkinningImportVersions& OutSkinningImportVersion) const;
 
-	/* Set the Versions of the geo and skinning data. We use those versions to answer to IsLODImportedDataBuildAvailable function. */
 	UE_DEPRECATED(5.4, "No equivalent provided since versioning is not surfaced for mesh description bulk data.")
 	ENGINE_API void SetLODImportedDataVersions(const int32 LODIndex, const ESkeletalMeshGeoImportVersions& InGeoImportVersion, const ESkeletalMeshSkinningImportVersions& InSkinningImportVersion);
 
-	/* Static function that copy the LOD import data from a source skeletal mesh to a destination skeletal mesh*/
 	UE_DEPRECATED(5.4, "Use GetMeshDescription and CreateMeshDescription instead.")
 	static ENGINE_API void CopyImportedData(int32 SrcLODIndex, USkeletalMesh* SrcSkeletalMesh, int32 DestLODIndex, USkeletalMesh* DestSkeletalMesh);
 
-	/* Allocate the space we need. Use this before calling this API in multithreaded. */
 	UE_DEPRECATED(5.4, "No equivalent provided. All LODs should be added up-front for multi-threaded use.")
 	ENGINE_API void ReserveLODImportData(int32 MaxLODIndex);
 	
 	UE_DEPRECATED(5.4, "No equivalent provided.")
 	ENGINE_API void ForceBulkDataResident(const int32 LODIndex);
 
-	/* Remove the import data for the specified LOD */
 	UE_DEPRECATED(5.4, "Use ClearMeshDescriptionAndBulkData instead.")
 	ENGINE_API void EmptyLODImportData(const int32 LODIndex);
 
-	/* Remove the import data for all the LODs */
 	UE_DEPRECATED(5.4, "Use ClearMeshDescriptionAndBulkData instead.")
 	ENGINE_API void EmptyAllImportData();
 
@@ -1419,16 +1427,6 @@ public:
 		PRAGMA_ENABLE_DEPRECATION_WARNINGS
 	}
 
-	UE_DEPRECATED(4.27, "Please do not use this non const function; use the combination of USkeletalMesh::CreateBodySetup() and USkeletalMesh::GetBodySetup() const. Cast the skeletal mesh caller to const to force the compiler to use the USkeletalMesh::GetBodySetup() const function and avoid the deprecation warning")
-	class UBodySetup* GetBodySetup()
-	{
-		WaitUntilAsyncPropertyReleased(ESkeletalMeshAsyncProperties::BodySetup);
-		PRAGMA_DISABLE_DEPRECATION_WARNINGS
-		CreateBodySetup();
-		return BodySetup;
-		PRAGMA_ENABLE_DEPRECATION_WARNINGS
-	}
-
 	void SetBodySetup(class UBodySetup* InBodySetup)
 	{
 		WaitUntilAsyncPropertyReleased(ESkeletalMeshAsyncProperties::BodySetup);
diff --git a/Engine/Source/Runtime/Engine/Classes/Engine/SkeletalMeshSourceModel.h b/Engine/Source/Runtime/Engine/Classes/Engine/SkeletalMeshSourceModel.h
--- a/Engine/Source/Runtime/Engine/Classes/Engine/SkeletalMeshSourceModel.h
+++ b/Engine/Source/Runtime/Engine/Classes/Engine/SkeletalMeshSourceModel.h
@@ -120,11 +120,21 @@ struct FSkeletalMeshSourceModel
 	{
 		return Bounds;
 	}
-	
+
+	/**
+	 * Returns the list of names of alternate skin weight profiles stored on the mesh, if any.
+	 */
+	TConstArrayView<FName> GetSkinWeightProfileNames() const;
+
+	/**
+	 * Returns the list names of morph targets stored on the mesh, if any.
+	 */
+	TConstArrayView<FName> GetMorphTargetNames() const;
+
 private:
 #if WITH_EDITOR
 	USkeletalMesh* GetOwner() const;
-
+	
 	/**
 	 * Ensures that any old raw mesh bulk data has been converted to new data. Used prior to saving out an
 	 * asset, since the old bulk data will not be saved out again and would otherwise be lost.
@@ -157,6 +167,12 @@ private:
 	 *   the statistics are reset to zero, for the counts, and an empty bounds.
 	 */
 	void UpdateCachedMeshStatistics(const FMeshDescription* InMeshDescription);
+
+	/**
+	 *  Update the cached mesh statistics from either resident mesh, or by loading
+	 *  and then unloading the bulkd data, to conserve memory.
+	 */
+	void UpdateCachedMeshStatisticsFromBulkIfNeeded();
 #endif
 	
 	UPROPERTY()
@@ -167,6 +183,17 @@ private:
 
 	UPROPERTY()
 	FBoxSphereBounds Bounds;
+
+	// List of all alternate skin weight profiles stored on the mesh,
+	// NOTE: The default value is single-entry array with NAME_None to mark
+	// that this value has not been initialized yet / TOptional does not work
+	// with containers. 
+	UPROPERTY()
+	TArray<FName> CachedSkinWeightProfileNames{ NAME_None };
+
+	// List of all morph targets stored on the mesh.
+	UPROPERTY()
+	TArray<FName> CachedMorphTargetNames{ NAME_None };
 	
 #if WITH_EDITORONLY_DATA
 	/**
diff --git a/Engine/Source/Runtime/Engine/Private/Animation/SkeletalMeshVertexAttribute.cpp b/Engine/Source/Runtime/Engine/Private/Animation/SkeletalMeshVertexAttribute.cpp
--- a/Engine/Source/Runtime/Engine/Private/Animation/SkeletalMeshVertexAttribute.cpp
+++ b/Engine/Source/Runtime/Engine/Private/Animation/SkeletalMeshVertexAttribute.cpp
@@ -7,11 +7,6 @@
 #include UE_INLINE_GENERATED_CPP_BY_NAME(SkeletalMeshVertexAttribute)
 
 
-FName FSkeletalMeshVertexAttributeInfo::GetRequirementName() const
-{
-	return FName(FString::Printf(TEXT("Vertex.%s"), *Name.ToString()));
-}
-
 bool FSkeletalMeshVertexAttributeInfo::IsEnabledForRender() const
 {
 	return EnabledForRender.GetValue();
diff --git a/Engine/Source/Runtime/Engine/Private/SkeletalMesh.cpp b/Engine/Source/Runtime/Engine/Private/SkeletalMesh.cpp
--- a/Engine/Source/Runtime/Engine/Private/SkeletalMesh.cpp
+++ b/Engine/Source/Runtime/Engine/Private/SkeletalMesh.cpp
@@ -2420,6 +2420,22 @@ void USkeletalMesh::RemoveSourceModelInternal(const int32 InLODIndex)
 }
 
 
+TConstArrayView<FSkeletalMeshSourceModel> USkeletalMesh::GetAllSourceModels() const
+{
+	WaitUntilAsyncPropertyReleased(ESkeletalMeshAsyncProperties::SourceModels);
+	PRAGMA_DISABLE_DEPRECATION_WARNINGS
+	return SourceModels;
+	PRAGMA_ENABLE_DEPRECATION_WARNINGS
+}
+
+TArrayView<FSkeletalMeshSourceModel> USkeletalMesh::GetAllSourceModels()
+{
+	WaitUntilAsyncPropertyReleased(ESkeletalMeshAsyncProperties::SourceModels);
+	PRAGMA_DISABLE_DEPRECATION_WARNINGS
+	return SourceModels;
+	PRAGMA_ENABLE_DEPRECATION_WARNINGS
+}
+
 const FSkeletalMeshSourceModel& USkeletalMesh::GetSourceModel(const int32 InLODIndex) const
 {
 	WaitUntilAsyncPropertyReleased(ESkeletalMeshAsyncProperties::SourceModels);
@@ -2516,6 +2532,107 @@ bool USkeletalMesh::CommitMeshDescription(
 	{
 		SetImportedBounds(SourceModel.GetBoundsFast());
 	}
+	
+	if (SourceModel.HasMeshDescription())
+	{
+		if (InParams.bUpdateSkinWeightProfiles)
+		{
+			static FCriticalSection ProfileUpdateMutex;
+
+			// Since SkinWeightProfiles is a USkeletalMesh member, we want to avoid multiple
+			// threads all mutating it at the same time, in case we have a geometry processor
+			// that is committing multiple meshes across differing LODs simultaneously.
+			FScopeLock ScopeLock(&ProfileUpdateMutex);
+			
+			FSkeletalMeshConstAttributes Attributes(*SourceModel.GetMeshDescription());
+
+			TArray<FSkinWeightProfileInfo>& ExistingProfiles = GetSkinWeightProfiles();
+			TMap<FName, int32> ProfileIndexes;
+
+			for (int32 Index = 0; Index < ExistingProfiles.Num(); Index++)
+			{
+				ProfileIndexes.Add(ExistingProfiles[Index].Name, Index);
+			}
+
+			// Get all profiles from all LODs, since we may have some that weren't defined 
+			TSet<FName> ValidProfiles;
+			for (const FSkeletalMeshSourceModel& OtherSourceModels: GetAllSourceModels())
+			{
+				ValidProfiles.Append(OtherSourceModels.GetSkinWeightProfileNames());
+			}
+			for (const FName& ProfileName: ValidProfiles)
+			{
+				if (!ProfileIndexes.Contains(ProfileName))
+				{
+					FSkinWeightProfileInfo& NewProfile = ExistingProfiles.AddDefaulted_GetRef();
+					NewProfile.Name = ProfileName;
+				}
+			}
+
+			ExistingProfiles.RemoveAll([ValidProfiles](const FSkinWeightProfileInfo& InProfileInfo)
+			{
+				return !ValidProfiles.Contains(InProfileInfo.Name);
+			});
+		}
+		
+		if (InParams.bUpdateVertexAttributes)
+		{
+			TSet<FName> MeshVertexAttributes;
+			
+			// NOTE: We're currently limited to just single-channel attributes for rendering.
+			SourceModel.GetMeshDescription()->VertexAttributes().ForEachByType<float>([&MeshVertexAttributes](const FName InAttributeName, TVertexAttributesConstRef<float> InAttributeRef)
+			{
+				if (!FSkeletalMeshAttributes::IsReservedAttributeName(InAttributeName))
+				{
+					MeshVertexAttributes.Add(InAttributeName);
+				}
+			});
+
+			bool bVertexAttributesChanged = false;
+			TArray<FSkeletalMeshVertexAttributeInfo>& ExistingVertexAttributes = GetLODInfo(InLODIndex)->VertexAttributes;
+			TSet<FName> ExistingVertexAttributeNames;
+			for (const FSkeletalMeshVertexAttributeInfo& AttributeInfo: ExistingVertexAttributes)
+			{
+				ExistingVertexAttributeNames.Add(AttributeInfo.Name);
+			}
+			
+			for (FName AttributeName: MeshVertexAttributes)
+			{
+				if (!ExistingVertexAttributeNames.Contains(AttributeName))
+				{
+					FSkeletalMeshVertexAttributeInfo& NewAttribute = ExistingVertexAttributes.AddDefaulted_GetRef();
+					NewAttribute.Name = AttributeName;
+
+					bVertexAttributesChanged = true;
+				}
+			}
+
+			// Remove all attributes from the LOD that no longer exist on the mesh.
+			if (ExistingVertexAttributes.RemoveAll([&MeshVertexAttributes](const FSkeletalMeshVertexAttributeInfo& InInfo)
+				{
+					return !MeshVertexAttributes.Contains(InInfo.Name);
+				}))
+			{
+				bVertexAttributesChanged = true;
+			}
+
+			if (bVertexAttributesChanged)
+			{
+#if WITH_EDITOR
+				// Notify UI and other systems of the change
+				// Dispatch it on the game thread for thread-safety as this can be called on a worker thread
+				FFunctionGraphTask::CreateAndDispatchWhenReady(
+					[WeakSkelMesh = TWeakObjectPtr<USkeletalMesh>(this)]()
+					{
+						if (USkeletalMesh* SkeletalMesh = WeakSkelMesh.Get())
+						{
+							SkeletalMesh->GetOnVertexAttributesArrayChanged().Broadcast();
+						}
+					}, TStatId(), nullptr, ENamedThreads::GameThread);	
+#endif			
+			}
+		}
+	}
 
 	if (ensure(GetImportedModel()->LODModels.IsValidIndex(InLODIndex)))
 	{
diff --git a/Engine/Source/Runtime/Engine/Private/SkeletalMeshSourceModel.cpp b/Engine/Source/Runtime/Engine/Private/SkeletalMeshSourceModel.cpp
--- a/Engine/Source/Runtime/Engine/Private/SkeletalMeshSourceModel.cpp
+++ b/Engine/Source/Runtime/Engine/Private/SkeletalMeshSourceModel.cpp
@@ -216,6 +216,24 @@ const FMeshDescriptionBulkData* FSkeletalMeshSourceModel::GetMeshDescriptionBulk
 }
 
 
+TConstArrayView<FName> FSkeletalMeshSourceModel::GetSkinWeightProfileNames() const
+{
+	if (CachedSkinWeightProfileNames.Num() == 1 && CachedSkinWeightProfileNames[0] == NAME_None)
+	{
+		const_cast<FSkeletalMeshSourceModel*>(this)->UpdateCachedMeshStatisticsFromBulkIfNeeded();
+	}
+	return CachedSkinWeightProfileNames;
+}
+
+TConstArrayView<FName> FSkeletalMeshSourceModel::GetMorphTargetNames() const
+{
+	if (CachedMorphTargetNames.Num() == 1 && CachedMorphTargetNames[0] == NAME_None)
+	{
+		const_cast<FSkeletalMeshSourceModel*>(this)->UpdateCachedMeshStatisticsFromBulkIfNeeded();
+	}
+	return CachedMorphTargetNames;
+}
+
 USkeletalMesh* FSkeletalMeshSourceModel::GetOwner() const
 {
 	if (!ensure(MeshDescriptionBulkData))
@@ -349,11 +367,57 @@ void FSkeletalMeshSourceModel::UpdateCachedMeshStatistics(const FMeshDescription
 		TriangleCount = InMeshDescription->Triangles().Num();
 		VertexCount = InMeshDescription->Vertices().Num();
 		Bounds = InMeshDescription->GetBounds();
+		
+		FSkeletalMeshConstAttributes Attributes(*InMeshDescription);
+		constexpr bool bInUserDefinedOnly = true;
+		CachedSkinWeightProfileNames = Attributes.GetSkinWeightProfileNames(bInUserDefinedOnly);
+		CachedMorphTargetNames = Attributes.GetMorphTargetNames();
 	}
 	else
 	{
 		TriangleCount = VertexCount = 0;
 		Bounds = FBoxSphereBounds{ForceInitToZero};
+		CachedSkinWeightProfileNames.Reset();
+		CachedMorphTargetNames.Reset();
+	}
+}
+
+void FSkeletalMeshSourceModel::UpdateCachedMeshStatisticsFromBulkIfNeeded()
+{
+	TRACE_CPUPROFILER_EVENT_SCOPE(FSkeletalMeshSourceData::UpdateCachedMeshStatisticsFromBulkIfNeeded);
+
+	FScopeLock Lock(&MeshDescriptionBulkDataMutex);
+
+	if (!ensure(MeshDescriptionBulkData))
+	{
+		return;
+	}
+	
+	if (MeshDescriptionBulkData->HasCachedMeshDescription())
+	{
+		UpdateCachedMeshStatistics(&MeshDescriptionBulkData->GetMeshDescription()->GetMeshDescription());
+	}
+
+	if (!MeshDescriptionBulkData->IsBulkDataValid())
+	{
+		if (RawMeshBulkData.IsValid())
+		{
+			// This call will implicitly update the mesh statistics.
+			ConvertRawMeshToMeshDescriptionBulkData();
+		}
+		else
+		{
+			UpdateCachedMeshStatistics(nullptr);
+		}
+	}
+	else
+	{
+		// Temporarily load the mesh description from the bulk data, we don't care about fixing up the morph targets
+		// in this instance (as would be the case with LoadMeshDescriptionFromBulkData).
+		FMeshDescription MeshDescription;
+		MeshDescriptionBulkData->GetBulkData().LoadMeshDescription(MeshDescription);
+
+		UpdateCachedMeshStatistics(&MeshDescription);
 	}
 }
 
diff --git a/Engine/Source/Runtime/Engine/Public/Animation/SkeletalMeshVertexAttribute.h b/Engine/Source/Runtime/Engine/Public/Animation/SkeletalMeshVertexAttribute.h
--- a/Engine/Source/Runtime/Engine/Public/Animation/SkeletalMeshVertexAttribute.h
+++ b/Engine/Source/Runtime/Engine/Public/Animation/SkeletalMeshVertexAttribute.h
@@ -41,9 +41,6 @@ struct FSkeletalMeshVertexAttributeInfo
 	UPROPERTY(EditAnywhere, Category="Vertex Attributes")
 	ESkeletalMeshVertexAttributeDataType DataType = ESkeletalMeshVertexAttributeDataType::Float;
 
-	/** Returns the name to use for this attribute when */ 
-	ENGINE_API FName GetRequirementName() const;
-
 	/** Returns true if this attribute is enabled for rendering on this platform */
 	ENGINE_API bool IsEnabledForRender() const;
 };
diff --git a/Engine/Source/Runtime/SkeletalMeshDescription/Private/SkeletalMeshAttributes.cpp b/Engine/Source/Runtime/SkeletalMeshDescription/Private/SkeletalMeshAttributes.cpp
--- a/Engine/Source/Runtime/SkeletalMeshDescription/Private/SkeletalMeshAttributes.cpp
+++ b/Engine/Source/Runtime/SkeletalMeshDescription/Private/SkeletalMeshAttributes.cpp
@@ -358,7 +358,7 @@ FSkinWeightsVertexAttributesConstRef FSkeletalMeshAttributesShared::GetVertexSki
 	return {};
 }
 
-TArray<FName> FSkeletalMeshAttributesShared::GetSkinWeightProfileNames() const
+TArray<FName> FSkeletalMeshAttributesShared::GetSkinWeightProfileNames(const bool bInUserDefinedOnly) const
 {
 	TArray<FName> AllAttributeNames;
 	MeshDescriptionShared.VertexAttributes().GetAttributeNames(AllAttributeNames);
@@ -379,7 +379,7 @@ TArray<FName> FSkeletalMeshAttributesShared::GetSkinWeightProfileNames() const
 	}
 
 	SkinWeightProfileNames.Sort([](const FName A, const FName B) -> bool { return A.FastLess(B); });
-	if (bHasDefault)
+	if (bHasDefault && !bInUserDefinedOnly)
 	{
 		SkinWeightProfileNames.Insert(DefaultSkinWeightProfileName, 0);
 	}
diff --git a/Engine/Source/Runtime/SkeletalMeshDescription/Public/SkeletalMeshAttributes.h b/Engine/Source/Runtime/SkeletalMeshDescription/Public/SkeletalMeshAttributes.h
--- a/Engine/Source/Runtime/SkeletalMeshDescription/Public/SkeletalMeshAttributes.h
+++ b/Engine/Source/Runtime/SkeletalMeshDescription/Public/SkeletalMeshAttributes.h
@@ -95,8 +95,11 @@ public:
 	// Skin Weights Methods
 	//
 
-	/// Returns the list of all registered skin weight profile names on this mesh.
-	SKELETALMESHDESCRIPTION_API TArray<FName> GetSkinWeightProfileNames() const;
+	/** Returns the list of all registered skin weight profile names on this mesh.
+	 *  \param bInUserDefinedOnly Only return user-defined profiles, not the default profile.
+	 *  \return The list of the profile names.
+	 */
+	SKELETALMESHDESCRIPTION_API TArray<FName> GetSkinWeightProfileNames(const bool bInUserDefinedOnly = false) const;
 	
 	/// Returns \c true if the given identifier is a valid profile name. If the name is empty, or matches the default profile,
 	/// then the profile name is considered invalid. 
