Fix issue with FActorDescArchivePatcher which can recursively call itself inadvertently writing out values read for patching to the OutArchive, and then later again writing out the patched value. This results in more data written than read which will cause corruption when loading.

We now guard from writing out serialized values if we are already patching

#rnx

diff --git a/Engine/Source/Runtime/Engine/Private/WorldPartition/WorldPartitionActorDescArchive.cpp b/Engine/Source/Runtime/Engine/Private/WorldPartition/WorldPartitionActorDescArchive.cpp
--- a/Engine/Source/Runtime/Engine/Private/WorldPartition/WorldPartitionActorDescArchive.cpp
+++ b/Engine/Source/Runtime/Engine/Private/WorldPartition/WorldPartitionActorDescArchive.cpp
@@ -135,19 +135,35 @@ FArchive& FActorDescArchive::operator<<(FSoftObjectPath& Value)
 
 FArchive& FActorDescArchivePatcher::operator<<(FName& Value)
 {
-	TGuardValue<bool> GuardIsPatching(bIsPatching, true);
-	FActorDescArchive::operator<<(Value);
-	AssetDataPatcher->DoPatch(Value);
-	OutAr << Value;
+	{
+		TGuardValue<bool> GuardIsPatching(bIsPatching, true);
+		FActorDescArchive::operator<<(Value);
+		AssetDataPatcher->DoPatch(Value);
+	}
+
+    // Only write out values if we aren't already patching since this function can be called
+	// from other patching functions which will perform the final write of the patched values
+	if (!bIsPatching)
+	{
+		OutAr << Value;
+	}
 	return *this;
 }
 
 FArchive& FActorDescArchivePatcher::operator<<(FSoftObjectPath& Value)
 {
-	TGuardValue<bool> GuardIsPatching(bIsPatching, true);
-	FActorDescArchive::operator<<(Value);
-	AssetDataPatcher->DoPatch(Value);
-	Value.SerializePathWithoutFixup(OutAr);
+	{
+		TGuardValue<bool> GuardIsPatching(bIsPatching, true);
+		FActorDescArchive::operator<<(Value);
+		AssetDataPatcher->DoPatch(Value);
+	}
+
+    // Only write out values if we aren't already patching since this function can be called
+	// from other patching functions which will perform the final write of the patched values
+	if (!bIsPatching)
+	{
+		Value.SerializePathWithoutFixup(OutAr);
+	}
 	return *this;
 }
 
