Added constructor zero initialization of bitfield member bDisablePreExposureScale in UMaterial.  Some other bitfield members in the class aren't initialized in the constructor, but most are, so it seems best to have.

#rnx

diff --git a/Engine/Source/Runtime/Engine/Private/Materials/Material.cpp b/Engine/Source/Runtime/Engine/Private/Materials/Material.cpp
--- a/Engine/Source/Runtime/Engine/Private/Materials/Material.cpp
+++ b/Engine/Source/Runtime/Engine/Private/Materials/Material.cpp
@@ -1017,6 +1017,7 @@ UMaterial::UMaterial(const FObjectInitializer& ObjectInitializer)
 	BlendableOutputAlpha = false;
 	bIsBlendable = true;
 	PreshaderGap = 0;
+	bDisablePreExposureScale = false;
 	bEnableStencilTest = false;
 	bUsedWithVolumetricCloud = false;
 	bUsedWithHeterogeneousVolumes = false;
