[World Partition Actor Descriptors] Add actor descriptor support for internal actors.
This allows standard levels (non-OFPA) to be used as level instances where their actors are placed in the streaming grid.
No future plans for supporting data layers, actor filters or property overrides yet.

fyi jeansebastien.guay

diff --git a/Engine/Plugins/Editor/DataValidation/Source/DataValidation/Private/WorldPartitionChangelistValidator.cpp b/Engine/Plugins/Editor/DataValidation/Source/DataValidation/Private/WorldPartitionChangelistValidator.cpp
--- a/Engine/Plugins/Editor/DataValidation/Source/DataValidation/Private/WorldPartitionChangelistValidator.cpp
+++ b/Engine/Plugins/Editor/DataValidation/Source/DataValidation/Private/WorldPartitionChangelistValidator.cpp
@@ -519,10 +519,10 @@ void UWorldPartitionChangelistValidator::OnLevelInstanceInvalidWorldAsset(const
 				FText::FromName(WorldAsset));
 			AssetFails(CurrentAsset, CurrentError);
 			break;
-		case ELevelInstanceInvalidReason::WorldAssetNotUsingExternalActors:
+		case ELevelInstanceInvalidReason::WorldAssetDontContainActorsMetadata:
 			// Not a validation error
 			break;
-		case ELevelInstanceInvalidReason::WorldAssetImcompatiblePartitioned:
+		case ELevelInstanceInvalidReason::WorldAssetIncompatiblePartitioned:
 			// Not a validation error
 			break;
 		case ELevelInstanceInvalidReason::WorldAssetHasInvalidContainer:
diff --git a/Engine/Source/Runtime/Engine/Classes/Engine/Level.h b/Engine/Source/Runtime/Engine/Classes/Engine/Level.h
--- a/Engine/Source/Runtime/Engine/Classes/Engine/Level.h
+++ b/Engine/Source/Runtime/Engine/Classes/Engine/Level.h
@@ -722,6 +722,9 @@ public:
 	ENGINE_API static bool GetIsLevelUsingExternalActorsFromAsset(const FAssetData& Asset);
 	ENGINE_API static bool GetIsLevelUsingExternalActorsFromPackage(FName LevelPackage);
 
+	ENGINE_API static bool GetIsLevelUsingActorsDescsFromAsset(const FAssetData& Asset);
+	ENGINE_API static bool GetIsLevelUsingActorsDescsFromPackage(FName LevelPackage);
+
 	ENGINE_API static bool GetIsUsingActorFoldersFromAsset(const FAssetData& Asset);
 	ENGINE_API static bool GetIsUsingActorFoldersFromPackage(FName LevelPackage);
 
diff --git a/Engine/Source/Runtime/Engine/Classes/Engine/World.h b/Engine/Source/Runtime/Engine/Classes/Engine/World.h
--- a/Engine/Source/Runtime/Engine/Classes/Engine/World.h
+++ b/Engine/Source/Runtime/Engine/Classes/Engine/World.h
@@ -2800,6 +2800,7 @@ public:
 	virtual void GetAssetRegistryTags(FAssetRegistryTagsContext Context) const override;
 	UE_DEPRECATED(5.4, "Implement the version that takes FAssetRegistryTagsContext instead.")
 	virtual void GetAssetRegistryTags(TArray<FAssetRegistryTag>& OutTags) const override;
+	virtual void GetExtendedAssetRegistryTagsForSave(const ITargetPlatform* TargetPlatform, TArray<FAssetRegistryTag>& OutTags) const override;
 protected:
 	virtual void ThreadedPostLoadAssetRegistryTagsOverride(FPostLoadAssetRegistryTagsContext& Context) const;
 public:
diff --git a/Engine/Source/Runtime/Engine/Private/Level.cpp b/Engine/Source/Runtime/Engine/Private/Level.cpp
--- a/Engine/Source/Runtime/Engine/Private/Level.cpp
+++ b/Engine/Source/Runtime/Engine/Private/Level.cpp
@@ -2514,6 +2514,25 @@ bool ULevel::GetIsLevelUsingExternalActorsFromPackage(FName LevelPackage)
 	});
 }
 
+bool ULevel::GetIsLevelUsingActorsDescsFromAsset(const FAssetData& Asset)
+{
+	FString ActorsMetaDataStr;
+	static FName NAME_ActorsMetaData(TEXT("ActorsMetaData"));
+	if (Asset.GetTagValue(NAME_ActorsMetaData, ActorsMetaDataStr))
+	{
+		return true;
+	}
+	return false;
+}
+
+bool ULevel::GetIsLevelUsingActorsDescsFromPackage(FName LevelPackage)
+{
+	return LevelAssetRegistryHelper::GetLevelInfoFromAssetRegistry(LevelPackage, [](const FAssetData& Asset)
+	{
+		return GetIsLevelUsingActorsDescsFromAsset(Asset);
+	});
+}
+
 bool ULevel::GetIsUsingActorFoldersFromAsset(const FAssetData& Asset)
 {
 	FString LevelIsUsingActorFoldersStr;
diff --git a/Engine/Source/Runtime/Engine/Private/World.cpp b/Engine/Source/Runtime/Engine/Private/World.cpp
--- a/Engine/Source/Runtime/Engine/Private/World.cpp
+++ b/Engine/Source/Runtime/Engine/Private/World.cpp
@@ -54,6 +54,7 @@
 #include "WorldPartition/WorldPartition.h"
 #include "WorldPartition/DataLayer/DataLayerManager.h"
 #include "WorldPartition/DataLayer/WorldDataLayers.h"
+#include "WorldPartition/WorldPartitionActorDescUtils.h"
 #include "GameFramework/GameNetworkManager.h"
 #include "GameFramework/CharacterMovementComponent.h"
 #include "Net/NetworkProfiler.h"
@@ -9328,6 +9329,31 @@ void UWorld::GetAssetRegistryTags(FAssetRegistryTagsContext Context) const
 	FWorldDelegates::GetAssetTagsWithContext.Broadcast(this, Context);
 }
 
+void UWorld::GetExtendedAssetRegistryTagsForSave(const ITargetPlatform* TargetPlatform, TArray<FAssetRegistryTag>& OutTags) const
+{
+	Super::GetExtendedAssetRegistryTagsForSave(TargetPlatform, OutTags);
+
+	if (!PersistentLevel->IsUsingExternalActors())
+	{
+		TArray<FString> ActorsMetaData;
+		for (AActor* Actor : PersistentLevel->Actors)
+		{
+			if (IsValid(Actor) && Actor->SupportsExternalPackaging())
+			{
+				FWorldPartitionActorDescUtils::FActorDescInitParams ActorDescInitParams(Actor);
+				ActorsMetaData.Add(ActorDescInitParams.ToString());
+			}
+		}
+
+		if (ActorsMetaData.Num())
+		{
+			static FName NAME_ActorsMetaData(TEXT("ActorsMetaData"));
+			const FString ActorsMetaDataStr = FString::Join(ActorsMetaData, TEXT(";"));
+			OutTags.Add(UObject::FAssetRegistryTag(NAME_ActorsMetaData, ActorsMetaDataStr, UObject::FAssetRegistryTag::TT_Hidden));
+		}
+	}
+}
+
 void UWorld::ThreadedPostLoadAssetRegistryTagsOverride(FPostLoadAssetRegistryTagsContext& Context) const
 {
 	Super::ThreadedPostLoadAssetRegistryTagsOverride(Context);
diff --git a/Engine/Source/Runtime/Engine/Private/WorldPartition/ActorDescContainer.cpp b/Engine/Source/Runtime/Engine/Private/WorldPartition/ActorDescContainer.cpp
--- a/Engine/Source/Runtime/Engine/Private/WorldPartition/ActorDescContainer.cpp
+++ b/Engine/Source/Runtime/Engine/Private/WorldPartition/ActorDescContainer.cpp
@@ -69,33 +69,68 @@ void UActorDescContainer::Initialize(const FInitializeParams& InitParams)
 		ContentBundleGuid = InitParams.ContentBundleGuid;
 	}
 
-	TArray<FAssetData> Assets;
+	TArray<FAssetData> ExternalAssets;
+	TArray<FString> InternalAssets;
 	if (!ContainerPackageName.IsNone())
 	{
 		const FString ContainerExternalActorsPath = GetExternalActorPath();
 
-		// Do a synchronous scan of the level external actors path.					
 		IAssetRegistry& AssetRegistry = FModuleManager::LoadModuleChecked<FAssetRegistryModule>(TEXT("AssetRegistry")).Get();
+
+		// Do a synchronous scan of the level external actors path.
 		{
 			TRACE_CPUPROFILER_EVENT_SCOPE(ScanSynchronous);
 			AssetRegistry.ScanSynchronous({ ContainerExternalActorsPath }, TArray<FString>());
 		}
 
-		FARFilter Filter;
-		Filter.bRecursivePaths = true;
-		Filter.bIncludeOnlyOnDiskAssets = true;
-		Filter.PackagePaths.Add(*ContainerExternalActorsPath);
+		// Gather external actors
+		{
+			TRACE_CPUPROFILER_EVENT_SCOPE(GetExternalAssets);
+
+			FARFilter Filter;
+			Filter.bRecursivePaths = true;
+			Filter.bIncludeOnlyOnDiskAssets = true;
+			Filter.PackagePaths.Add(*ContainerExternalActorsPath);
+
+			FExternalPackageHelper::GetSortedAssets(Filter, ExternalAssets);
+		}
 
-		TRACE_CPUPROFILER_EVENT_SCOPE(GetAssets);
-		FExternalPackageHelper::GetSortedAssets(Filter, Assets);
+		// Gather non-external actors
+		{
+			TRACE_CPUPROFILER_EVENT_SCOPE(GetInternalAssets);
+
+			FARFilter Filter;
+			Filter.bIncludeOnlyOnDiskAssets = true;
+			Filter.PackageNames.Add(ContainerPackageName);
+
+			TArray<FAssetData> WorldAssetData;
+			AssetRegistry.GetAssets(Filter, WorldAssetData);
+
+			// Transform world assets
+			static FName NAME_ActorsMetaData(TEXT("ActorsMetaData"));
+			for (const FAssetData& AssetData : WorldAssetData)
+			{
+				FString ActorsMetaDataStr;
+				if (AssetData.GetTagValue(NAME_ActorsMetaData, ActorsMetaDataStr))
+				{
+					TArray<FString> ActorsMetaData;
+					if (ActorsMetaDataStr.ParseIntoArray(ActorsMetaData, TEXT(";")))
+					{
+						InternalAssets.Append(ActorsMetaData);
+					}
+				}
+			}
+		}
 	}
 
+	UE_LOG(LogWorldPartition, Verbose, TEXT("Parsed actor descriptor container package '%s': %d external actors, %d internal actors"), *InitParams.PackageName.ToString(), ExternalAssets.Num(), InternalAssets.Num());
+
 	FWorldPartitionClassDescRegistry& ClassDescRegistry = FWorldPartitionClassDescRegistry::Get();
 	{
 		TRACE_CPUPROFILER_EVENT_SCOPE(GatherDescriptorsClass);
 		
 		TSet<FTopLevelAssetPath> ClassPaths;
-		for (const FAssetData& Asset : Assets)
+		for (const FAssetData& Asset : ExternalAssets)
 		{
 			ClassPaths.Add(Asset.AssetClassPath);
 		}
@@ -108,7 +143,7 @@ void UActorDescContainer::Initialize(const FInitializeParams& InitParams)
 		TRACE_CPUPROFILER_EVENT_SCOPE(CreateDescriptors);
 
 		TMap<FName, FWorldPartitionActorDesc*> ActorDescsByPackage;
-		for (const FAssetData& Asset : Assets)
+		for (const FAssetData& Asset : ExternalAssets)
 		{
 			TUniquePtr<FWorldPartitionActorDesc> ActorDesc = FWorldPartitionActorDescUtils::GetActorDescriptorFromAssetData(Asset);
 								
@@ -172,6 +207,16 @@ void UActorDescContainer::Initialize(const FInitializeParams& InitParams)
 				ValidActorDescs.Add(ActorDesc->GetGuid(), MoveTemp(ActorDesc));
 			}
 		}
+
+		for (const FString& InternalAsset : InternalAssets)
+		{
+			FWorldPartitionActorDescUtils::FActorDescInitParams ActorDescInitParams(InternalAsset);
+
+			TUniquePtr<FWorldPartitionActorDesc> ActorDesc = FWorldPartitionActorDescUtils::GetActorDescriptorFromInitParams(ActorDescInitParams, ContainerPackageName);
+
+			ActorDescsByPackage.Add(ActorDesc->GetActorPackage(), ActorDesc.Get());
+			ValidActorDescs.Add(ActorDesc->GetGuid(), MoveTemp(ActorDesc));
+		}
 	}
 
 	{
@@ -338,8 +383,10 @@ void UActorDescContainer::OnDeletedObjectPlaceholderCreated(const UDeletedObject
 	{
 		if (ShouldHandleDeletedObjectPlaceholderEvent(InDeletedObjectPlaceholder))
 		{
-			check(GetActorDescriptor(Actor->GetActorGuid()));
-			DeletedObjectPlaceholdersAnnotation.AddAnnotation(Actor, FDeletedObjectPlaceholderAnnotation(InDeletedObjectPlaceholder, GetContainerName()));
+			if (GetActorDescriptor(Actor->GetActorGuid()))
+			{
+				DeletedObjectPlaceholdersAnnotation.AddAnnotation(Actor, FDeletedObjectPlaceholderAnnotation(InDeletedObjectPlaceholder, GetContainerName()));
+			}
 		}
 	}
 }
diff --git a/Engine/Source/Runtime/Engine/Private/WorldPartition/ErrorHandling/WorldPartitionStreamingGenerationLogErrorHandler.cpp b/Engine/Source/Runtime/Engine/Private/WorldPartition/ErrorHandling/WorldPartitionStreamingGenerationLogErrorHandler.cpp
--- a/Engine/Source/Runtime/Engine/Private/WorldPartition/ErrorHandling/WorldPartitionStreamingGenerationLogErrorHandler.cpp
+++ b/Engine/Source/Runtime/Engine/Private/WorldPartition/ErrorHandling/WorldPartitionStreamingGenerationLogErrorHandler.cpp
@@ -116,10 +116,10 @@ void FStreamingGenerationLogErrorHandler::OnLevelInstanceInvalidWorldAsset(const
 	case ELevelInstanceInvalidReason::WorldAssetNotFound:
 		UE_ASSET_LOG_ACTORDESCVIEW(LogWorldPartition, Log, ActorDescView, TEXT("Actor %s has an invalid world asset %s"), *ActorName, *WorldAsset.ToString());
 		break;
-	case ELevelInstanceInvalidReason::WorldAssetNotUsingExternalActors:
-		UE_ASSET_LOG_ACTORDESCVIEW(LogWorldPartition, Log, ActorDescView, TEXT("Actor %s world asset %s is not using external actors"), *ActorName, *WorldAsset.ToString());
+	case ELevelInstanceInvalidReason::WorldAssetDontContainActorsMetadata:
+		UE_ASSET_LOG_ACTORDESCVIEW(LogWorldPartition, Log, ActorDescView, TEXT("Actor %s world asset %s is not using external actors, resave level to add compatibility."), *ActorName, *WorldAsset.ToString());
 		break;
-	case ELevelInstanceInvalidReason::WorldAssetImcompatiblePartitioned:
+	case ELevelInstanceInvalidReason::WorldAssetIncompatiblePartitioned:
 		UE_ASSET_LOG_ACTORDESCVIEW(LogWorldPartition, Log, ActorDescView, TEXT("Actor %s world asset %s is partitioned but not marked as compatible"), *ActorName, *WorldAsset.ToString());
 		break;
 	case ELevelInstanceInvalidReason::WorldAssetHasInvalidContainer:
diff --git a/Engine/Source/Runtime/Engine/Private/WorldPartition/ErrorHandling/WorldPartitionStreamingGenerationTokenizedMessageErrorHandler.cpp b/Engine/Source/Runtime/Engine/Private/WorldPartition/ErrorHandling/WorldPartitionStreamingGenerationTokenizedMessageErrorHandler.cpp
--- a/Engine/Source/Runtime/Engine/Private/WorldPartition/ErrorHandling/WorldPartitionStreamingGenerationTokenizedMessageErrorHandler.cpp
+++ b/Engine/Source/Runtime/Engine/Private/WorldPartition/ErrorHandling/WorldPartitionStreamingGenerationTokenizedMessageErrorHandler.cpp
@@ -228,11 +228,11 @@ void ITokenizedMessageErrorHandler::OnLevelInstanceInvalidWorldAsset(const IWorl
 		MessageSeverity = EMessageSeverity::Error;
 		ReasonText = LOCTEXT("TokenMessage_WorldPartition_HasInvalidWorldAsset", "has an invalid world asset");
 		break;
-	case ELevelInstanceInvalidReason::WorldAssetNotUsingExternalActors:
+	case ELevelInstanceInvalidReason::WorldAssetDontContainActorsMetadata:
 		MessageSeverity = EMessageSeverity::Error;
-		ReasonText = LOCTEXT("TokenMessage_WorldPartition_WorldAssetIsNotUsingExternalActors", "is not using external actors");
+		ReasonText = LOCTEXT("TokenMessage_WorldPartition_WorldAssetIsNotUsingExternalActors", "is not using external actors, resave level to add compatibility");
 		break;
-	case ELevelInstanceInvalidReason::WorldAssetImcompatiblePartitioned:
+	case ELevelInstanceInvalidReason::WorldAssetIncompatiblePartitioned:
 		MessageSeverity = EMessageSeverity::Error;
 		ReasonText = LOCTEXT("TokenMessage_WorldPartition_WorldAssetIsPartitionedIncompatible", "is partitioned but not marked as compatible");
 		break;
diff --git a/Engine/Source/Runtime/Engine/Private/WorldPartition/LevelInstance/LevelInstanceActorDesc.cpp b/Engine/Source/Runtime/Engine/Private/WorldPartition/LevelInstance/LevelInstanceActorDesc.cpp
--- a/Engine/Source/Runtime/Engine/Private/WorldPartition/LevelInstance/LevelInstanceActorDesc.cpp
+++ b/Engine/Source/Runtime/Engine/Private/WorldPartition/LevelInstance/LevelInstanceActorDesc.cpp
@@ -203,8 +203,9 @@ bool FLevelInstanceActorDesc::IsChildContainerInstanceInternal() const
 	{
 		return false;
 	}
-	
-	if (!ULevel::GetIsLevelUsingExternalActorsFromPackage(GetChildContainerPackage()))
+
+	if (!ULevel::GetIsLevelUsingExternalActorsFromPackage(GetChildContainerPackage()) && 
+		!ULevel::GetIsLevelUsingActorsDescsFromPackage(GetChildContainerPackage()))
 	{
 		return false;
 	}
@@ -245,17 +246,24 @@ void FLevelInstanceActorDesc::CheckForErrors(const IWorldPartitionActorDescInsta
 	{
 		ErrorHandler->OnLevelInstanceInvalidWorldAsset(*InActorDescView, ChildContainerPackage, IStreamingGenerationErrorHandler::ELevelInstanceInvalidReason::WorldAssetNotFound);
 	}
-	else if (!ULevel::GetIsLevelUsingExternalActorsFromPackage(ChildContainerPackage))
+	else
 	{
-		if (DesiredRuntimeBehavior != ELevelInstanceRuntimeBehavior::LevelStreaming)
+		if (!ULevel::GetIsLevelUsingExternalActorsFromPackage(ChildContainerPackage))
+		{
+			if (DesiredRuntimeBehavior != ELevelInstanceRuntimeBehavior::LevelStreaming)
+			{
+				if (!ULevel::GetIsLevelUsingActorsDescsFromPackage(ChildContainerPackage))
+				{
+					ErrorHandler->OnLevelInstanceInvalidWorldAsset(*InActorDescView, ChildContainerPackage, IStreamingGenerationErrorHandler::ELevelInstanceInvalidReason::WorldAssetDontContainActorsMetadata);
+				}
+			}
+		}
+	
+		if (!ValidateCircularReference(InActorDescView->GetContainerInstance(), ChildContainerPackage))
 		{
-			ErrorHandler->OnLevelInstanceInvalidWorldAsset(*InActorDescView, ChildContainerPackage, IStreamingGenerationErrorHandler::ELevelInstanceInvalidReason::WorldAssetNotUsingExternalActors);
+			ErrorHandler->OnLevelInstanceInvalidWorldAsset(*InActorDescView, ChildContainerPackage, IStreamingGenerationErrorHandler::ELevelInstanceInvalidReason::CirculalReference);
 		}
 	}
-	else if (!ValidateCircularReference(InActorDescView->GetContainerInstance(), ChildContainerPackage))
-	{
-		ErrorHandler->OnLevelInstanceInvalidWorldAsset(*InActorDescView, ChildContainerPackage, IStreamingGenerationErrorHandler::ELevelInstanceInvalidReason::CirculalReference);
-	}
 }
 
 void FLevelInstanceActorDesc::TransferFrom(const FWorldPartitionActorDesc* From)
diff --git a/Engine/Source/Runtime/Engine/Private/WorldPartition/WorldPartitionActorDesc.cpp b/Engine/Source/Runtime/Engine/Private/WorldPartition/WorldPartitionActorDesc.cpp
--- a/Engine/Source/Runtime/Engine/Private/WorldPartition/WorldPartitionActorDesc.cpp
+++ b/Engine/Source/Runtime/Engine/Private/WorldPartition/WorldPartitionActorDesc.cpp
@@ -82,7 +82,6 @@ void FWorldPartitionActorDesc::Init(const AActor* InActor)
 	}
 	else
 	{
-		check(InActor->IsPackageExternal());
 		check(InActor->GetActorGuid().IsValid());
 		Guid = InActor->GetActorGuid();
 	}
@@ -176,8 +175,7 @@ void FWorldPartitionActorDesc::Init(const AActor* InActor)
 			ParentActor = AttachParentActor->GetActorGuid();
 		}
 
-		const ActorsReferencesUtils::FGetActorReferencesParams Params = ActorsReferencesUtils::FGetActorReferencesParams(const_cast<AActor*>(InActor))
-			.SetRequiredFlags(RF_HasExternalPackage);
+		const ActorsReferencesUtils::FGetActorReferencesParams Params = ActorsReferencesUtils::FGetActorReferencesParams(const_cast<AActor*>(InActor));
 		TArray<ActorsReferencesUtils::FActorReference> ActorReferences = ActorsReferencesUtils::GetActorReferences(Params);
 
 		if (ActorReferences.Num())
diff --git a/Engine/Source/Runtime/Engine/Private/WorldPartition/WorldPartitionActorDescUtils.cpp b/Engine/Source/Runtime/Engine/Private/WorldPartition/WorldPartitionActorDescUtils.cpp
--- a/Engine/Source/Runtime/Engine/Private/WorldPartition/WorldPartitionActorDescUtils.cpp
+++ b/Engine/Source/Runtime/Engine/Private/WorldPartition/WorldPartitionActorDescUtils.cpp
@@ -21,30 +21,78 @@
 static FName NAME_ActorMetaDataClass(TEXT("ActorMetaDataClass"));
 static FName NAME_ActorMetaData(TEXT("ActorMetaData"));
 
-FName FWorldPartitionActorDescUtils::ActorMetaDataClassTagName()
+namespace FWorldPartitionActorDescUtilsPrivate
 {
-	return NAME_ActorMetaDataClass;
+	static FString ResolveClassRedirector(const FString& InClassName)
+	{
+		FString ClassName;
+		FString ClassPackageName;
+		if (!InClassName.Split(TEXT("."), &ClassPackageName, &ClassName))
+		{
+			ClassName = *InClassName;
+		}
+
+		// Look for a class redirectors
+		const FCoreRedirectObjectName OldClassName = FCoreRedirectObjectName(*ClassName, NAME_None, *ClassPackageName);
+		const FCoreRedirectObjectName NewClassName = FCoreRedirects::GetRedirectedName(ECoreRedirectFlags::Type_Class, OldClassName);
+
+		return NewClassName.ToString();
+	}
+
+	static UClass* GetActorNativeClassFromClassName(FName InClassName)
+	{
+		// Avoid an assert when calling StaticFindObject during save to retrieve the actor's class.
+		// Since we are only looking for a native class, the call to StaticFindObject is legit.
+		TGuardValue<bool> GIsSavingPackageGuard(GIsSavingPackage, false);
+
+		// Look for a class redirectors
+		const FString ActorNativeClassName = FWorldPartitionActorDescUtilsPrivate::ResolveClassRedirector(InClassName.ToString());
+		
+		// Handle deprecated short class names
+		const FTopLevelAssetPath ClassPath = FAssetData::TryConvertShortClassNameToPathName(*ActorNativeClassName, ELogVerbosity::Log);
+
+		// Lookup the native class
+		return UClass::TryFindTypeSlow<UClass>(ClassPath.ToString(), EFindFirstObjectOptions::ExactClass);
+	}
 }
 
-FName FWorldPartitionActorDescUtils::ActorMetaDataTagName()
+FWorldPartitionActorDescUtils::FActorDescInitParams::FActorDescInitParams(const AActor* InActor)
 {
-	return NAME_ActorMetaData;
+	// Avoid an assert when calling StaticFindObject during save to retrieve the actor's class.
+	// Since we are only looking for a native class, the call to StaticFindObject is legit.
+	TGuardValue<bool> GIsSavingPackageGuard(GIsSavingPackage, false);
+
+	TUniquePtr<FWorldPartitionActorDesc> ActorDesc(InActor->CreateActorDesc());
+
+	PathName = *InActor->GetPathName();
+	NativeClassName = *(GetParentNativeClass(InActor->GetClass()))->GetPathName();
+	AssetData = FWorldPartitionActorDescUtils::GetAssetDataFromActorDescriptor(ActorDesc);
 }
 
-static FString ResolveClassRedirector(const FString& InClassName)
+FWorldPartitionActorDescUtils::FActorDescInitParams::FActorDescInitParams(const FString& InString)
 {
-	FString ClassName;
-	FString ClassPackageName;
-	if (!InClassName.Split(TEXT("."), &ClassPackageName, &ClassName))
+	TArray<FString> ParsedAssetData;
+	if (InString.ParseIntoArray(ParsedAssetData, TEXT(",")))
 	{
-		ClassName = *InClassName;
+		PathName = *ParsedAssetData[0];
+		NativeClassName = *ParsedAssetData[1];
+		AssetData = *ParsedAssetData[2];
 	}
+}
 
-	// Look for a class redirectors
-	const FCoreRedirectObjectName OldClassName = FCoreRedirectObjectName(*ClassName, NAME_None, *ClassPackageName);
-	const FCoreRedirectObjectName NewClassName = FCoreRedirects::GetRedirectedName(ECoreRedirectFlags::Type_Class, OldClassName);
+FString FWorldPartitionActorDescUtils::FActorDescInitParams::ToString() const
+{
+	return FString::Printf(TEXT("%s,%s,%s"), *PathName.ToString(), *NativeClassName.ToString(), *AssetData);
+}
 
-	return NewClassName.ToString();
+FName FWorldPartitionActorDescUtils::ActorMetaDataClassTagName()
+{
+	return NAME_ActorMetaDataClass;
+}
+
+FName FWorldPartitionActorDescUtils::ActorMetaDataTagName()
+{
+	return NAME_ActorMetaData;
 }
 
 bool FWorldPartitionActorDescUtils::IsValidActorDescriptorFromAssetData(const FAssetData& InAssetData)
@@ -65,11 +113,11 @@ UClass* FWorldPartitionActorDescUtils::GetActorNativeClassFromAssetData(const FA
 UClass* FWorldPartitionActorDescUtils::GetActorNativeClassFromString(const FString& InClassPath)
 {
 	// Avoid an assert when calling StaticFindObject during save to retrieve the actor's class.
-		// Since we are only looking for a native class, the call to StaticFindObject is legit.
+	// Since we are only looking for a native class, the call to StaticFindObject is legit.
 	TGuardValue<bool> GIsSavingPackageGuard(GIsSavingPackage, false);
 
 	// Look for a class redirectors
-	const FString ActorNativeClassName = ResolveClassRedirector(InClassPath);
+	const FString ActorNativeClassName = FWorldPartitionActorDescUtilsPrivate::ResolveClassRedirector(InClassPath);
 
 	// Handle deprecated short class names
 	const FTopLevelAssetPath ClassPath = FAssetData::TryConvertShortClassNameToPathName(*ActorNativeClassName, ELogVerbosity::Log);
@@ -78,35 +126,46 @@ UClass* FWorldPartitionActorDescUtils::GetActorNativeClassFromString(const FStri
 	return UClass::TryFindTypeSlow<UClass>(ClassPath.ToString(), EFindFirstObjectOptions::ExactClass);
 }
 
-TUniquePtr<FWorldPartitionActorDesc> FWorldPartitionActorDescUtils::GetActorDescriptorFromAssetData(const FAssetData& InAssetData)
+TUniquePtr<FWorldPartitionActorDesc> FWorldPartitionActorDescUtils::GetActorDescriptorFromInitParams(const FActorDescInitParams& InActorDescInitParams, FName InPackageName)
 {
-	if (IsValidActorDescriptorFromAssetData(InAssetData))
-	{
-		// Fixing here for now until following jira tasks are closed to unblock users that are renaming plugins:
-		// @todo_ow: remove once https://jira.it.epicgames.com/browse/UE-168245 & https://jira.it.epicgames.com/browse/UE-144431 are closed
-		FSoftObjectPath ActorPath = InAssetData.GetSoftObjectPath();
-		ActorPath.FixupCoreRedirects();
-
-		FWorldPartitionActorDescInitData ActorDescInitData = FWorldPartitionActorDescInitData()
-			.SetNativeClass(GetActorNativeClassFromAssetData(InAssetData))
-			.SetPackageName(InAssetData.PackageName)
-			.SetActorPath(ActorPath);
+	FWorldPartitionActorDescInitData ActorDescInitData = FWorldPartitionActorDescInitData()
+		.SetNativeClass(FWorldPartitionActorDescUtilsPrivate::GetActorNativeClassFromClassName(InActorDescInitParams.NativeClassName))
+		.SetPackageName(InPackageName)
+		.SetActorPath(InActorDescInitParams.PathName.ToString());
 
-		FString ActorMetaDataStr;
-		verify(InAssetData.GetTagValue(NAME_ActorMetaData, ActorMetaDataStr));
-		verify(FBase64::Decode(ActorMetaDataStr, ActorDescInitData.GetSerializedData()));
+	verify(FBase64::Decode(InActorDescInitParams.AssetData, ActorDescInitData.GetSerializedData()));
 
-		TUniquePtr<FWorldPartitionActorDesc> NewActorDesc(AActor::StaticCreateClassActorDesc(ActorDescInitData.NativeClass ? ActorDescInitData.NativeClass : AActor::StaticClass()));
+	TUniquePtr<FWorldPartitionActorDesc> NewActorDesc(AActor::StaticCreateClassActorDesc(ActorDescInitData.NativeClass ? ActorDescInitData.NativeClass : AActor::StaticClass()));
 
-		NewActorDesc->Init(ActorDescInitData);
+	NewActorDesc->Init(ActorDescInitData);
 			
-		if (!ActorDescInitData.NativeClass)
+	if (!ActorDescInitData.NativeClass)
+	{
+		UE_LOG(LogWorldPartition, Warning, TEXT("Invalid class for actor guid `%s` ('%s') from package '%s'"), *NewActorDesc->GetGuid().ToString(), *NewActorDesc->GetActorName().ToString(), *NewActorDesc->GetActorPackage().ToString());
+		NewActorDesc->NativeClass.Reset();
+	}
+
+	return NewActorDesc;
+}
+
+TUniquePtr<FWorldPartitionActorDesc> FWorldPartitionActorDescUtils::GetActorDescriptorFromAssetData(const FAssetData& InAssetData)
+{
+	if (IsValidActorDescriptorFromAssetData(InAssetData))
+	{
+		FString ActorMetaDataStr;
+		if (InAssetData.GetTagValue(NAME_ActorMetaData, ActorMetaDataStr))
 		{
-			UE_LOG(LogWorldPartition, Warning, TEXT("Invalid class for actor guid `%s` ('%s') from package '%s'"), *NewActorDesc->GetGuid().ToString(), *NewActorDesc->GetActorNameString(), *NewActorDesc->GetActorPackage().ToString());
-			NewActorDesc->NativeClass.Reset();
+			// Fixing here for now until following jira tasks are closed to unblock users that are renaming plugins:
+			// @todo_ow: remove once https://jira.it.epicgames.com/browse/UE-168245 & https://jira.it.epicgames.com/browse/UE-144431 are closed
+			FSoftObjectPath ActorPath = InAssetData.GetSoftObjectPath();
+			ActorPath.FixupCoreRedirects();
+
+			FActorDescInitParams ActorDescInitParams;
+			ActorDescInitParams.PathName = *ActorPath.ToString();
+			ActorDescInitParams.NativeClassName = *GetActorNativeClassFromAssetData(InAssetData)->GetPathName();
+			ActorDescInitParams.AssetData = ActorMetaDataStr;
+			return GetActorDescriptorFromInitParams(ActorDescInitParams, InAssetData.PackageName);
 		}
-
-		return NewActorDesc;
 	}
 
 	return nullptr;
@@ -224,4 +283,4 @@ bool FWorldPartitionActorDescUtils::FixupRedirectedAssetPath(FName& InOutAssetPa
 	UAssetRegistryHelpers::FixupRedirectedAssetPath(InOutAssetPath);
 	return true;
 }
-#endif
+#endif
\ No newline at end of file
diff --git a/Engine/Source/Runtime/Engine/Private/WorldPartition/WorldPartitionLevelHelper.cpp b/Engine/Source/Runtime/Engine/Private/WorldPartition/WorldPartitionLevelHelper.cpp
--- a/Engine/Source/Runtime/Engine/Private/WorldPartition/WorldPartitionLevelHelper.cpp
+++ b/Engine/Source/Runtime/Engine/Private/WorldPartition/WorldPartitionLevelHelper.cpp
@@ -237,8 +237,8 @@ void FWorldPartitionLevelHelper::MoveExternalActorsToLevel(const TArray<FWorldPa
 		AActor* Actor = FindObject<AActor>(nullptr, *PackageObjectMapping.LoadedPath.ToString());
 		if (Actor)
 		{
-			UPackage* ActorExternalPackage = Actor->GetExternalPackage();
-			check(ActorExternalPackage);
+			UPackage* ActorPackage = Actor->GetPackage();
+			check(ActorPackage);
 
 			const bool bSameOuter = (InLevel == Actor->GetOuter());
 			Actor->SetPackageExternal(false, false);
@@ -247,7 +247,7 @@ void FWorldPartitionLevelHelper::MoveExternalActorsToLevel(const TArray<FWorldPa
 			// (UObject::Rename doesn't check if Rename is being called with existing outer and assigns new name)
 			if (!bSameOuter)
 			{
-				Actor->Rename(nullptr, InLevel);
+				Actor->Rename(nullptr, InLevel, REN_NonTransactional | REN_DoNotDirty | REN_DontCreateRedirectors);
 
 				// AActor::Rename will register components but doesn't call RerunConstructionScripts like AddLoadedActors does.
 				// If bIsWorldInitialized is false. RerunConstructionScripts will get called as part of UEditorEngine::InitializePhysicsSceneForSaveIfNecessary during Cell package save
@@ -265,46 +265,51 @@ void FWorldPartitionLevelHelper::MoveExternalActorsToLevel(const TArray<FWorldPa
 			check(Actor->GetPackage() == LevelPackage);
 
 			// Process objects found in the source actor package
-			TArray<UObject*> Objects;
-			const bool bIncludeNestedSubobjects = false;
-			// Skip Garbage objects as the initial Rename on an actor with an ChildActorComponent can destroy its child actors.
-			// This happens when the component has bNeedsRecreate set to true (when it has a valid ChildActorTemplate).
-			GetObjectsWithPackage(ActorExternalPackage, Objects, bIncludeNestedSubobjects, RF_NoFlags, EInternalObjectFlags::Garbage);
-			for (UObject* Object : Objects)
+			if (Actor->IsPackageExternal())
 			{
-				if (Object->GetFName() != NAME_PackageMetaData)
+				TArray<UObject*> Objects;
+
+				// Skip Garbage objects as the initial Rename on an actor with an ChildActorComponent can destroy its child actors.
+				// This happens when the component has bNeedsRecreate set to true (when it has a valid ChildActorTemplate).
+				const bool bIncludeNestedSubobjects = false;
+				GetObjectsWithPackage(ActorPackage, Objects, bIncludeNestedSubobjects, RF_NoFlags, EInternalObjectFlags::Garbage);
+
+				for (UObject* Object : Objects)
 				{
-					if (Object->GetOuter()->IsA<ULevel>())
+					if (Object->GetFName() != NAME_PackageMetaData)
 					{
-						// Move objects that are outered the level in the destination level
-						AActor* NestedActor = Cast<AActor>(Object);
-						if (InLevel != Object->GetOuter())
-						{
-							Object->Rename(nullptr, InLevel);
-						}
-						else if (NestedActor && !InLevel->Actors.Contains(NestedActor))
+						if (Object->GetOuter()->IsA<ULevel>())
 						{
-							LoadedActors.Emplace(NestedActor);
+							// Move objects that are outered the level in the destination level
+							AActor* NestedActor = Cast<AActor>(Object);
+							if (InLevel != Object->GetOuter())
+							{
+								Object->Rename(nullptr, InLevel, REN_NonTransactional | REN_DoNotDirty | REN_DontCreateRedirectors);
+							}
+							else if (NestedActor && !InLevel->Actors.Contains(NestedActor))
+							{
+								LoadedActors.Emplace(NestedActor);
+							}
+							if (NestedActor)
+							{
+								LevelActors.Add(NestedActor->GetFName());
+							}
 						}
-						if (NestedActor)
+						else
 						{
-							LevelActors.Add(NestedActor->GetFName());
+							// Move objects in the destination level package
+							Object->Rename(nullptr, LevelPackage, REN_NonTransactional | REN_DoNotDirty | REN_DontCreateRedirectors);
 						}
 					}
-					else
-					{
-						// Move objects in the destination level package
-						Object->Rename(nullptr, LevelPackage);
-					}
 				}
-			}
 
-			// Trash this package to guarantee that any potential future load of this actor won't find the old empty package
-			// @todo_ow: Decide if we want to support actor reloads during cook. If not, remove this code, detect the reload and report an error.
-			FName NewPackageName = MakeUniqueObjectName(nullptr, UPackage::StaticClass(), FName(*FString::Printf(TEXT("%s_Trashed"), *ActorExternalPackage->GetName())));
-			ActorExternalPackage->Rename(*NewPackageName.ToString(), nullptr, REN_DontCreateRedirectors | REN_NonTransactional | REN_DoNotDirty);
+				// Trash this package to guarantee that any potential future load of this actor won't find the old empty package
+				// @todo_ow: Decide if we want to support actor reloads during cook. If not, remove this code, detect the reload and report an error.
+				FName NewPackageName = MakeUniqueObjectName(nullptr, UPackage::StaticClass(), FName(*FString::Printf(TEXT("%s_Trashed"), *ActorPackage->GetName())));
+				ActorPackage->Rename(*NewPackageName.ToString(), nullptr, REN_DontCreateRedirectors | REN_NonTransactional | REN_DoNotDirty);
+			}
 
-			OutModifiedPackages.Add(ActorExternalPackage);
+			OutModifiedPackages.Add(ActorPackage);
 			LevelActors.Add(Actor->GetFName());
 		}
 		else
@@ -341,24 +346,6 @@ void FWorldPartitionLevelHelper::RemapLevelSoftObjectPaths(ULevel* InLevel, UWor
 	FixupSerializer.Fixup(InLevel);
 }
 
-FString FWorldPartitionLevelHelper::GetContainerPackage(const FActorContainerID& InContainerID, const FString& InPackageName, const FString& InDestLevelPackageName)
-{
-	// Generate a unique name to load a Level Instance embedded actor if there are multiple instances of this Level Instance and possibly across multiple instances of the World Partition world
-	// InContainerID will distinguish between instances of the same Level Instance
-	// InDestLevelPackageName will distinguish between instances of the same top level World Partition world (Only needed in PIE, In Cook we always cook the source WP and not an instance and Actor packages no longer exist at runtime)
-	uint64 DestLevelID = 0;
-	TStringBuilder<512> PackageNameBuilder;
-	PackageNameBuilder.Appendf(TEXT("/Temp%s_%s"), *InPackageName, *InContainerID.ToShortString());
-		
-	if (!InDestLevelPackageName.IsEmpty())
-	{
-		DestLevelID = CityHash64((const char*)*InDestLevelPackageName, InDestLevelPackageName.Len() * sizeof(TCHAR));
-		PackageNameBuilder.Appendf(TEXT("_%016llx"), DestLevelID);
-	}
-
-	return PackageNameBuilder.ToString();
-}
-
 FSoftObjectPath FWorldPartitionLevelHelper::RemapActorPath(const FActorContainerID& InContainerID, const FString& InSourceWorldPath, const FSoftObjectPath& InActorPath)
 {
 	// If Path is in an instanced package it will now be remapped to its source package
@@ -606,6 +593,23 @@ bool FWorldPartitionLevelHelper::LoadActorsInternal(FLoadActorsParams&& InParams
 	TArray<FWorldPartitionRuntimeCellObjectMapping*> ActorPackagesToLoad;
 	TMap<FActorContainerID, FLinkerInstancingContext> LinkerInstancingContexts;
 
+	auto GetContainerPackage = [](const FActorContainerID& InContainerID, const FString& InPackageName, const UObject* InContextObject)
+	{
+		// Generate a unique name to load a Level Instance embedded actor if there are multiple instances of this Level Instance and possibly across multiple instances of the World Partition world
+		// InContainerID will distinguish between instances of the same Level Instance
+		// InContextObject will distinguish between instances of the same top level World Partition world (Only needed in PIE, In cook we always cook the source WP and not an instance and Actor packages no longer exist at runtime)
+		uint64 DestLevelID = 0;
+		TStringBuilder<512> PackageNameBuilder;
+		PackageNameBuilder.Appendf(TEXT("/Temp%s_%s"), *InPackageName, *InContainerID.ToShortString());
+		
+		if (InContextObject)
+		{
+			PackageNameBuilder.Appendf(TEXT("_%016llx"), InContextObject);
+		}
+
+		return PackageNameBuilder.ToString();
+	};
+
 	if (!InParams.ActorPackages.IsEmpty())
 	{
 		ActorPackagesToLoad.Reserve(InParams.ActorPackages.Num());
@@ -620,17 +624,17 @@ bool FWorldPartitionLevelHelper::LoadActorsInternal(FLoadActorsParams&& InParams
 			{
 				check(!PackageObjectMapping.ContainerID.IsMainContainer());
 		
-				const FString DestLevelPackageName = InParams.DestLevel ? InParams.DestLevel->GetPackage()->GetName() : FString();
-				const FName ContainerPackageInstanceName(GetContainerPackage(PackageObjectMapping.ContainerID, PackageObjectMapping.ContainerPackage.ToString(), DestLevelPackageName));
-
 				FLinkerInstancingContext& NewContext = LinkerInstancingContexts.Add(PackageObjectMapping.ContainerID);
 
 				// Make sure here we don't remap the SoftObjectPaths through the linker when loading the embedded actor packages. 
 				// A remapping will happen in the packaged loaded callback later in this method.
 				NewContext.SetSoftObjectPathRemappingEnabled(false); 
-			
+		
+				// Don't load external objects as we are going to individually load them
 				NewContext.AddTag(ULevel::DontLoadExternalObjectsTag);
 				NewContext.AddTag(ULevel::DontLoadExternalFoldersTag);
+
+				const FName ContainerPackageInstanceName(GetContainerPackage(PackageObjectMapping.ContainerID, PackageObjectMapping.ContainerPackage.ToString(), InParams.OuterWorld));
 				NewContext.AddPackageMapping(PackageObjectMapping.ContainerPackage, ContainerPackageInstanceName);
 				Context = &NewContext;
 			}
@@ -670,25 +674,44 @@ bool FWorldPartitionLevelHelper::LoadActorsInternal(FLoadActorsParams&& InParams
 
 	for (FWorldPartitionRuntimeCellObjectMapping* PackageObjectMapping : ActorPackagesToLoad)
 	{
+		const FName PackageToLoad(*FPackageName::ObjectPathToPackageName(PackageObjectMapping->Package.ToString()));
+		const FLinkerInstancingContext ContainerInstancingContext = FLinkerInstancingContext::DuplicateContext(LinkerInstancingContexts.FindChecked(PackageObjectMapping->ContainerID));
+		const FName PackageName = ContainerInstancingContext.RemapPackage(PackageToLoad);
+
 		FLoadPackageAsyncDelegate CompletionCallback = FLoadPackageAsyncDelegate::CreateLambda([LoadProgress, PackageObjectMapping, LoadedOverrides = InLoadedPropertyOverrides, PackageReferencer = InParams.PackageReferencer, OuterWorld = InParams.OuterWorld, DestLevel = InParams.DestLevel, CompletionCallback = InParams.CompletionCallback](const FName& LoadedPackageName, UPackage* LoadedPackage, EAsyncLoadingResult::Type Result)
 		{
 			const FName ActorName = *FPaths::GetExtension(PackageObjectMapping->Path.ToString());
 			check(LoadProgress->NumPendingLoadRequests);
 			LoadProgress->NumPendingLoadRequests--;
 
-			// In PIE, we make sure to clear RF_Standalone flag on objects in external packages (UMetaData) 
-			// This guarantees that external packages of actors that are destroyed during the PIE session will
-			// properly get GC'ed and will allow future edits/modifications of OFPA actors.
-			if (LoadedPackage && DestLevel && DestLevel->GetPackage()->HasAnyPackageFlags(PKG_PlayInEditor))
+			AActor* Actor = nullptr;
+
+			if (LoadedPackage)
 			{
-				ForEachObjectWithPackage(LoadedPackage, [](UObject* Object)
+				if (LoadedPackage->ContainsMap())
 				{
-					Object->ClearFlags(RF_Standalone);
-					return true;
-				}, false);
-			}
+					if (UWorld* LoadedWorld = UWorld::FindWorldInPackage(LoadedPackage))
+					{
+						Actor = FindObject<AActor>(LoadedWorld->PersistentLevel, *ActorName.ToString());
+					}
+				}
+				else
+				{
+					// In PIE, we make sure to clear RF_Standalone flag on objects in external packages (UMetaData) 
+					// This guarantees that external packages of actors that are destroyed during the PIE session will
+					// properly get GC'ed and will allow future edits/modifications of OFPA actors.
+					if (DestLevel && DestLevel->GetPackage()->HasAnyPackageFlags(PKG_PlayInEditor))
+					{
+						ForEachObjectWithPackage(LoadedPackage, [](UObject* Object)
+						{
+							Object->ClearFlags(RF_Standalone);
+							return true;
+						}, false);
+					}
 
-			AActor* Actor = LoadedPackage ? FindObject<AActor>(LoadedPackage, *ActorName.ToString()) : nullptr;
+					Actor = FindObject<AActor>(LoadedPackage, *ActorName.ToString());
+				}
+			}
 
 			if (Actor)
 			{
@@ -812,7 +835,6 @@ bool FWorldPartitionLevelHelper::LoadActorsInternal(FLoadActorsParams&& InParams
 						Actor->SetFolderPath(*SrcActorFolderPath);
 					}
 
-					check(Actor->IsPackageExternal());
 					DestLevel->Actors.Add(Actor);
 					checkf(Actor->GetLevel() == DestLevel, TEXT("Levels mismatch, got : %s, expected: %s\nActor: %s\nActorFullName: %s\nActorPackage: %s"), *DestLevel->GetFullName(), *Actor->GetLevel()->GetFullName(), *Actor->GetActorNameOrLabel(), *Actor->GetFullName(), *Actor->GetPackage()->GetFullName());
 
@@ -831,21 +853,18 @@ bool FWorldPartitionLevelHelper::LoadActorsInternal(FLoadActorsParams&& InParams
 			}
 			else
 			{
+				UE_LOG(LogWorldPartition, Warning, TEXT("Failed to find actor '%s' in package '%s'."), *ActorName.ToString(), *LoadedPackageName.ToString());
+
 				if (LoadedPackage)
 				{
-					UE_LOG(LogWorldPartition, Warning, TEXT("Failed to find actor in package %s. Package Content:"), *LoadedPackageName.ToString());
-					ForEachObjectWithPackage(LoadedPackage, [](UObject* Object)
+					UE_LOG(LogWorldPartition, Warning, TEXT("\tPackage Content:"), *ActorName.ToString(), *LoadedPackageName.ToString());
+					ForEachObjectWithOuter(LoadedPackage, [](UObject* Object)
 					{
-						UE_LOG(LogWorldPartition, Warning, TEXT("\t Object %s, Flags 0x%llx"), *Object->GetName(), static_cast<uint64>(Object->GetFlags()));
+						UE_LOG(LogWorldPartition, Warning, TEXT("\t\tObject %s, Flags 0x%llx"), *Object->GetPathName(), static_cast<uint64>(Object->GetFlags()));
 						return true;
-					}, false);
-				}
-				else
-				{
-					UE_LOG(LogWorldPartition, Warning, TEXT("Failed to load actor package %s"), *LoadedPackageName.ToString());
+					}, true);
 				}
 
-				//@todo_ow: cumulate and process when NumPendingActorRequests == 0
 				LoadProgress->NumFailedLoadedRequests++;
 			}
 
@@ -855,11 +874,13 @@ bool FWorldPartitionLevelHelper::LoadActorsInternal(FLoadActorsParams&& InParams
 			}
 		});
 
-		FName PackageToLoad(*FPackageName::ObjectPathToPackageName(PackageObjectMapping->Package.ToString()));
-		const FLinkerInstancingContext ContainerInstancingContext = FLinkerInstancingContext::DuplicateContext(LinkerInstancingContexts.FindChecked(PackageObjectMapping->ContainerID));
-		FName PackageName = ContainerInstancingContext.RemapPackage(PackageToLoad);
-
-		if (InParams.bLoadAsync)
+		// If the package already exists, we are loading actors from a non-OFPA level package, just fire the completion callback in this case as all actors are
+		// already loaded in.
+		if (UPackage* ExistingPackage = FindPackage(nullptr, *PackageName.ToString()))
+		{
+			CompletionCallback.Execute(PackageToLoad, ExistingPackage, EAsyncLoadingResult::Succeeded);
+		}
+		else if (InParams.bLoadAsync)
 		{
 			check(InParams.DestLevel);
 			const UPackage* DestPackage = InParams.DestLevel->GetPackage();
diff --git a/Engine/Source/Runtime/Engine/Private/WorldPartition/WorldPartitionStreamingGeneration.cpp b/Engine/Source/Runtime/Engine/Private/WorldPartition/WorldPartitionStreamingGeneration.cpp
--- a/Engine/Source/Runtime/Engine/Private/WorldPartition/WorldPartitionStreamingGeneration.cpp
+++ b/Engine/Source/Runtime/Engine/Private/WorldPartition/WorldPartitionStreamingGeneration.cpp
@@ -1105,8 +1105,7 @@ class FWorldPartitionStreamingGenerator
 			if (WorldPartitionContext)
 			{
 				// Gather all references to external actors from the world and make them non-spatially loaded
-				const ActorsReferencesUtils::FGetActorReferencesParams Params = ActorsReferencesUtils::FGetActorReferencesParams(WorldPartitionContext->GetTypedOuter<UWorld>())
-					.SetRequiredFlags(RF_HasExternalPackage);
+				const ActorsReferencesUtils::FGetActorReferencesParams Params = ActorsReferencesUtils::FGetActorReferencesParams(WorldPartitionContext->GetTypedOuter<UWorld>());
 				TArray<ActorsReferencesUtils::FActorReference> WorldExternalActorReferences = ActorsReferencesUtils::GetActorReferences(Params);
 				Algo::Transform(WorldExternalActorReferences, WorldReferences, [](const ActorsReferencesUtils::FActorReference& ActorReference) { return ActorReference.Actor->GetActorGuid(); });
 
diff --git a/Engine/Source/Runtime/Engine/Public/WorldPartition/ErrorHandling/WorldPartitionStreamingGenerationErrorHandler.h b/Engine/Source/Runtime/Engine/Public/WorldPartition/ErrorHandling/WorldPartitionStreamingGenerationErrorHandler.h
--- a/Engine/Source/Runtime/Engine/Public/WorldPartition/ErrorHandling/WorldPartitionStreamingGenerationErrorHandler.h
+++ b/Engine/Source/Runtime/Engine/Public/WorldPartition/ErrorHandling/WorldPartitionStreamingGenerationErrorHandler.h
@@ -111,8 +111,8 @@ public:
 	enum class ELevelInstanceInvalidReason
 	{
 		WorldAssetNotFound,
-		WorldAssetNotUsingExternalActors,
-		WorldAssetImcompatiblePartitioned,
+		WorldAssetDontContainActorsMetadata,
+		WorldAssetIncompatiblePartitioned,
 		WorldAssetHasInvalidContainer,
 		CirculalReference
 	};
diff --git a/Engine/Source/Runtime/Engine/Public/WorldPartition/WorldPartitionActorDescUtils.h b/Engine/Source/Runtime/Engine/Public/WorldPartition/WorldPartitionActorDescUtils.h
--- a/Engine/Source/Runtime/Engine/Public/WorldPartition/WorldPartitionActorDescUtils.h
+++ b/Engine/Source/Runtime/Engine/Public/WorldPartition/WorldPartitionActorDescUtils.h
@@ -14,9 +14,23 @@ struct FWorldPartitionActorDescUtils
 {
 	/** @return The asset registry tag name for ActorMetaDataClass */
 	static ENGINE_API FName ActorMetaDataClassTagName();
+
 	/** @return The asset registry tag name for ActorMetaData */
 	static ENGINE_API FName ActorMetaDataTagName();
 
+	struct FActorDescInitParams
+	{
+		FActorDescInitParams() {}
+		FActorDescInitParams(const AActor* InActor);
+		FActorDescInitParams(const FString& InString);
+
+		FString ToString() const;
+
+		FName PathName;
+		FName NativeClassName;
+		FString AssetData;
+	};
+
 	/** 
 	 * Checks if the provided asset data contains a valid actor descriptor.
 	 * @param InAssetData	The asset data to look into.
@@ -38,6 +52,15 @@ struct FWorldPartitionActorDescUtils
 	 */
 	static ENGINE_API UClass* GetActorNativeClassFromString(const FString& InClassPath);
 
+	/** 
+	 * Creates a valid actor descriptor from the provided actor descriptor init params.
+	 * @param InActorDescInitParams	The actor dfescriptor input params.
+	 * @param InPackageName			The package name the actor descriptor comes from.
+	 * @return						Actor descriptor retrieved from the provided asset data.
+	 */
+	static ENGINE_API TUniquePtr<FWorldPartitionActorDesc> GetActorDescriptorFromInitParams(const FActorDescInitParams& InActorDescInitParams, FName InPackageName);
+
+
 	/** 
 	 * Creates a valid actor descriptor from the provided asset data.
 	 * @param InAssetData	The asset data to look into.
diff --git a/Engine/Source/Runtime/Engine/Public/WorldPartition/WorldPartitionLevelHelper.h b/Engine/Source/Runtime/Engine/Public/WorldPartition/WorldPartitionLevelHelper.h
--- a/Engine/Source/Runtime/Engine/Public/WorldPartition/WorldPartitionLevelHelper.h
+++ b/Engine/Source/Runtime/Engine/Public/WorldPartition/WorldPartitionLevelHelper.h
@@ -89,10 +89,8 @@ private:
 	void RemoveReferences(FPackageReferencer* InReferencer);
 	void PreGarbageCollect();
 
-	static FString GetContainerPackage(const FActorContainerID& InContainerID, const FString& InPackageName, const FString& InDestLevelPackageName = FString());
 	static UWorld::InitializationValues GetWorldInitializationValues();
 
-
 	friend class FContentBundleEditor;
 	static bool RemapLevelCellPathInContentBundle(ULevel* Level, const class FContentBundleEditor* ContentBundleEditor, const UWorldPartitionRuntimeCell* Cell);
 
