Fix address of stack memory associated with local variable 'PackageNameBuilder' returned.

diff --git a/Engine/Source/Runtime/Engine/Private/WorldPartition/WorldPartitionLevelHelper.cpp b/Engine/Source/Runtime/Engine/Private/WorldPartition/WorldPartitionLevelHelper.cpp
--- a/Engine/Source/Runtime/Engine/Private/WorldPartition/WorldPartitionLevelHelper.cpp
+++ b/Engine/Source/Runtime/Engine/Private/WorldPartition/WorldPartitionLevelHelper.cpp
@@ -593,7 +593,7 @@ bool FWorldPartitionLevelHelper::LoadActorsInternal(FLoadActorsParams&& InParams
 	TArray<FWorldPartitionRuntimeCellObjectMapping*> ActorPackagesToLoad;
 	TMap<FActorContainerID, FLinkerInstancingContext> LinkerInstancingContexts;
 
-	auto GetContainerPackage = [](const FActorContainerID& InContainerID, const FString& InPackageName, const UObject* InContextObject)
+	auto GetContainerPackage = [](const FActorContainerID& InContainerID, const FString& InPackageName, const UObject* InContextObject) -> FName
 	{
 		// Generate a unique name to load a Level Instance embedded actor if there are multiple instances of this Level Instance and possibly across multiple instances of the World Partition world
 		// InContainerID will distinguish between instances of the same Level Instance
