Path Tracer: Fix typo in display name

Forgot to rename the field in the UI when the parameter was renamed.

#rnx

diff --git a/Engine/Source/Runtime/Engine/Classes/Engine/Scene.h b/Engine/Source/Runtime/Engine/Classes/Engine/Scene.h
--- a/Engine/Source/Runtime/Engine/Classes/Engine/Scene.h
+++ b/Engine/Source/Runtime/Engine/Classes/Engine/Scene.h
@@ -2228,7 +2228,7 @@ struct FPostProcessSettings
 	int32 PathTracingSamplesPerPixel;
 
 	/** Sets the maximum intensity of indirect samples to reduce fireflies. Lowering this value reduces noise at the expense of accuracy. Increasing it is more accurate but may lead to more noise. */
-	UPROPERTY(interp, EditAnywhere, BlueprintReadWrite, Category = "Path Tracing", meta = (ClampMin = "1.0", ClampMax = "65504.0", editcondition = "bOverride_PathTracingMaxPathIntensity", DisplayName = "Max Path Exposure"))
+	UPROPERTY(interp, EditAnywhere, BlueprintReadWrite, Category = "Path Tracing", meta = (ClampMin = "1.0", ClampMax = "65504.0", editcondition = "bOverride_PathTracingMaxPathIntensity", DisplayName = "Max Path Intensity"))
 	float PathTracingMaxPathIntensity;
 
 	/** Should emissive materials contribute to scene lighting? */
