[World Partition] Prevent handling unsaved actors when streaming partitioned worlds in PIE.

diff --git a/Engine/Source/Runtime/Engine/Private/WorldPartition/WorldPartitionStreamingGeneration.cpp b/Engine/Source/Runtime/Engine/Private/WorldPartition/WorldPartitionStreamingGeneration.cpp
--- a/Engine/Source/Runtime/Engine/Private/WorldPartition/WorldPartitionStreamingGeneration.cpp
+++ b/Engine/Source/Runtime/Engine/Private/WorldPartition/WorldPartitionStreamingGeneration.cpp
@@ -1931,7 +1931,7 @@ bool UWorldPartition::GenerateContainerStreaming(const FGenerateStreamingParams&
 
 	FWorldPartitionStreamingGenerator::FWorldPartitionStreamingGeneratorParams StreamingGeneratorParams = FWorldPartitionStreamingGenerator::FWorldPartitionStreamingGeneratorParams()
 		.SetWorldPartitionContext(this)
-		.SetHandleUnsavedActors(bIsPIE)
+		.SetHandleUnsavedActors(bIsPIE && !GetTypedOuter<UWorld>()->IsGameWorld())
 		.SetIsValidGrid([this](FName GridName, const UClass* ActorClass) { return RuntimeHash->IsValidGrid(GridName, ActorClass); })
 		.SetIsValidHLODLayer([this](FName GridName, const FSoftObjectPath& HLODLayerPath) { return RuntimeHash->IsValidHLODLayer(GridName, HLODLayerPath); })
 		.SetErrorHandler(ErrorHandlerSelector.Get())
