Include all UTexture objects in iteration to find virtual textures.
We want to catch virtual textures that don't inherit from UTexture2D.
#rnx
#tests Using r.vt.poolsizescale and checking that the number of vt pools returns to the base line (visualizing with r.vt.residency.show 1)

diff --git a/Engine/Source/Runtime/Engine/Private/VT/VirtualTextureRecreate.cpp b/Engine/Source/Runtime/Engine/Private/VT/VirtualTextureRecreate.cpp
--- a/Engine/Source/Runtime/Engine/Private/VT/VirtualTextureRecreate.cpp
+++ b/Engine/Source/Runtime/Engine/Private/VT/VirtualTextureRecreate.cpp
@@ -3,7 +3,7 @@
 #include "VT/VirtualTextureRecreate.h"
 
 #include "Components/RuntimeVirtualTextureComponent.h"
-#include "Engine/Texture2D.h"
+#include "Engine/Texture.h"
 #include "TextureResource.h"
 #include "VT/RuntimeVirtualTexture.h"
 #include "UObject/UObjectIterator.h"
@@ -16,7 +16,7 @@ namespace VirtualTexture
 	void Recreate()
 	{
 		// Reinit streaming virtual textures.
-		for (TObjectIterator<UTexture2D> It(RF_ClassDefaultObject, true, EInternalObjectFlags::Garbage); It; ++It)
+		for (TObjectIterator<UTexture> It(RF_ClassDefaultObject, true, EInternalObjectFlags::Garbage); It; ++It)
 		{
 			if (It->IsCurrentlyVirtualTextured())
 			{
@@ -37,7 +37,7 @@ namespace VirtualTexture
 	void Recreate(TConstArrayView < TEnumAsByte<EPixelFormat> > InFormat)
 	{
 		// Reinit streaming virtual textures that match one of our passed in format arrays.
-		for (TObjectIterator<UTexture2D> It(RF_ClassDefaultObject, true, EInternalObjectFlags::Garbage); It; ++It)
+		for (TObjectIterator<UTexture> It(RF_ClassDefaultObject, true, EInternalObjectFlags::Garbage); It; ++It)
 		{
 			if (It->IsCurrentlyVirtualTextured())
 			{
