Add virtual texture adapter object.
This wraps a regular UTexture and allows it to be referenced as a virtual texture in materials.
The virtual texture pages are filled on demand by copying from the wrapped texture.
This isn't memory efficient or performant, but can be useful for certain debugging or previewing modes where we don't want to change the material, but do want to bind a non-virtual texture to an existing virtual texture sample.
The first intended use is to support mesh painting onto virtual textures where we want to paint to a 2d render target but view the painting through the material which is reading a virtual texture.

diff --git a/Engine/Source/Runtime/Engine/Classes/VT/VirtualTextureAdapter.h b/Engine/Source/Runtime/Engine/Classes/VT/VirtualTextureAdapter.h
new file mode 100644
--- /dev/null
+++ b/Engine/Source/Runtime/Engine/Classes/VT/VirtualTextureAdapter.h
@@ -0,0 +1,53 @@
+// Copyright Epic Games, Inc. All Rights Reserved.
+
+#pragma once
+
+#include "Engine/Texture.h"
+#include "VirtualTextureAdapter.generated.h"
+
+/** 
+ * Adapter object that wraps a regular UTexture and allows it to be referenced as a virtual texture in materials.
+ * The virtual texture pages are filled on demand by copying from the wrapped texture.
+ * This isn't memory efficient or performant, but can be useful for certain debugging or previewing modes where we don't
+ * want to change the material, but do want to bind a non-virtual texture to an existing virtual texture sample.
+ */
+UCLASS(ClassGroup = Rendering, BlueprintType, MinimalAPI)
+class UVirtualTextureAdapter : public UTexture
+{
+	GENERATED_UCLASS_BODY()
+
+	//~ Begin UTexture Interface.
+	virtual FTextureResource* CreateResource() override;
+	virtual ETextureClass GetTextureClass() const override;
+	virtual EMaterialValueType GetMaterialType() const override;
+	virtual bool IsCurrentlyVirtualTextured() const override { return true; }
+	virtual float GetSurfaceWidth() const override;
+	virtual float GetSurfaceHeight() const override;
+	virtual float GetSurfaceDepth() const override { return 0; }
+	virtual uint32 GetSurfaceArraySize() const override { return 0; }
+	//~ End UTexture Interface.
+
+public:
+	/** The UTexture object to wrap. */
+	UPROPERTY(EditAnywhere, Category = Texture)
+	TObjectPtr<UTexture> Texture;
+
+	/** Use a compressed format for the virtual texture even when the wrapped Texture does not have a compressed format. */
+	UPROPERTY(EditAnywhere, Category = Texture)
+	bool bUseCompressedFormat = false;
+
+	/** Set to true if we want to use the default project virtual texture tile settings. */
+	UPROPERTY(EditAnywhere, Category = Texture)
+	bool bUseDefaultTileSizes = true;
+
+	/** Page tile size. (Will be rounded up to power of 2). */
+	UPROPERTY(EditAnywhere, Category = Texture, meta = (UIMin = "0", UIMax = "512", editcondition="!bUseDefaultTileSizes"))
+	int32 TileSize = 0;
+
+	/** Page tile border size. (Will be rounded up to multiple of 2). */
+	UPROPERTY(EditAnywhere, Category = Texture, meta = (UIMin = "0", UIMax = "8", editcondition = "!bUseDefaultTileSizes"))
+	int32 TileBorderSize = 0;
+
+	/** Flush the virtual texture page contents. Call this whenever the wrapped UTexture is modified. */
+	ENGINE_API void Flush(FBox2f const& UVRect);
+};
diff --git a/Engine/Source/Runtime/Engine/Private/VT/VirtualTextureAdapter.cpp b/Engine/Source/Runtime/Engine/Private/VT/VirtualTextureAdapter.cpp
new file mode 100644
--- /dev/null
+++ b/Engine/Source/Runtime/Engine/Private/VT/VirtualTextureAdapter.cpp
@@ -0,0 +1,613 @@
+// Copyright Epic Games, Inc. All Rights Reserved.
+
+#include "VT/VirtualTextureAdapter.h"
+
+#include "Engine/Texture2D.h"
+#include "Engine/TextureRenderTarget2D.h"
+#include "EngineModule.h"
+#include "GlobalShader.h"
+#include "RendererInterface.h"
+#include "RenderGraphBuilder.h"
+#include "RenderGraphUtils.h"
+#include "RHIStaticStates.h"
+#include "ShaderParameterStruct.h"
+#include "TextureResource.h"
+#include "VirtualTexturing.h"
+#include "VT/VirtualTextureBuildSettings.h"
+
+#include UE_INLINE_GENERATED_CPP_BY_NAME(VirtualTextureAdapter)
+
+namespace VirtualTextureAdapter
+{
+	/** Copy with compression shader. */
+	class FCopyCompressCS : public FGlobalShader
+	{
+	public:
+		DECLARE_GLOBAL_SHADER(FCopyCompressCS);
+		SHADER_USE_PARAMETER_STRUCT(FCopyCompressCS, FGlobalShader);
+
+		BEGIN_SHADER_PARAMETER_STRUCT(FParameters, )
+			SHADER_PARAMETER_SRV(Texture2D, SourceTextureA)
+			SHADER_PARAMETER_RDG_TEXTURE_SRV(Texture2D, SourceTextureB)
+			SHADER_PARAMETER_SAMPLER(SamplerState, TextureSampler)
+			SHADER_PARAMETER_RDG_TEXTURE_UAV(RWTexture2D, DestTexture)
+			SHADER_PARAMETER_RDG_TEXTURE_UAV(RWTexture2D, DestCompressTexture_64bit)
+			SHADER_PARAMETER_RDG_TEXTURE_UAV(RWTexture2D, DestCompressTexture_128bit)
+			SHADER_PARAMETER(FVector2f, SourceUV)
+			SHADER_PARAMETER(FVector2f, TexelSize)
+			SHADER_PARAMETER(FVector2f, TexelOffsets)
+			SHADER_PARAMETER(FIntVector4, DestRect)
+		END_SHADER_PARAMETER_STRUCT()
+
+		static const int32 GroupSize = 8;
+
+		class FSourceTextureSelector : SHADER_PERMUTATION_BOOL("SOURCE_TEXTURE_A");
+		class FDestSrgb : SHADER_PERMUTATION_BOOL("TEXTURE_SRGB");
+		class FCompressionFormatDim : SHADER_PERMUTATION_INT("COMPRESSION_FORMAT", 7);
+		using FPermutationDomain = TShaderPermutationDomain<FSourceTextureSelector, FDestSrgb, FCompressionFormatDim>;
+
+		/** Get index to use for FCompressionFormatDim. */
+		static int32 GetCompressionPermutation(EPixelFormat InFormat)
+		{
+			const EPixelFormat CompressedFormats[] = { PF_DXT1, PF_DXT5, PF_BC4, PF_BC5, PF_BC6H, PF_BC7 };
+			return MakeArrayView(CompressedFormats).Find(InFormat) + 1;
+		}
+	};
+	IMPLEMENT_GLOBAL_SHADER(FCopyCompressCS, "/Engine/Private/VirtualTextureAdapter.usf", "CopyCompressCS", SF_Compute);
+
+	/** Final copy to output RDG parameters. */
+	BEGIN_SHADER_PARAMETER_STRUCT(FCopyToOutputParameters, )
+		RDG_TEXTURE_ACCESS(Input, ERHIAccess::CopySrc)
+	END_SHADER_PARAMETER_STRUCT()
+
+
+	/** Get the final virtual texture format from the wrapped texture format. */
+	static EPixelFormat GetFinalFormat(EPixelFormat InSourceFormat, bool bInUseCompressedFormat)
+	{
+		if (!bInUseCompressedFormat || IsBlockCompressedFormat(InSourceFormat) || IsInteger(InSourceFormat))
+		{
+			return InSourceFormat;
+		}
+
+		if (IsHDR(InSourceFormat))
+		{
+			return PF_BC6H;
+		}
+
+		EPixelFormatChannelFlags ChannelFlags = GetPixelFormatValidChannels(InSourceFormat);
+		if (EnumHasAnyFlags(ChannelFlags, EPixelFormatChannelFlags::A))
+		{
+			return PF_DXT5;
+		}
+		if (!EnumHasAnyFlags(ChannelFlags, EPixelFormatChannelFlags::G | EPixelFormatChannelFlags::B | EPixelFormatChannelFlags::A))
+		{
+			return PF_BC4;
+		}
+
+		return PF_DXT1;
+	}
+
+	/** Get the intermediate texture format that we use for transient intermediate targets. */
+	static EPixelFormat GetIntermediateFormat(EPixelFormat InSourceFormat, EPixelFormat InDestFormat)
+	{
+		if (IsBlockCompressedFormat(InSourceFormat))
+		{
+			return IsHDR(InSourceFormat) ? PF_FloatRGBA : PF_R8G8B8A8;
+		}
+
+		return InSourceFormat;
+	}
+
+	/** Copy from source to destination. Handles downsample and optional texture compression. */
+	static void RenderTile(
+		FRDGBuilder& GraphBuilder, 
+		FShaderResourceViewRHIRef& InSourceSRV,
+		FRHITexture* InDestTexture, 
+		IPooledRenderTarget* InDestRenderTarget, 
+		EPixelFormat InSourceFormat,
+		EPixelFormat InIntermediateFormat,
+		EPixelFormat InDestFormat,
+		uint32 InLevel, 
+		FBox2f const& InUVRange, 
+		FIntRect const& InDestRect)
+	{
+		FGlobalShaderMap* GlobalShaderMap = GetGlobalShaderMap(GMaxRHIFeatureLevel);
+
+		// Need to do a compute shader copy when source UV extends outside texture (because of tile borders). This is to get correct clamping behavior.
+		// For other cases we can just rely on the final RHICopyTexture (which can also handle direct copy of compressed formats).
+		const bool bUseCopyStep = InLevel == 0 && (InUVRange.Min.X < 0 || InUVRange.Min.Y < 0 || InUVRange.Max.X > 1 || InUVRange.Max.Y > 1);
+		// Need downsampling steps if we are copying to higher mip level.
+		// Note that if source texture has mips then we should have passed in the correct SRV and modified InLevel to avoid the downsampling.
+		const bool bUseDownsampleStep = InLevel > 0 && !bUseCopyStep;
+		// Need a compresssion step if we have a compressed format that isn't simply copied with the final RHICopyTexture.
+		const bool bUseCompressionStep = IsBlockCompressedFormat(InDestFormat) && (InSourceFormat != InDestFormat || bUseCopyStep || bUseDownsampleStep);
+
+		const uint32 FinalTexelCountX = InDestRect.Max.X - InDestRect.Min.X;
+		const uint32 FinalTexelCountY = InDestRect.Max.Y - InDestRect.Min.Y;
+
+		const bool bSrgb = InSourceSRV->GetDesc().Texture.SRV.GetViewInfo(InSourceSRV->GetTexture()).bSRGB;
+		const ETextureCreateFlags TextureCreateFlagsSrgb = bSrgb ? ETextureCreateFlags::SRGB : ETextureCreateFlags::None;
+
+		FRDGTextureRef CurrentOutput = nullptr;
+		FBox2f CurrentUVRange(InUVRange);
+		bool bIsDone = false;
+
+		if (bUseCopyStep)
+		{
+			const bool bIsFinalPass = !bUseCompressionStep;
+			if (bIsFinalPass)
+			{
+				CurrentOutput = GraphBuilder.RegisterExternalTexture(InDestRenderTarget, ERDGTextureFlags::None);
+			}
+			else
+			{
+				const FRDGTextureDesc Desc(FRDGTextureDesc::Create2D(FIntPoint(FinalTexelCountX, FinalTexelCountY), InIntermediateFormat, FClearValueBinding::None, ETextureCreateFlags::ShaderResource | ETextureCreateFlags::UAV | TextureCreateFlagsSrgb));
+				CurrentOutput = GraphBuilder.CreateTexture(Desc, TEXT("VirtualTextureAdapter.Downsample"));
+			}
+
+			FCopyCompressCS::FPermutationDomain PermutationVector;
+			PermutationVector.Set<FCopyCompressCS::FSourceTextureSelector>(true);
+			PermutationVector.Set<FCopyCompressCS::FDestSrgb>(bSrgb);
+			PermutationVector.Set<FCopyCompressCS::FCompressionFormatDim>(0);
+			TShaderMapRef<FCopyCompressCS> Shader(GlobalShaderMap, PermutationVector);
+
+			FCopyCompressCS::FParameters* Parameters = GraphBuilder.AllocParameters<FCopyCompressCS::FParameters>();
+			Parameters->SourceTextureA = InSourceSRV;
+			Parameters->TextureSampler = TStaticSamplerState<SF_Point, AM_Clamp, AM_Clamp, AM_Clamp>::GetRHI();
+			Parameters->DestTexture = GraphBuilder.CreateUAV(CurrentOutput);
+			Parameters->SourceUV = CurrentUVRange.Min;
+			Parameters->TexelSize = (CurrentUVRange.Max - CurrentUVRange.Min) / FVector2f(FinalTexelCountX, FinalTexelCountY);
+			Parameters->TexelOffsets = FVector2f(1.0f, 0.0f);
+			Parameters->DestRect = bIsFinalPass ? FIntVector4(InDestRect.Min.X, InDestRect.Min.Y, InDestRect.Max.X, InDestRect.Max.Y) : FIntVector4(0, 0, FinalTexelCountX, FinalTexelCountY);
+
+			const FIntVector GroupCount = FComputeShaderUtils::GetGroupCount(FIntPoint(FinalTexelCountX, FinalTexelCountY), FCopyCompressCS::GroupSize);
+
+			FComputeShaderUtils::AddPass(
+				GraphBuilder,
+				RDG_EVENT_NAME("VirtualTextureAdapterDownsample"),
+				Shader, Parameters, GroupCount);
+
+			CurrentUVRange = FBox2f(FVector2f(0, 0), FVector2f(1, 1));
+			bIsDone = bIsFinalPass;
+		}
+
+		if (bUseDownsampleStep)
+		{
+			for (uint32 LevelIndex = 0; LevelIndex < InLevel; ++LevelIndex)
+			{
+				FRDGTextureRef LastOutput = CurrentOutput;
+				const bool bUseSourceTextureA = LastOutput == nullptr;
+			
+				const uint32 DownsampleInputSizeX = FinalTexelCountX << (InLevel - LevelIndex);
+				const uint32 DownsampleInputSizeY = FinalTexelCountY << (InLevel - LevelIndex);
+				const uint32 DownsampleOutputSizeX = FinalTexelCountX << (InLevel - LevelIndex - 1);
+				const uint32 DownsampleOutputSizeY = FinalTexelCountY << (InLevel - LevelIndex - 1);
+
+				const bool bIsFinalPass = (LevelIndex == InLevel - 1) && !bUseCompressionStep;
+				if (bIsFinalPass)
+				{
+					CurrentOutput = GraphBuilder.RegisterExternalTexture(InDestRenderTarget, ERDGTextureFlags::None);
+				}
+				else
+				{
+					const FRDGTextureDesc Desc(FRDGTextureDesc::Create2D(FIntPoint(DownsampleOutputSizeX, DownsampleOutputSizeY), InIntermediateFormat, FClearValueBinding::None, ETextureCreateFlags::ShaderResource | ETextureCreateFlags::UAV | TextureCreateFlagsSrgb));
+					CurrentOutput = GraphBuilder.CreateTexture(Desc, TEXT("VirtualTextureAdapter.Downsample"));
+				}
+
+				FCopyCompressCS::FPermutationDomain PermutationVector;
+				PermutationVector.Set<FCopyCompressCS::FSourceTextureSelector>(bUseSourceTextureA);
+				PermutationVector.Set<FCopyCompressCS::FDestSrgb>(bSrgb);
+				PermutationVector.Set<FCopyCompressCS::FCompressionFormatDim>(0);
+				TShaderMapRef<FCopyCompressCS> Shader(GlobalShaderMap, PermutationVector);
+
+				FCopyCompressCS::FParameters* Parameters = GraphBuilder.AllocParameters<FCopyCompressCS::FParameters>();
+				Parameters->SourceTextureA = bUseSourceTextureA ? InSourceSRV : nullptr;
+				Parameters->SourceTextureB = bUseSourceTextureA ? nullptr : GraphBuilder.CreateSRV(FRDGTextureSRVDesc(LastOutput));
+				Parameters->TextureSampler = TStaticSamplerState<SF_Bilinear, AM_Clamp, AM_Clamp, AM_Clamp>::GetRHI();
+				Parameters->DestTexture = GraphBuilder.CreateUAV(CurrentOutput);
+				Parameters->SourceUV = CurrentUVRange.Min;
+				Parameters->TexelSize = (CurrentUVRange.Max - CurrentUVRange.Min) / FVector2f(DownsampleInputSizeX, DownsampleInputSizeY);
+				Parameters->TexelOffsets = FVector2f(2.0f, 0.5f);
+				Parameters->DestRect = bIsFinalPass ? FIntVector4(InDestRect.Min.X, InDestRect.Min.Y, InDestRect.Max.X, InDestRect.Max.Y) : FIntVector4(0, 0, DownsampleOutputSizeX, DownsampleOutputSizeY);
+
+				const FIntVector GroupCount = FComputeShaderUtils::GetGroupCount(FIntPoint(DownsampleOutputSizeX, DownsampleOutputSizeY), FCopyCompressCS::GroupSize);
+
+				FComputeShaderUtils::AddPass(
+					GraphBuilder,
+					RDG_EVENT_NAME("VirtualTextureAdapterDownsample"),
+					Shader, Parameters, GroupCount);
+
+				CurrentUVRange = FBox2f(FVector2f(0, 0), FVector2f(1, 1));
+				bIsDone = bIsFinalPass;
+			}
+		}
+
+		if (bUseCompressionStep)
+		{
+			const int32 CompressionPermutation = FCopyCompressCS::GetCompressionPermutation(InDestFormat);
+			check(CompressionPermutation > 0)
+			
+			FRDGTextureRef LastOutput = CurrentOutput;
+			const bool bUseSourceTextureA = LastOutput == nullptr;
+
+			FCopyCompressCS::FPermutationDomain PermutationVector;
+			PermutationVector.Set<FCopyCompressCS::FSourceTextureSelector>(bUseSourceTextureA);
+			PermutationVector.Set<FCopyCompressCS::FDestSrgb>(bSrgb);
+			PermutationVector.Set<FCopyCompressCS::FCompressionFormatDim>(CompressionPermutation);
+			TShaderMapRef<FCopyCompressCS> Shader(GlobalShaderMap, PermutationVector);
+
+			const EPixelFormat AliasFormat64bit = PF_R32G32_UINT;
+			const EPixelFormat AliasFormat128bit = PF_R32G32B32A32_UINT;
+			const bool bAliasTo64bit = InDestFormat == PF_DXT1 || InDestFormat == PF_BC4;
+			const EPixelFormat AliasFormat = bAliasTo64bit ? AliasFormat64bit : AliasFormat128bit;
+				
+			const bool bDirectAliasing = GRHISupportsUAVFormatAliasing;
+			if (bDirectAliasing)
+			{
+				CurrentOutput = GraphBuilder.RegisterExternalTexture(InDestRenderTarget, ERDGTextureFlags::None);
+			}
+			else
+			{
+				const FRDGTextureDesc Desc(FRDGTextureDesc::Create2D(FIntPoint(FinalTexelCountX, FinalTexelCountY) / 4, AliasFormat, FClearValueBinding::None, ETextureCreateFlags::ShaderResource | ETextureCreateFlags::UAV));
+				CurrentOutput = GraphBuilder.CreateTexture(Desc, TEXT("VirtualTextureAdapter.Compress"));
+			}
+
+			FCopyCompressCS::FParameters* Parameters = GraphBuilder.AllocParameters<FCopyCompressCS::FParameters>();
+			Parameters->SourceTextureA = bUseSourceTextureA ? InSourceSRV : nullptr;
+			Parameters->SourceTextureB = bUseSourceTextureA ? nullptr : GraphBuilder.CreateSRV(FRDGTextureSRVDesc(LastOutput));
+			Parameters->TextureSampler = TStaticSamplerState<SF_Bilinear, AM_Clamp, AM_Clamp, AM_Clamp>::GetRHI();
+			Parameters->DestCompressTexture_64bit = bAliasTo64bit ? GraphBuilder.CreateUAV(FRDGTextureUAVDesc(CurrentOutput, 0, AliasFormat64bit)) : nullptr;
+			Parameters->DestCompressTexture_128bit = bAliasTo64bit ? nullptr : GraphBuilder.CreateUAV(FRDGTextureUAVDesc(CurrentOutput, 0, AliasFormat128bit));
+			Parameters->SourceUV = CurrentUVRange.Min;
+			Parameters->TexelSize = (CurrentUVRange.Max - CurrentUVRange.Min) / FVector2f(FinalTexelCountX, FinalTexelCountY);
+			Parameters->TexelOffsets = FVector2f(4.0f, 0.5f);
+			Parameters->DestRect = bDirectAliasing ? FIntVector4(InDestRect.Min.X, InDestRect.Min.Y, InDestRect.Max.X, InDestRect.Max.Y) / 4 : FIntVector4(0, 0, FinalTexelCountX, FinalTexelCountY);
+
+			const FIntVector GroupCount = FComputeShaderUtils::GetGroupCount(FIntPoint(FinalTexelCountX, FinalTexelCountY) / 4, FCopyCompressCS::GroupSize);
+
+			FComputeShaderUtils::AddPass(
+				GraphBuilder,
+				RDG_EVENT_NAME("VirtualTextureAdapterCompress"),
+				Shader, Parameters, GroupCount);
+
+			bIsDone = bDirectAliasing;
+		}
+
+		// Final copy if we didn't already directly write to the physical texture.
+		if (!bIsDone)
+		{
+			FRHICopyTextureInfo CopyInfo;
+			CopyInfo.DestPosition = FIntVector(InDestRect.Min.X, InDestRect.Min.Y, 0);
+			CopyInfo.Size = FIntVector(FinalTexelCountX, FinalTexelCountY, 0);
+
+			const bool bUseSourceTextureA = CurrentOutput == nullptr;
+			if (bUseSourceTextureA)
+			{
+				// Copying directly from the source.
+				FIntPoint SourceTextureSize = InSourceSRV->GetTexture()->GetDesc().Extent;
+				uint32 SourceTextureMip = InSourceSRV->GetDesc().Texture.SRV.MipRange.First;
+				int32 SourcePositionX = FMath::FloorToInt(InUVRange.Min.X * (float)((uint32)SourceTextureSize.X >> SourceTextureMip));
+				int32 SourcePositionY = FMath::FloorToInt(InUVRange.Min.Y * (float)((uint32)SourceTextureSize.Y >> SourceTextureMip));
+				CopyInfo.SourcePosition = FIntVector(SourcePositionX, SourcePositionY, 0);
+				CopyInfo.SourceMipIndex = SourceTextureMip;
+			}
+
+			if (bUseCompressionStep)
+			{
+				// Take aliased format size difference into account.
+				CopyInfo.SourcePosition = CopyInfo.SourcePosition / 4;
+				CopyInfo.Size = CopyInfo.Size / 4;
+			}
+
+			FCopyToOutputParameters* Parameters = GraphBuilder.AllocParameters<FCopyToOutputParameters>();
+			Parameters->Input = CurrentOutput;
+
+			GraphBuilder.AddPass(
+				RDG_EVENT_NAME("VirtualTextureAdapterCopyToOutput"),
+				Parameters,
+				ERDGPassFlags::Copy | ERDGPassFlags::NeverCull,
+				[bUseSourceTextureA, InputTextureA = InSourceSRV, InputTextureB = CurrentOutput, OutputTexture = InDestTexture, CopyInfo](FRHICommandList& RHICmdList)
+				{
+					FRHITexture* InputTexture = bUseSourceTextureA ? InputTextureA->GetTexture() : InputTextureB->GetRHI();
+
+					RHICmdList.Transition(FRHITransitionInfo(OutputTexture, ERHIAccess::SRVMask, ERHIAccess::CopyDest));
+					RHICmdList.CopyTexture(InputTexture, OutputTexture, CopyInfo);
+					RHICmdList.Transition(FRHITransitionInfo(OutputTexture, ERHIAccess::CopyDest, ERHIAccess::SRVMask));
+				});
+		}
+	}
+}
+
+/** IVirtualTextureFinalizer implementation that renders the virtual texture pages on demand. */
+class FVirtualTextureAdapterFinalizer : public IVirtualTextureFinalizer
+{
+public:
+	FVirtualTextureAdapterFinalizer(FRHICommandListBase& RHICmdList, FRHITexture* InSourceTexture, FVTProducerDescription const& InProducerDesc)
+		: SourceTexture(InSourceTexture)
+		, ProducerDesc(InProducerDesc)
+	{
+		SourceFormat = SourceTexture->GetDesc().Format;
+		DestFormat = ProducerDesc.LayerFormat[0];
+		IntermediateFormat = VirtualTextureAdapter::GetIntermediateFormat(SourceFormat, DestFormat);
+
+		// Store off SRVs which we use to access individual source mips.
+		SourceSRVs.Reserve(SourceTexture->GetNumMips());
+		for (uint32 MipIndex = 0; MipIndex < SourceTexture->GetNumMips(); ++MipIndex)
+		{
+			SourceSRVs.Add(RHICmdList.CreateShaderResourceView(SourceTexture, MipIndex));
+		}
+	}
+
+	virtual ~FVirtualTextureAdapterFinalizer() = default;
+
+	/** A description for a single tile to render. */
+	struct FTileEntry
+	{
+		FVTProduceTargetLayer Target;
+		uint64 vAddress = 0;
+		uint8 vLevel = 0;
+	};
+
+	void AddTile(FTileEntry const& InEntry)
+	{
+		TilesToRender.Add(InEntry);
+	}
+
+protected:
+	/** Source RHI texture to copy into virtual texture pages. */
+	FRHITexture* SourceTexture = nullptr;
+	/** Producer description of our virtual texture. */
+	const FVTProducerDescription ProducerDesc;
+	/** Format of SourceTexture. */
+	EPixelFormat SourceFormat;
+	/** Destination format for tile generation. */
+	EPixelFormat DestFormat;
+	/** Pixel format used for intermedite downsample buffers. */
+	EPixelFormat IntermediateFormat;
+	/** SRVs of source RHI texture to copy into virtual texture pages. */
+	TArray<FShaderResourceViewRHIRef> SourceSRVs;
+	/** Array of tiles in the queue to finalize. */
+	TArray<FTileEntry> TilesToRender;
+
+	//~ Begin IVirtualTextureFinalizer Interface.
+	virtual void Finalize(FRDGBuilder& GraphBuilder) override 
+	{
+		for (FTileEntry const& Tile : TilesToRender)
+		{
+			FRHITexture* DestTexture = Tile.Target.TextureRHI != nullptr ? Tile.Target.TextureRHI->GetTexture2D() : nullptr;
+			IPooledRenderTarget* DestRenderTarget = Tile.Target.PooledRenderTarget;
+
+			const float X = (float)FMath::ReverseMortonCode2_64(Tile.vAddress);
+			const float Y = (float)FMath::ReverseMortonCode2_64(Tile.vAddress >> 1);
+			const float DivisorX = (float)ProducerDesc.BlockWidthInTiles / (float)(1 << Tile.vLevel);
+			const float DivisorY = (float)ProducerDesc.BlockHeightInTiles / (float)(1 << Tile.vLevel);
+
+			const FVector2f UV(X / DivisorX, Y / DivisorY);
+			const FVector2f UVSize(1.f / DivisorX, 1.f / DivisorY);
+			const FVector2f UVBorder = UVSize * ((float)ProducerDesc.TileBorderSize / (float)ProducerDesc.TileSize);
+			const FBox2f UVRange(UV - UVBorder, UV + UVSize + UVBorder);
+
+			const int32 TileSize = ProducerDesc.TileSize + 2 * ProducerDesc.TileBorderSize;
+			const FIntPoint DestinationPos(Tile.Target.pPageLocation.X * TileSize, Tile.Target.pPageLocation.Y * TileSize);
+			const FIntRect DestRect(DestinationPos, DestinationPos + FIntPoint(TileSize, TileSize));
+		
+			const bool bHasSRVForLevel = SourceSRVs.IsValidIndex(Tile.vLevel);
+			const int32 SRVIndex = bHasSRVForLevel ? Tile.vLevel : 0;
+			const int32 vLevel = bHasSRVForLevel ? 0 : Tile.vLevel;
+
+			VirtualTextureAdapter::RenderTile(GraphBuilder, SourceSRVs[SRVIndex], DestTexture, DestRenderTarget, SourceFormat, IntermediateFormat, DestFormat, vLevel, UVRange, DestRect);
+		}
+
+		TilesToRender.Reset();
+	}
+	//~ End IVirtualTextureFinalizer Interface.
+};
+
+/** IVirtualTexture implementation that is handling runtime rendered page data requests. */
+class FVirtualTextureAdapterProducer : public IVirtualTexture
+{
+public:
+	FVirtualTextureAdapterProducer(FRHICommandListBase& RHICmdList, FRHITexture* InSourceTexture, int32 InUnstreamedMipCount, FVTProducerDescription const& InProducerDesc)
+		: Finalizer(RHICmdList, InSourceTexture, InProducerDesc)
+		, UnstreamedMipCount(InUnstreamedMipCount)
+	{}
+
+	virtual ~FVirtualTextureAdapterProducer() = default;
+
+	//~ Begin IVirtualTexture Interface.
+	virtual bool IsPageStreamed(uint8 vLevel, uint32 vAddress) const override
+	{
+		return false;
+	}
+
+	virtual FVTRequestPageResult RequestPageData(
+		FRHICommandList& RHICmdList,
+		const FVirtualTextureProducerHandle& ProducerHandle,
+		uint8 LayerMask,
+		uint8 vLevel,
+		uint64 vAddress,
+		EVTRequestPagePriority Priority
+	) override
+	{
+		EVTRequestPageStatus Status = vLevel < UnstreamedMipCount ? EVTRequestPageStatus::Invalid : EVTRequestPageStatus::Available;
+		return FVTRequestPageResult(Status, 0u);
+	}
+
+	virtual IVirtualTextureFinalizer* ProducePageData(
+		FRHICommandList& RHICmdList,
+		ERHIFeatureLevel::Type FeatureLevel,
+		EVTProducePageFlags Flags,
+		const FVirtualTextureProducerHandle& ProducerHandle,
+		uint8 LayerMask,
+		uint8 vLevel,
+		uint64 vAddress,
+		uint64 RequestHandle,
+		const FVTProduceTargetLayer* TargetLayers
+	) override
+	{
+		FVirtualTextureAdapterFinalizer::FTileEntry Tile;
+		Tile.Target = TargetLayers[0];
+		Tile.vAddress = vAddress;
+		Tile.vLevel = vLevel;
+		Finalizer.AddTile(Tile);
+		return &Finalizer;
+	}
+	//~ End IVirtualTexture Interface.
+
+private:
+	/** The Finalizer object to write the virtual texture pages. */
+	FVirtualTextureAdapterFinalizer Finalizer;
+	/** The number of mips that are not streamed which is also the number of mips that we can't produce. */
+	int32 UnstreamedMipCount = 0;
+
+};
+
+/** Implementation of FVirtualTextureResource that instantiates an FVirtualTextureAdapterProducer. */
+class FVirtualTextureAdapterRenderResource : public FVirtualTexture2DResource
+{
+	FTextureResource* SourceResource = nullptr;
+	EPixelFormat Format = PF_Unknown;
+	uint32 TileSize = 0;
+	uint32 TileBorderSize = 0;
+	uint32 NumTilesX = 0;
+	uint32 NumTilesY = 0;
+	uint32 MaxLevel = 0;
+	uint32 NumSourceMips = 1;
+
+public:
+	FVirtualTextureAdapterRenderResource(UVirtualTextureAdapter const* InOwner, UTexture* InTexture, int32 InTileSize, int32 InTileBorderSize, bool bUseCompressedFormat)
+	{
+		TextureName = InOwner->GetFName();
+		PackageName = InOwner->GetOutermost()->GetFName();
+
+		SourceResource = InTexture->GetResource();
+
+		EPixelFormat SourceFormat = PF_Unknown;
+		if (UTexture2D* Texture2D = Cast<UTexture2D>(InTexture))
+		{
+			SourceFormat = Texture2D->GetPixelFormat(0);
+			NumSourceMips = Texture2D->GetNumMips();
+		}
+		else if (UTextureRenderTarget2D* RenderTarget2D = Cast<UTextureRenderTarget2D>(InTexture))
+		{
+			SourceFormat = RenderTarget2D->GetFormat();
+		}
+		Format = VirtualTextureAdapter::GetFinalFormat(SourceFormat, bUseCompressedFormat);
+
+		TileSize = InTileSize;
+		TileBorderSize = InTileBorderSize;
+		NumTilesX = FMath::DivideAndRoundUp((uint32)InTexture->GetSurfaceWidth(), TileSize);
+		NumTilesY = FMath::DivideAndRoundUp((uint32)InTexture->GetSurfaceHeight(), TileSize);
+		MaxLevel = FMath::CeilLogTwo(FMath::Max(NumTilesX, NumTilesY));
+	}
+
+	//~ Begin FVirtualTexture2DResource Interface.
+	virtual uint32 GetNumLayers() const override { return 1; }
+	virtual EPixelFormat GetFormat(uint32 LayerIndex) const override { return Format; }
+	virtual uint32 GetTileSize() const override { return TileSize; }
+	virtual uint32 GetBorderSize() const override { return TileBorderSize; }
+	virtual uint32 GetNumTilesX() const override { return NumTilesX; }
+	virtual uint32 GetNumTilesY() const override { return NumTilesY; }
+	virtual uint32 GetNumMips() const override { return MaxLevel + 1; }
+	virtual FIntPoint GetSizeInBlocks() const override { return 1; }
+	//~ End FVirtualTexture2DResource Interface.
+
+	//~ Begin FRenderResource Interface.
+	virtual void InitRHI(FRHICommandListBase& RHICmdList) override
+	{
+		TRACE_CPUPROFILER_EVENT_SCOPE(FVirtualTextureAdapterRenderResource::InitRHI);
+
+		FSamplerStateInitializerRHI SamplerStateInitializer;
+		SamplerStateInitializer.Filter = SF_Bilinear;
+		SamplerStateRHI = GetOrCreateSamplerState(SamplerStateInitializer);
+
+		const int32 UnstreamedMipCount = NumSourceMips - SourceResource->GetCurrentMipCount();
+
+		FVTProducerDescription ProducerDesc;
+		ProducerDesc.Name = TextureName;
+		ProducerDesc.FullNameHash = GetTypeHash(TextureName);
+		ProducerDesc.bContinuousUpdate = false;
+		ProducerDesc.Dimensions = 2;
+		ProducerDesc.TileSize = TileSize;
+		ProducerDesc.TileBorderSize = TileBorderSize;
+		ProducerDesc.BlockWidthInTiles = NumTilesX;
+		ProducerDesc.BlockHeightInTiles = NumTilesY;
+		ProducerDesc.DepthInTiles = 1u;
+		ProducerDesc.MaxLevel = MaxLevel;
+		ProducerDesc.NumTextureLayers = 1;
+		ProducerDesc.NumPhysicalGroups = 1;
+		ProducerDesc.LayerFormat[0] = Format;
+
+		FVirtualTextureAdapterProducer* VirtualTexture = new FVirtualTextureAdapterProducer(RHICmdList, SourceResource->GetTexture2DRHI(), UnstreamedMipCount, ProducerDesc);
+		ProducerHandle = GetRendererModule().RegisterVirtualTextureProducer(RHICmdList, ProducerDesc, VirtualTexture);
+	}
+	//~ End FRenderResource Interface.
+};
+
+UVirtualTextureAdapter::UVirtualTextureAdapter(const FObjectInitializer& ObjectInitializer)
+	: Super(ObjectInitializer)
+{
+}
+
+FTextureResource* UVirtualTextureAdapter::CreateResource()
+{
+	// Only support 2D texture or render target.
+	if (Cast<UTexture2D>(Texture) == nullptr && Cast<UTextureRenderTarget2D>(Texture) == nullptr)
+	{
+		return nullptr;
+	}
+	// Can only wrap regular textures.
+	if (Texture->VirtualTextureStreaming)
+	{
+		return nullptr;
+	}
+
+	FVirtualTextureBuildSettings DefaultSettings;
+	DefaultSettings.Init();
+
+	const uint32 FinalTileSize = bUseDefaultTileSizes ? DefaultSettings.TileSize : FMath::Min(FMath::RoundUpToPowerOfTwo(TileSize), 512u);
+	const uint32 FinalTileBorderSize = bUseDefaultTileSizes ? DefaultSettings.TileBorderSize : FMath::Min((((uint32)TileBorderSize + 1) & ~1u), 8u);
+	
+	return new FVirtualTextureAdapterRenderResource(this, Texture, FinalTileSize, FinalTileBorderSize, bUseCompressedFormat);
+}
+
+
+ETextureClass UVirtualTextureAdapter::GetTextureClass() const
+{
+	return ETextureClass::TwoD;
+}
+
+EMaterialValueType UVirtualTextureAdapter::GetMaterialType() const
+{
+	return MCT_TextureVirtual;
+}
+
+float UVirtualTextureAdapter::GetSurfaceWidth() const
+{
+	return Texture ? Texture->GetSurfaceWidth() : 0;
+}
+
+float UVirtualTextureAdapter::GetSurfaceHeight() const
+{
+	return Texture ? Texture->GetSurfaceHeight() : 0;
+}
+
+void UVirtualTextureAdapter::Flush(FBox2f const& UVRect)
+{
+	FTextureResource* Resource = GetResource();
+	if (Resource == nullptr)
+	{
+		return;
+	}
+
+	ENQUEUE_RENDER_COMMAND(VirtualTextureAdapterFlush)([Resource, UVRect](FRHICommandListBase&)
+	{
+		FVirtualTexture2DResource* VTResource = Resource->GetVirtualTexture2DResource();
+		IAllocatedVirtualTexture* AllocatedVT = VTResource ? VTResource->AcquireAllocatedVT() : nullptr;
+
+		if (AllocatedVT != nullptr)
+		{
+			GetRendererModule().FlushVirtualTextureCache(AllocatedVT, UVRect.Min, UVRect.Max);
+		}
+	});
+}
