Make non-stereo views devoid of stereo instanced views.

- Snapshots of a view (via FViewInfo::CreateSnapshot()) still link into the original view family. Even if the code (e.g. water rendering) disables stereo for them, they are still considered instanced and have instanced view UB set up. While possibly no harm in doing that, this feels wrong.

diff --git a/Engine/Source/Runtime/Engine/Private/SceneView.cpp b/Engine/Source/Runtime/Engine/Private/SceneView.cpp
--- a/Engine/Source/Runtime/Engine/Private/SceneView.cpp
+++ b/Engine/Source/Runtime/Engine/Private/SceneView.cpp
@@ -2818,8 +2818,8 @@ const FSceneView* FSceneView::GetPrimarySceneView() const
 
 const FSceneView* FSceneView::GetInstancedSceneView() const
 {
-	// if we don't have ISR (or MMV) enabled, we don't have instanced views
-	if (bIsMultiViewportEnabled || bIsMobileMultiViewEnabled)
+	// if we don't have ISR or MMV enabled, or we're rendering a full pass, we don't have instanced views
+	if ((bIsMultiViewportEnabled || bIsMobileMultiViewEnabled) && StereoPass != EStereoscopicPass::eSSP_FULL)
 	{
 		// If called on the first secondary view it'll return itself.
 		if (Family && Family->Views.IsValidIndex(PrimaryViewIndex + 1))
diff --git a/Engine/Source/Runtime/Renderer/Private/SceneRendering.cpp b/Engine/Source/Runtime/Renderer/Private/SceneRendering.cpp
--- a/Engine/Source/Runtime/Renderer/Private/SceneRendering.cpp
+++ b/Engine/Source/Runtime/Renderer/Private/SceneRendering.cpp
@@ -2276,16 +2276,6 @@ const FViewInfo* FViewInfo::GetPrimaryView() const
 	return this;
 }
 
-const FViewInfo* FViewInfo::GetInstancedView() const
-{
-	// Extra checks are needed because some code relies on this function to return NULL if ISR is disabled.
-	if (bIsInstancedStereoEnabled || bIsMobileMultiViewEnabled)
-	{
-		return static_cast<const FViewInfo*>(GetInstancedSceneView());
-	}
-	return nullptr;
-}
-
 struct FViewInfoSnapshotCache
 {
 	// These are not real view infos, just dumb memory blocks
diff --git a/Engine/Source/Runtime/Renderer/Private/SceneRendering.h b/Engine/Source/Runtime/Renderer/Private/SceneRendering.h
--- a/Engine/Source/Runtime/Renderer/Private/SceneRendering.h
+++ b/Engine/Source/Runtime/Renderer/Private/SceneRendering.h
@@ -1697,7 +1697,7 @@ public:
 	const FViewInfo* GetPrimaryView() const;
 
 	/** Returns the instanced view associated with the input view, or null if none exists. */
-	const FViewInfo* GetInstancedView() const;
+	inline const FViewInfo* GetInstancedView() const { return static_cast<const FViewInfo*>(GetInstancedSceneView()); }
 
 	/** Create a snapshot of this view info on the scene allocator. */
 	FViewInfo* CreateSnapshot() const;
