Activate ODSC's new mode by default
This means that we won't compile all shaders for a given material, but only the permutations needed for rendering

#rnx

diff --git a/Engine/Source/Runtime/Engine/Private/ODSC/ODSCManager.cpp b/Engine/Source/Runtime/Engine/Private/ODSC/ODSCManager.cpp
--- a/Engine/Source/Runtime/Engine/Private/ODSC/ODSCManager.cpp
+++ b/Engine/Source/Runtime/Engine/Private/ODSC/ODSCManager.cpp
@@ -22,10 +22,10 @@ DEFINE_LOG_CATEGORY(LogODSC);
 
 static TAutoConsoleVariable<int32> CVarODSCRecompileMode(
 	TEXT("ODSC.recompilemode"),
-	0,
+	1,
 	TEXT("Highly experimental - Changes how recompileshaders behaves in cooked build\n")
-	TEXT("0 (default): Gathers all visible materials in a single frame and compiles all permutations for them\n")
-	TEXT("1: Compile only the permutations that are requested by the renderer. Faster iteration but more prone to hitching because of MDC recaching\n")
+	TEXT("0 (legacy): Gathers all visible materials in a single frame and compiles all permutations for them\n")
+	TEXT("1 (default): Compile only the permutations that are requested by the renderer. Faster iteration but more prone to hitching because of MDC recaching\n")
 	);
 
 FODSCManager* GODSCManager = nullptr;
