[Material Editor] New material translator push.

- Moved IRModule creation entirely to translation time, because that's the only moment where we have all information necessary for complete translation (e.g. static parameter values, target shader model)
- Implemented scopes and automatic instruction scoping.
- Added branch instruction.
- Added some comparison binary instructions.
- Extended and improved the FEmitter's API.
- Added support for UMaterialExpressionIf node.
- Better error reporting.

diff --git a/Engine/Source/Runtime/Engine/Private/Materials/Material.cpp b/Engine/Source/Runtime/Engine/Private/Materials/Material.cpp
--- a/Engine/Source/Runtime/Engine/Private/Materials/Material.cpp
+++ b/Engine/Source/Runtime/Engine/Private/Materials/Material.cpp
@@ -2081,24 +2081,14 @@ void UMaterial::UpdateCachedExpressionData()
 
 	FMaterialCachedExpressionData* LocalCachedExpressionData = new FMaterialCachedExpressionData();
 	FMaterialCachedHLSLTree* LocalCachedTree = nullptr;
-	if (IsUsingNewHLSLGenerator())
+	if (IsUsingNewHLSLGenerator() && !IsUsingNewMaterialTranslatorPrototype())
 	{
 		// Relinks function call inputs. Otherwise, we can get invalid inputs and they will cause errors when generating the syntax tree
 		UpdateTransientExpressionData();
 
-		if (IsUsingNewMaterialTranslatorPrototype())
-		{
-			FMaterialIRModuleBuilder Builder;
-			Builder.Build(this, &IRModule);
-
-			LocalCachedExpressionData->AnalyzeMaterial(*this); // Temporary
-		}
-		else
-		{
-			LocalCachedTree = new FMaterialCachedHLSLTree();
-			LocalCachedTree->GenerateTree(this, nullptr, nullptr);
-			LocalCachedExpressionData->UpdateForCachedHLSLTree(*LocalCachedTree, nullptr, this);
-		}
+		LocalCachedTree = new FMaterialCachedHLSLTree();
+		LocalCachedTree->GenerateTree(this, nullptr, nullptr);
+		LocalCachedExpressionData->UpdateForCachedHLSLTree(*LocalCachedTree, nullptr, this);
 	}
 	else
 	{
diff --git a/Engine/Source/Runtime/Engine/Private/Materials/MaterialExpressionsIR.cpp b/Engine/Source/Runtime/Engine/Private/Materials/MaterialExpressionsIR.cpp
--- a/Engine/Source/Runtime/Engine/Private/Materials/MaterialExpressionsIR.cpp
+++ b/Engine/Source/Runtime/Engine/Private/Materials/MaterialExpressionsIR.cpp
@@ -260,25 +260,25 @@ void UMaterialExpression::Build(IR::FEmitter& Emitter)
 
 void UMaterialExpressionConstant::Build(IR::FEmitter& Emitter)
 {
-	IR::FValuePtr Value = Emitter.EmitConstantFloat1(R);
+	IR::FValue* Value = Emitter.EmitConstantFloat1(R);
 	Emitter.Put(GetOutput(0), Value);
 }
 
 void UMaterialExpressionConstant2Vector::Build(IR::FEmitter& Emitter)
 {
-	IR::FValuePtr Value = Emitter.EmitConstantFloat2({ R, G });
+	IR::FValue* Value = Emitter.EmitConstantFloat2({ R, G });
 	Emitter.Put(GetOutput(0), Value);
 }
 
 void UMaterialExpressionConstant3Vector::Build(IR::FEmitter& Emitter)
 {
-	IR::FValuePtr Value = Emitter.EmitConstantFloat3({ Constant.R, Constant.G, Constant.B });
+	IR::FValue* Value = Emitter.EmitConstantFloat3({ Constant.R, Constant.G, Constant.B });
 	Emitter.Put(GetOutput(0), Value);
 }
 
 void UMaterialExpressionConstant4Vector::Build(IR::FEmitter& Emitter)
 {
-	IR::FValuePtr Value = Emitter.EmitConstantFloat4(Constant);
+	IR::FValue* Value = Emitter.EmitConstantFloat4(Constant);
 	Emitter.Put(GetOutput(0), Value);
 }
 
@@ -293,37 +293,33 @@ static void BuildBinaryArithmeticOperator(
 	float RhsConst,
 	FExpressionOutput* Output)
 {
-	// Get values flowing into inputs and check that they're arithmetic.
-	IR::FValuePtr LhsValue = Emitter.GetAndCheckInputTypeKindIs(LhsInput, IR::TK_Arithmetic);
-	IR::FValuePtr RhsValue = Emitter.GetAndCheckInputTypeKindIs(RhsInput, IR::TK_Arithmetic);
-	if (Emitter.IsInvalid()) { return; }
+	// Default inputs to their relative constants if disconnected, then get each input after checking it has arithmetic type.
+	IR::FValue* LhsValue = Emitter.DefaultTo(LhsInput, LhsConst).TryGetArithmetic(LhsInput);
+	IR::FValue* RhsValue = Emitter.DefaultTo(RhsInput, RhsConst).TryGetArithmetic(RhsInput);
+
+	if (Emitter.IsInvalid())
+	{
+		return;
+	}
 
 	// Determine operation input/output type by looking at the first connected input and picking float1 otherwise.
 	IR::FTypePtr ResultType = LhsValue ? LhsValue->Type
 		: RhsValue ? RhsValue->Type
 		: IR::FArithmeticType::GetScalar(IR::SK_Float);
 
-	// If lhs operand is disconnected, create a value for its constant property instead.
-	if (!LhsValue)
-	{
-		LhsValue = Emitter.EmitConstantFloat1(LhsConst);
-	}
-
-	// If rhs operand is disconnected, create a value for its constant property instead.
-	if (!RhsValue)
-	{
-		RhsValue = Emitter.EmitConstantFloat1(RhsConst);
-	}
-
 	// Convert operand values to determined result type. 
 	LhsValue = Emitter.TryEmitConvert(LhsValue, ResultType);
 	RhsValue = Emitter.TryEmitConvert(RhsValue, ResultType);
-	if (Emitter.IsInvalid()) { return; }
+
+	if (Emitter.IsInvalid())
+	{
+		return;
+	}
 
 	// Finally emit the binary operator.
-	IR::FValuePtr Value = Emitter.EmitBinaryOperator(Op, LhsValue, RhsValue);
+	IR::FValue* Value = Emitter.EmitBinaryOperator(Op, LhsValue, RhsValue);
 
-	// And flow it out from our only output.
+	// And flow it out of the expression's only output.
 	Emitter.Put(Output, Value);
 }
 
@@ -347,4 +343,67 @@ void UMaterialExpressionDivide::Build(IR::FEmitter& Emitter)
 	BuildBinaryArithmeticOperator(Emitter, IR::BO_Divide, &A, ConstA, &B, ConstB, GetOutput(0));
 }
 
+void UMaterialExpressionIf::Build(IR::FEmitter& Emitter)
+{
+	// Create default values flowing into disconnected inputs
+	Emitter.DefaultToFloatZero(&A);
+	Emitter.DefaultTo(&B, ConstB);
+	Emitter.DefaultToFloatZero(&AGreaterThanB);
+	Emitter.DefaultToFloatZero(&AEqualsB);
+	Emitter.DefaultToFloatZero(&ALessThanB);
+
+	// Get input values and check their types are what we expect.
+	IR::FValue* AValue = Emitter.TryGetScalar(&A);
+	IR::FValue* BValue = Emitter.TryGetScalar(&B);
+	IR::FValue* AGreaterThanBValue = Emitter.TryGetArithmetic(&AGreaterThanB);
+	IR::FValue* AEqualsBValue = Emitter.TryGetArithmetic(&AEqualsB);
+	IR::FValue* ALessThanBValue = Emitter.TryGetArithmetic(&ALessThanB);
+
+	if (Emitter.IsInvalid())
+	{
+		return;
+	}
+
+	// Get the arithmetic common type between the conditional arguments (e.g. if inputs are int and float, it will return float).
+	IR::FArithmeticTypePtr ConditionArgsType = Emitter.TryGetCommonArithmeticType(AValue->Type->AsArithmetic(), BValue->Type->AsArithmetic());
+
+	if (Emitter.IsInvalid())
+	{
+		return;
+	}
+
+	// Convert both conditional argument values to the common type. This is expected to work as the common type could be found.
+	AValue = Emitter.TryEmitConvert(AValue, ConditionArgsType);
+	BValue = Emitter.TryEmitConvert(BValue, ConditionArgsType);
+
+	// Now determine the output type by taking the common arithmetic type between result values.
+	IR::FArithmeticTypePtr OutputType = Emitter.TryGetCommonArithmeticType(AGreaterThanBValue->Type->AsArithmetic(), AEqualsBValue->Type->AsArithmetic());
+	OutputType = Emitter.TryGetCommonArithmeticType(OutputType, ALessThanBValue->Type->AsArithmetic());
+
+	if (Emitter.IsInvalid())
+	{
+		return;
+	}
+
+	// Convert result values to the common result type.
+	AGreaterThanBValue = Emitter.TryEmitConvert(AGreaterThanBValue, OutputType);
+	AEqualsBValue = Emitter.TryEmitConvert(AEqualsBValue, OutputType);
+	ALessThanBValue = Emitter.TryEmitConvert(ALessThanBValue, OutputType);
+
+	if (Emitter.IsInvalid())
+	{
+		return;
+	}
+
+	// Emit the comparison expressions.
+	IR::FValue* ALessThanBConditionValue = Emitter.EmitBinaryOperator(IR::BO_Lower, AValue, BValue);
+	IR::FValue* AEqualsBConditionValue = Emitter.EmitBinaryOperator(IR::BO_Equals, AValue, BValue);
+
+	// And finally emit the full conditional expression.
+	IR::FValue* OutputValue = Emitter.EmitBranch(AEqualsBConditionValue, AEqualsBValue, AGreaterThanBValue);
+	OutputValue = Emitter.EmitBranch(ALessThanBConditionValue, ALessThanBValue, OutputValue);
+
+	Emitter.Put(GetOutput(0), OutputValue);
+}
+
 #endif // WITH_EDITOR
diff --git a/Engine/Source/Runtime/Engine/Private/Materials/MaterialIR.cpp b/Engine/Source/Runtime/Engine/Private/Materials/MaterialIR.cpp
--- a/Engine/Source/Runtime/Engine/Private/Materials/MaterialIR.cpp
+++ b/Engine/Source/Runtime/Engine/Private/Materials/MaterialIR.cpp
@@ -8,12 +8,17 @@
 namespace UE::MIR
 {
 
-FInstructionPtr FValue::AsInstruction() const
+FInstruction* FValue::AsInstruction()
 {
-	return (Kind >= VK_InstructionBegin && Kind < VK_InstructionEnd) ? static_cast<FInstructionPtr>(this) : nullptr;
+	return (Kind >= VK_InstructionBegin && Kind < VK_InstructionEnd) ? static_cast<FInstruction*>(this) : nullptr;
 }
 
-bool FValue::Equals(FValuePtr Other) const
+const FInstruction* FValue::AsInstruction() const
+{
+	return const_cast<FValue*>(this)->AsInstruction();
+}
+
+bool FValue::Equals(const FValue* Other) const
 {
 	// If kinds are different the two values are surely different.
 	if (Kind != Other->Kind)
@@ -34,28 +39,97 @@ uint32 FValue::GetSizeInBytes() const
 	switch (Kind)
 	{
 		case VK_ScalarConstant: return sizeof(FScalarConstant);
-		case VK_SetMaterialOutput: return sizeof(FScalarConstant);
+		case VK_SetMaterialOutput: return sizeof(FSetMaterialOutput);
 		case VK_BinaryOperator: return sizeof(FBinaryOperator);
-		case VK_Dimensional: return sizeof(FDimensional) + sizeof(FValuePtr) * static_cast<const FDimensional*>(this)->GetComponents().Num();
+		case VK_Dimensional: return sizeof(FDimensional) + sizeof(FValue*) * static_cast<const FDimensional*>(this)->GetComponents().Num();
+		case VK_Branch: return sizeof(FBranch);
+		default: UE_MIR_UNREACHABLE();
+	}
+}
+
+TArrayView<FValue*> FValue::GetUses()
+{
+	switch (Kind)
+	{
+		case VK_ScalarConstant:
+		case VK_Dimensional:
+			return {};
+
+		case VK_SetMaterialOutput:
+		{
+			auto This = static_cast<FSetMaterialOutput*>(this);
+			return { &This->Arg, 1 };
+		}
+
+		case VK_BinaryOperator:
+		{
+			auto This = static_cast<FBinaryOperator*>(this);
+			return { &This->LhsArg, 2 };
+		}
+
+		case VK_Branch:
+		{
+			auto This = static_cast<FBranch*>(this);
+			return { &This->ConditionArg, 3 };
+		}
+
 		default: UE_MIR_UNREACHABLE();
 	}
 }
 
-TArrayView<const FValuePtr> FDimensional::GetComponents() const
+TArrayView<FValue* const> FDimensional::GetComponents() const
 {
-	TArrayView<FValuePtr> Components = const_cast<FDimensional*>(this)->GetMutableComponents();
+	TArrayView<FValue*> Components = const_cast<FDimensional*>(this)->GetMutableComponents();
 	return { Components.GetData(), Components.Num() };
 }
 
-TArrayView<FValuePtr> FDimensional::GetMutableComponents()
+TArrayView<FValue*> FDimensional::GetMutableComponents()
 {
-	FArithmeticTypePtr ArithmeticType = Type->ToArithmetic();
+	FArithmeticTypePtr ArithmeticType = Type->AsArithmetic();
 	check(ArithmeticType);
 
-	FValuePtr* Ptr = (FValuePtr*)static_cast<TDimensional<1>*>(this)->Components;
+	FValue** Ptr = (FValue**)static_cast<TDimensional<1>*>(this)->Components;
 	return { Ptr, ArithmeticType->NumRows };
 }
 
+bool FInstruction::GetInnerBlock(int32 Index, FValue*& OutArg, FBlock*& OutBlock)
+{
+	switch (Kind)
+	{
+		case VK_Branch:
+		{
+			auto This = static_cast<FBranch*>(this);
+			if (Index == 0)
+			{
+				OutArg = This->TrueArg;
+				OutBlock = &This->TrueBlock;
+				return true;
+			}
+			else if (Index == 1)
+			{
+				OutArg = This->FalseArg;
+				OutBlock = &This->FalseBlock;
+				return true;
+			}
+			break;
+		}
+
+		default:
+		{
+			TArrayView<FValue*> Uses = GetUses();
+			if (Uses.IsValidIndex(Index))
+			{
+				OutArg = Uses[Index];
+				OutBlock = Block;
+				return true;
+			}
+			break;
+		}
+	}
+
+	return false;
+}
+
 } // namespace UE::MIR
 
 #endif // #if WITH_EDITOR
diff --git a/Engine/Source/Runtime/Engine/Private/Materials/MaterialIREmitter.cpp b/Engine/Source/Runtime/Engine/Private/Materials/MaterialIREmitter.cpp
--- a/Engine/Source/Runtime/Engine/Private/Materials/MaterialIREmitter.cpp
+++ b/Engine/Source/Runtime/Engine/Private/Materials/MaterialIREmitter.cpp
@@ -28,7 +28,7 @@ struct FEmitter::FPrivate : FEmitter {
 	{
 		check(!Type->IsScalar());
 		int Dimensions = Type->NumRows * Type->NumColumns;
-		FDimensional* Value = (FDimensional*)FMemory::Malloc(sizeof(FDimensional) + sizeof(FValuePtr) * Dimensions);
+		FDimensional* Value = (FDimensional*)FMemory::Malloc(sizeof(FDimensional) + sizeof(FValue*) * Dimensions);
 		Value->Kind = VK_Dimensional;
 		Value->Type = Type;
 		return Value;
@@ -38,37 +38,38 @@ struct FEmitter::FPrivate : FEmitter {
 	// If none found, it creates a new value as a copy of the prototype, adds it to
 	// the module then returns it.
 	template <typename TValueType>
-	static FValuePtr EmitPrototype(FEmitter* Emitter, const TValueType& Prototype)
+	static FValue* EmitPrototype(FEmitter* Emitter, const TValueType& Prototype)
 	{
-		if (FValuePtr Existing = FindValue(Emitter, &Prototype))
+		if (FValue* Existing = FindValue(Emitter, &Prototype))
 		{
 			return Existing;
 		}
 
 		TValueType* Value = new TValueType{ Prototype };
-		Emitter->Module->Values.Add(Value);
+		PushNewValue(Emitter, Value);
 		return Value;
 	}
 
 	// Emits specified newly created `Value`. If the exact value already exists,
 	// specified one is *destroyed* and existing one is returned instead.
-	static FValuePtr EmitNew(FEmitter* Emitter, FValuePtr Value)
+	static FValue* EmitNew(FEmitter* Emitter, FValue* Value)
 	{
-		if (FValuePtr Existing = FindValue(Emitter, Value))
+		if (FValue* Existing = FindValue(Emitter, Value))
 		{
 			delete Value;
 			return Existing;
 		}
 
-		Emitter->Module->Values.Add(Value);
+		PushNewValue(Emitter, Value);
 		return Value;
 	}
 
 	// Looks for an existing value in the module that matches `Prototype` and returns it if found.
-	static FValuePtr FindValue(const FEmitter* Emitter, FValuePtr Prototype)
+	static FValue* FindValue(FEmitter* Emitter, const FValue* Prototype)
 	{
 		/* todo: improve this with hashmap */
-		for (FValuePtr CurrValue : Emitter->Module->Values)
+
+		for (FValue* CurrValue : Emitter->Module->Values)
 		{
 			if (CurrValue->Equals(Prototype))
 			{
@@ -78,6 +79,11 @@ struct FEmitter::FPrivate : FEmitter {
 
 		return nullptr;
 	}
+
+	static void PushNewValue(FEmitter* Emitter, FValue* Value)
+	{
+		Emitter->Module->Values.Add(Value);
+	}
 };
 
 FEmitter::FEmitter(FMaterialIRModuleBuilder* InBuilder, UMaterial* InMaterial, FMaterialIRModule* InModule)
@@ -85,29 +91,95 @@ FEmitter::FEmitter(FMaterialIRModuleBuilder* InBuilder, UMaterial* InMaterial, F
 	Builder = InBuilder;
 	Material = InMaterial;
 	Module = InModule;
+}
+
+FValue* FEmitter::Get(const FExpressionInput* Input)
+{
+	FValue** Value = Builder->InputValues.Find(Input);
+	return Value ? *Value : nullptr;
+}
+	
+void FEmitter::Put(const FExpressionOutput* Output, FValue* Value)
+{
+	check(Output);
+	Builder->OutputValues.Add(Output, Value);
+}
+
+FEmitter& FEmitter::DefaultToFloatZero(const FExpressionInput* Input)
+{
+	return DefaultTo(Input, 0.0f);
+}
+
+FEmitter& FEmitter::DefaultTo(const FExpressionInput* Input, float Float)
+{
+	if (!Input->IsConnected())
+	{
+		Builder->InputValues.Add(Input, EmitConstantFloat1(Float));
+	}
+
+	return *this;
+}
+
+FValue* FEmitter::TryGetFloat(const FExpressionInput* Input)
+{
+	FValue* Value = Get(Input);
+	CheckInputIsScalar(Input, Value, SK_Float);
+	return Value;
+}
+
+FValue* FEmitter::TryGetScalar(const FExpressionInput* Input)
+{
+	FValue* Value = Get(Input);
+	CheckInputIsScalar(Input, Value);
+	return Value;
+}
+
+FValue* FEmitter::TryGetArithmetic(const FExpressionInput* Input)
+{
+	FValue* Value = Get(Input);
+	CheckInputTypeIs(Input, Value, TK_Arithmetic);
+	return Value;
+}
+
+FValue* FEmitter::TryGetOfType(const FExpressionInput* Input, ETypeKind Kind)
+{
+	FValue* Value = Get(Input);
+	if (!Value)
+	{
+		return nullptr;
+	}
+
+	CheckInputTypeIs(Input, Value, Kind);
+	return Value;
+}
 
-	ZeroIntValue = EmitConstantScalarZero(SK_Int);
-	ZeroFloatValue = EmitConstantScalarZero(SK_Float);
+void FEmitter::CheckInputIsScalar(const FExpressionInput* Input, FValue* InputValue)
+{
+	FArithmeticTypePtr ArithmeticType = InputValue->Type->AsArithmetic();
+	if (!ArithmeticType || !ArithmeticType->IsScalar())
+	{
+		Errorf(TEXT("Input '%s' expected to be a scalar. It is %s instead."), *Input->InputName.ToString(), InputValue->Type->GetSpelling().GetData());
+	}
 }
 
-FValuePtr FEmitter::GetAndCheckInputTypeKindIs(const FExpressionInput* Input, ETypeKind Kind)
+void FEmitter::CheckInputIsScalar(const FExpressionInput* Input, FValue* InputValue, EScalarKind Kind)
 {
-	FValuePtr InputValue = Get(Input);
-	CheckInputTypeIs(Input, InputValue, Kind);
-	return InputValue;
+	FArithmeticTypePtr ArithmeticType = InputValue->Type->AsArithmetic();
+	if (!ArithmeticType || !ArithmeticType->IsScalar() || ArithmeticType->ScalarKind != Kind)
+	{
+		Errorf(TEXT("Input '%s' expected to be a %s scalar. It is %s instead."), *Input->InputName.ToString(), ScalarKindToString(Kind), InputValue->Type->GetSpelling().GetData());
+	}
 }
 
-bool FEmitter::CheckInputTypeIs(const FExpressionInput* Input, FValuePtr InputValue, ETypeKind Kind)
+void FEmitter::CheckInputTypeIs(const FExpressionInput* Input, FValue* InputValue, ETypeKind Kind)
 {
 	if (InputValue->Type->Kind != Kind)
 	{
-		Errorf(TEXT("Unexpected input '%s' value type (%s). Expected %s type."), *Input->InputName.ToString(), TEXT(""), TEXT(""));
-		return false;
+		Errorf(TEXT("Input '%s' expected to be have type %s. It is %s instead."), *Input->InputName.ToString(), TypeKindToString(Kind), InputValue->Type->GetSpelling().GetData());
 	}
-	return true;
 }
 
-FValuePtr FEmitter::EmitConstantFromShaderValue(const UE::Shader::FValue& InValue)
+FValue* FEmitter::EmitConstantFromShaderValue(const UE::Shader::FValue& InValue)
 {
 	using namespace UE::Shader;
 
@@ -125,104 +197,102 @@ FValuePtr FEmitter::EmitConstantFromShaderValue(const UE::Shader::FValue& InValu
 	}
 
 	UE_MIR_UNREACHABLE();
-}
+} 
 
-FValuePtr FEmitter::EmitConstantScalarZero(EScalarKind Kind)
+FValue* FEmitter::EmitConstantScalarZero(EScalarKind Kind)
 {
 	switch (Kind)
 	{
 		case SK_Bool: UE_MIR_TODO();
-		case SK_Int: return ZeroIntValue;
-		case SK_Float: return ZeroFloatValue;
+		case SK_Int: return EmitConstantInt1(0);
+		case SK_Float: return EmitConstantFloat1(0.0f);
 		default: UE_MIR_UNREACHABLE();
 	}
 }
 
-FValuePtr FEmitter::EmitConstantFloat1(float InX)
+FValue* FEmitter::EmitConstantFloat1(float InX)
 {
 	FScalarConstant Scalar = InitValue<FScalarConstant>(FArithmeticType::GetScalar(SK_Float));
 	Scalar.Float = InX;
-
 	return FPrivate::EmitPrototype(this, Scalar);
 }
 
-FValuePtr FEmitter::EmitConstantFloat2(const FVector2f& InValue)
+FValue* FEmitter::EmitConstantFloat2(const FVector2f& InValue)
 {
-	FValuePtr X = EmitConstantFloat1(InValue.X);
-	FValuePtr Y = EmitConstantFloat1(InValue.Y);
+	FValue* X = EmitConstantFloat1(InValue.X);
+	FValue* Y = EmitConstantFloat1(InValue.Y);
 	return EmitVector2(X, Y);
 }
 
-FValuePtr FEmitter::EmitConstantFloat3(const FVector3f& InValue)
+FValue* FEmitter::EmitConstantFloat3(const FVector3f& InValue)
 {
-	FValuePtr X = EmitConstantFloat1(InValue.X);
-	FValuePtr Y = EmitConstantFloat1(InValue.Y);
-	FValuePtr Z = EmitConstantFloat1(InValue.Z);
+	FValue* X = EmitConstantFloat1(InValue.X);
+	FValue* Y = EmitConstantFloat1(InValue.Y);
+	FValue* Z = EmitConstantFloat1(InValue.Z);
 	return EmitVector3(X, Y, Z);
 }
 
-FValuePtr FEmitter::EmitConstantFloat4(const FVector4f& InValue)
+FValue* FEmitter::EmitConstantFloat4(const FVector4f& InValue)
 {
-	FValuePtr X = EmitConstantFloat1(InValue.X);
-	FValuePtr Y = EmitConstantFloat1(InValue.Y);
-	FValuePtr Z = EmitConstantFloat1(InValue.Z);
-	FValuePtr W = EmitConstantFloat1(InValue.W);
+	FValue* X = EmitConstantFloat1(InValue.X);
+	FValue* Y = EmitConstantFloat1(InValue.Y);
+	FValue* Z = EmitConstantFloat1(InValue.Z);
+	FValue* W = EmitConstantFloat1(InValue.W);
 	return EmitVector4(X, Y, Z, W);
 }
 
-FValuePtr FEmitter::EmitConstantInt1(int InX)
+FValue* FEmitter::EmitConstantInt1(int InX)
 {
 	FScalarConstant Scalar = InitValue<FScalarConstant>(FArithmeticType::GetScalar(SK_Int));
 	Scalar.Integer = InX;
-
 	return FPrivate::EmitPrototype(this, Scalar);
 }
 
-FValuePtr FEmitter::EmitConstantInt2(const FIntVector2& InValue)
+FValue* FEmitter::EmitConstantInt2(const FIntVector2& InValue)
 {
-	FValuePtr X = EmitConstantInt1(InValue.X);
-	FValuePtr Y = EmitConstantInt1(InValue.Y);
+	FValue* X = EmitConstantInt1(InValue.X);
+	FValue* Y = EmitConstantInt1(InValue.Y);
 	return EmitVector2(X, Y);
 }
 
-FValuePtr FEmitter::EmitConstantInt3(const FIntVector3& InValue)
+FValue* FEmitter::EmitConstantInt3(const FIntVector3& InValue)
 {
-	FValuePtr X = EmitConstantInt1(InValue.X);
-	FValuePtr Y = EmitConstantInt1(InValue.Y);
-	FValuePtr Z = EmitConstantInt1(InValue.Z);
+	FValue* X = EmitConstantInt1(InValue.X);
+	FValue* Y = EmitConstantInt1(InValue.Y);
+	FValue* Z = EmitConstantInt1(InValue.Z);
 	return EmitVector3(X, Y, Z);
 }
 
-FValuePtr FEmitter::EmitConstantInt4(const FIntVector4& InValue)
+FValue* FEmitter::EmitConstantInt4(const FIntVector4& InValue)
 {
-	FValuePtr X = EmitConstantInt1(InValue.X);
-	FValuePtr Y = EmitConstantInt1(InValue.Y);
-	FValuePtr Z = EmitConstantInt1(InValue.Z);
-	FValuePtr W = EmitConstantInt1(InValue.W);
+	FValue* X = EmitConstantInt1(InValue.X);
+	FValue* Y = EmitConstantInt1(InValue.Y);
+	FValue* Z = EmitConstantInt1(InValue.Z);
+	FValue* W = EmitConstantInt1(InValue.W);
 	return EmitVector4(X, Y, Z, W);
 }
 
-FValuePtr FEmitter::EmitVector2(FValuePtr InX, FValuePtr InY)
+FValue* FEmitter::EmitVector2(FValue* InX, FValue* InY)
 {
-	check(InX->Type->ToScalar());
+	check(InX->Type->AsScalar());
 	check(InX->Type == InY->Type);
 
-	TDimensional<2> Vector = InitValue<TDimensional<2>>(FArithmeticType::GetVector(InX->Type->ToArithmetic()->ScalarKind, 2));
-	TArrayView<FValuePtr> Components = Vector.GetMutableComponents();
+	TDimensional<2> Vector = InitValue<TDimensional<2>>(FArithmeticType::GetVector(InX->Type->AsArithmetic()->ScalarKind, 2));
+	TArrayView<FValue*> Components = Vector.GetMutableComponents();
 	Components[0] = InX;
 	Components[1] = InY;
 
 	return FPrivate::EmitPrototype(this, Vector);
 }
 
-FValuePtr FEmitter::EmitVector3(FValuePtr InX, FValuePtr InY, FValuePtr InZ)
+FValue* FEmitter::EmitVector3(FValue* InX, FValue* InY, FValue* InZ)
 {
-	check(InX->Type->ToScalar());
+	check(InX->Type->AsScalar());
 	check(InX->Type == InY->Type);
 	check(InY->Type == InZ->Type);
 
-	TDimensional<3> Vector = InitValue<TDimensional<3>>(FArithmeticType::GetVector(InX->Type->ToArithmetic()->ScalarKind, 3));
-	TArrayView<FValuePtr> Components = Vector.GetMutableComponents();
+	TDimensional<3> Vector = InitValue<TDimensional<3>>(FArithmeticType::GetVector(InX->Type->AsArithmetic()->ScalarKind, 3));
+	TArrayView<FValue*> Components = Vector.GetMutableComponents();
 	Components[0] = InX;
 	Components[1] = InY;
 	Components[2] = InZ;
@@ -230,15 +300,15 @@ FValuePtr FEmitter::EmitVector3(FValuePtr InX, FValuePtr InY, FValuePtr InZ)
 	return FPrivate::EmitPrototype(this, Vector);
 }
 
-FValuePtr FEmitter::EmitVector4(FValuePtr InX, FValuePtr InY, FValuePtr InZ, FValuePtr InW)
+FValue* FEmitter::EmitVector4(FValue* InX, FValue* InY, FValue* InZ, FValue* InW)
 {
-	check(InX->Type->ToScalar());
+	check(InX->Type->AsScalar());
 	check(InX->Type == InY->Type);
 	check(InY->Type == InZ->Type);
 	check(InZ->Type == InW->Type);
 
-	TDimensional<4> Vector = InitValue<TDimensional<4>>(FArithmeticType::GetVector(InX->Type->ToArithmetic()->ScalarKind, 4));
-	TArrayView<FValuePtr> Components = Vector.GetMutableComponents();
+	TDimensional<4> Vector = InitValue<TDimensional<4>>(FArithmeticType::GetVector(InX->Type->AsArithmetic()->ScalarKind, 4));
+	TArrayView<FValue*> Components = Vector.GetMutableComponents();
 	Components[0] = InX;
 	Components[1] = InY;
 	Components[2] = InZ;
@@ -247,9 +317,9 @@ FValuePtr FEmitter::EmitVector4(FValuePtr InX, FValuePtr InY, FValuePtr InZ, FVa
 	return FPrivate::EmitPrototype(this, Vector);
 }
 
-FValuePtr FEmitter::EmitArithmetic(FArithmeticTypePtr Type, FValuePtr Scalar)
+FValue* FEmitter::EmitArithmetic(FArithmeticTypePtr Type, FValue* Scalar)
 {
-	check(Scalar->Type->ToScalar());
+	check(Scalar->Type->AsScalar());
 	if (Type->IsScalar())
 	{
 		check(Type == Scalar->Type);
@@ -257,7 +327,7 @@ FValuePtr FEmitter::EmitArithmetic(FArithmeticTypePtr Type, FValuePtr Scalar)
 	}
 
 	FDimensional* Value = FPrivate::NewDimensionalValue(this, Type);
-	for (FValuePtr& Component : Value->GetMutableComponents())
+	for (FValue*& Component : Value->GetMutableComponents())
 	{
 		Component = Scalar;
 	}
@@ -269,8 +339,9 @@ FSetMaterialOutput* FEmitter::EmitSetMaterialOutput(EMaterialProperty InProperty
 {
 	FSetMaterialOutput* Instr = new FSetMaterialOutput;
 	Instr->Kind 	= VK_SetMaterialOutput;
+	Instr->Block 	= Module->RootBlock;
 	Instr->Property = InProperty;
-	Instr->ArgValue = InArgValue;
+	Instr->Arg 		= InArgValue;
 
 	Module->Values.Add(Instr);
 	Module->Outputs.Add(Instr);
@@ -288,27 +359,46 @@ static FTypePtr GetBinaryOperatorResultType(EBinaryOperator Operator, FTypePtr L
 			check(LhsType == RhsType);
 			return LhsType;
 		
+		case BO_Greater:
+		case BO_Lower:
+		case BO_Equals:
+			return FArithmeticType::GetBool();
+
 		default:
 			UE_MIR_UNREACHABLE();
 	}
 }
 
-FValuePtr FEmitter::EmitBinaryOperator(EBinaryOperator Operator, FValuePtr Lhs, FValuePtr Rhs)
+FValue* FEmitter::EmitBinaryOperator(EBinaryOperator Operator, FValue* Lhs, FValue* Rhs)
 {
 	check(Lhs->Type == Rhs->Type);
 
 	/* todo: operation folding */
 
 	FTypePtr ResultType = GetBinaryOperatorResultType(Operator, Lhs->Type, Rhs->Type);
-	FBinaryOperator Prototype = InitValue<FBinaryOperator>(ResultType);
-	Prototype.Operator = Operator;
-	Prototype.Lhs= Lhs;
-	Prototype.Rhs = Rhs;
+	FBinaryOperator Proto = InitValue<FBinaryOperator>(ResultType);
+	Proto.Operator = Operator;
+	Proto.LhsArg= Lhs;
+	Proto.RhsArg = Rhs;
 
-	return FPrivate::EmitPrototype(this, Prototype);
+	return FPrivate::EmitPrototype(this, Proto);
 }
 
-static FValuePtr ConvertScalarValue(FEmitter* Emitter, FValuePtr ScalarValue, EScalarKind ScalarKind, EScalarKind TargetKind)
+FValue* FEmitter::EmitBranch(FValue* Condition, FValue* True, FValue* False)
+{
+	/* todo: operation folding */
+
+	check(True->Type == False->Type);
+
+	FBranch Proto = InitValue<FBranch>(True->Type);
+	Proto.ConditionArg = Condition;
+	Proto.TrueArg = True;
+	Proto.FalseArg = False;
+
+	return FPrivate::EmitPrototype(this, Proto);
+}
+
+static FValue* ConvertScalarValue(FEmitter* Emitter, FValue* ScalarValue, EScalarKind ScalarKind, EScalarKind TargetKind)
 {
 	if (ScalarKind == TargetKind)
 	{
@@ -351,12 +441,12 @@ static FValuePtr ConvertScalarValue(FEmitter* Emitter, FValuePtr ScalarValue, ES
 	UE_MIR_UNREACHABLE();
 }
 
-static FValuePtr ConvertArithmeticValue(FEmitter* Emitter, FValuePtr Value, FArithmeticTypePtr TargetArithmeticType)
+static FValue* ConvertArithmeticValue(FEmitter* Emitter, FValue* Value, FArithmeticTypePtr TargetArithmeticType)
 {
 	FArithmeticTypePtr ValueArithmeticType = static_cast<FArithmeticTypePtr>(Value->Type);
 	if (ValueArithmeticType->IsScalar())
 	{
-		FValuePtr ConvertedScalar = ConvertScalarValue(Emitter, Value, ValueArithmeticType->ScalarKind, TargetArithmeticType->ScalarKind);
+		FValue* ConvertedScalar = ConvertScalarValue(Emitter, Value, ValueArithmeticType->ScalarKind, TargetArithmeticType->ScalarKind);
 		return Emitter->EmitArithmetic(TargetArithmeticType, ConvertedScalar);
 	}
 	else
@@ -365,7 +455,7 @@ static FValuePtr ConvertArithmeticValue(FEmitter* Emitter, FValuePtr Value, FAri
 	}
 }
 
-FValuePtr FEmitter::TryEmitConvert(FValuePtr Value, FTypePtr TargetType)
+FValue* FEmitter::TryEmitConvert(FValue* Value, FTypePtr TargetType)
 {
 	// If types already match, nothing else to do, simply return the same value.
 	FTypePtr ValueType = Value->Type;
@@ -374,11 +464,11 @@ FValuePtr FEmitter::TryEmitConvert(FValuePtr Value, FTypePtr TargetType)
 		return Value;
 	}
 
-	FValuePtr ResultValue{};
+	FValue* ResultValue{};
 
 	// Converting from an arithmetic value to another arithmetic type.
-	FArithmeticTypePtr ValueArithmeticType = Value->Type->ToArithmetic();
-	FArithmeticTypePtr TargetArithmeticType = TargetType->ToArithmetic();
+	FArithmeticTypePtr ValueArithmeticType = Value->Type->AsArithmetic();
+	FArithmeticTypePtr TargetArithmeticType = TargetType->AsArithmetic();
 	if (ValueArithmeticType && TargetArithmeticType)
 	{
 		ResultValue = ConvertArithmeticValue(this, Value, TargetArithmeticType);
@@ -395,17 +485,26 @@ FValuePtr FEmitter::TryEmitConvert(FValuePtr Value, FTypePtr TargetType)
 	return ResultValue;
 }
 
-FValuePtr FEmitter::Get(const FExpressionInput* Input)
+FArithmeticTypePtr FEmitter::TryGetCommonArithmeticType(FArithmeticTypePtr A, FArithmeticTypePtr B)
 {
-	check(Input);
-	FValuePtr* Value = Builder->InputValues.Find(Input);
-	return *Value;
-}
+	// Trivial case: types are equal
+	if (A == B)
+	{
+		return A;
+	}
 	
-void FEmitter::Put(const FExpressionOutput* Output, FValuePtr Value)
-{
-	check(Output);
-	Builder->OutputValues.Add(Output, Value);
+	//	
+	if (A->IsMatrix() != B->IsMatrix())
+	{
+		Errorf(TEXT("No common arithmetic type between `%s` and `%s`."), A->Spelling.GetData(), B->Spelling.GetData());
+		return nullptr;
+	}
+
+	EScalarKind ScalarKind = FMath::Max(A->ScalarKind, B->ScalarKind);
+	int NumRows = FMath::Max(A->NumRows, B->NumRows);
+	int NumColumns = FMath::Max(A->NumColumns, B->NumColumns);
+
+	return FArithmeticType::Get(ScalarKind, NumRows, NumColumns);
 }
 
 void FEmitter::Error(FString Message)
diff --git a/Engine/Source/Runtime/Engine/Private/Materials/MaterialIRModule.cpp b/Engine/Source/Runtime/Engine/Private/Materials/MaterialIRModule.cpp
--- a/Engine/Source/Runtime/Engine/Private/Materials/MaterialIRModule.cpp
+++ b/Engine/Source/Runtime/Engine/Private/Materials/MaterialIRModule.cpp
@@ -7,14 +7,23 @@
 
 namespace IR = UE::MIR;
 
+FMaterialIRModule::FMaterialIRModule()
+{
+	RootBlock = new IR::FBlock;
+}
+
 FMaterialIRModule::~FMaterialIRModule()
 {
 	Empty();
+
+	delete RootBlock;
 }
 
 void FMaterialIRModule::Empty()
 {
-	for (IR::FValuePtr Value : Values)
+	RootBlock->Instructions = nullptr;
+
+	for (IR::FValue* Value : Values)
 	{
 		delete Value;
 	}
diff --git a/Engine/Source/Runtime/Engine/Private/Materials/MaterialIRModuleBuilder.cpp b/Engine/Source/Runtime/Engine/Private/Materials/MaterialIRModuleBuilder.cpp
--- a/Engine/Source/Runtime/Engine/Private/Materials/MaterialIRModuleBuilder.cpp
+++ b/Engine/Source/Runtime/Engine/Private/Materials/MaterialIRModuleBuilder.cpp
@@ -9,22 +9,30 @@
 #include "Materials/MaterialIRTypes.h"
 #include "Materials/MaterialIREmitter.h"
 #include "MaterialIRUtility.h"
-
 #include "Materials/MaterialAttributeDefinitionMap.h"
 #include "Materials/Material.h"
 #include "MaterialExpressionIO.h"
 #include "MaterialShared.h"
 #include "Materials/MaterialExpression.h"
+#include "Async/ParallelFor.h"
 
 namespace IR = UE::MIR;
 
 struct FMaterialIRModuleBuilder::FPrivate
 {
-	static void Build_GenerateOutputInstructions(FMaterialIRModuleBuilder& Builder, UE::MIR::FEmitter& Emitter)
+	FMaterialIRModuleBuilder& Builder;
+	UMaterial& Material;
+	FMaterialIRTargetParams Params;
+	FMaterialIRModule& Module;
+	IR::FEmitter& Emitter;
+	TArray<UMaterialExpression*> ExpressionAnalysisStack;
+	TArray<IR::FInstruction*> InstructionStack;
+
+	void Build_GenerateOutputInstructions()
 	{
 		// Prepare the array of FSetMaterialOutputInstr outputs from the material attributes inputs.
 		FMaterialInputDescription Input;
-		for (int32 Index = 0; UE::Utility::NextMaterialAttributeInput(Builder.BaseMaterial, Index, Input); ++Index)
+		for (int32 Index = 0; UE::Utility::NextMaterialAttributeInput(&Material, Index, Input); ++Index)
 		{
 			EMaterialProperty Property = (EMaterialProperty)Index;
 
@@ -32,31 +40,31 @@ struct FMaterialIRModuleBuilder::FPrivate
 
 			if (Input.bUseConstant)
 			{
-				Output->ArgValue = Emitter.EmitConstantFromShaderValue(Input.ConstantValue);
+				Output->Arg = Emitter.EmitConstantFromShaderValue(Input.ConstantValue);
 			}
 			else if (!Input.Input->IsConnected())
 			{
-				Output->ArgValue = UE::Utility::CreateMaterialAttributeDefaultValue(Emitter, Builder.BaseMaterial, Property);
+				Output->Arg = UE::Utility::CreateMaterialAttributeDefaultValue(Emitter, &Material, Property);
 			}
 			else
 			{
-				Builder.ExpressionAnalysisStack.Add(Input.Input->Expression);
+				ExpressionAnalysisStack.Add(Input.Input->Expression);
 			}
 		}
 	}
 
-	static void Build_AnalyzeExpressionGraph(FMaterialIRModuleBuilder& Builder, UE::MIR::FEmitter & Emitter)
+	void Build_AnalyzeExpressionGraph()
 	{
 		TSet<UMaterialExpression*> BuiltExpressions;
 
-		while (!Builder.ExpressionAnalysisStack.IsEmpty())
+		while (!ExpressionAnalysisStack.IsEmpty())
 		{
-			Emitter.Expression = Builder.ExpressionAnalysisStack.Last();
+			Emitter.Expression = ExpressionAnalysisStack.Last();
 
 			// If expression is clean, nothing to be done.
 			if (BuiltExpressions.Contains(Emitter.Expression))
 			{
-				Builder.ExpressionAnalysisStack.Pop(EAllowShrinking::No);
+				ExpressionAnalysisStack.Pop(EAllowShrinking::No);
 				continue;
 			}
 
@@ -69,16 +77,16 @@ struct FMaterialIRModuleBuilder::FPrivate
 					continue;
 				}
 
-				Builder.ExpressionAnalysisStack.Push(It->Expression);
+				ExpressionAnalysisStack.Push(It->Expression);
 			}
 
 			// If on top of the stack there's a different expression, we have a dependency to analyze first.
-			if (Builder.ExpressionAnalysisStack.Last() != Emitter.Expression) {
+			if (ExpressionAnalysisStack.Last() != Emitter.Expression) {
 				continue;
 			}
 
 			// 
-			Builder.ExpressionAnalysisStack.Pop();
+			ExpressionAnalysisStack.Pop();
 			BuiltExpressions.Add(Emitter.Expression);
 
 			//
@@ -87,7 +95,7 @@ struct FMaterialIRModuleBuilder::FPrivate
 				if (FExpressionOutput* ConnectedOutput = It->GetConnectedOutput())
 				{
 					// Fetch the value flowing through connected output.
-					IR::FValuePtr* ValuePtr = Builder.OutputValues.Find(ConnectedOutput);
+					IR::FValue** ValuePtr = Builder.OutputValues.Find(ConnectedOutput);
 					check(ValuePtr && TEXT("Output value not found."));
 
 					// Set the value flowing into this input.
@@ -104,35 +112,150 @@ struct FMaterialIRModuleBuilder::FPrivate
 		}
 	}
 
-	static void Build_LinkMaterialOutputsToIncomingValues(FMaterialIRModuleBuilder& Builder, UE::MIR::FEmitter& Emitter)
+	void Build_LinkMaterialOutputsToIncomingValues()
 	{
 		FMaterialInputDescription Input;
-		for (IR::FSetMaterialOutput* Output : Builder.Module->Outputs)
+		for (IR::FSetMaterialOutput* Output : Module.Outputs)
 		{
-			if (Output->ArgValue || !ensure(Builder.BaseMaterial->GetExpressionInputDescription(Output->Property, Input))) {
+			if (Output->Arg || !ensure(Material.GetExpressionInputDescription(Output->Property, Input))) {
 				continue;
 			}
 
-			IR::FValuePtr* ValuePtr = Builder.OutputValues.Find(Input.Input->Expression->GetOutput(0));
+			IR::FValue** ValuePtr = Builder.OutputValues.Find(Input.Input->Expression->GetOutput(0));
 			check(ValuePtr);
 
-			Output->ArgValue = *ValuePtr;
+			Output->Arg = *ValuePtr;
+		}
+	}
+
+	void Build_FinalizeValueGraph()
+	{
+		InstructionStack.Reserve(64);
+
+		for (IR::FSetMaterialOutput* Output : Module.Outputs)
+		{
+			InstructionStack.Push(Output);
+		}
+
+		while (!InstructionStack.IsEmpty())
+		{
+			IR::FValue* Instr = InstructionStack.Pop();
+			for (IR::FValue* UseValue : Instr->GetUses())
+			{
+				IR::FInstruction* Use = UseValue->AsInstruction();
+				if (!Use)
+				{
+					continue;
+				}
+
+				Use->NumUsers += 1;
+
+				if (!(Use->Flags & IR::IF_Counted))
+				{
+					Use->SetFlags(IR::IF_Counted);
+					InstructionStack.Push(Use);
+				}
+			}
 		}
 	}
+
+	void Build_PopulateBlock()
+	{
+		// This function walks the instruction graph and puts each instruction into the inner most possible block.
+		InstructionStack.Empty(InstructionStack.Max());
+
+		for (IR::FSetMaterialOutput* Output : Module.Outputs)
+		{
+			Output->Block = Module.RootBlock;
+			InstructionStack.Add(Output);
+		}
+
+		while (!InstructionStack.IsEmpty()) {
+			IR::FInstruction* Instr = InstructionStack.Pop();
+
+			// Push the instruction in its block
+			Instr->Next = Instr->Block->Instructions;
+			Instr->Block->Instructions = Instr;
+
+			IR::FValue* UseValue;
+			IR::FBlock* InnerBlock;
+			for (int32 i = 0; Instr->GetInnerBlock(i, UseValue, InnerBlock); ++i)
+			{
+				IR::FInstruction* Use = UseValue->AsInstruction();
+				if (!Use)
+				{
+					continue;
+				}
+
+				// Update dependency's block to be a child of current instruction's block.
+				if (InnerBlock != Instr->Block)
+				{
+					InnerBlock->Parent = Instr->Block;
+					InnerBlock->Level = Instr->Block->Level + 1;
+				}
+
+				// Set the dependency's block to the common block betwen its current block and this one.
+				Use->Block = Use->Block
+					? FindCommonParentBlock(Use->Block, InnerBlock)
+					: InnerBlock;
+
+				// Increase the number of times this dependency instruction has been considered.
+				// When all of its users have processed, we can carry on visiting this instruction.
+				++Use->NumProcessedUsers;
+				check(Use->NumProcessedUsers <= Use->NumUsers);
+
+				// If all dependants have been processed, we can carry the processing from this dependency.
+				if (Use->NumProcessedUsers == Use->NumUsers)
+				{
+					InstructionStack.Push(Use);
+				}
+			}
+		}
+	}
+	
+    static IR::FBlock* FindCommonParentBlock(IR::FBlock* A, IR::FBlock* B)
+    {
+        if (A == B) {
+            return A;
+        }
+
+        while (A->Level > B->Level) {
+            A = A->Parent;
+        }
+
+        while (B->Level > A->Level) {
+            B = B->Parent;
+        }
+
+        while (A != B) {
+            A = A->Parent;
+            B = B->Parent;
+        }
+
+        return A;
+    }
 };
 
-bool FMaterialIRModuleBuilder::Build(UMaterial* InMaterial, FMaterialIRModule* TargetModule)
+bool FMaterialIRModuleBuilder::Build(UMaterial* InMaterial, const FMaterialIRTargetParams& TargetParams, FMaterialIRModule* TargetModule)
 {
-	BaseMaterial = InMaterial;
+	TargetModule->Empty();
+	TargetModule->ShaderPlatform = TargetParams.ShaderPlatform;
+
+	IR::FEmitter Emitter{ this, InMaterial, TargetModule };
 
-	Module = TargetModule;
-	Module->Empty();
+	FPrivate Private{ *this, *InMaterial, TargetParams, *TargetModule, Emitter };
 
-	IR::FEmitter Emitter{ this, BaseMaterial, Module };
+	Private.Build_GenerateOutputInstructions();
+	Private.Build_AnalyzeExpressionGraph();
+
+	if (Private.Emitter.IsInvalid())
+	{
+		return false;
+	}
 
-	FPrivate::Build_GenerateOutputInstructions(*this, Emitter);
-	FPrivate::Build_AnalyzeExpressionGraph(*this, Emitter);
-	FPrivate::Build_LinkMaterialOutputsToIncomingValues(*this, Emitter);
+	Private.Build_LinkMaterialOutputsToIncomingValues();
+	Private.Build_FinalizeValueGraph();
+	Private.Build_PopulateBlock();
 
 	return true;
 }
diff --git a/Engine/Source/Runtime/Engine/Private/Materials/MaterialIRToHLSLTranslator.cpp b/Engine/Source/Runtime/Engine/Private/Materials/MaterialIRToHLSLTranslator.cpp
--- a/Engine/Source/Runtime/Engine/Private/Materials/MaterialIRToHLSLTranslator.cpp
+++ b/Engine/Source/Runtime/Engine/Private/Materials/MaterialIRToHLSLTranslator.cpp
@@ -19,10 +19,18 @@
 
 namespace IR = UE::MIR;
 
+enum ENoOp { NoOp };
+enum ENewLine { NewLine };
+enum EEndOfStatement { EndOfStatement };
+enum EOpenBrace { OpenBrace };
+enum ECloseBrace { CloseBrace };
+enum EIndentation { Indentation };
+
 struct FHLSLPrinter
 {
-	FString& Buffer;
+	FString Buffer;
 	bool bFirstListItem = false;
+	int32 Tabs = 0;
 
 	template <int N, typename... Types>
 	void Printf(const TCHAR (&Format)[N], Types... Args)
@@ -32,28 +40,78 @@ struct FHLSLPrinter
 
 	FHLSLPrinter& operator<<(const TCHAR* Text)
 	{
- 		Buffer.Append(Text);
+		Buffer.Append(Text);
+		return *this;
+	}
+	
+	FHLSLPrinter& operator<<(const FString& Text)
+	{
+		Buffer.Append(Text);
 		return *this;
 	}
 
 	FHLSLPrinter& operator<<(int32 Value)
 	{
- 		Buffer.Appendf(TEXT("%d"), Value);
+		Buffer.Appendf(TEXT("%d"), Value);
 		return *this;
 	}
 
 	FHLSLPrinter& operator<<(uint32 Value)
 	{
- 		Buffer.Appendf(TEXT("%u"), Value);
+		Buffer.Appendf(TEXT("%u"), Value);
 		return *this;
 	}
 
 	FHLSLPrinter& operator<<(float Value)
 	{
- 		Buffer.Appendf(TEXT("%.5ff"), Value);
+		Buffer.Appendf(TEXT("%.5ff"), Value);
+		return *this;
+	}
+
+    FHLSLPrinter& operator<<(ENoOp)
+	{
+		return *this;
+	}
+
+    FHLSLPrinter& operator<<(ENewLine)
+    {
+		Buffer.AppendChar('\n');
+		operator<<(Indentation);
+        return *this;
+    }
+
+	FHLSLPrinter& operator<<(EIndentation)
+	{
+		for (int i = 0; i < Tabs; ++i)
+		{
+			Buffer.AppendChar('\t');
+		}
 		return *this;
 	}
 
+	FHLSLPrinter& operator<<(EEndOfStatement)
+	{
+		Buffer.AppendChar(';');
+        *this << NewLine;
+        return *this;
+	}
+
+    FHLSLPrinter& operator<<(EOpenBrace)
+    {
+        Buffer.Append("{");
+        ++Tabs;
+        *this << NewLine;
+        return *this;
+    }
+
+    FHLSLPrinter& operator<<(ECloseBrace)
+    {
+        --Tabs;
+        Buffer.LeftChopInline(1); // undo tab
+        Buffer.AppendChar('}');
+        return *this;
+    }
+
 	void BeginList()
 	{
 		bFirstListItem = true;
@@ -120,25 +178,35 @@ static const TCHAR* GetShadingModelParameterName(EMaterialShadingModel InModel)
 	}
 }
 
-struct FMaterialIRToHLSLTranslator::FPrivate
+static bool IsFoldable(const IR::FInstruction* Instr)
+{
+	if (auto Branch = Instr->As<IR::FBranch>())
+	{
+		return !Branch->TrueBlock.Instructions && !Branch->FalseBlock.Instructions;
+	}
+
+	return true;
+}
+
+struct FTranslator : FMaterialIRToHLSLTranslation
 {
-	FMaterialIRToHLSLTranslator* Translator;
+	int32 NumLocals{};
+	TMap<const IR::FInstruction*, FString> LocalToIdentifier;
+	FHLSLPrinter Printer;
+	FString PixelAttributesHLSL;
+	FString EvaluateOtherMaterialAttributesHLSL;
 
-	void GenerateHLSL(const FMaterial& InMaterial,
-					  const FMaterialIRModule& InModule,
-					  FParametersMap& InParams)
+	void GenerateHLSL()
 	{
-		FString PixelAttributes;
-		FString EvaluateOtherMaterialAttributesHLSL;
-		FHLSLPrinter Printer{ EvaluateOtherMaterialAttributesHLSL };
+		Printer.Tabs = 1;
+		Printer << Indentation;
 
-		for (const IR::FSetMaterialOutput* Output : InModule.GetOutputs())
-		{
-			LowerValue(Printer, Output);
-		}
+		LowerBlock(Module->GetRootBlock());
+
+		Printer << TEXT("PixelMaterialInputs.FrontMaterial = GetInitialisedSubstrateData()") << EndOfStatement;
+		Printer << TEXT("PixelMaterialInputs.Subsurface = 0") << EndOfStatement;
 
-		EvaluateOtherMaterialAttributesHLSL.Append(TEXT("\tPixelMaterialInputs.FrontMaterial = GetInitialisedSubstrateData();\n"));
-		EvaluateOtherMaterialAttributesHLSL.Append(TEXT("\tPixelMaterialInputs.Subsurface = 0;\n"));
+		EvaluateOtherMaterialAttributesHLSL = MoveTemp(Printer.Buffer);
 
 		for (int32 PropertyIndex = 0; PropertyIndex < MP_MAX; ++PropertyIndex)
 		{
@@ -155,67 +223,254 @@ struct FMaterialIRToHLSLTranslator::FPrivate
 			EMaterialValueType Type = (Property == MP_SubsurfaceColor) ? MCT_Float4 : FMaterialAttributeDefinitionMap::GetValueType(Property);
 			check(PropertyName.Len() > 0);
 
-			PixelAttributes.Appendf(TEXT("\t%s %s;\n"), GetHLSLTypeString(Type), *PropertyName);
+			PixelAttributesHLSL.Appendf(TEXT("\t%s %s;\n"), GetHLSLTypeString(Type), *PropertyName);
 		}
+	}
 	
-		InParams.Add(TEXT("pixel_material_inputs"), MoveTemp(PixelAttributes));
+	ENoOp LowerBlock(const IR::FBlock& Block)
+	{
+		int OldNumLocals = NumLocals;
+        for (IR::FInstruction* Instr = Block.Instructions; Instr; Instr = Instr->Next)
+		{
+            if (Instr->NumUsers == 1 && IsFoldable(Instr))
+			{
+                continue;
+            }
+            
+			if (Instr->NumUsers >= 1)
+			{
+                FString LocalStr = FString::Printf(TEXT("l%d"), NumLocals);
+                ++NumLocals;
 
-		InParams.Add(TEXT("calc_pixel_material_inputs_normal"), TEXT("\tPixelMaterialInputs.Normal = MaterialFloat3(0.00000000, 0.00000000, 1.00000000);"));
-		InParams.Add(TEXT("calc_pixel_material_inputs_other_inputs"), EvaluateOtherMaterialAttributesHLSL);
+                Printer << InlineType(Instr->Type) << TEXT(" ") << LocalStr;
 
-		InParams.Add(TEXT("calc_pixel_material_inputs_analytic_derivatives_normal"), TEXT("\tPixelMaterialInputs.Normal = MaterialFloat3(0.00000000, 0.00000000, 1.00000000);"));
-		InParams.Add(TEXT("calc_pixel_material_inputs_analytic_derivatives_other_inputs"), MoveTemp(EvaluateOtherMaterialAttributesHLSL));
-	
-		InParams.Add(TEXT("material_declarations"), TEXT("struct FMaterialAttributes {};"));
-	
-		auto SetParamInt = [&] (const TCHAR* InParamName, int InValue)
+                LocalToIdentifier.Add(Instr, MoveTemp(LocalStr));
+                if (IsFoldable(Instr))
+				{
+                    Printer << TEXT(" = ";)
+                }
+            }
+
+			LowerInstruction(Instr);
+
+            if (Printer.Buffer.EndsWith(TEXT("}")))
+			{
+                Printer << NewLine;
+            }
+            else
+			{
+                Printer << EndOfStatement;
+            }
+        }
+
+        NumLocals = OldNumLocals;
+
+		return NoOp;
+	}
+
+	ENoOp LowerInstruction(const IR::FInstruction* Instr)
+	{
+		switch (Instr->Kind)
 		{
-			InParams.Add(InParamName, FString::Printf(TEXT("%d"), InValue));
-		};
+			case IR::VK_SetMaterialOutput:
+			{
+				auto Output = static_cast<const IR::FSetMaterialOutput*>(Instr);
 
-		SetParamInt(TEXT("num_material_texcoords_vertex"), 0);
-		SetParamInt(TEXT("num_material_texcoords"), 0);
-		SetParamInt(TEXT("num_custom_vertex_interpolators"), 0);
-		SetParamInt(TEXT("num_tex_coord_interpolators"), 0);
+				// Special case MP_SubsurfaceColor as the actual property is a combination of the color and the profile but we don't want to expose the profile
+				const FString& PropertyName = (Output->Property == MP_SubsurfaceColor) ? "Subsurface" : FMaterialAttributeDefinitionMap::GetAttributeName(Output->Property);
+
+				Printer << TEXT("PixelMaterialInputs.") << PropertyName << TEXT(" = ") << InlineValue(Output->Arg);
+
+				break;
+			}
+
+			case IR::VK_BinaryOperator:
+			{
+				auto BinaryOperator = static_cast<const IR::FBinaryOperator*>(Instr);
+
+				InlineValue(BinaryOperator->LhsArg);
+
+				const TCHAR* OpString;
+				switch (BinaryOperator->Operator)
+				{
+					case IR::BO_Add: OpString = TEXT(" + "); break;
+					case IR::BO_Subtract: OpString = TEXT(" - "); break;
+					case IR::BO_Multiply: OpString = TEXT(" * "); break;
+					case IR::BO_Divide: OpString = TEXT(" / "); break;
+					case IR::BO_Greater: OpString = TEXT(" > "); break;
+					case IR::BO_Lower: OpString = TEXT(" < "); break;
+					case IR::BO_Equals: OpString = TEXT(" == "); break;
+					default: UE_MIR_UNREACHABLE();
+				}
+				Printer << OpString;
+
+				InlineValue(BinaryOperator->RhsArg);
+
+				break;
+			}
+
+			case IR::VK_Branch:
+			{
+				auto Branch = static_cast<const IR::FBranch*>(Instr);
+
+				if (IsFoldable(Branch))
+				{
+					Printer << InlineValue(Branch->ConditionArg)
+					<< TEXT(" ? ") << InlineValue(Branch->TrueArg)
+					<< TEXT(" : ") << InlineValue(Branch->FalseArg);
+				}
+				else
+				{
+					Printer << EndOfStatement;
+					Printer << TEXT("if (") << InlineValue(Branch->ConditionArg) << TEXT(")") << NewLine << OpenBrace;
+					Printer << LowerBlock(Branch->TrueBlock);
+					Printer << LocalToIdentifier[Instr] << " = " << InlineValue(Branch->TrueArg) << EndOfStatement;
+					Printer << CloseBrace << NewLine;
+					Printer << TEXT("else") << NewLine << OpenBrace;
+					Printer << LowerBlock(Branch->FalseBlock);
+					Printer << LocalToIdentifier[Instr] << " = " << InlineValue(Branch->FalseArg) << EndOfStatement;
+					Printer << CloseBrace;
+				}
+				break;
+			}
+
+			default:
+				UE_MIR_UNREACHABLE();
+		}
+
+		return NoOp;
+	}
+
+	ENoOp InlineValue(const IR::FValue* InValue)
+	{
+		if (const IR::FInstruction* Instr = InValue->AsInstruction())
+		{
+			if (Instr->NumUsers <= 1 && IsFoldable(Instr))
+			{
+                LowerInstruction(Instr);
+            }
+            else
+			{
+                Printer << LocalToIdentifier[Instr];
+            }
+		}
+		else if (const IR::FScalarConstant* Scalar = InValue->As<IR::FScalarConstant>())
+		{
+			IR::FArithmeticTypePtr ArithType = Scalar->Type->AsArithmetic();
+			check(ArithType && ArithType->IsScalar());
+
+			switch (ArithType->ScalarKind)
+			{
+				case IR::SK_Bool: Printer << Scalar->Boolean; break;
+				case IR::SK_Int: Printer << Scalar->Integer; break;
+				case IR::SK_Float: Printer << Scalar->Float; break;
+			}
+		}
+		else if (const IR::FDimensional* Vector = InValue->As<IR::FDimensional>())
+		{
+			IR::FArithmeticTypePtr ArithType = Vector->Type->AsArithmetic();
+			check(ArithType && ArithType->IsVector());
+
+			Printer << ScalarKindToString(ArithType->ScalarKind) << ArithType->NumRows << TEXT("(");
+
+			Printer.BeginList();
+			for (IR::FValue* Component : Vector->GetComponents())
+			{
+				Printer.PrintListSeparator();
+				InlineValue(Component);
+			}
+
+			Printer << TEXT(")");
+		}
+		else
+		{
+			UE_MIR_UNREACHABLE();
+		}
+
+		return NoOp;
+	}
+
+	ENoOp InlineType(const IR::FType* Type)
+	{
+		if (auto ArithmeticType = Type->AsArithmetic())
+		{
+			switch (ArithmeticType->ScalarKind)
+			{
+				case IR::SK_Bool:	Printer << TEXT("bool"); break;
+				case IR::SK_Int: 	Printer << TEXT("int"); break;
+				case IR::SK_Float:	Printer << TEXT("float"); break;
+			}
+
+			if (ArithmeticType->NumRows > 1)
+			{
+				Printer << ArithmeticType->NumRows;
+			}
+			
+			if (ArithmeticType->NumColumns > 1)
+			{
+				Printer << TEXT("x") << ArithmeticType->NumColumns;
+			}
+		}
+		else
+		{
+			UE_MIR_UNREACHABLE();
+		}
+
+		return NoOp;
 	}
 
-	void SetMaterialParameters(const FMaterial& InMaterial, FParametersMap& InParams)
+	void SetMaterialParameters(TMap<FString, FString>& Params)
 	{
+		auto SetParamInt = [&] (const TCHAR* InParamName, int InValue)
+		{
+			Params.Add(InParamName, FString::Printf(TEXT("%d"), InValue));
+		};
+		
 		auto SetParamReturnFloat = [&] (const TCHAR* InParamName, float InValue)
 		{
-			InParams.Add(InParamName, FString::Printf(TEXT("\treturn %.5f"), InValue));
+			Params.Add(InParamName, FString::Printf(TEXT("\treturn %.5f"), InValue));
 		};
 
+		Params.Add(TEXT("pixel_material_inputs"), MoveTemp(PixelAttributesHLSL));
+		Params.Add(TEXT("calc_pixel_material_inputs_initial_calculations"), EvaluateOtherMaterialAttributesHLSL);
+		Params.Add(TEXT("calc_pixel_material_inputs_analytic_derivatives_initial"), MoveTemp(EvaluateOtherMaterialAttributesHLSL));
+		Params.Add(TEXT("material_declarations"), TEXT("struct FMaterialAttributes {};"));
+		SetParamInt(TEXT("num_material_texcoords_vertex"), 0);
+		SetParamInt(TEXT("num_material_texcoords"), 0);
+		SetParamInt(TEXT("num_custom_vertex_interpolators"), 0);
+		SetParamInt(TEXT("num_tex_coord_interpolators"), 0);
+
 		SetParamReturnFloat(TEXT("get_material_emissive_for_cs"), 0.f);
-		SetParamReturnFloat(TEXT("get_material_translucency_directional_lighting_intensity"), InMaterial.GetTranslucencyDirectionalLightingIntensity());
-		SetParamReturnFloat(TEXT("get_material_translucent_shadow_density_scale"), InMaterial.GetTranslucentShadowDensityScale());
-		SetParamReturnFloat(TEXT("get_material_translucent_self_shadow_density_scale"), InMaterial.GetTranslucentSelfShadowDensityScale());
-		SetParamReturnFloat(TEXT("get_material_translucent_self_shadow_second_density_scale"), InMaterial.GetTranslucentSelfShadowSecondDensityScale());
-		SetParamReturnFloat(TEXT("get_material_translucent_self_shadow_second_opacity"), InMaterial.GetTranslucentSelfShadowSecondOpacity());
-		SetParamReturnFloat(TEXT("get_material_translucent_backscattering_exponent"), InMaterial.GetTranslucentBackscatteringExponent());
-
-		FLinearColor Extinction = InMaterial.GetTranslucentMultipleScatteringExtinction();
-		InParams.Add(TEXT("get_material_translucent_multiple_scattering_extinction"), FString::Printf(TEXT("\treturn MaterialFloat3(%.5f, %.5f, %.5f)"), Extinction.R, Extinction.G, Extinction.B));
-
-		SetParamReturnFloat(TEXT("get_material_opacity_mask_clip_value"), InMaterial.GetOpacityMaskClipValue());
-		InParams.Add(TEXT("get_material_world_position_offset_raw"), TEXT("\treturn 0; // todo"));
-		InParams.Add(TEXT("get_material_previous_world_position_offset_raw"), TEXT("\treturn 0; // todo"));
+		SetParamReturnFloat(TEXT("get_material_translucency_directional_lighting_intensity"), Material->GetTranslucencyDirectionalLightingIntensity());
+		SetParamReturnFloat(TEXT("get_material_translucent_shadow_density_scale"), Material->GetTranslucentShadowDensityScale());
+		SetParamReturnFloat(TEXT("get_material_translucent_self_shadow_density_scale"), Material->GetTranslucentSelfShadowDensityScale());
+		SetParamReturnFloat(TEXT("get_material_translucent_self_shadow_second_density_scale"), Material->GetTranslucentSelfShadowSecondDensityScale());
+		SetParamReturnFloat(TEXT("get_material_translucent_self_shadow_second_opacity"), Material->GetTranslucentSelfShadowSecondOpacity());
+		SetParamReturnFloat(TEXT("get_material_translucent_backscattering_exponent"), Material->GetTranslucentBackscatteringExponent());
+
+		FLinearColor Extinction = Material->GetTranslucentMultipleScatteringExtinction();
+		Params.Add(TEXT("get_material_translucent_multiple_scattering_extinction"), FString::Printf(TEXT("\treturn MaterialFloat3(%.5f, %.5f, %.5f)"), Extinction.R, Extinction.G, Extinction.B));
+
+		SetParamReturnFloat(TEXT("get_material_opacity_mask_clip_value"), Material->GetOpacityMaskClipValue());
+		Params.Add(TEXT("get_material_world_position_offset_raw"), TEXT("\treturn 0; // todo"));
+		Params.Add(TEXT("get_material_previous_world_position_offset_raw"), TEXT("\treturn 0; // todo"));
 	
 		// CustomData0/1 are named ClearCoat/ClearCoatRoughness
-		InParams.Add(TEXT("get_material_custom_data0"), TEXT("\treturn 1.0f; // todo"));
-		InParams.Add(TEXT("get_material_custom_data1"), TEXT("\treturn 0.1f; // todo"));
+		Params.Add(TEXT("get_material_custom_data0"), TEXT("\treturn 1.0f; // todo"));
+		Params.Add(TEXT("get_material_custom_data1"), TEXT("\treturn 0.1f; // todo"));
 
 		FString EvaluateMaterialDeclaration;
 		EvaluateMaterialDeclaration.Append(TEXT("void EvaluateVertexMaterialAttributes(in out FMaterialVertexParameters Parameters)\n{\n"));
 		EvaluateMaterialDeclaration.Append(TEXT("\n}\n"));
-		InParams.Add(TEXT("evaluate_material_attributes"), EvaluateMaterialDeclaration);
+		Params.Add(TEXT("evaluate_material_attributes"), EvaluateMaterialDeclaration);
 	}
 
-	void GetShaderCompilerEnvironment(const FMaterial& InMaterial, const FMaterialIRModule& InModule, FShaderCompilerEnvironment& OutEnvironment)
+	void GetShaderCompilerEnvironment(FShaderCompilerEnvironment& OutEnvironment)
 	{
-		const FMaterialCompilationOutput& CompilationOutput = InModule.GetCompilationOutput();
+		const FMaterialCompilationOutput& CompilationOutput = Module->GetCompilationOutput();
+		EShaderPlatform ShaderPlatform = Module->GetShaderPlatform();
 
-		OutEnvironment.TargetPlatform = Translator->TargetPlatform;
+		OutEnvironment.TargetPlatform = TargetPlatform;
 		OutEnvironment.SetDefine(TEXT("ENABLE_NEW_HLSL_GENERATOR"), 1);
 		OutEnvironment.SetDefine(TEXT("MATERIAL_ATMOSPHERIC_FOG"), false);
 		OutEnvironment.SetDefine(TEXT("MATERIAL_SKY_ATMOSPHERE"), false);
@@ -229,26 +484,26 @@ struct FMaterialIRToHLSLTranslator::FPrivate
 		OutEnvironment.SetDefineAndCompileArgument(TEXT("USES_WORLD_POSITION_OFFSET"), (bool)CompilationOutput.bUsesWorldPositionOffset);
 		OutEnvironment.SetDefineAndCompileArgument(TEXT("USES_DISPLACEMENT"), false);
 		OutEnvironment.SetDefine(TEXT("USES_EMISSIVE_COLOR"), false);
-		OutEnvironment.SetDefine(TEXT("USES_DISTORTION"), InMaterial.IsDistorted());
-		OutEnvironment.SetDefine(TEXT("MATERIAL_ENABLE_TRANSLUCENCY_FOGGING"), InMaterial.ShouldApplyFogging());
-		OutEnvironment.SetDefine(TEXT("MATERIAL_ENABLE_TRANSLUCENCY_CLOUD_FOGGING"), InMaterial.ShouldApplyCloudFogging());
-		OutEnvironment.SetDefine(TEXT("MATERIAL_IS_SKY"), InMaterial.IsSky());
-		OutEnvironment.SetDefine(TEXT("MATERIAL_COMPUTE_FOG_PER_PIXEL"), InMaterial.ComputeFogPerPixel());
+		OutEnvironment.SetDefine(TEXT("USES_DISTORTION"), Material->IsDistorted());
+		OutEnvironment.SetDefine(TEXT("MATERIAL_ENABLE_TRANSLUCENCY_FOGGING"), Material->ShouldApplyFogging());
+		OutEnvironment.SetDefine(TEXT("MATERIAL_ENABLE_TRANSLUCENCY_CLOUD_FOGGING"), Material->ShouldApplyCloudFogging());
+		OutEnvironment.SetDefine(TEXT("MATERIAL_IS_SKY"), Material->IsSky());
+		OutEnvironment.SetDefine(TEXT("MATERIAL_COMPUTE_FOG_PER_PIXEL"), Material->ComputeFogPerPixel());
 		OutEnvironment.SetDefine(TEXT("MATERIAL_FULLY_ROUGH"), false);
 		OutEnvironment.SetDefine(TEXT("MATERIAL_USES_ANISOTROPY"), false);
-		OutEnvironment.SetDefine(TEXT("MATERIAL_NEURAL_POST_PROCESS"), (CompilationOutput.bUsedWithNeuralNetworks || InMaterial.IsUsedWithNeuralNetworks()) && InMaterial.IsPostProcessMaterial());
+		OutEnvironment.SetDefine(TEXT("MATERIAL_NEURAL_POST_PROCESS"), (CompilationOutput.bUsedWithNeuralNetworks || Material->IsUsedWithNeuralNetworks()) && Material->IsPostProcessMaterial());
 		OutEnvironment.SetDefine(TEXT("NUM_VIRTUALTEXTURE_SAMPLES"), 0);
 		OutEnvironment.SetDefine(TEXT("MATERIAL_VIRTUALTEXTURE_FEEDBACK"), false);
 		OutEnvironment.SetDefine(TEXT("IS_MATERIAL_SHADER"), true);
 
-		FMaterialShadingModelField ShadingModels = InMaterial.GetShadingModels();
+		FMaterialShadingModelField ShadingModels = Material->GetShadingModels();
 		ensure(ShadingModels.IsValid());
 
 		int32 NumActiveShadingModels = 0;
 		if (ShadingModels.IsLit())
 		{
 			// This is to have platforms use the simple single layer water shading similar to mobile: no dynamic lights, only sun and sky, no distortion, no colored transmittance on background, no custom depth read.
-			const bool bSingleLayerWaterUsesSimpleShading = FDataDrivenShaderPlatformInfo::GetWaterUsesSimpleForwardShading(Translator->ShaderPlatform) && IsForwardShadingEnabled(Translator->ShaderPlatform);
+			const bool bSingleLayerWaterUsesSimpleShading = FDataDrivenShaderPlatformInfo::GetWaterUsesSimpleForwardShading(ShaderPlatform) && IsForwardShadingEnabled(ShaderPlatform);
 
 			for (int i = 0; i < MSM_NUM; ++i)
 			{
@@ -258,7 +513,7 @@ struct FMaterialIRToHLSLTranslator::FPrivate
 					continue;
 				}
 
-				if (Model == MSM_SingleLayerWater && !FDataDrivenShaderPlatformInfo::GetRequiresDisableForwardLocalLights(Translator->ShaderPlatform))
+				if (Model == MSM_SingleLayerWater && !FDataDrivenShaderPlatformInfo::GetRequiresDisableForwardLocalLights(ShaderPlatform))
 				{
 					continue;
 				}
@@ -276,8 +531,8 @@ struct FMaterialIRToHLSLTranslator::FPrivate
 		else
 		{
 			// Unlit shading model can only exist by itself
-			OutEnvironment.SetDefine(TEXT("MATERIAL_SINGLE_SHADINGMODEL"), true);
 			OutEnvironment.SetDefine(GetShadingModelParameterName(MSM_Unlit), true);
+			NumActiveShadingModels += 1;
 		}
 
 		if (NumActiveShadingModels == 1)
@@ -295,10 +550,10 @@ struct FMaterialIRToHLSLTranslator::FPrivate
 		OutEnvironment.SetDefine(TEXT("WSVECTOR_IS_TILEOFFSET"), true);
 		OutEnvironment.SetDefine(TEXT("WSVECTOR_IS_DOUBLEFLOAT"), false);
 
-		if (InMaterial.GetMaterialDomain() == MD_Volume)
+		if (Material->GetMaterialDomain() == MD_Volume)
 		{
 			TArray<const UMaterialExpressionVolumetricAdvancedMaterialOutput*> VolumetricAdvancedExpressions;
-			InMaterial.GetMaterialInterface()->GetMaterial()->GetAllExpressionsOfType(VolumetricAdvancedExpressions);
+			Material->GetMaterialInterface()->GetMaterial()->GetAllExpressionsOfType(VolumetricAdvancedExpressions);
 			if (VolumetricAdvancedExpressions.Num() > 0)
 			{
 				if (VolumetricAdvancedExpressions.Num() > 1)
@@ -316,7 +571,7 @@ struct FMaterialIRToHLSLTranslator::FPrivate
 				OutEnvironment.SetDefine(TEXT("MATERIAL_VOLUMETRIC_ADVANCED_CLAMP_MULTISCATTERING_CONTRIBUTION"), VolumetricAdvancedNode->bClampMultiScatteringContribution);
 				OutEnvironment.SetDefine(TEXT("MATERIAL_VOLUMETRIC_ADVANCED_MULTISCATTERING_OCTAVE_COUNT"), VolumetricAdvancedNode->GetMultiScatteringApproximationOctaveCount());
 				OutEnvironment.SetDefine(TEXT("MATERIAL_VOLUMETRIC_ADVANCED_CONSERVATIVE_DENSITY"), VolumetricAdvancedNode->ConservativeDensity.IsConnected());
-				OutEnvironment.SetDefine(TEXT("MATERIAL_VOLUMETRIC_ADVANCED_OVERRIDE_AMBIENT_OCCLUSION"), InMaterial.HasAmbientOcclusionConnected());
+				OutEnvironment.SetDefine(TEXT("MATERIAL_VOLUMETRIC_ADVANCED_OVERRIDE_AMBIENT_OCCLUSION"), Material->HasAmbientOcclusionConnected());
 				OutEnvironment.SetDefine(TEXT("MATERIAL_VOLUMETRIC_ADVANCED_GROUND_CONTRIBUTION"), VolumetricAdvancedNode->bGroundContribution);
 			}
 		}
@@ -325,98 +580,16 @@ struct FMaterialIRToHLSLTranslator::FPrivate
 		OutEnvironment.SetDefine(TEXT("DUAL_SOURCE_COLOR_BLENDING_ENABLED"), false);
 		OutEnvironment.SetDefine(TEXT("TEXTURE_SAMPLE_DEBUG"), false);
 	}
-
-	void LowerValue( FHLSLPrinter& Printer, const IR::FValue* InValue)
-	{
-		if (const IR::FScalarConstant* Scalar = InValue->As<IR::FScalarConstant>())
-		{
-			IR::FArithmeticTypePtr ArithType = Scalar->Type->ToArithmetic();
-			check(ArithType && ArithType->IsScalar());
-
-			switch (ArithType->ScalarKind)
-			{
-				case IR::SK_Bool: Printer << Scalar->Boolean; break;
-				case IR::SK_Int: Printer << Scalar->Integer; break;
-				case IR::SK_Float: Printer << Scalar->Float; break;
-			}
-		}
-		else if (const IR::FDimensional* Vector = InValue->As<IR::FDimensional>())
-		{
-			IR::FArithmeticTypePtr ArithType = Vector->Type->ToArithmetic();
-			check(ArithType && ArithType->IsVector());
-
-			Printer << ScalarKindToString(ArithType->ScalarKind) << ArithType->NumRows << TEXT("(");
-
-			Printer.BeginList();
-			for (IR::FValuePtr Component : Vector->GetComponents())
-			{
-				Printer.PrintListSeparator();
-				LowerValue(Printer, Component);
-			}
-
-			Printer << TEXT(")");
-		}
-		else if (const IR::FSetMaterialOutput* Output = InValue->As<IR::FSetMaterialOutput>())
-		{
-			// Special case MP_SubsurfaceColor as the actual property is a combination of the color and the profile but we don't want to expose the profile
-			const FString& PropertyName = (Output->Property == MP_SubsurfaceColor) ? "Subsurface" : FMaterialAttributeDefinitionMap::GetAttributeName(Output->Property);
-
-			Printer.Printf(TEXT("\tPixelMaterialInputs.%s = "), *PropertyName);
-		
-			LowerValue(Printer, Output->ArgValue);
-
-			Printer << TEXT(";\n");
-		}
-		else if (const IR::FBinaryOperator* BinOp = InValue->As<IR::FBinaryOperator>())
-		{
-			LowerValue(Printer, BinOp->Lhs);
-
-			const TCHAR* OpString;
-			switch (BinOp->Operator)
-			{
-				case IR::BO_Add: OpString = TEXT(" + "); break;
-				case IR::BO_Subtract: OpString = TEXT(" - "); break;
-				case IR::BO_Multiply: OpString = TEXT(" * "); break;
-				case IR::BO_Divide: OpString = TEXT(" / "); break;
-				default: UE_MIR_UNREACHABLE();
-			}
-			Printer << OpString;
-
-			LowerValue(Printer, BinOp->Rhs);
-		}
-		else
-		{
-			UE_MIR_UNREACHABLE();
-		}
-	}
 };
 
-FMaterialIRToHLSLTranslator::FMaterialIRToHLSLTranslator()
+void FMaterialIRToHLSLTranslation::Run(TMap<FString, FString>& OutParameters, FShaderCompilerEnvironment& OutEnvironment)
 {
-}
+	OutParameters.Empty();
 
-void FMaterialIRToHLSLTranslator::SetTarget(EShaderPlatform InShaderPlatform, const ITargetPlatform* InTargetPlatform, ERHIFeatureLevel::Type InFeatureLevel)
-{
-	ShaderPlatform = InShaderPlatform;
-	TargetPlatform = InTargetPlatform;
-	FeatureLevel = InFeatureLevel;
-}
-
-bool FMaterialIRToHLSLTranslator::Translate(
-	const FMaterial& InMaterial,
-	const FStaticParameterSet& StaticParameters,
-	const FMaterialIRModule& InModule,
-	FParametersMap& OutParametersMap,
-	FShaderCompilerEnvironment& OutEnvironment)
-{
-	OutParametersMap.Empty();
-
-	FPrivate Private{ this };
-	Private.GenerateHLSL(InMaterial, InModule, OutParametersMap);
-	Private.SetMaterialParameters(InMaterial, OutParametersMap);
-	Private.GetShaderCompilerEnvironment(InMaterial, InModule, OutEnvironment);
-
-	return true;
+	FTranslator Translator{ *this };
+	Translator.GenerateHLSL();
+	Translator.SetMaterialParameters(OutParameters);
+	Translator.GetShaderCompilerEnvironment(OutEnvironment);
 }
 
 #endif // #if WITH_EDITOR
diff --git a/Engine/Source/Runtime/Engine/Private/Materials/MaterialIRTypes.cpp b/Engine/Source/Runtime/Engine/Private/Materials/MaterialIRTypes.cpp
--- a/Engine/Source/Runtime/Engine/Private/Materials/MaterialIRTypes.cpp
+++ b/Engine/Source/Runtime/Engine/Private/Materials/MaterialIRTypes.cpp
@@ -7,6 +7,16 @@
 namespace UE::MIR
 {
 
+const TCHAR* TypeKindToString(ETypeKind Kind)
+{
+	switch (Kind)
+	{
+		case TK_Void: return TEXT("void");
+		case TK_Arithmetic: return TEXT("arithmetic");
+		default: UE_MIR_UNREACHABLE();
+	}
+}
+
 FTypePtr FType::FromShaderType(const UE::Shader::FType& InShaderType)
 {
 	check(!InShaderType.IsStruct());
@@ -46,26 +56,36 @@ FTypePtr FType::GetVoid()
 	return &Type;
 }
 
-FArithmeticTypePtr FType::ToArithmetic() const
+FStringView FType::GetSpelling() const
+{
+	if (auto Arithmetic = AsArithmetic())
+	{
+		return Arithmetic->Spelling;
+	}
+	
+	UE_MIR_UNREACHABLE();
+}
+
+FArithmeticTypePtr FType::AsArithmetic() const
 {
 	return Kind == TK_Arithmetic ? static_cast<FArithmeticTypePtr>(this) : nullptr; 
 }
 
-FArithmeticTypePtr FType::ToScalar() const
+FArithmeticTypePtr FType::AsScalar() const
 {
-	FArithmeticTypePtr Type = ToArithmetic();
+	FArithmeticTypePtr Type = AsArithmetic();
 	return Type->IsScalar() ? Type : nullptr;
 }
 
-FArithmeticTypePtr FType::ToVector() const
+FArithmeticTypePtr FType::AsVector() const
 {
-	FArithmeticTypePtr Type = ToArithmetic();
+	FArithmeticTypePtr Type = AsArithmetic();
 	return Type->IsVector() ? Type : nullptr;
 }
 
-FArithmeticTypePtr FType::ToMatrix() const
+FArithmeticTypePtr FType::AsMatrix() const
 {
-	FArithmeticTypePtr Type = ToArithmetic();
+	FArithmeticTypePtr Type = AsArithmetic();
 	return Type->IsMatrix() ? Type : nullptr;
 }
 
@@ -80,82 +100,99 @@ const TCHAR* ScalarKindToString(EScalarKind Kind)
 	}
 }
 
-static const FArithmeticType* GetNumericalType(EScalarKind InScalarKind, int NumRows, int NumColumns)
+FArithmeticTypePtr FArithmeticType::GetBool()
 {
-	check(InScalarKind >= 0 && InScalarKind <= SK_Float);
-	
-	static const FArithmeticType Types[] {
-		{ { TK_Arithmetic }, SK_Bool, 1, 1 },
-		{ { TK_Arithmetic }, SK_Bool, 1, 2 }, 
-		{ { TK_Arithmetic }, SK_Bool, 1, 3 },
-		{ { TK_Arithmetic }, SK_Bool, 1, 4 },
-		{ { TK_Arithmetic }, SK_Bool, 2, 1 },
-		{ { TK_Arithmetic }, SK_Bool, 2, 2 },
-		{ { TK_Arithmetic }, SK_Bool, 2, 3 },
-		{ { TK_Arithmetic }, SK_Bool, 2, 4 },
-		{ { TK_Arithmetic }, SK_Bool, 3, 1 },
-		{ { TK_Arithmetic }, SK_Bool, 3, 2 },
-		{ { TK_Arithmetic }, SK_Bool, 3, 3 },
-		{ { TK_Arithmetic }, SK_Bool, 3, 4 },
-		{ { TK_Arithmetic }, SK_Bool, 4, 1 },
-		{ { TK_Arithmetic }, SK_Bool, 4, 2 },
-		{ { TK_Arithmetic }, SK_Bool, 4, 3 },
-		{ { TK_Arithmetic }, SK_Bool, 4, 4 },
-		{ { TK_Arithmetic }, SK_Int, 1, 1 },
-		{ { TK_Arithmetic }, SK_Int, 1, 2 },
-		{ { TK_Arithmetic }, SK_Int, 1, 3 },
-		{ { TK_Arithmetic }, SK_Int, 1, 4 },
-		{ { TK_Arithmetic }, SK_Int, 2, 1 },
-		{ { TK_Arithmetic }, SK_Int, 2, 2 },
-		{ { TK_Arithmetic }, SK_Int, 2, 3 },
-		{ { TK_Arithmetic }, SK_Int, 2, 4 },
-		{ { TK_Arithmetic }, SK_Int, 3, 1 },
-		{ { TK_Arithmetic }, SK_Int, 3, 2 },
-		{ { TK_Arithmetic }, SK_Int, 3, 3 },
-		{ { TK_Arithmetic }, SK_Int, 3, 4 },
-		{ { TK_Arithmetic }, SK_Int, 4, 1 },
-		{ { TK_Arithmetic }, SK_Int, 4, 2 },
-		{ { TK_Arithmetic }, SK_Int, 4, 3 },
-		{ { TK_Arithmetic }, SK_Int, 4, 4 },
-		{ { TK_Arithmetic }, SK_Float, 1, 1 },
-		{ { TK_Arithmetic }, SK_Float, 1, 2 },
-		{ { TK_Arithmetic }, SK_Float, 1, 3 },
-		{ { TK_Arithmetic }, SK_Float, 1, 4 },
-		{ { TK_Arithmetic }, SK_Float, 2, 1 },
-		{ { TK_Arithmetic }, SK_Float, 2, 2 },
-		{ { TK_Arithmetic }, SK_Float, 2, 3 },
-		{ { TK_Arithmetic }, SK_Float, 2, 4 },
-		{ { TK_Arithmetic }, SK_Float, 3, 1 },
-		{ { TK_Arithmetic }, SK_Float, 3, 2 },
-		{ { TK_Arithmetic }, SK_Float, 3, 3 },
-		{ { TK_Arithmetic }, SK_Float, 3, 4 },
-		{ { TK_Arithmetic }, SK_Float, 4, 1 },
-		{ { TK_Arithmetic }, SK_Float, 4, 2 },
-		{ { TK_Arithmetic }, SK_Float, 4, 3 },
-		{ { TK_Arithmetic }, SK_Float, 4, 4 },
-	};
+	return GetScalar(SK_Bool);
+}
 
-	int Index = InScalarKind * 4 * 4 + (NumRows - 1) * 4 + (NumColumns - 1);
-	check(Index < UE_ARRAY_COUNT(Types));
-	return &Types[Index];
+FArithmeticTypePtr FArithmeticType::GetInt()
+{
+	return GetScalar(SK_Int);
+}
+
+FArithmeticTypePtr FArithmeticType::GetFloat()
+{
+	return GetScalar(SK_Float);
 }
 
 const FArithmeticType* FArithmeticType::GetScalar(EScalarKind InScalarKind)
 {
-	return GetNumericalType(InScalarKind, 1, 1);
+	return Get(InScalarKind, 1, 1);
 }
 
 const FArithmeticType* FArithmeticType::GetVector(EScalarKind InScalarKind, int NumComponents)
 {
 	check(NumComponents >= 1 && NumComponents <= 4);
-	return GetNumericalType(InScalarKind, NumComponents, 1);
+	return Get(InScalarKind, NumComponents, 1);
 }
 
 const FArithmeticType* FArithmeticType::GetMatrix(EScalarKind InScalarKind, int NumRows, int NumColumns)
 {
 	check(NumColumns > 1 && NumColumns <= 4);
 	check(NumRows > 1 && NumRows <= 4);
-	return GetNumericalType(InScalarKind, NumRows, NumColumns);
+	return Get(InScalarKind, NumRows, NumColumns);
+}
+
+FArithmeticTypePtr FArithmeticType::Get(EScalarKind InScalarKind, int NumRows, int NumColumns)
+{
+	check(InScalarKind >= 0 && InScalarKind <= SK_Float);
+	
+	static const FStringView Invalid = TEXT("invalid");
+
+	static const FArithmeticType Types[] {
+		{ { TK_Arithmetic }, { TEXT("bool") }, 		SK_Bool, 1, 1 },
+		{ { TK_Arithmetic }, Invalid, 				SK_Bool, 1, 2 }, 
+		{ { TK_Arithmetic }, Invalid, 				SK_Bool, 1, 3 },
+		{ { TK_Arithmetic }, Invalid, 				SK_Bool, 1, 4 },
+		{ { TK_Arithmetic }, { TEXT("bool2") },   	SK_Bool, 2, 1 },
+		{ { TK_Arithmetic }, { TEXT("bool2x2") }, 	SK_Bool, 2, 2 },
+		{ { TK_Arithmetic }, { TEXT("bool2x3") }, 	SK_Bool, 2, 3 },
+		{ { TK_Arithmetic }, { TEXT("bool2x4") }, 	SK_Bool, 2, 4 },
+		{ { TK_Arithmetic }, { TEXT("bool3") },   	SK_Bool, 3, 1 },
+		{ { TK_Arithmetic }, { TEXT("bool3x2") }, 	SK_Bool, 3, 2 },
+		{ { TK_Arithmetic }, { TEXT("bool3x3") }, 	SK_Bool, 3, 3 },
+		{ { TK_Arithmetic }, { TEXT("bool3x4") }, 	SK_Bool, 3, 4 },
+		{ { TK_Arithmetic }, { TEXT("bool4") },   	SK_Bool, 4, 1 },
+		{ { TK_Arithmetic }, { TEXT("bool4x2") }, 	SK_Bool, 4, 2 },
+		{ { TK_Arithmetic }, { TEXT("bool4x3") }, 	SK_Bool, 4, 3 },
+		{ { TK_Arithmetic }, { TEXT("bool4x4") }, 	SK_Bool, 4, 4 },
+		{ { TK_Arithmetic }, { TEXT("int") }, 		SK_Int, 1, 1 },
+		{ { TK_Arithmetic }, Invalid, 				SK_Int, 1, 2 },
+		{ { TK_Arithmetic }, Invalid, 				SK_Int, 1, 3 },
+		{ { TK_Arithmetic }, Invalid, 				SK_Int, 1, 4 },
+		{ { TK_Arithmetic }, { TEXT("int2") },   	SK_Int, 2, 1 },
+		{ { TK_Arithmetic }, { TEXT("int2x2") }, 	SK_Int, 2, 2 },
+		{ { TK_Arithmetic }, { TEXT("int2x3") }, 	SK_Int, 2, 3 },
+		{ { TK_Arithmetic }, { TEXT("int2x4") }, 	SK_Int, 2, 4 },
+		{ { TK_Arithmetic }, { TEXT("int3") },   	SK_Int, 3, 1 },
+		{ { TK_Arithmetic }, { TEXT("int3x2") }, 	SK_Int, 3, 2 },
+		{ { TK_Arithmetic }, { TEXT("int3x3") }, 	SK_Int, 3, 3 },
+		{ { TK_Arithmetic }, { TEXT("int3x4") }, 	SK_Int, 3, 4 },
+		{ { TK_Arithmetic }, { TEXT("int4") },   	SK_Int, 4, 1 },
+		{ { TK_Arithmetic }, { TEXT("int4x2") }, 	SK_Int, 4, 2 },
+		{ { TK_Arithmetic }, { TEXT("int4x3") }, 	SK_Int, 4, 3 },
+		{ { TK_Arithmetic }, { TEXT("int4x4") }, 	SK_Int, 4, 4 },
+		{ { TK_Arithmetic }, { TEXT("float") }, 	SK_Float, 1, 1 },
+		{ { TK_Arithmetic }, Invalid, 				SK_Float, 1, 2 },
+		{ { TK_Arithmetic }, Invalid, 				SK_Float, 1, 3 },
+		{ { TK_Arithmetic }, Invalid, 				SK_Float, 1, 4 },
+		{ { TK_Arithmetic }, { TEXT("float2") },   	SK_Float, 2, 1 },
+		{ { TK_Arithmetic }, { TEXT("float2x2") }, 	SK_Float, 2, 2 },
+		{ { TK_Arithmetic }, { TEXT("float2x3") }, 	SK_Float, 2, 3 },
+		{ { TK_Arithmetic }, { TEXT("float2x4") }, 	SK_Float, 2, 4 },
+		{ { TK_Arithmetic }, { TEXT("float3") },   	SK_Float, 3, 1 },
+		{ { TK_Arithmetic }, { TEXT("float3x2") }, 	SK_Float, 3, 2 },
+		{ { TK_Arithmetic }, { TEXT("float3x3") }, 	SK_Float, 3, 3 },
+		{ { TK_Arithmetic }, { TEXT("float3x4") }, 	SK_Float, 3, 4 },
+		{ { TK_Arithmetic }, { TEXT("float4") },   	SK_Float, 4, 1 },
+		{ { TK_Arithmetic }, { TEXT("float4x2") }, 	SK_Float, 4, 2 },
+		{ { TK_Arithmetic }, { TEXT("float4x3") }, 	SK_Float, 4, 3 },
+		{ { TK_Arithmetic }, { TEXT("float4x4") }, 	SK_Float, 4, 4 },
+	};
+
+	int Index = InScalarKind * 4 * 4 + (NumRows - 1) * 4 + (NumColumns - 1);
+	check(Index < UE_ARRAY_COUNT(Types));
+	return &Types[Index];
 }
 
 } // namespace UE::MIR
diff --git a/Engine/Source/Runtime/Engine/Private/Materials/MaterialIRUtility.cpp b/Engine/Source/Runtime/Engine/Private/Materials/MaterialIRUtility.cpp
--- a/Engine/Source/Runtime/Engine/Private/Materials/MaterialIRUtility.cpp
+++ b/Engine/Source/Runtime/Engine/Private/Materials/MaterialIRUtility.cpp
@@ -57,7 +57,7 @@ bool NextMaterialAttributeInput(UMaterial* BaseMaterial, int32& PropertyIndex, F
 	return false;
 }
 
-IR::FValuePtr CreateMaterialAttributeDefaultValue(IR::FEmitter& Emitter, const UMaterial* Material, EMaterialProperty Property)
+IR::FValue* CreateMaterialAttributeDefaultValue(IR::FEmitter& Emitter, const UMaterial* Material, EMaterialProperty Property)
 {
 	EMaterialValueType Type = FMaterialAttributeDefinitionMap::GetValueType(Property);
 	FVector4f DefaultValue = FMaterialAttributeDefinitionMap::GetDefaultValue(Property);
diff --git a/Engine/Source/Runtime/Engine/Private/Materials/MaterialIRUtility.h b/Engine/Source/Runtime/Engine/Private/Materials/MaterialIRUtility.h
--- a/Engine/Source/Runtime/Engine/Private/Materials/MaterialIRUtility.h
+++ b/Engine/Source/Runtime/Engine/Private/Materials/MaterialIRUtility.h
@@ -14,7 +14,7 @@ bool IsMaterialPropertyShared(EMaterialProperty InProperty);
 bool NextMaterialAttributeInput(UMaterial* BaseMaterial, int32& PropertyIndex, FMaterialInputDescription& Input);
 
 //
-UE::MIR::FValuePtr CreateMaterialAttributeDefaultValue(UE::MIR::FEmitter& Emitter, const UMaterial* Material, EMaterialProperty Property);
+UE::MIR::FValue* CreateMaterialAttributeDefaultValue(UE::MIR::FEmitter& Emitter, const UMaterial* Material, EMaterialProperty Property);
 
 } // namespace Utility
 
diff --git a/Engine/Source/Runtime/Engine/Private/Materials/MaterialShared.cpp b/Engine/Source/Runtime/Engine/Private/Materials/MaterialShared.cpp
--- a/Engine/Source/Runtime/Engine/Private/Materials/MaterialShared.cpp
+++ b/Engine/Source/Runtime/Engine/Private/Materials/MaterialShared.cpp
@@ -3395,55 +3395,65 @@ bool FMaterial::Translate_Legacy(const FMaterialShaderMapId& ShaderMapId,
 
 bool FMaterial::Translate_New(const FMaterialShaderMapId& InShaderMapId,
 	const FStaticParameterSet& InStaticParameters,
-	EShaderPlatform InPlatform,
+	EShaderPlatform InShaderPlatform,
 	const ITargetPlatform* InTargetPlatform,
 	FMaterialCompilationOutput& OutCompilationOutput,
 	TRefCountPtr<FSharedShaderCompilerEnvironment>& OutMaterialEnvironment)
 {
-	if (IsUsingNewMaterialTranslatorPrototype())
+	if (!IsUsingNewMaterialTranslatorPrototype())
 	{
-		const FMaterialIRModule& Module = this->GetMaterialInterface()->GetIRModule();
+		const FMaterialCompileTargetParameters TargetParams(InShaderPlatform, InShaderMapId.FeatureLevel, InTargetPlatform);
+		return MaterialEmitHLSL(TargetParams, InStaticParameters, *this, OutCompilationOutput, OutMaterialEnvironment);
+	}
+
+	FMaterialIRModule Module;
 
-		if (!Module.GetErrors().IsEmpty())
+	FMaterialIRTargetParams TargetParams{
+		.ShaderPlatform = InShaderPlatform,
+		.TargetPlatform = InTargetPlatform,
+		.StaticParameters = InStaticParameters,
+	};
+
+	// Build the material
+	FMaterialIRModuleBuilder Builder;
+	if (!Builder.Build(GetMaterialInterface()->GetMaterial(), TargetParams, &Module))
+	{
+		for (const FMaterialIRModule::FError& Error : Module.GetErrors())
 		{
-			// TODO: This could be moved elsewhere, when the Module is built.
-			for (const FMaterialIRModule::FError& Error : Module.GetErrors())
-			{
-				ErrorExpressions.Push(Error.Expression);
-				CompileErrors.Push(Error.Message);
-			}
-			return false;
+			ErrorExpressions.Push(Error.Expression);
+			CompileErrors.Push(Error.Message);
 		}
+
+		return false;
+	}
 	
-		// Copy over the compilation output
-		OutCompilationOutput = Module.GetCompilationOutput();
-		OutMaterialEnvironment = new FSharedShaderCompilerEnvironment();
+	// Copy over the compilation output
+	OutCompilationOutput = Module.GetCompilationOutput();
+	OutMaterialEnvironment = new FSharedShaderCompilerEnvironment();
 
-		// Translate the material IR module to HLSL template string parameters and material environment
-		FMaterialIRToHLSLTranslator::FParametersMap ShaderStringParameters;
-		FMaterialIRToHLSLTranslator Translator;
-		if (!Translator.Translate(*this, InStaticParameters, Module, ShaderStringParameters, *OutMaterialEnvironment))
-		{
-			return false;
-		}
+	// Translate the material IR module to HLSL template string parameters and material environment
+	TMap<FString, FString> ShaderStringParameters;
 
-		// Interpolate HLSL parameters with the material shader template to produce the final shader source
-		int32 LineNumber;
-		FStringTemplateResolver Resolver = FMaterialSourceTemplate::Get().BeginResolve(InPlatform, &LineNumber);
-		ShaderStringParameters.Add({TEXT("line_number"), FString::Printf(TEXT("%u"), LineNumber)});
-		Resolver.SetParameterMap(&ShaderStringParameters);
+	FMaterialIRToHLSLTranslation Translation{
+		.Module = &Module,
+		.Material = this,
+		.StaticParameters = &InStaticParameters,
+		.TargetPlatform = InTargetPlatform,
+	};
 
-		// Interpolate the final material shader source string
-		FString MaterialShaderCode = Resolver.Finalize();
-		OutMaterialEnvironment->IncludeVirtualPathToContentsMap.Add(TEXT("/Engine/Generated/Material.ush"), MoveTemp(MaterialShaderCode));
+	Translation.Run(ShaderStringParameters, *OutMaterialEnvironment);
 
-		return true;
-	}
-	else
-	{
-		const FMaterialCompileTargetParameters TargetParams(InPlatform, InShaderMapId.FeatureLevel, InTargetPlatform);
-		return MaterialEmitHLSL(TargetParams, InStaticParameters, *this, OutCompilationOutput, OutMaterialEnvironment);
-	}
+	// Interpolate HLSL parameters with the material shader template to produce the final shader source
+	int32 LineNumber;
+	FStringTemplateResolver Resolver = FMaterialSourceTemplate::Get().BeginResolve(InShaderPlatform, &LineNumber);
+	ShaderStringParameters.Add({TEXT("line_number"), FString::Printf(TEXT("%u"), LineNumber)});
+	Resolver.SetParameterMap(&ShaderStringParameters);
+
+	// Interpolate the final material shader source string
+	FString MaterialShaderCode = Resolver.Finalize();
+	OutMaterialEnvironment->IncludeVirtualPathToContentsMap.Add(TEXT("/Engine/Generated/Material.ush"), MoveTemp(MaterialShaderCode));
+
+	return true;
 }
 
 bool FMaterial::Translate(const FMaterialShaderMapId& InShaderMapId,
diff --git a/Engine/Source/Runtime/Engine/Public/Materials/MaterialExpressionIf.h b/Engine/Source/Runtime/Engine/Public/Materials/MaterialExpressionIf.h
--- a/Engine/Source/Runtime/Engine/Public/Materials/MaterialExpressionIf.h
+++ b/Engine/Source/Runtime/Engine/Public/Materials/MaterialExpressionIf.h
@@ -42,6 +42,7 @@ class UMaterialExpressionIf : public UMaterialExpression
 
 	//~ Begin UMaterialExpression Interface
 #if WITH_EDITOR
+	virtual void Build(UE::MIR::FEmitter& Emitter) override;
 	virtual int32 Compile(class FMaterialCompiler* Compiler, int32 OutputIndex) override;
 	virtual void GetCaption(TArray<FString>& OutCaptions) const override;
 	virtual uint32 GetInputType(int32 InputIndex) override;
diff --git a/Engine/Source/Runtime/Engine/Public/Materials/MaterialIR.h b/Engine/Source/Runtime/Engine/Public/Materials/MaterialIR.h
--- a/Engine/Source/Runtime/Engine/Public/Materials/MaterialIR.h
+++ b/Engine/Source/Runtime/Engine/Public/Materials/MaterialIR.h
@@ -20,6 +20,7 @@ enum EValueKind
 
 	VK_SetMaterialOutput = VK_InstructionBegin,
 	VK_BinaryOperator,
+	VK_Branch,
 
 	VK_InstructionEnd,
 };
@@ -32,13 +33,14 @@ struct FValue
 	FTypePtr  Type{};
 
 	bool IsA(EValueKind InKind) const { return Kind == InKind; }
-	FInstructionPtr AsInstruction() const;
-	bool Equals(FValuePtr Other) const;
+	FInstruction* AsInstruction();
+	const FInstruction* AsInstruction() const;
+	bool Equals(const FValue* Other) const;
 	uint32 GetSizeInBytes() const;
+	TArrayView<FValue*> GetUses();
 
 	template <typename T>
 	const T* As() const { return this && IsA(T::TypeKind) ? static_cast<const T*>(this) : nullptr; }
-
 };
 
 template <EValueKind TTypeKind>
@@ -59,22 +61,41 @@ struct FScalarConstant : TValue<VK_ScalarConstant>
 
 struct FDimensional : TValue<VK_Dimensional>
 {
-	TArrayView<const FValuePtr> GetComponents() const;
-	TArrayView<FValuePtr> GetMutableComponents();
-	uint32 GetSizeInBytes() const;
+	TArrayView<FValue* const> GetComponents() const;
+	TArrayView<FValue*> GetMutableComponents();
 };
 
 template <int TDimension>
 struct TDimensional : FDimensional
 {
-	FValuePtr Components[TDimension];
+	FValue* Components[TDimension];
 };
 
 /* Instructions */
 
+enum EInstructionFlags
+{
+	IF_None = 0,
+	IF_Counted = 1,
+};
+
 struct FInstruction : FValue
 {
+	EInstructionFlags Flags = IF_None;
 	FInstruction* Next{};
+	FBlock* Block{};
+	uint32 NumUsers{};
+	uint32 NumProcessedUsers{};
+
+	void SetFlags(EInstructionFlags InFlags) { Flags = (EInstructionFlags)(Flags | InFlags); }
+	bool GetInnerBlock(int32 Index, FValue*& OutArg, FBlock*& OutBlock); 
+};
+
+struct FBlock
+{
+	FBlock* Parent{};
+	FInstruction* Instructions{};
+	int32 Level{};
 };
 
 template <EValueKind TTypeKind>
@@ -86,7 +107,7 @@ struct TInstruction : FInstruction
 struct FSetMaterialOutput : TInstruction<VK_SetMaterialOutput>
 {
 	EMaterialProperty Property;
-	FValuePtr ArgValue;
+	FValue* Arg;
 };
 
 enum EBinaryOperator
@@ -96,13 +117,25 @@ enum EBinaryOperator
 	BO_Subtract,
 	BO_Multiply,
 	BO_Divide,
+	BO_Greater,
+	BO_Lower,
+	BO_Equals,
 };
 
 struct FBinaryOperator : TInstruction<VK_BinaryOperator>
 {
 	EBinaryOperator Operator = BO_Invalid;
-	FValuePtr Lhs{};
-	FValuePtr Rhs{};
+	FValue* LhsArg{};
+	FValue* RhsArg{};
+};
+
+struct FBranch : TInstruction<VK_Branch>
+{
+	FValue* ConditionArg{};
+	FValue* TrueArg{};
+	FValue* FalseArg{};
+	FBlock TrueBlock{};
+	FBlock FalseBlock{};
 };
 
 } // namespace UE::MIR
diff --git a/Engine/Source/Runtime/Engine/Public/Materials/MaterialIRCommon.h b/Engine/Source/Runtime/Engine/Public/Materials/MaterialIRCommon.h
--- a/Engine/Source/Runtime/Engine/Public/Materials/MaterialIRCommon.h
+++ b/Engine/Source/Runtime/Engine/Public/Materials/MaterialIRCommon.h
@@ -49,6 +49,7 @@ using FInstructionPtr = const FInstruction*;
 
 /* Others */
 class FEmitter;
+struct FBlock;
 
 }
 
diff --git a/Engine/Source/Runtime/Engine/Public/Materials/MaterialIREmitter.h b/Engine/Source/Runtime/Engine/Public/Materials/MaterialIREmitter.h
--- a/Engine/Source/Runtime/Engine/Public/Materials/MaterialIREmitter.h
+++ b/Engine/Source/Runtime/Engine/Public/Materials/MaterialIREmitter.h
@@ -15,40 +15,79 @@ class FEmitter
 public:
 	FEmitter(FMaterialIRModuleBuilder* InBuilder, UMaterial* InMaterial, FMaterialIRModule* InModule);
 	
+	/* IO */
+
+	// Gets and returns the value flowing into specified `Input`. If disconnected,
+	// it returns nullptr. 
+	FValue* Get(const FExpressionInput* Input);
+
+	// Flows specified `Value` out of specified expression `Output`.
+	void 	Put(const FExpressionOutput* Output, FValue* Value);
+
+	/* IO Helpers */
+
+	//
+	FEmitter& DefaultToFloatZero(const FExpressionInput* Input);
+
+	//
+	FEmitter& DefaultTo(const FExpressionInput* Input, float Float);
+
+	// It gets the value flowing into it and checks that its type is float scalar.
+	FValue* TryGetFloat(const FExpressionInput* Input);
+
+	//
+	FValue* TryGetScalar(const FExpressionInput* Input);
+
+	//
+	FValue* TryGetArithmetic(const FExpressionInput* Input);
+
+	// Gets the value flowing into `Input` and returns it after checking that its
+	// type matches `Kind`.
+	FValue* TryGetOfType(const FExpressionInput* Input, ETypeKind Kind);
+
 	/* Analysis */
-	FValuePtr GetAndCheckInputTypeKindIs(const FExpressionInput* Input, ETypeKind Kind);
-	bool CheckInputTypeIs(const FExpressionInput* Input, FValuePtr InputValue, ETypeKind Kind);
+
+	//
+	void CheckInputIsScalar(const FExpressionInput* Input, FValue* InputValue);
+
+	//
+	void CheckInputIsScalar(const FExpressionInput* Input, FValue* InputValue, EScalarKind Kind);
+
+	// Checks that the type of the value `InputValue` flowing into `Input` is of
+	// specified type `Kind. If it isn't it reports an error. You may check whether
+	// an error occurred with `IsInvalid()`.
+	void CheckInputTypeIs(const FExpressionInput* Input, FValue* InputValue, ETypeKind Kind);
 
 	/* Constants */
 
-	FValuePtr EmitConstantFromShaderValue(const UE::Shader::FValue& InValue);
-	FValuePtr EmitConstantScalarZero(EScalarKind Kind);
-	FValuePtr EmitConstantFloat1(float InX);
-	FValuePtr EmitConstantFloat2(const FVector2f& InValue);
-	FValuePtr EmitConstantFloat3(const FVector3f& InValue);
-	FValuePtr EmitConstantFloat4(const FVector4f& InValue);
-	FValuePtr EmitConstantInt1(int InX);
-	FValuePtr EmitConstantInt2(const FIntVector2& InValue);
-	FValuePtr EmitConstantInt3(const FIntVector3& InValue);
-	FValuePtr EmitConstantInt4(const FIntVector4& InValue);
-	FValuePtr EmitVector2(FValuePtr InX, FValuePtr InY);
-	FValuePtr EmitVector3(FValuePtr InX, FValuePtr InY, FValuePtr InZ);
-	FValuePtr EmitVector4(FValuePtr InX, FValuePtr InY, FValuePtr InZ, FValuePtr InW);
+	FValue* EmitConstantFromShaderValue(const UE::Shader::FValue& InValue);
+	FValue* EmitConstantScalarZero(EScalarKind Kind);
+	FValue* EmitConstantFloat1(float InX);
+	FValue* EmitConstantFloat2(const FVector2f& InValue);
+	FValue* EmitConstantFloat3(const FVector3f& InValue);
+	FValue* EmitConstantFloat4(const FVector4f& InValue);
+	FValue* EmitConstantInt1(int InX);
+	FValue* EmitConstantInt2(const FIntVector2& InValue);
+	FValue* EmitConstantInt3(const FIntVector3& InValue);
+	FValue* EmitConstantInt4(const FIntVector4& InValue);
+	FValue* EmitVector2(FValue* InX, FValue* InY);
+	FValue* EmitVector3(FValue* InX, FValue* InY, FValue* InZ);
+	FValue* EmitVector4(FValue* InX, FValue* InY, FValue* InZ, FValue* InW);
 
 	/* Other Values */
 
-	FValuePtr EmitArithmetic(FArithmeticTypePtr Type, FValuePtr Scalar);
+	FValue* EmitArithmetic(FArithmeticTypePtr Type, FValue* Scalar);
 
 	/* Instructions */
 
 	FSetMaterialOutput* EmitSetMaterialOutput(EMaterialProperty InProperty, FValue* InArgValue);
-	FValuePtr EmitBinaryOperator(EBinaryOperator Operator, FValuePtr Lhs, FValuePtr Rhs);
-	FValuePtr TryEmitConvert(FValuePtr Value, FTypePtr TargetType);
+	FValue* EmitBinaryOperator(EBinaryOperator Operator, FValue* Lhs, FValue* Rhs);
+	FValue* EmitBranch(FValue* Condition, FValue* True, FValue* False);
+	FValue* TryEmitConvert(FValue* Value, FTypePtr TargetType);
 
-	/* IO */
+	/* Types */
 
-	FValuePtr Get(const FExpressionInput* Input);
-	void Put(const FExpressionOutput* Output, FValuePtr Value);
+	FArithmeticTypePtr TryGetCommonArithmeticType(FArithmeticTypePtr A, FArithmeticTypePtr B);
 
 	/* Error reporting */
 	bool IsInvalid() const { return bHasExprBuildError; }
@@ -66,9 +105,7 @@ private:
 	FMaterialIRModule* Module{};
 	UMaterialExpression* Expression{};
 	FMaterialIRModuleBuilder* Builder{};
-	FValuePtr ZeroIntValue{};
-	FValuePtr ZeroFloatValue{};
-	bool bHasExprBuildError = true;
+	bool bHasExprBuildError = false;
 
 	struct FPrivate;
 	friend FMaterialIRModuleBuilder;
diff --git a/Engine/Source/Runtime/Engine/Public/Materials/MaterialIRModule.h b/Engine/Source/Runtime/Engine/Public/Materials/MaterialIRModule.h
--- a/Engine/Source/Runtime/Engine/Public/Materials/MaterialIRModule.h
+++ b/Engine/Source/Runtime/Engine/Public/Materials/MaterialIRModule.h
@@ -16,18 +16,22 @@ public:
 	};
 
 public:
+	FMaterialIRModule();
 	~FMaterialIRModule();
 	void Empty();
+	EShaderPlatform GetShaderPlatform() const { return ShaderPlatform; }
 	const FMaterialCompilationOutput& GetCompilationOutput() const { return CompilationOutput; }
 	TArrayView<const UE::MIR::FSetMaterialOutput* const> GetOutputs() const { return Outputs; }
+	const UE::MIR::FBlock& GetRootBlock() const { return *RootBlock; }
 	TArrayView<const FError> GetErrors() const { return Errors; }
 
 private:
+	EShaderPlatform ShaderPlatform;
 	FMaterialCompilationOutput CompilationOutput;
-	TArray<UE::MIR::FValuePtr> Values;
+	TArray<UE::MIR::FValue*> Values;
 	TArray<UE::MIR::FSetMaterialOutput*> Outputs;
 	TArray<FError> Errors;
-
+	UE::MIR::FBlock* RootBlock;
 
 	friend UE::MIR::FEmitter;
 	friend FMaterialIRModuleBuilder;
diff --git a/Engine/Source/Runtime/Engine/Public/Materials/MaterialIRModuleBuilder.h b/Engine/Source/Runtime/Engine/Public/Materials/MaterialIRModuleBuilder.h
--- a/Engine/Source/Runtime/Engine/Public/Materials/MaterialIRModuleBuilder.h
+++ b/Engine/Source/Runtime/Engine/Public/Materials/MaterialIRModuleBuilder.h
@@ -5,18 +5,22 @@
 
 #if WITH_EDITOR
 
+struct FMaterialIRTargetParams
+{
+	EShaderPlatform ShaderPlatform;
+	const ITargetPlatform* TargetPlatform;
+	const FStaticParameterSet& StaticParameters;
+};
+
 //
 class FMaterialIRModuleBuilder
 {
 public:
-	bool Build(UMaterial* InMaterial, FMaterialIRModule* TargetModule);
+	bool Build(UMaterial* InMaterial, const FMaterialIRTargetParams& TargetParams, FMaterialIRModule* TargetModule);
 
 private:
-	UMaterial* BaseMaterial;
-	FMaterialIRModule* Module;
-	TArray<UMaterialExpression*> ExpressionAnalysisStack;
-	TMap<const FExpressionInput*, UE::MIR::FValuePtr> InputValues;
-	TMap<const FExpressionOutput*, UE::MIR::FValuePtr> OutputValues;
+	TMap<const FExpressionInput*, UE::MIR::FValue*> InputValues;
+	TMap<const FExpressionOutput*, UE::MIR::FValue*> OutputValues;
 
 	friend class UE::MIR::FEmitter;
 	struct FPrivate;
diff --git a/Engine/Source/Runtime/Engine/Public/Materials/MaterialIRToHLSLTranslator.h b/Engine/Source/Runtime/Engine/Public/Materials/MaterialIRToHLSLTranslator.h
--- a/Engine/Source/Runtime/Engine/Public/Materials/MaterialIRToHLSLTranslator.h
+++ b/Engine/Source/Runtime/Engine/Public/Materials/MaterialIRToHLSLTranslator.h
@@ -6,22 +6,23 @@
 
 #if WITH_EDITOR
 
-class FMaterialIRToHLSLTranslator
+//
+struct FMaterialIRToHLSLTranslation
 {
-public:
-	using FParametersMap = TMap<FString, FString>;
-
-public:
-	FMaterialIRToHLSLTranslator();
-	void SetTarget(EShaderPlatform InShaderPlatform, const ITargetPlatform* InTargetPlatform, ERHIFeatureLevel::Type InFeatureLevel);
-	bool Translate(const FMaterial& InMaterial, const FStaticParameterSet& StaticParameters, const FMaterialIRModule& InModule, FParametersMap& OutParameters, FShaderCompilerEnvironment& OutEnvironment);
-
-private:
-	EShaderPlatform ShaderPlatform{};
-	const ITargetPlatform* TargetPlatform{};
-	ERHIFeatureLevel::Type FeatureLevel{};
-
-	struct FPrivate;
+	//
+	const FMaterialIRModule* Module;
+
+	//
+	const FMaterial* Material;
+	
+	//
+	const FStaticParameterSet* StaticParameters;
+	
+	//
+	const ITargetPlatform* TargetPlatform;
+	
+	//
+	void Run(TMap<FString, FString>& OutParameters, FShaderCompilerEnvironment& OutEnvironment);
 };
 
 #endif // #if WITH_EDITOR
diff --git a/Engine/Source/Runtime/Engine/Public/Materials/MaterialIRTypes.h b/Engine/Source/Runtime/Engine/Public/Materials/MaterialIRTypes.h
--- a/Engine/Source/Runtime/Engine/Public/Materials/MaterialIRTypes.h
+++ b/Engine/Source/Runtime/Engine/Public/Materials/MaterialIRTypes.h
@@ -15,17 +15,20 @@ enum ETypeKind
 	TK_Arithmetic
 };
 
+const TCHAR* TypeKindToString(ETypeKind Kind);
+
 struct FType
 {
-	ETypeKind Kind;
+	ETypeKind 	Kind;
 
 	static FTypePtr FromShaderType(const UE::Shader::FType& InShaderType);
 	static FTypePtr GetVoid();
 
-	FArithmeticTypePtr ToArithmetic() const;
-	FArithmeticTypePtr ToScalar() const;
-	FArithmeticTypePtr ToVector() const;
-	FArithmeticTypePtr ToMatrix() const;
+	FStringView GetSpelling() const;
+	FArithmeticTypePtr AsArithmetic() const;
+	FArithmeticTypePtr AsScalar() const;
+	FArithmeticTypePtr AsVector() const;
+	FArithmeticTypePtr AsMatrix() const;
 };
 
 enum EScalarKind
@@ -37,13 +40,18 @@ const TCHAR* ScalarKindToString(EScalarKind Kind);
 
 struct FArithmeticType : FType
 {
+	FStringView Spelling;
 	EScalarKind ScalarKind;
 	int NumRows;
 	int NumColumns;
 
-	static const FArithmeticType* GetScalar(EScalarKind InScalarKind);
-	static const FArithmeticType* GetVector(EScalarKind InScalarKind, int NumRows);
-	static const FArithmeticType* GetMatrix(EScalarKind InScalarKind, int NumColumns, int NumRows);
+	static FArithmeticTypePtr GetBool();
+	static FArithmeticTypePtr GetInt();
+	static FArithmeticTypePtr GetFloat();
+	static FArithmeticTypePtr GetScalar(EScalarKind InScalarKind);
+	static FArithmeticTypePtr GetVector(EScalarKind InScalarKind, int NumRows);
+	static FArithmeticTypePtr GetMatrix(EScalarKind InScalarKind, int NumColumns, int NumRows);
+	static FArithmeticTypePtr Get(EScalarKind InScalarKind, int NumRows, int NumColumns);
 
 	bool IsScalar() const { return NumRows == 1 && NumColumns == 1; }
 	bool IsVector() const { return NumRows > 1 && NumColumns == 1; }
diff --git a/Engine/Source/Runtime/Engine/Public/Materials/MaterialInterface.h b/Engine/Source/Runtime/Engine/Public/Materials/MaterialInterface.h
--- a/Engine/Source/Runtime/Engine/Public/Materials/MaterialInterface.h
+++ b/Engine/Source/Runtime/Engine/Public/Materials/MaterialInterface.h
@@ -434,7 +434,6 @@ public:
 
 #if WITH_EDITOR
 	ENGINE_API virtual const FMaterialCachedHLSLTree& GetCachedHLSLTree(TMicRecursionGuard RecursionGuard = TMicRecursionGuard()) const;
-	ENGINE_API const FMaterialIRModule& GetIRModule() const { return IRModule; }
 #endif
 
 	ENGINE_API bool IsUsingNewHLSLGenerator() const;
@@ -1198,7 +1197,6 @@ protected:
 
 #if WITH_EDITOR
 	TUniquePtr<FMaterialCachedHLSLTree> CachedHLSLTree;
-	FMaterialIRModule IRModule;
 #endif // WITH_EDITOR
 private:
 	/**
