Improvements to r.RayTracing.UseReferenceBasedResidency
- added pool to control memory used by resident geometries (also enables keeping some unreferenced geometries resident to reduce overhead when they're referenced again).
- support always resident geometries.
- keep track of evictable geometries to initialize NotReferencedResidentGeometries with the relevant subset of RegisteredGeometries.
- r.RayTracing.UseReferenceBasedResidency can be toggled at runtime.

diff --git a/Engine/Source/Runtime/Engine/Private/Rendering/RayTracingGeometryManager.cpp b/Engine/Source/Runtime/Engine/Private/Rendering/RayTracingGeometryManager.cpp
--- a/Engine/Source/Runtime/Engine/Private/Rendering/RayTracingGeometryManager.cpp
+++ b/Engine/Source/Runtime/Engine/Private/Rendering/RayTracingGeometryManager.cpp
@@ -14,6 +14,8 @@
 
 #include "Serialization/MemoryReader.h"
 
+#include "Math/UnitConversion.h"
+
 #if RHI_RAYTRACING
 
 static bool bHasRayTracingEnableChanged = false;
@@ -49,7 +51,7 @@ static TAutoConsoleVariable<bool> CVarRayTracingUseReferenceBasedResidency(
 				}
 			);
 		}),
-	ECVF_ReadOnly
+	ECVF_RenderThreadSafe
 );
 
 static int32 GRayTracingStreamingMaxPendingRequests = 128;
@@ -60,6 +62,35 @@ static FAutoConsoleVariableRef CVarNaniteStreamingMaxPendingRequests(
 	ECVF_ReadOnly
 );
 
+static int32 GRayTracingResidentGeometryMemoryPoolSizeInMB = 256;
+static FAutoConsoleVariableRef CVarRayTracingResidentGeometryMemoryPoolSizeInMB(
+	TEXT("r.RayTracing.ResidentGeometryMemoryPoolSizeInMB"),
+	GRayTracingResidentGeometryMemoryPoolSizeInMB,
+	TEXT("Size of the ray tracing geometry pool.\n")
+	TEXT("If pool size is larger than the requested geometry size, some unreferenced geometries will stay resident to reduce build overhead when they are requested again."),
+	ECVF_RenderThreadSafe
+);
+
+static bool bRefreshAlwaysResidentRayTracingGeometries = false;
+
+static int32 GRayTracingNumAlwaysResidentLODs = 1;
+static FAutoConsoleVariableRef CVarRayTracingNumAlwaysResidentLODs(
+	TEXT("r.RayTracing.NumAlwaysResidentLODs"),
+	GRayTracingNumAlwaysResidentLODs,
+	TEXT("Number of LODs per ray tracing geometry group to always keep resident (even when not referenced by TLAS).\n")
+	TEXT("Doesn't apply when ray tracing is disabled, in which case all ray tracing geometry is evicted."),
+	FConsoleVariableDelegate::CreateLambda([](IConsoleVariable* InVariable)
+		{
+			ENQUEUE_RENDER_COMMAND(RefreshAlwaysResidentRayTracingGeometriesCmd)(
+				[](FRHICommandListImmediate&)
+				{
+					bRefreshAlwaysResidentRayTracingGeometries = true;
+				}
+			);
+		}),
+	ECVF_RenderThreadSafe
+);
+
 static int32 GRayTracingMaxBuiltPrimitivesPerFrame = -1;
 static FAutoConsoleVariableRef CVarRayTracingMaxBuiltPrimitivesPerFrame(
 	TEXT("r.RayTracing.Geometry.MaxBuiltPrimitivesPerFrame"),
@@ -82,6 +113,8 @@ DECLARE_DWORD_ACCUMULATOR_STAT(TEXT("Geometry Count"), STAT_RayTracingGeometryCo
 DECLARE_DWORD_ACCUMULATOR_STAT(TEXT("Geometry Group Count"), STAT_RayTracingGeometryGroupCount, STATGROUP_RayTracingGeometry);
 
 DECLARE_MEMORY_STAT(TEXT("Resident Memory"), STAT_RayTracingGeometryResidentMemory, STATGROUP_RayTracingGeometry);
+DECLARE_MEMORY_STAT(TEXT("Always Resident Memory"), STAT_RayTracingGeometryAlwaysResidentMemory, STATGROUP_RayTracingGeometry);
+DECLARE_MEMORY_STAT(TEXT("Requested Memory"), STAT_RayTracingGeometryRequestedMemory, STATGROUP_RayTracingGeometry);
 
 DECLARE_DWORD_ACCUMULATOR_STAT(TEXT("Pending Builds"), STAT_RayTracingPendingBuilds, STATGROUP_RayTracingGeometry);
 DECLARE_DWORD_ACCUMULATOR_STAT(TEXT("Pending Build Primitives"), STAT_RayTracingPendingBuildPrimitives, STATGROUP_RayTracingGeometry);
@@ -151,7 +184,7 @@ RayTracing::GeometryGroupHandle FRayTracingGeometryManager::RegisterRayTracingGe
 	FScopeLock ScopeLock(&MainCS);
 
 	FRayTracingGeometryGroup Group;
-	Group.Geometries.AddDefaulted(NumLODs);
+	Group.GeometryHandles.Init(INDEX_NONE, NumLODs);
 	Group.NumReferences = 1;
 	Group.CurrentFirstLODIdx = CurrentFirstLODIdx;
 
@@ -179,9 +212,9 @@ void FRayTracingGeometryManager::ReleaseRayTracingGeometryGroupReference(RayTrac
 
 	if (Group.NumReferences == 0)
 	{
-		for (FRayTracingGeometry* Geometry : Group.Geometries)
+		for (RayTracingGeometryHandle GeometryHandle : Group.GeometryHandles)
 		{
-			checkf(Geometry == nullptr, TEXT("All FRayTracingGeometry in a group must be unregistered before releasing the group."));
+			checkf(GeometryHandle == INDEX_NONE, TEXT("All FRayTracingGeometry in a group must be unregistered before releasing the group."));
 		}
 
 		check(Group.ProxiesWithCachedRayTracingState.IsEmpty());
@@ -203,6 +236,7 @@ FRayTracingGeometryManager::RayTracingGeometryHandle FRayTracingGeometryManager:
 
 	FRegisteredGeometry& RegisteredGeometry = RegisteredGeometries[Handle];
 	RegisteredGeometry.Geometry = InGeometry;
+	RegisteredGeometry.LastReferencedFrame = 0;
 
 	if (InGeometry->GroupHandle != INDEX_NONE)
 	{
@@ -210,13 +244,20 @@ FRayTracingGeometryManager::RayTracingGeometryHandle FRayTracingGeometryManager:
 
 		FRayTracingGeometryGroup& Group = RegisteredGroups[InGeometry->GroupHandle];
 
-		checkf(InGeometry->LODIndex >= 0 && InGeometry->LODIndex < Group.Geometries.Num(), TEXT("FRayTracingGeometry assigned to a group must have a valid LODIndex"));
-		checkf(Group.Geometries[InGeometry->LODIndex] == nullptr, TEXT("Each LOD inside a FRayTracingGeometryGroup can only be associated with a single FRayTracingGeometry"));
+		checkf(InGeometry->LODIndex >= 0 && InGeometry->LODIndex < Group.GeometryHandles.Num(), TEXT("FRayTracingGeometry assigned to a group must have a valid LODIndex"));
+		checkf(Group.GeometryHandles[InGeometry->LODIndex] == INDEX_NONE, TEXT("Each LOD inside a FRayTracingGeometryGroup can only be associated with a single FRayTracingGeometry"));
 
-		Group.Geometries[InGeometry->LODIndex] = InGeometry;
+		Group.GeometryHandles[InGeometry->LODIndex] = Handle;
 		++Group.NumReferences;
 
-		if (InGeometry->LODIndex >= Group.CurrentFirstLODIdx && IsRayTracingEnabled() && !IsRayTracingUsingReferenceBasedResidency())
+		const bool bAlwaysResident = InGeometry->LODIndex >= Group.GeometryHandles.Num() - GRayTracingNumAlwaysResidentLODs;
+
+		if (bAlwaysResident)
+		{
+			AlwaysResidentGeometries.Add(Handle);
+		}
+
+		if (IsRayTracingEnabled() && InGeometry->LODIndex >= Group.CurrentFirstLODIdx && (!IsRayTracingUsingReferenceBasedResidency() || bAlwaysResident))
 		{
 			PendingStreamingRequests.Add(Handle);
 			INC_DWORD_STAT(STAT_RayTracingPendingStreamingRequests);
@@ -246,21 +287,34 @@ void FRayTracingGeometryManager::ReleaseRayTracingGeometryHandle(RayTracingGeome
 
 		FRayTracingGeometryGroup& Group = RegisteredGroups[RegisteredGeometry.Geometry->GroupHandle];
 
-		checkf(RegisteredGeometry.Geometry->LODIndex >= 0 && RegisteredGeometry.Geometry->LODIndex < Group.Geometries.Num(), TEXT("FRayTracingGeometry assigned to a group must have a valid LODIndex"));
-		checkf(Group.Geometries[RegisteredGeometry.Geometry->LODIndex] == RegisteredGeometry.Geometry, TEXT("Unexpected mismatch of FRayTracingGeometry in FRayTracingGeometryGroup"));
+		checkf(RegisteredGeometry.Geometry->LODIndex >= 0 && RegisteredGeometry.Geometry->LODIndex < Group.GeometryHandles.Num(), TEXT("FRayTracingGeometry assigned to a group must have a valid LODIndex"));
+		checkf(Group.GeometryHandles[RegisteredGeometry.Geometry->LODIndex] == Handle, TEXT("Unexpected mismatch of FRayTracingGeometry in FRayTracingGeometryGroup"));
 
-		Group.Geometries[RegisteredGeometry.Geometry->LODIndex] = nullptr;
+		Group.GeometryHandles[RegisteredGeometry.Geometry->LODIndex] = INDEX_NONE;
 
 		ReleaseRayTracingGeometryGroupReference(RegisteredGeometry.Geometry->GroupHandle);
 	}
 
-	int32 NumRemoved = ResidentGeometries.Remove(RegisteredGeometry.Geometry);
+	{
+		int32 NumRemoved = ResidentGeometries.Remove(Handle);
+
+		if (NumRemoved > 0)
+		{
+			TotalResidentSize -= RegisteredGeometry.Size;
+		}
+	}
 
-	if (NumRemoved > 0)
 	{
-		TotalResidentSize -= RegisteredGeometry.Size;
+		int32 NumRemoved = AlwaysResidentGeometries.Remove(Handle);
+
+		if (NumRemoved > 0)
+		{
+			TotalAlwaysResidentSize -= RegisteredGeometry.Size;
+		}
 	}
 
+	EvictableGeometries.Remove(Handle);
+
 	RegisteredGeometries.RemoveAt(Handle);
 	ReferencedGeometryHandles.Remove(Handle);
 	if (PendingStreamingRequests.Remove(Handle) > 0)
@@ -296,11 +350,20 @@ void FRayTracingGeometryManager::SetRayTracingGeometryGroupCurrentFirstLODIndex(
 		FRHIResourceReplaceBatcher Batcher(RHICmdList, NewCurrentFirstLODIdx - Group.CurrentFirstLODIdx);
 		for (int32 LODIdx = Group.CurrentFirstLODIdx; LODIdx < NewCurrentFirstLODIdx; ++LODIdx)
 		{
+			RayTracingGeometryHandle GeometryHandle = Group.GeometryHandles[LODIdx];
+
 			// some LODs might be stripped during cook
 			// skeletal meshes only create static LOD when rendering as static
-			if (Group.Geometries[LODIdx] && !Group.Geometries[LODIdx]->IsEvicted())
+			if (GeometryHandle == INDEX_NONE)
+			{
+				continue;
+			}
+
+			FRegisteredGeometry& RegisteredGeometry = RegisteredGeometries[GeometryHandle];
+
+			if (!RegisteredGeometry.Geometry->IsEvicted())
 			{
-				Group.Geometries[LODIdx]->ReleaseRHIForStreaming(Batcher);
+				RegisteredGeometry.Geometry->ReleaseRHIForStreaming(Batcher);
 			}
 		}
 	}
@@ -308,9 +371,10 @@ void FRayTracingGeometryManager::SetRayTracingGeometryGroupCurrentFirstLODIndex(
 	{
 		for (int32 LODIdx = NewCurrentFirstLODIdx; LODIdx < Group.CurrentFirstLODIdx; ++LODIdx)
 		{
-			if (Group.Geometries[LODIdx])
+			if (Group.GeometryHandles[LODIdx] != INDEX_NONE)
 			{
-				PendingStreamingRequests.Add(Group.Geometries[LODIdx]->RayTracingGeometryHandle);
+				// TODO: should do this for always resident mips even when using reference based residency
+				PendingStreamingRequests.Add(Group.GeometryHandles[LODIdx]);
 				INC_DWORD_STAT(STAT_RayTracingPendingStreamingRequests);
 			}
 		}
@@ -333,7 +397,7 @@ void FRayTracingGeometryManager::RefreshRegisteredGeometry(RayTracingGeometryHan
 			RegisteredGeometry.Size = RegisteredGeometry.Geometry->GetRHI()->GetSizeInfo().ResultSize;
 
 			bool bAlreadyInSet;
-			ResidentGeometries.Add(RegisteredGeometry.Geometry, &bAlreadyInSet);
+			ResidentGeometries.Add(Handle, &bAlreadyInSet);
 
 			if (bAlreadyInSet)
 			{
@@ -341,17 +405,37 @@ void FRayTracingGeometryManager::RefreshRegisteredGeometry(RayTracingGeometryHan
 			}
 
 			TotalResidentSize += RegisteredGeometry.Size;
+			
+			if (AlwaysResidentGeometries.Contains(Handle))
+			{
+				TotalAlwaysResidentSize -= OldSize;
+				TotalAlwaysResidentSize += RegisteredGeometry.Size;
+			}
+
+			if (RegisteredGeometry.Geometry->GroupHandle != INDEX_NONE)
+			{
+				const FRayTracingGeometryGroup& Group = RegisteredGroups[RegisteredGeometry.Geometry->GroupHandle];
+
+				if (RegisteredGeometry.Geometry->LODIndex < Group.GeometryHandles.Num() - GRayTracingNumAlwaysResidentLODs) // don't want to evict lowest LODs
+				{
+					EvictableGeometries.Add(Handle);
+				}
+			}
 		}
 		else
 		{
-			int32 NumRemoved = ResidentGeometries.Remove(RegisteredGeometry.Geometry);
+			int32 NumRemoved = ResidentGeometries.Remove(Handle);
 
 			if (NumRemoved > 0)
 			{
 				TotalResidentSize -= RegisteredGeometry.Size;
 			}
+
+			EvictableGeometries.Remove(Handle);
 		}
 
+		checkf(!AlwaysResidentGeometries.Contains(Handle) || !RegisteredGeometry.Geometry->IsEvicted() || !IsRayTracingEnabled(), TEXT("Always resident geometries can't be evicted"));
+
 		if (RegisteredGeometry.Geometry->Initializer.Type == ERayTracingGeometryInitializerType::StreamingDestination)
 		{
 			RegisteredGeometry.Status = FRegisteredGeometry::FStatus::StreamedOut;
@@ -393,6 +477,47 @@ void FRayTracingGeometryManager::Tick(FRHICommandList& RHICmdList)
 	checkf(IsRayTracingUsingReferenceBasedResidency() || (ReferencedGeometryHandles.IsEmpty() && ReferencedGeometryGroups.IsEmpty()),
 		TEXT("ReferencedGeometryHandles and ReferencedGeometryGroups are expected to be empty when not using reference based residency"));
 
+	if (bRefreshAlwaysResidentRayTracingGeometries)
+	{
+		bRefreshAlwaysResidentRayTracingGeometries = false;
+
+		AlwaysResidentGeometries.Empty();
+		TotalAlwaysResidentSize = 0;
+
+		for (FRegisteredGeometry& RegisteredGeometry : RegisteredGeometries)
+		{
+			if (RegisteredGeometry.Geometry->GroupHandle == INDEX_NONE)
+			{
+				continue;
+			}
+
+			const FRayTracingGeometryGroup& Group = RegisteredGroups[RegisteredGeometry.Geometry->GroupHandle];
+			if (RegisteredGeometry.Geometry->LODIndex >= Group.GeometryHandles.Num() - GRayTracingNumAlwaysResidentLODs)
+			{
+				AlwaysResidentGeometries.Add(RegisteredGeometry.Geometry->RayTracingGeometryHandle);
+				TotalAlwaysResidentSize += RegisteredGeometry.Size;
+
+				if (RegisteredGeometry.Geometry->IsEvicted())
+				{
+					RegisteredGeometry.Geometry->MakeResident(RHICmdList);
+				}
+
+				if (!RequestRayTracingGeometryStreamIn(RHICmdList, RegisteredGeometry.Geometry->RayTracingGeometryHandle))
+				{
+					PendingStreamingRequests.Add(RegisteredGeometry.Geometry->RayTracingGeometryHandle);
+
+					INC_DWORD_STAT(STAT_RayTracingPendingStreamingRequests);
+				}
+
+				EvictableGeometries.Remove(RegisteredGeometry.Geometry->RayTracingGeometryHandle);
+			}
+			else if (IsRayTracingUsingReferenceBasedResidency() && RegisteredGeometry.Geometry->GetRHI() != nullptr)
+			{
+				RegisteredGeometry.Geometry->Evict();
+			}
+		}
+	}
+
 	if (!IsRayTracingEnabled())
 	{
 		if (bHasRayTracingEnableChanged)
@@ -428,6 +553,8 @@ void FRayTracingGeometryManager::Tick(FRHICommandList& RHICmdList)
 		);
 
 		check(PendingStreamingRequests.IsEmpty());
+
+		SET_MEMORY_STAT(STAT_RayTracingGeometryRequestedMemory, 0);
 	}
 	else if (IsRayTracingUsingReferenceBasedResidency())
 	{
@@ -443,40 +570,88 @@ void FRayTracingGeometryManager::Tick(FRHICommandList& RHICmdList)
 
 		bRenderedFrame = false;
 
-		TSet<FRayTracingGeometry*> NotReferencedResidentGeometries = ResidentGeometries;
+		if (bHasRayTracingEnableChanged)
+		{
+			// make always resident geometries actually resident
 
-		TArray<FRayTracingGeometry*> ReferencedGeometries;
+			for (RayTracingGeometryHandle GeometryHandle : AlwaysResidentGeometries)
+			{
+				FRegisteredGeometry& RegisteredGeometry = RegisteredGeometries[GeometryHandle];
 
-		// 1st step
-		// - map RayTracingGeometryHandle to FRayTracingGeometry
-		//		- we use handles to track referenced geometries because some of the higher level code uses const FRayTracingGeometry*
+				if (RegisteredGeometry.Geometry->IsEvicted())
+				{
+					RegisteredGeometry.Geometry->MakeResident(RHICmdList);
+				}
+				if (!RequestRayTracingGeometryStreamIn(RHICmdList, GeometryHandle))
+				{
+					PendingStreamingRequests.Add(GeometryHandle);
+
+					INC_DWORD_STAT(STAT_RayTracingPendingStreamingRequests);
+				}
+			}
+		}
+
+		TSet<RayTracingGeometryHandle> NotReferencedResidentGeometries = EvictableGeometries;
+
+		TArray<RayTracingGeometryHandle> ReferencedGeometries;
+
+		uint64 RequestedSize = 0;
+		uint64 RequestedButEvictedSize = 0;
+
+		// Step 1
+		// - update LastReferencedFrame of referenced geometries and calculate memory required to make evicted geometries resident
 		for (RayTracingGeometryHandle GeometryHandle : ReferencedGeometryHandles)
 		{
 			FRegisteredGeometry& RegisteredGeometry = RegisteredGeometries[GeometryHandle];
+			RegisteredGeometry.LastReferencedFrame = GFrameCounterRenderThread;
+
+			ReferencedGeometries.Add(GeometryHandle);
+			NotReferencedResidentGeometries.Remove(GeometryHandle);
 
-			ReferencedGeometries.Add(RegisteredGeometry.Geometry);
-			NotReferencedResidentGeometries.Remove(RegisteredGeometry.Geometry);
+			RequestedSize += RegisteredGeometry.Size;
+
+			if (RegisteredGeometry.Geometry->IsEvicted())
+			{
+				RequestedButEvictedSize += RegisteredGeometry.Size;
+			}
 		}
 
-		// 2nd step
+		// Step 2
 		// - add all geometries in referenced groups to ReferencedGeometries
 		//		- need to make all geometries in group resident otherwise might not have valid geometry when reducing LOD
 		//		- TODO: Could track TargetLOD and only make [TargetLOD ... LastLOD] range resident
+		// - also update LastReferencedFrame and calculate memory required to make evicted geometries resident
 		for (RayTracing::GeometryGroupHandle GroupHandle : ReferencedGeometryGroups)
 		{
 			checkf(RegisteredGroups.IsValidIndex(GroupHandle), TEXT("RayTracingGeometryGroupHandle must be valid"));
 
 			const FRayTracingGeometryGroup& Group = RegisteredGroups[GroupHandle];
 
-			for (uint8 LODIndex = Group.CurrentFirstLODIdx; LODIndex < Group.Geometries.Num(); ++LODIndex)
+			for (uint8 LODIndex = Group.CurrentFirstLODIdx; LODIndex < Group.GeometryHandles.Num(); ++LODIndex)
 			{
-				FRayTracingGeometry* Geometry = Group.Geometries[LODIndex];
+				RayTracingGeometryHandle GeometryHandle = Group.GeometryHandles[LODIndex];
 
-				// some LODs might be stripped during cook
-				if (Geometry != nullptr)
+				if (GeometryHandle != INDEX_NONE) // some LODs might be stripped during cook
 				{
-					ReferencedGeometries.Add(Geometry);
-					NotReferencedResidentGeometries.Remove(Geometry);
+					FRegisteredGeometry& RegisteredGeometry = RegisteredGeometries[GeometryHandle];
+					RegisteredGeometry.LastReferencedFrame = GFrameCounterRenderThread;
+
+					if (RegisteredGeometry.Geometry->LODIndex >= Group.GeometryHandles.Num() - GRayTracingNumAlwaysResidentLODs)
+					{
+						checkf(!RegisteredGeometry.Geometry->IsEvicted(), TEXT("Always resident ray tracing geometry was unexpectely evicted."));
+					}
+					else
+					{
+						ReferencedGeometries.Add(GeometryHandle);
+						NotReferencedResidentGeometries.Remove(GeometryHandle);
+
+						RequestedSize += RegisteredGeometry.Size;
+
+						if (RegisteredGeometry.Geometry->IsEvicted())
+						{
+							RequestedButEvictedSize += RegisteredGeometry.Size;
+						}
+					}
 				}
 			}
 		}
@@ -485,28 +660,86 @@ void FRayTracingGeometryManager::Tick(FRHICommandList& RHICmdList)
 		ensure(ReferencedGeometries.Num() == TSet(ReferencedGeometries).Num());
 #endif
 
-		// 3rd step
-		// - make referenced geometries resident
-		for (FRayTracingGeometry* Geometry : ReferencedGeometries)
+		SET_MEMORY_STAT(STAT_RayTracingGeometryRequestedMemory, RequestedSize);
+
+		uint64 ResidentGeometryMemoryPoolSize = FUnitConversion::Convert(GRayTracingResidentGeometryMemoryPoolSizeInMB, EUnit::Megabytes, EUnit::Bytes);
+
+		// Step 3
+		// - if making requested geometries resident will put us over budget -> evict some geometry not referenced by TLAS
+		if (TotalResidentSize + RequestedButEvictedSize > ResidentGeometryMemoryPoolSize)
 		{
-			NotReferencedResidentGeometries.Remove(Geometry);
+			TArray<RayTracingGeometryHandle> NotReferencedResidentGeometriesArray = NotReferencedResidentGeometries.Array();
 
-			if (Geometry->IsEvicted())
+			// Step 3.1
+			// - sort to evict geometries in the following order:
+			//		- least recently used
+			//		- largest geometries
+			Algo::Sort(NotReferencedResidentGeometriesArray, [this](RayTracingGeometryHandle& LHSHandle, RayTracingGeometryHandle& RHSHandle)
+				{
+					FRegisteredGeometry& LHS = RegisteredGeometries[LHSHandle];
+					FRegisteredGeometry& RHS = RegisteredGeometries[RHSHandle];
+
+					// TODO: evict unreferenced dynamic geometries using shared buffers first since they need to be rebuild anyway
+					// (and then dynamic geometries requiring update?
+
+					// 1st - last referenced frame
+					if (LHS.LastReferencedFrame != RHS.LastReferencedFrame)
+					{
+						return LHS.LastReferencedFrame < RHS.LastReferencedFrame;
+					}
+
+					// 2nd - size
+					return LHS.Size > RHS.Size;
+				});
+
+			// Step 3.2
+			// - evict geometries until we are in budget
+			int32 Index = 0;
+			while (TotalResidentSize + RequestedSize > ResidentGeometryMemoryPoolSize && Index < NotReferencedResidentGeometriesArray.Num())
 			{
-				Geometry->MakeResident(RHICmdList);
-			}
+				RayTracingGeometryHandle GeometryHandle = NotReferencedResidentGeometriesArray[Index];
+				FRegisteredGeometry& RegisteredGeometry = RegisteredGeometries[GeometryHandle];
+
+				check(RegisteredGeometry.Geometry->IsValid() && !RegisteredGeometry.Geometry->IsEvicted());
 
-			RequestRayTracingGeometryStreamIn(RHICmdList, Geometry->RayTracingGeometryHandle);
+				RegisteredGeometry.Geometry->Evict();
+
+				++Index;
+			}
 		}
 
-		// 4th step
-		// - evict geometries not referenced by TLAS
-		// - TODO: keep unreferenced geometries up to some pool size resident to prevent resident/evicted loops
-		for (FRayTracingGeometry* Geometry : NotReferencedResidentGeometries)
+		// Step 4
+		// - make referenced geometries resident until we go over budget
+		if (TotalResidentSize < ResidentGeometryMemoryPoolSize)
 		{
-			if (Geometry->GetRHI() != nullptr)
+			// Step 4.1
+			//	- sort by size to prioritize smaller geometries
+			Algo::Sort(ReferencedGeometries, [this](RayTracingGeometryHandle& LHSHandle, RayTracingGeometryHandle& RHSHandle)
+				{
+					FRegisteredGeometry& LHS = RegisteredGeometries[LHSHandle];
+					FRegisteredGeometry& RHS = RegisteredGeometries[RHSHandle];
+
+					return LHS.Size < RHS.Size;
+				});
+
+			// Step 3.2
+			// - make geometries resident until we go over budget
+			int32 Index = 0;
+			while (TotalResidentSize < ResidentGeometryMemoryPoolSize && Index < ReferencedGeometries.Num())
 			{
-				Geometry->Evict();
+				// if referenced this frame, mark for eviction and add to pending list
+
+				RayTracingGeometryHandle GeometryHandle = ReferencedGeometries[Index];
+				FRegisteredGeometry& RegisteredGeometry = RegisteredGeometries[GeometryHandle];
+
+				if (RegisteredGeometry.Geometry->IsEvicted())
+				{
+					RegisteredGeometry.Geometry->MakeResident(RHICmdList);
+				}
+
+				RequestRayTracingGeometryStreamIn(RHICmdList, GeometryHandle);
+
+				++Index;
 			}
 		}
 	}
@@ -541,23 +774,27 @@ void FRayTracingGeometryManager::Tick(FRHICommandList& RHICmdList)
 				checkf(!RegisteredGeometry.Geometry->IsEvicted(), TEXT("Ray tracing geometry should not be evicted when ray tracing is enabled."));
 			}
 #endif
-			
-			TSet<RayTracingGeometryHandle> CurrentPendingStreamingRequests;
-			Swap(CurrentPendingStreamingRequests, PendingStreamingRequests);
-			PendingStreamingRequests.Reserve(CurrentPendingStreamingRequests.Num());
+		}
+
+		SET_MEMORY_STAT(STAT_RayTracingGeometryRequestedMemory, TotalResidentSize);
+	}
 
-			for (RayTracingGeometryHandle GeometryHandle : CurrentPendingStreamingRequests)
+	{
+		TSet<RayTracingGeometryHandle> CurrentPendingStreamingRequests;
+		Swap(CurrentPendingStreamingRequests, PendingStreamingRequests);
+		PendingStreamingRequests.Reserve(CurrentPendingStreamingRequests.Num());
+
+		for (RayTracingGeometryHandle GeometryHandle : CurrentPendingStreamingRequests)
+		{
+			if (!RequestRayTracingGeometryStreamIn(RHICmdList, GeometryHandle))
 			{
-				if (!RequestRayTracingGeometryStreamIn(RHICmdList, GeometryHandle))
-				{
-					PendingStreamingRequests.Add(GeometryHandle);
-				}
+				PendingStreamingRequests.Add(GeometryHandle);
 			}
-
-			SET_DWORD_STAT(STAT_RayTracingPendingStreamingRequests, PendingStreamingRequests.Num());
 		}
 	}
 
+	SET_DWORD_STAT(STAT_RayTracingPendingStreamingRequests, PendingStreamingRequests.Num());
+
 	ProcessCompletedStreamingRequests(RHICmdList);
 
 	ReferencedGeometryHandles.Reset();
@@ -566,6 +803,7 @@ void FRayTracingGeometryManager::Tick(FRHICommandList& RHICmdList)
 	bHasRayTracingEnableChanged = false;
 
 	SET_MEMORY_STAT(STAT_RayTracingGeometryResidentMemory, TotalResidentSize);
+	SET_MEMORY_STAT(STAT_RayTracingGeometryAlwaysResidentMemory, TotalAlwaysResidentSize);
 }
 
 bool FRayTracingGeometryManager::RequestRayTracingGeometryStreamIn(FRHICommandList& RHICmdList, RayTracingGeometryHandle GeometryHandle)
diff --git a/Engine/Source/Runtime/Engine/Public/Rendering/RayTracingGeometryManager.h b/Engine/Source/Runtime/Engine/Public/Rendering/RayTracingGeometryManager.h
--- a/Engine/Source/Runtime/Engine/Public/Rendering/RayTracingGeometryManager.h
+++ b/Engine/Source/Runtime/Engine/Public/Rendering/RayTracingGeometryManager.h
@@ -88,7 +88,7 @@ private:
 
 	struct FRayTracingGeometryGroup
 	{
-		TArray<FRayTracingGeometry*> Geometries;
+		TArray<RayTracingGeometryHandle> GeometryHandles;
 
 		TSet<FPrimitiveSceneProxy*> ProxiesWithCachedRayTracingState;
 
@@ -105,6 +105,7 @@ private:
 	struct FRegisteredGeometry
 	{
 		FRayTracingGeometry* Geometry = nullptr;
+		uint64 LastReferencedFrame = 0;
 		uint32 Size = 0;
 
 		FByteBulkData* StreamableData = nullptr;
@@ -134,9 +135,14 @@ private:
 	// Used for keeping track of geometries when ray tracing is dynamic
 	TSparseArray<FRegisteredGeometry> RegisteredGeometries;
 
-	TSet<FRayTracingGeometry*> ResidentGeometries;
+	TSet<RayTracingGeometryHandle> ResidentGeometries;
 	uint64 TotalResidentSize = 0;
 
+	TSet<RayTracingGeometryHandle> AlwaysResidentGeometries;
+	uint64 TotalAlwaysResidentSize = 0;
+
+	TSet<RayTracingGeometryHandle> EvictableGeometries;
+
 	TSet<RayTracingGeometryHandle> ReferencedGeometryHandles;
 	TSet<RayTracing::GeometryGroupHandle> ReferencedGeometryGroups;
 
