[HISM]
 - Run game-thread task to a safe place in the frame using the ticker

diff --git a/Engine/Source/Runtime/Engine/Private/HierarchicalInstancedStaticMesh.cpp b/Engine/Source/Runtime/Engine/Private/HierarchicalInstancedStaticMesh.cpp
--- a/Engine/Source/Runtime/Engine/Private/HierarchicalInstancedStaticMesh.cpp
+++ b/Engine/Source/Runtime/Engine/Private/HierarchicalInstancedStaticMesh.cpp
@@ -19,6 +19,7 @@
 #include "UObject/UObjectIterator.h"
 #include "RenderUtils.h"
 #include "UnrealEngine.h"
+#include "Containers/Ticker.h"
 #include "InstancedStaticMeshDelegates.h"
 #include "UObject/ReleaseObjectVersion.h"
 #include "ComponentRecreateRenderStateContext.h"
@@ -2871,11 +2872,25 @@ void UHierarchicalInstancedStaticMeshComponent::BuildTreeAsync()
 
 		// add a dependent task to run on the main thread when build is complete
 		FGraphEventRef PostBuildTreeAsyncResult(
-			FDelegateGraphTask::CreateAndDispatchWhenReady(
-			FDelegateGraphTask::FDelegate::CreateUObject(this, &UHierarchicalInstancedStaticMeshComponent::ApplyBuildTreeAsync, Builder, StartTime), GET_STATID(STAT_FoliageBuildTime),
-			BuildTreeAsyncResult, ENamedThreads::GameThread, ENamedThreads::GameThread
+			FFunctionGraphTask::CreateAndDispatchWhenReady(
+				[WeakPtr = TWeakObjectPtr(this), Builder, StartTime]()
+				{
+					ExecuteOnGameThread(
+						TEXT("ApplyBuildTreeAsync_GameThread"),
+						[WeakPtr, Builder, StartTime]()
+						{
+							if (UHierarchicalInstancedStaticMeshComponent* Ptr = WeakPtr.Get())
+							{
+								FGraphEventRef Dummy;
+								Ptr->ApplyBuildTreeAsync(ENamedThreads::GameThread, Dummy, Builder, StartTime);
+							}
+						}
+					);
+				},
+				GET_STATID(STAT_FoliageBuildTime),
+				BuildTreeAsyncResult
 			)
-			);
+		);
 
 		BuildTreeAsyncTasks.Add(PostBuildTreeAsyncResult);
 	}
