Disallow adding actors to a child data layer if the parent is an External Data Layer and the actors are not part of it.

diff --git a/Engine/Source/Runtime/Engine/Private/WorldPartition/DataLayer/DataLayerInstance.cpp b/Engine/Source/Runtime/Engine/Private/WorldPartition/DataLayer/DataLayerInstance.cpp
--- a/Engine/Source/Runtime/Engine/Private/WorldPartition/DataLayer/DataLayerInstance.cpp
+++ b/Engine/Source/Runtime/Engine/Private/WorldPartition/DataLayer/DataLayerInstance.cpp
@@ -331,6 +331,18 @@ bool UDataLayerInstance::CanAddActor(AActor* InActor, FText* OutReason) const
 		return false;
 	}
 
+	const UExternalDataLayerInstance* ExternalDataLayerInstance = GetRootExternalDataLayerInstance();
+	if (ExternalDataLayerInstance && (ExternalDataLayerInstance != this))
+	{
+		if (InActor->GetExternalDataLayerAsset() != ExternalDataLayerInstance->GetAsset())
+		{
+			if (OutReason)
+			{
+				*OutReason = LOCTEXT("CantAddActorMismatchRootExternalDataLayer", "Can't assign actors to a Data Layer child of an External Data Layer that is not already assigned to the actors");
+			}
+			return false;
+		}
+	}
 	return true;
 }
 
