[World Partition Runtime HashSet] Fix missing check for null HLOD layer when registering streaming world asset.

diff --git a/Engine/Source/Runtime/Engine/Private/WorldPartition/RuntimeHashSet/WorldPartitionRuntimeHashSet.cpp b/Engine/Source/Runtime/Engine/Private/WorldPartition/RuntimeHashSet/WorldPartitionRuntimeHashSet.cpp
--- a/Engine/Source/Runtime/Engine/Private/WorldPartition/RuntimeHashSet/WorldPartitionRuntimeHashSet.cpp
+++ b/Engine/Source/Runtime/Engine/Private/WorldPartition/RuntimeHashSet/WorldPartitionRuntimeHashSet.cpp
@@ -670,6 +670,11 @@ FGuid UWorldPartitionRuntimeHashSet::RegisterWorldAssetStreaming(const UWorldPar
 			continue;
 		}
 
+		if (bIsHLODPass && InParams.WorldAssetHLOD.IsNull())
+		{
+			continue;
+		}
+
 		FRuntimePartitionStreamingData StreamingData;
 		StreamingData.LoadingRange = -1;
 
