Fixed Level Instance not properly propagating its visibility flag driven by [REDACTED_PERSON] state.

diff --git a/Engine/Source/Runtime/Engine/Private/LevelInstance/LevelInstanceLevelStreaming.cpp b/Engine/Source/Runtime/Engine/Private/LevelInstance/LevelInstanceLevelStreaming.cpp
--- a/Engine/Source/Runtime/Engine/Private/LevelInstance/LevelInstanceLevelStreaming.cpp
+++ b/Engine/Source/Runtime/Engine/Private/LevelInstance/LevelInstanceLevelStreaming.cpp
@@ -231,6 +231,9 @@ void ULevelStreamingLevelInstance::InitializeActors(const TArray<AActor*>& InAct
 	{
 		if (ILevelInstanceInterface* LevelInstance = GetLevelInstance())
 		{
+			const AActor* LevelInstanceActor = CastChecked<AActor>(LevelInstance);
+			const bool bIsHiddenEdLayer = LevelInstanceActor->bHiddenEdLayer;
+			const bool bIsHiddenEdTemporary = LevelInstanceActor->IsTemporarilyHiddenInEditor();
 			const FActorContainerID ContainerID = LevelInstance->GetLevelInstanceID().GetContainerID();
 		
 			for (AActor* Actor : InActors)
@@ -252,6 +255,9 @@ void ULevelStreamingLevelInstance::InitializeActors(const TArray<AActor*>& InAct
 						}
 					}
 
+					Actor->SetIsHiddenEdLayer(bIsHiddenEdLayer);
+					Actor->SetIsTemporarilyHiddenInEditor(bIsHiddenEdTemporary);
+
 					// Must happen before the actors are registered with the world, which is the case for this delegate.
 					FSetActorInstanceGuid SetActorInstanceGuid(Actor, ContainerID.GetActorGuid(Actor->GetActorGuid()));
 
@@ -270,6 +276,9 @@ void ULevelStreamingLevelInstance::OnLoadedActorsAddedToLevelPostEvent(const TAr
 		{
 			if (ILevelInstanceInterface* LevelInstance = GetLevelInstance())
 			{
+				const AActor* LevelInstanceActor = CastChecked<AActor>(LevelInstance);
+				const bool bIsInEditLevelInstanceHierarchy = LevelInstanceActor->IsInEditLevelInstanceHierarchy();
+
 				for (AActor* Actor : InActors)
 				{
 					if (IsValid(Actor))
@@ -280,10 +289,7 @@ void ULevelStreamingLevelInstance::OnLoadedActorsAddedToLevelPostEvent(const TAr
 						}
 
 						Actor->PushSelectionToProxies();
-						if (LevelInstance)
-						{
-							Actor->PushLevelInstanceEditingStateToProxies(CastChecked<AActor>(LevelInstance)->IsInEditLevelInstanceHierarchy());
-						}
+						Actor->PushLevelInstanceEditingStateToProxies(bIsInEditLevelInstanceHierarchy);
 
 						if (LevelInstanceEditorInstanceActor.IsValid())
 						{
