Correctly mark an old seamless travel world as having stopped play before final cleanup, this fixes a warning and a missed callback

diff --git a/Engine/Source/Runtime/Engine/Private/World.cpp b/Engine/Source/Runtime/Engine/Private/World.cpp
--- a/Engine/Source/Runtime/Engine/Private/World.cpp
+++ b/Engine/Source/Runtime/Engine/Private/World.cpp
@@ -7861,6 +7861,7 @@ UWorld* FSeamlessTravelHandler::Tick()
 			}
 
 			// only consider session ended if we're making the final switch so that HUD, etc. UI elements stay around until the end
+			CurrentWorld->SetBegunPlay(false);
 			CurrentWorld->CleanupWorld(bSwitchedToDefaultMap);
 			CurrentWorld->RemoveFromRoot();
 			CurrentWorld->ClearFlags(RF_Standalone);
