[NetPhysics] Fix issue where IsLocallyControlled didn't get marshaled to the physics thread if not set on BeginPlay.

#changelist validated
#virtualized

diff --git a/Engine/Source/Runtime/Engine/Private/PhysicsEngine/NetworkPhysicsComponent.cpp b/Engine/Source/Runtime/Engine/Private/PhysicsEngine/NetworkPhysicsComponent.cpp
--- a/Engine/Source/Runtime/Engine/Private/PhysicsEngine/NetworkPhysicsComponent.cpp
+++ b/Engine/Source/Runtime/Engine/Private/PhysicsEngine/NetworkPhysicsComponent.cpp
@@ -470,7 +470,7 @@ void UNetworkPhysicsComponent::BeginPlay()
 	Super::BeginPlay();
 
 	// Update async component with current component properties
-	UpdateAsyncComponent();
+	UpdateAsyncComponent(true);
 }
 
 void UNetworkPhysicsComponent::InitializeComponent()
@@ -644,7 +644,6 @@ void UNetworkPhysicsComponent::NetworkMarshaledData()
 			{
 				bIsRelayingLocalInputs = false;
 				bStopRelayingLocalInputsDeferred = false;
-				UpdateAsyncComponent();
 			}
 		}
 	}
@@ -897,7 +896,6 @@ bool UNetworkPhysicsComponent::IsLocallyControlled() const
 void UNetworkPhysicsComponent::SetIsRelayingLocalInputs(bool bInRelayingLocalInputs)
 {
 	bIsRelayingLocalInputs = bInRelayingLocalInputs;
-	UpdateAsyncComponent();
 }
 
 APlayerController* UNetworkPhysicsComponent::GetPlayerController() const
@@ -942,10 +940,11 @@ void UNetworkPhysicsComponent::UpdateAsyncComponent(const bool bFullUpdate)
 				}
 			}
 
+			// bIsLocallyControlled is marshaled outside of the bFullUpdate because it's not always set at BeginPlay when last bFullUpdate is called.
+			AsyncInput->bIsLocallyControlled = IsLocallyControlled();
+
 			if (bFullUpdate)
 			{
-				AsyncInput->bIsLocallyControlled = IsLocallyControlled();
-
 				if (UWorld* World = GetWorld())
 				{ 
 					AsyncInput->NetMode = World->GetNetMode();
@@ -962,7 +961,6 @@ void UNetworkPhysicsComponent::UpdateAsyncComponent(const bool bFullUpdate)
 				{
 					AsyncInput->ActorComponent = ActorComponent;
 				}
-
 			}
 		}
 	}
diff --git a/Engine/Source/Runtime/Engine/Public/Physics/NetworkPhysicsComponent.h b/Engine/Source/Runtime/Engine/Public/Physics/NetworkPhysicsComponent.h
--- a/Engine/Source/Runtime/Engine/Public/Physics/NetworkPhysicsComponent.h
+++ b/Engine/Source/Runtime/Engine/Public/Physics/NetworkPhysicsComponent.h
@@ -823,7 +823,7 @@ private:
 	void NetworkMarshaledData();
 
 	/** Update the async component with GT properties */
-	void UpdateAsyncComponent(const bool bFullUpdate = true);
+	void UpdateAsyncComponent(const bool bFullUpdate);
 
 	/** Setup data type and history in async input and trigger the creation of data history in the async component */
 	ENGINE_API void CreateAsyncDataHistory();
