Exposed a util to gather a FStringTable hosted within an asset

#rnx

diff --git a/Engine/Source/Runtime/Engine/Private/Internationalization/StringTable.cpp b/Engine/Source/Runtime/Engine/Private/Internationalization/StringTable.cpp
--- a/Engine/Source/Runtime/Engine/Private/Internationalization/StringTable.cpp
+++ b/Engine/Source/Runtime/Engine/Private/Internationalization/StringTable.cpp
@@ -14,14 +14,12 @@
 
 #include UE_INLINE_GENERATED_CPP_BY_NAME(StringTable)
 
-#if WITH_EDITORONLY_DATA
-namespace
+namespace UE::StringTable
 {
-	static void GatherStringTableForLocalization(const UObject* const Object, FPropertyLocalizationDataGatherer& PropertyLocalizationDataGatherer, const EPropertyLocalizationGathererTextFlags GatherTextFlags)
+#if WITH_EDITORONLY_DATA
+	void GatherForLocalization(const FString& SourceLocation, const FStringTableConstRef& StringTable, FPropertyLocalizationDataGatherer& PropertyLocalizationDataGatherer, const EPropertyLocalizationGathererTextFlags GatherTextFlags)
 	{
-		FStringTableConstRef StringTable = CastChecked<UStringTable>(Object)->GetStringTable();
-
-		auto FindOrAddTextData = [&](const FString& InText) -> FGatherableTextData&
+		auto FindOrAddTextData = [&StringTable, &PropertyLocalizationDataGatherer](const FString& InText) -> FGatherableTextData&
 		{
 			check(!InText.IsEmpty());
 
@@ -45,9 +43,7 @@ namespace
 			return *GatherableTextData;
 		};
 
-		const FString SourceLocation = Object->GetPathName();
-
-		StringTable->EnumerateSourceStrings([&](const FString& InKey, const FString& InSourceString) -> bool
+		StringTable->EnumerateSourceStrings([&SourceLocation, &StringTable, &FindOrAddTextData](const FString& InKey, const FString& InSourceString) -> bool
 		{
 			if (!InSourceString.IsEmpty())
 			{
@@ -69,8 +65,18 @@ namespace
 			return true; // continue enumeration
 		});
 	}
+
+	namespace Private
+	{
+		void GatherAssetForLocalization(const UObject* const Object, FPropertyLocalizationDataGatherer& PropertyLocalizationDataGatherer, const EPropertyLocalizationGathererTextFlags GatherTextFlags)
+		{
+			FStringTableConstRef StringTable = CastChecked<UStringTable>(Object)->GetStringTable();
+			const FString SourceLocation = Object->GetPathName();
+			GatherForLocalization(SourceLocation, StringTable, PropertyLocalizationDataGatherer, GatherTextFlags);
+		}
+	}
+#endif	// WITH_EDITORONLY_DATA
 }
-#endif
 
 class FStringTableEngineBridge : public IStringTableEngineBridge, public FGCObject
 {
@@ -354,7 +360,7 @@ UStringTable::UStringTable()
 	}
 
 #if WITH_EDITORONLY_DATA
-	{ static const FAutoRegisterLocalizationDataGatheringCallback AutomaticRegistrationOfLocalizationGatherer(UStringTable::StaticClass(), &GatherStringTableForLocalization); }
+	{ static const FAutoRegisterLocalizationDataGatheringCallback AutomaticRegistrationOfLocalizationGatherer(UStringTable::StaticClass(), &UE::StringTable::Private::GatherAssetForLocalization); }
 #endif
 }
 
@@ -421,4 +427,3 @@ FStringTableRef UStringTable::GetMutableStringTable() const
 {
 	return StringTable.ToSharedRef();
 }
-
diff --git a/Engine/Source/Runtime/Engine/Public/Internationalization/StringTable.h b/Engine/Source/Runtime/Engine/Public/Internationalization/StringTable.h
--- a/Engine/Source/Runtime/Engine/Public/Internationalization/StringTable.h
+++ b/Engine/Source/Runtime/Engine/Public/Internationalization/StringTable.h
@@ -7,6 +7,22 @@
 #include "UObject/Object.h"
 #include "StringTable.generated.h"
 
+class FPropertyLocalizationDataGatherer;
+enum class EPropertyLocalizationGathererTextFlags : uint8;
+
+namespace UE::StringTable
+{
+
+#if WITH_EDITORONLY_DATA
+/**
+ * Utility to gather the string table data stored within an asset.
+ * @note Primarily used internally by UStringTable, but may be used by any other assets that host a FStringTable and need to handle gathering it for localization via a FPropertyLocalizationDataGatherer.
+ */
+ENGINE_API void GatherForLocalization(const FString& SourceLocation, const FStringTableConstRef& StringTable, FPropertyLocalizationDataGatherer& PropertyLocalizationDataGatherer, const EPropertyLocalizationGathererTextFlags GatherTextFlags);
+#endif	// WITH_EDITORONLY_DATA
+
+}
+
 /** String table wrapper asset */
 UCLASS(MinimalAPI)
 class UStringTable : public UObject
