Fixed regression in Level Instances Bounds : Computed LI Editor Bounds was equal to its LevelInstanceComponent bounds which is the pivot, this caused LIs to only load if their pivot was intersecting the loading region. Also no bounds were displayed for Level Instances in the WP editor (when option to show actor bounds was enabled).

diff --git a/Engine/Source/Runtime/Engine/Private/LevelInstance/LevelInstanceActor.cpp b/Engine/Source/Runtime/Engine/Private/LevelInstance/LevelInstanceActor.cpp
--- a/Engine/Source/Runtime/Engine/Private/LevelInstance/LevelInstanceActor.cpp
+++ b/Engine/Source/Runtime/Engine/Private/LevelInstance/LevelInstanceActor.cpp
@@ -293,6 +293,17 @@ FBox ALevelInstance::GetStreamingBounds() const
 	return Super::GetStreamingBounds();
 }
 
+FBox ALevelInstance::GetStreamingBoundsEditor() const
+{
+	FBox LevelInstanceEditorBounds;
+	if (LevelInstanceActorImpl.GetEditorBounds(LevelInstanceEditorBounds))
+	{
+		return LevelInstanceEditorBounds;
+	}
+
+	return Super::GetStreamingBoundsEditor();
+}
+
 bool ALevelInstance::IsUserManaged() const
 {
 	return Super::IsUserManaged() || LevelInstanceActorImpl.IsUserManaged();
diff --git a/Engine/Source/Runtime/Engine/Private/LevelInstance/LevelInstanceActorImpl.cpp b/Engine/Source/Runtime/Engine/Private/LevelInstance/LevelInstanceActorImpl.cpp
--- a/Engine/Source/Runtime/Engine/Private/LevelInstance/LevelInstanceActorImpl.cpp
+++ b/Engine/Source/Runtime/Engine/Private/LevelInstance/LevelInstanceActorImpl.cpp
@@ -385,6 +385,22 @@ bool FLevelInstanceActorImpl::GetBounds(FBox& OutBounds) const
 	return false;
 }
 
+bool FLevelInstanceActorImpl::GetEditorBounds(FBox& OutBounds) const
+{
+	if (LevelInstance->IsLoadingEnabled())
+	{
+		if (ULevelInstanceSubsystem* LevelInstanceSubsystem = LevelInstance->GetLevelInstanceSubsystem())
+		{
+			if (LevelInstanceSubsystem->GetLevelInstanceEditorBounds(LevelInstance, OutBounds))
+			{
+				return true;
+			}
+		}
+	}
+
+	return false;
+}
+
 void FLevelInstanceActorImpl::PushSelectionToProxies()
 {
 	// Actors of the LevelInstance need to reflect the LevelInstance actor's selected state
diff --git a/Engine/Source/Runtime/Engine/Private/LevelInstance/LevelInstanceSubsystem.cpp b/Engine/Source/Runtime/Engine/Private/LevelInstance/LevelInstanceSubsystem.cpp
--- a/Engine/Source/Runtime/Engine/Private/LevelInstance/LevelInstanceSubsystem.cpp
+++ b/Engine/Source/Runtime/Engine/Private/LevelInstance/LevelInstanceSubsystem.cpp
@@ -998,6 +998,16 @@ void ULevelInstanceSubsystem::OnExitEditorModeInternal(bool bForceExit)
 }
 
 bool ULevelInstanceSubsystem::GetLevelInstanceBounds(const ILevelInstanceInterface* LevelInstance, FBox& OutBounds) const
+{
+	return GetLevelInstanceBoundsInternal(LevelInstance, false, OutBounds);
+}
+
+bool ULevelInstanceSubsystem::GetLevelInstanceEditorBounds(const ILevelInstanceInterface* LevelInstance, FBox& OutBounds) const
+{
+	return GetLevelInstanceBoundsInternal(LevelInstance, true, OutBounds);
+}
+
+bool ULevelInstanceSubsystem::GetLevelInstanceBoundsInternal(const ILevelInstanceInterface* LevelInstance, bool bIsEditorBounds, FBox& OutBounds) const
 {
 	if (IsLoaded(LevelInstance))
 	{
@@ -1027,7 +1037,7 @@ bool ULevelInstanceSubsystem::GetLevelInstanceBounds(const ILevelInstanceInterfa
 
 		FString LevelPackage = LevelInstance->GetWorldAssetPackage();
 
-		if (FBox ContainerBounds = UActorDescContainerSubsystem::GetChecked().GetContainerBounds(*LevelPackage); ContainerBounds.IsValid)
+		if (FBox ContainerBounds = UActorDescContainerSubsystem::GetChecked().GetContainerBounds(*LevelPackage, bIsEditorBounds); ContainerBounds.IsValid)
 		{
 			FTransform LevelInstancePivotOffsetTransform = FTransform(ULevel::GetLevelInstancePivotOffsetFromPackage(*LevelPackage));
 			FTransform LevelTransform = LevelInstancePivotOffsetTransform * CastChecked<AActor>(LevelInstance)->GetActorTransform();
diff --git a/Engine/Source/Runtime/Engine/Private/WorldPartition/ActorDescContainerSubsystem.cpp b/Engine/Source/Runtime/Engine/Private/WorldPartition/ActorDescContainerSubsystem.cpp
--- a/Engine/Source/Runtime/Engine/Private/WorldPartition/ActorDescContainerSubsystem.cpp
+++ b/Engine/Source/Runtime/Engine/Private/WorldPartition/ActorDescContainerSubsystem.cpp
@@ -37,6 +37,7 @@ void UActorDescContainerSubsystem::FContainerManager::FRegisteredContainer::AddR
 
 void UActorDescContainerSubsystem::FContainerManager::FRegisteredContainer::UpdateBounds()
 {
+	EditorBounds.Init();
 	Bounds.Init();
 	for (FActorDescList::TIterator<> ActorDescIt(Container); ActorDescIt; ++ActorDescIt)
 	{
@@ -51,6 +52,11 @@ void UActorDescContainerSubsystem::FContainerManager::FRegisteredContainer::Upda
 		{
 			Bounds += RuntimeBounds;
 		}
+		const FBox ActorDescEditorBounds = ActorDesc->GetEditorBounds();
+		if (ActorDescEditorBounds.IsValid)
+		{
+			EditorBounds += ActorDescEditorBounds;
+		}
 	}
 }
 
@@ -74,11 +80,11 @@ void UActorDescContainerSubsystem::FContainerManager::UnregisterContainer(UActor
 	}
 }
 
-FBox UActorDescContainerSubsystem::FContainerManager::GetContainerBounds(const FString& ContainerName) const
+FBox UActorDescContainerSubsystem::FContainerManager::GetContainerBounds(const FString& ContainerName, bool bIsEditorBounds) const
 {
 	if (const FRegisteredContainer* RegisteredContainer = RegisteredContainers.Find(ContainerName))
 	{
-		return RegisteredContainer->Bounds;
+		return bIsEditorBounds ? RegisteredContainer->EditorBounds : RegisteredContainer->Bounds;
 	}
 	return FBox(ForceInit);
 }
diff --git a/Engine/Source/Runtime/Engine/Private/WorldPartition/LevelInstance/LevelInstanceActorDesc.cpp b/Engine/Source/Runtime/Engine/Private/WorldPartition/LevelInstance/LevelInstanceActorDesc.cpp
--- a/Engine/Source/Runtime/Engine/Private/WorldPartition/LevelInstance/LevelInstanceActorDesc.cpp
+++ b/Engine/Source/Runtime/Engine/Private/WorldPartition/LevelInstance/LevelInstanceActorDesc.cpp
@@ -130,7 +130,8 @@ void FLevelInstanceActorDesc::UpdateBounds()
 {
 	if (UActorDescContainer* ChildContainerPtr = ChildContainer.Get())
 	{
-		RuntimeBounds = UActorDescContainerSubsystem::GetChecked().GetContainerBounds(GetChildContainerName()).TransformBy(GetChildContainerTransform());
+		RuntimeBounds = UActorDescContainerSubsystem::GetChecked().GetContainerBounds(GetChildContainerName(), false).TransformBy(GetChildContainerTransform());
+		EditorBounds = UActorDescContainerSubsystem::GetChecked().GetContainerBounds(GetChildContainerName(), true).TransformBy(GetChildContainerTransform());
 	}
 }
 
diff --git a/Engine/Source/Runtime/Engine/Public/LevelInstance/LevelInstanceActor.h b/Engine/Source/Runtime/Engine/Public/LevelInstance/LevelInstanceActor.h
--- a/Engine/Source/Runtime/Engine/Public/LevelInstance/LevelInstanceActor.h
+++ b/Engine/Source/Runtime/Engine/Public/LevelInstance/LevelInstanceActor.h
@@ -103,6 +103,7 @@ public:
 	ENGINE_API virtual void PushLevelInstanceEditingStateToProxies(bool bInEditingState) override;
 	ENGINE_API virtual FBox GetComponentsBoundingBox(bool bNonColliding = false, bool bIncludeFromChildActors = false) const override;
 	ENGINE_API virtual FBox GetStreamingBounds() const override;
+	ENGINE_API virtual FBox GetStreamingBoundsEditor() const override;
 	ENGINE_API virtual bool IsLockLocation() const override;
 	ENGINE_API virtual bool IsActorLabelEditable() const override;
 	ENGINE_API virtual bool GetReferencedContentObjects(TArray<UObject*>& Objects) const override;
diff --git a/Engine/Source/Runtime/Engine/Public/LevelInstance/LevelInstanceActorImpl.h b/Engine/Source/Runtime/Engine/Public/LevelInstance/LevelInstanceActorImpl.h
--- a/Engine/Source/Runtime/Engine/Public/LevelInstance/LevelInstanceActorImpl.h
+++ b/Engine/Source/Runtime/Engine/Public/LevelInstance/LevelInstanceActorImpl.h
@@ -94,6 +94,7 @@ public:
 	ENGINE_API virtual bool IsUserManaged() const;
 	ENGINE_API virtual bool ShouldExport() const;
 	ENGINE_API virtual bool GetBounds(FBox& OutBounds) const;
+	ENGINE_API virtual bool GetEditorBounds(FBox& OutBounds) const;
 	ENGINE_API virtual void PushSelectionToProxies();
 	ENGINE_API virtual void PushLevelInstanceEditingStateToProxies(bool bInEditingState);
 	ENGINE_API virtual void CheckForErrors();
diff --git a/Engine/Source/Runtime/Engine/Public/LevelInstance/LevelInstanceSubsystem.h b/Engine/Source/Runtime/Engine/Public/LevelInstance/LevelInstanceSubsystem.h
--- a/Engine/Source/Runtime/Engine/Public/LevelInstance/LevelInstanceSubsystem.h
+++ b/Engine/Source/Runtime/Engine/Public/LevelInstance/LevelInstanceSubsystem.h
@@ -117,6 +117,7 @@ public:
 	ENGINE_API bool IsEditingLevelInstanceDirty(const ILevelInstanceInterface* LevelInstance) const;
 	bool IsEditingLevelInstance(const ILevelInstanceInterface* LevelInstance) const { return GetLevelInstanceEdit(LevelInstance) != nullptr; }
 	
+	ENGINE_API bool GetLevelInstanceEditorBounds(const ILevelInstanceInterface* LevelInstance, FBox& OutBounds) const;
 	ENGINE_API bool GetLevelInstanceBounds(const ILevelInstanceInterface* LevelInstance, FBox& OutBounds) const;
 	static ENGINE_API bool GetLevelInstanceBoundsFromPackage(const FTransform& InstanceTransform, FName LevelPackage, FBox& OutBounds);
 	
@@ -201,6 +202,8 @@ private:
 
 	ENGINE_API static void RegisterPrimitiveColorHandler();
 	ENGINE_API static void UnregisterPrimitiveColorHandler();
+
+	ENGINE_API bool GetLevelInstanceBoundsInternal(const ILevelInstanceInterface* LevelInstance, bool bIsEditorBounds, FBox& OutBounds) const;
 #endif
 	friend class ULevelStreamingLevelInstance;
 
diff --git a/Engine/Source/Runtime/Engine/Public/WorldPartition/ActorDescContainerSubsystem.h b/Engine/Source/Runtime/Engine/Public/WorldPartition/ActorDescContainerSubsystem.h
--- a/Engine/Source/Runtime/Engine/Public/WorldPartition/ActorDescContainerSubsystem.h
+++ b/Engine/Source/Runtime/Engine/Public/WorldPartition/ActorDescContainerSubsystem.h
@@ -45,7 +45,7 @@ public:
 	const UActorDescContainer* GetActorDescContainer(const FString& InContainerName) const { return ContainerManager.GetActorDescContainer(InContainerName); }
 	void RegisterContainer(UActorDescContainer* Container) { ContainerManager.RegisterContainer(Container); }
 	void UnregisterContainer(UActorDescContainer* Container) { ContainerManager.UnregisterContainer(Container); }
-	FBox GetContainerBounds(const FString& ContainerName) const { return ContainerManager.GetContainerBounds(ContainerName); }
+	FBox GetContainerBounds(const FString& ContainerName, bool bIsEditorBounds = false) const { return ContainerManager.GetContainerBounds(ContainerName, bIsEditorBounds); }
 	void UpdateContainerBounds(const FString& ContainerName) { ContainerManager.UpdateContainerBounds(ContainerName); }
 	void SetContainerPackage(UActorDescContainer* Container, FName ContainerPackageName) { ContainerManager.SetContainerPackage(Container, ContainerPackageName); }
 	void NotifyContainerUpdated(FName ContainerPackage)
@@ -89,6 +89,7 @@ private:
 				: Container(nullptr)
 				, RefCount(0)
 				, Bounds(ForceInit)
+				, EditorBounds(ForceInit)
 			{}
 
 			void AddReferencedObjects(FReferenceCollector& Collector);
@@ -97,6 +98,7 @@ private:
 			TObjectPtr<UActorDescContainer> Container;
 			uint32 RefCount;
 			FBox Bounds;
+			FBox EditorBounds;
 		};
 
 		void AddReferencedObjects(FReferenceCollector& Collector);
@@ -145,7 +147,7 @@ private:
 				
 		void UnregisterContainer(UActorDescContainer* Container);
 
-		FBox GetContainerBounds(const FString& ContainerName) const;
+		FBox GetContainerBounds(const FString& ContainerName, bool bIsEditorBounds = false) const;
 		void UpdateContainerBounds(const FString& ContainerName);
 		void UpdateContainerBoundsFromPackage(FName ContainerPackage);
 
