Add alignment to virtual texture tile size when initializing from cvars.
This is already enforced in compilation, but there are now other parts of the code that access this structure, so we should do this as early as possible.

diff --git a/Engine/Source/Runtime/Engine/Classes/VT/VirtualTextureBuildSettings.h b/Engine/Source/Runtime/Engine/Classes/VT/VirtualTextureBuildSettings.h
--- a/Engine/Source/Runtime/Engine/Classes/VT/VirtualTextureBuildSettings.h
+++ b/Engine/Source/Runtime/Engine/Classes/VT/VirtualTextureBuildSettings.h
@@ -19,4 +19,9 @@ struct FVirtualTextureBuildSettings
 
 	/** Initialize with the default build settings. These are defined by the current project setup. */
 	ENGINE_API void Init();
+
+	/** Helper to clamp tile size and align to next power of two. */
+	ENGINE_API static int32 ClampAndAlignTileSize(int32 InTileSize);
+	/** Helper to clamp tile border size and align to next multiple of two. */
+	ENGINE_API static int32 ClampAndAlignTileBorderSize(int32 InTileBorderSize);
 };
diff --git a/Engine/Source/Runtime/Engine/Private/VT/VirtualTextureBuildSettings.cpp b/Engine/Source/Runtime/Engine/Private/VT/VirtualTextureBuildSettings.cpp
--- a/Engine/Source/Runtime/Engine/Private/VT/VirtualTextureBuildSettings.cpp
+++ b/Engine/Source/Runtime/Engine/Private/VT/VirtualTextureBuildSettings.cpp
@@ -15,12 +15,22 @@ static TAutoConsoleVariable<int32> CVarVTTileSize(
 static TAutoConsoleVariable<int32> CVarVTTileBorderSize(
 	TEXT("r.VT.TileBorderSize"),
 	4,
-	TEXT("Size in pixels to use for virtual texture tiles borders (rounded to next power-of-2)")
+	TEXT("Size in pixels to use for virtual texture tiles borders (rounded to multiple-of-2)")
 );
 
-void FVirtualTextureBuildSettings::Init()
+int32 FVirtualTextureBuildSettings::ClampAndAlignTileSize(int32 InTileSize)
+{
+	return FMath::RoundUpToPowerOfTwo(FMath::Clamp<uint32>(InTileSize, 16u, 1024u));
+}
+
+int32 FVirtualTextureBuildSettings::ClampAndAlignTileBorderSize(int32 InTileBorderSize)
 {
-	TileSize = CVarVTTileSize.GetValueOnAnyThread();
-	TileBorderSize = CVarVTTileBorderSize.GetValueOnAnyThread();
+	// Round tile border size up to multiple of 2 to ensure that block compressed formats will be OK.
+	return (FMath::Clamp<uint32>(InTileBorderSize, 0u, 8u) + 1u) & ~1u;
 }
 
+void FVirtualTextureBuildSettings::Init()
+{
+	TileSize = ClampAndAlignTileSize(CVarVTTileSize.GetValueOnAnyThread());
+	TileBorderSize = ClampAndAlignTileBorderSize(CVarVTTileBorderSize.GetValueOnAnyThread());
+}
