Change logic for virtual texture recreation so that streaming textures do a full resource recreate instead of just a Release/Init RHI pair.
This will allow us to hook the event in the UTexture::UpdateResource() where previously there was no game thread functions called to hook.
Hooking is necessary for tracking changes for the incoming mesh painting virtual textures.

diff --git a/Engine/Source/Runtime/Engine/Private/VT/VirtualTextureRecreate.cpp b/Engine/Source/Runtime/Engine/Private/VT/VirtualTextureRecreate.cpp
--- a/Engine/Source/Runtime/Engine/Private/VT/VirtualTextureRecreate.cpp
+++ b/Engine/Source/Runtime/Engine/Private/VT/VirtualTextureRecreate.cpp
@@ -20,10 +20,7 @@ namespace VirtualTexture
 		{
 			if (It->IsCurrentlyVirtualTextured())
 			{
-				if (FVirtualTexture2DResource* Resource = static_cast<FVirtualTexture2DResource*>(It->GetResource()))
-				{
-					BeginUpdateResourceRHI(Resource);
-				}
+				It->UpdateResource();
 			}
 		}
 
@@ -55,7 +52,7 @@ namespace VirtualTexture
 
 					if (bIsMatchingFormat)
 					{
-						BeginUpdateResourceRHI(Resource);
+						It->UpdateResource();
 					}
 				}
 			}
