Cloth - Made the USkeletalMeshComponent::RecreateClothingActors() function BlueprintCallable.

#rnx
#tests Created a blueprint on a cloth skeletal mesh, and tested the RecreateClothingActors function with a change of scale.

diff --git a/Engine/Source/Runtime/Engine/Classes/Components/SkeletalMeshComponent.h b/Engine/Source/Runtime/Engine/Classes/Components/SkeletalMeshComponent.h
--- a/Engine/Source/Runtime/Engine/Classes/Components/SkeletalMeshComponent.h
+++ b/Engine/Source/Runtime/Engine/Classes/Components/SkeletalMeshComponent.h
@@ -2270,6 +2270,7 @@ public:
 	ENGINE_API void ApplyDeltaToAllPhysicsTransforms(const FVector& DeltaLocation, const FQuat& DeltaRotation);
 
 	/** Destroys and recreates the clothing actors in the current simulation */
+	UFUNCTION(BlueprintCallable, Category = "Clothing")
 	ENGINE_API void RecreateClothingActors();
 
 	/** Given bounds InOutBounds, expand them to also enclose the clothing simulation mesh */
