Added bTargetArraySlicesIndependently to UTextureRenderTarget (relevant for 2D arrays and cube or cube array) so that each slice can be individually rendered to by a pixel shader

diff --git a/Engine/Source/Runtime/Engine/Classes/Engine/TextureRenderTarget.h b/Engine/Source/Runtime/Engine/Classes/Engine/TextureRenderTarget.h
--- a/Engine/Source/Runtime/Engine/Classes/Engine/TextureRenderTarget.h
+++ b/Engine/Source/Runtime/Engine/Classes/Engine/TextureRenderTarget.h
@@ -32,6 +32,9 @@ class UTextureRenderTarget : public UTexture
 	/** If true, it will be possible to create a FUnorderedAccessViewRHIRef using RHICreateUnorderedAccessView and the internal FTextureRHIRef. */
 	uint32 bCanCreateUAV : 1;
 
+	/** If true, RTVs will be created for each slice (texture array and/or cube). */
+	uint32 bTargetArraySlicesIndependently : 1;
+
 	/**
 	 * Render thread: Access the render target resource for this texture target object
 	 * @return pointer to resource or NULL if not initialized
diff --git a/Engine/Source/Runtime/Engine/Private/TextureRenderTarget.cpp b/Engine/Source/Runtime/Engine/Private/TextureRenderTarget.cpp
--- a/Engine/Source/Runtime/Engine/Private/TextureRenderTarget.cpp
+++ b/Engine/Source/Runtime/Engine/Private/TextureRenderTarget.cpp
@@ -30,6 +30,7 @@ UTextureRenderTarget::UTextureRenderTarget(const FObjectInitializer& ObjectIniti
 	LODGroup = TEXTUREGROUP_RenderTarget;	
 	bNeedsTwoCopies = false;
 	bCanCreateUAV = false;
+	bTargetArraySlicesIndependently = false;
 #if WITH_EDITORONLY_DATA
 	CompressionNone = true;
 #endif // #if WITH_EDITORONLY_DATA
diff --git a/Engine/Source/Runtime/Engine/Private/TextureRenderTarget2D.cpp b/Engine/Source/Runtime/Engine/Private/TextureRenderTarget2D.cpp
--- a/Engine/Source/Runtime/Engine/Private/TextureRenderTarget2D.cpp
+++ b/Engine/Source/Runtime/Engine/Private/TextureRenderTarget2D.cpp
@@ -610,7 +610,7 @@ void FTextureRenderTarget2DResource::ReleaseRHI()
 
 /**
  * Updates (resolves) the render target texture.
- * Optionally clears the contents of the render target to green.
+ * Optionally clears the contents of the render target to the clear color.
  * This is only called by the rendering thread.
  */
 void FTextureRenderTarget2DResource::UpdateDeferredResource( FRHICommandListImmediate& RHICmdList, bool bClearRenderTarget/*=true*/ )
@@ -633,7 +633,7 @@ void FTextureRenderTarget2DResource::UpdateDeferredResource( FRHICommandListImme
 	FRDGTextureRef RenderTargetTextureRDG = GraphBuilder.RegisterExternalTexture(MipGenerationCache);
 	FRDGTextureRef TextureRDG = GraphBuilder.RegisterExternalTexture(CreateRenderTarget(TextureRHI, TEXT("TextureRenderTarget2DResource")));
 
- 	// clear the target surface to green
+ 	// clear the target surface to the clear color
 	if (bClearRenderTarget)
 	{
 		ensure(RenderTargetTextureRHI.IsValid() && (RenderTargetTextureRHI->GetClearColor() == ClearColor));
diff --git a/Engine/Source/Runtime/Engine/Private/TextureRenderTarget2DArray.cpp b/Engine/Source/Runtime/Engine/Private/TextureRenderTarget2DArray.cpp
--- a/Engine/Source/Runtime/Engine/Private/TextureRenderTarget2DArray.cpp
+++ b/Engine/Source/Runtime/Engine/Private/TextureRenderTarget2DArray.cpp
@@ -15,6 +15,7 @@
 #include "RHIUtilities.h"
 #include "UObject/Package.h"
 #include "ProfilingDebugging/AssetMetadataTrace.h"
+#include "RenderGraphUtils.h"
 
 #include UE_INLINE_GENERATED_CPP_BY_NAME(TextureRenderTarget2DArray)
 
@@ -33,7 +34,7 @@ UTextureRenderTarget2DArray::UTextureRenderTarget2DArray(const FObjectInitialize
 
 EPixelFormat UTextureRenderTarget2DArray::GetFormat() const
 {
-	if(OverrideFormat == PF_Unknown)
+	if (OverrideFormat == PF_Unknown)
 	{
 		return bHDR ? PF_FloatRGBA : PF_B8G8R8A8;
 	}
@@ -48,10 +49,10 @@ bool UTextureRenderTarget2DArray::IsSRGB() const
 	bool bIsSRGB = true;
 
 	// if render target gamma used was 1.0 then disable SRGB for the static texture
-	if(FMath::Abs(GetDisplayGamma() - 1.0f) < UE_KINDA_SMALL_NUMBER)
+	if (FMath::Abs(GetDisplayGamma() - 1.0f) < UE_KINDA_SMALL_NUMBER)
 	{
 		bIsSRGB = false;
-	}	
+	}
 
 	return bIsSRGB;
 }
@@ -96,10 +97,9 @@ void UTextureRenderTarget2DArray::UpdateResourceImmediate(bool bClearRenderTarge
 		FTextureRenderTarget2DArrayResource* InResource = static_cast<FTextureRenderTarget2DArrayResource*>(GetResource());
 		ENQUEUE_RENDER_COMMAND(UpdateResourceImmediate)(
 			[InResource, bClearRenderTarget](FRHICommandListImmediate& RHICmdList)
-			{
-				InResource->UpdateDeferredResource(RHICmdList, bClearRenderTarget);
-			}
-		);
+		{
+			InResource->UpdateDeferredResource(RHICmdList, bClearRenderTarget);
+		});
 	}
 }
 void UTextureRenderTarget2DArray::GetResourceSizeEx(FResourceSizeEx& CumulativeResourceSize)
@@ -108,12 +108,12 @@ void UTextureRenderTarget2DArray::GetResourceSizeEx(FResourceSizeEx& CumulativeR
 
 	// Calculate size based on format.
 	const EPixelFormat Format = GetFormat();
-	const int64 BlockSizeX	= GPixelFormats[Format].BlockSizeX;
-	const int64 BlockSizeY	= GPixelFormats[Format].BlockSizeY;
-	const int64 BlockBytes	= GPixelFormats[Format].BlockBytes;
-	const int64 NumBlocksX	= (SizeX + BlockSizeX - 1) / BlockSizeX;
-	const int64 NumBlocksY	= (SizeY + BlockSizeY - 1) / BlockSizeY;
-	const int64 NumBytes	= NumBlocksX * NumBlocksY * Slices * BlockBytes;
+	const int64 BlockSizeX = GPixelFormats[Format].BlockSizeX;
+	const int64 BlockSizeY = GPixelFormats[Format].BlockSizeY;
+	const int64 BlockBytes = GPixelFormats[Format].BlockBytes;
+	const int64 NumBlocksX = (SizeX + BlockSizeX - 1) / BlockSizeX;
+	const int64 NumBlocksY = (SizeY + BlockSizeY - 1) / BlockSizeY;
+	const int64 NumBytes = NumBlocksX * NumBlocksY * Slices * BlockBytes;
 
 	CumulativeResourceSize.AddUnknownMemoryBytes(NumBytes);
 }
@@ -157,7 +157,7 @@ void UTextureRenderTarget2DArray::PostLoad()
 
 FString UTextureRenderTarget2DArray::GetDesc()
 {
-	return FString::Printf( TEXT("Render to Texture 2DArray %dx%d[%s]"), SizeX, SizeX, GPixelFormats[GetFormat()].Name);
+	return FString::Printf(TEXT("Render to Texture 2DArray %dx%d[%s]"), SizeX, SizeX, GPixelFormats[GetFormat()].Name);
 }
 
 TSubclassOf<UTexture> UTextureRenderTarget2DArray::GetTextureUClass() const
@@ -224,7 +224,7 @@ void FTextureRenderTarget2DArrayResource::InitRHI(FRHICommandListBase& RHICmdLis
 	LLM_SCOPE_DYNAMIC_STAT_OBJECTPATH(Owner->GetPackage(), ELLMTagSet::Assets);
 	UE_TRACE_METADATA_SCOPE_ASSET_FNAME(NAME_None, NAME_None, Owner->GetPackage()->GetFName());
 
-	if((Owner->SizeX > 0) && (Owner->SizeY > 0) && (Owner->Slices > 0))
+	if ((Owner->SizeX > 0) && (Owner->SizeY > 0) && (Owner->Slices > 0))
 	{
 		// Create the RHI texture. Only one mip is used and the texture is targetable for resolve.
 		ETextureCreateFlags TexCreateFlags = Owner->IsSRGB() ? ETextureCreateFlags::SRGB : ETextureCreateFlags::None;
@@ -233,6 +233,11 @@ void FTextureRenderTarget2DArrayResource::InitRHI(FRHICommandListBase& RHICmdLis
 			TexCreateFlags |= ETextureCreateFlags::UAV;
 		}
 
+		if (Owner->bTargetArraySlicesIndependently)
+		{
+			TexCreateFlags |= ETextureCreateFlags::TargetArraySlicesIndependently;
+		}
+
 		{
 			const FRHITextureCreateDesc Desc =
 				FRHITextureCreateDesc::Create2DArray(TEXT("FTextureRenderTarget2DArrayResource"))
@@ -291,7 +296,7 @@ void FTextureRenderTarget2DArrayResource::ReleaseRHI()
 
 /**
  * Updates (resolves) the render target texture.
- * Optionally clears each face of the render target to green.
+ * Optionally clears each face of the render target to the clear color.
  * This is only called by the rendering thread.
  */
 void FTextureRenderTarget2DArrayResource::UpdateDeferredResource(FRHICommandListImmediate& RHICmdList, bool bClearRenderTarget/*=true*/)
@@ -306,12 +311,26 @@ void FTextureRenderTarget2DArrayResource::UpdateDeferredResource(FRHICommandList
 		return;
 	}
 
-	RHICmdList.Transition(FRHITransitionInfo(TextureRHI, ERHIAccess::Unknown, ERHIAccess::RTV));
-	ClearRenderTarget(RHICmdList, TextureRHI);
-	RHICmdList.Transition(FRHITransitionInfo(TextureRHI, ERHIAccess::RTV, ERHIAccess::SRVMask));
+	// If the render target has the TargetArraySlicesIndependently flag, we can do the right thing and clear all slices :
+	if (EnumHasAnyFlags(TextureRHI->GetDesc().Flags, ETextureCreateFlags::TargetArraySlicesIndependently))
+	{
+		FRDGTextureClearInfo TextureClearInfo;
+		TextureClearInfo.NumSlices = TextureRHI->GetDesc().ArraySize;
+		FRDGBuilder GraphBuilder(RHICmdList);
+		FRDGTextureRef TextureRDG = GraphBuilder.RegisterExternalTexture(CreateRenderTarget(TextureRHI, TEXT("TextureRenderTarget2DArrayResource")));
+		AddClearRenderTargetPass(GraphBuilder, TextureRDG, TextureClearInfo);
+		GraphBuilder.Execute();
+	}
+	else
+	{
+		// Otherwise, just (probably erroneously) clear the first slice for backwards compatibility :
+		RHICmdList.Transition(FRHITransitionInfo(TextureRHI, ERHIAccess::Unknown, ERHIAccess::RTV));
+		ClearRenderTarget(RHICmdList, TextureRHI);
+		RHICmdList.Transition(FRHITransitionInfo(TextureRHI, ERHIAccess::RTV, ERHIAccess::SRVMask));
+	}
 }
 
-/** 
+/**
  * @return width of target
  */
 uint32 FTextureRenderTarget2DArrayResource::GetSizeX() const
@@ -319,7 +338,7 @@ uint32 FTextureRenderTarget2DArrayResource::GetSizeX() const
 	return Owner->SizeX;
 }
 
-/** 
+/**
  * @return height of target
  */
 uint32 FTextureRenderTarget2DArrayResource::GetSizeY() const
@@ -327,7 +346,7 @@ uint32 FTextureRenderTarget2DArrayResource::GetSizeY() const
 	return Owner->SizeX;
 }
 
-/** 
+/**
  * @return dimensions of target surface
  */
 FIntPoint FTextureRenderTarget2DArrayResource::GetSizeXY() const
@@ -337,12 +356,12 @@ FIntPoint FTextureRenderTarget2DArrayResource::GetSizeXY() const
 
 float UTextureRenderTarget2DArray::GetDisplayGamma() const
 {
-	if(TargetGamma > UE_KINDA_SMALL_NUMBER * 10.0f)
+	if (TargetGamma > UE_KINDA_SMALL_NUMBER * 10.0f)
 	{
 		return TargetGamma;
 	}
 	EPixelFormat Format = GetFormat();
-	if(Format == PF_FloatRGB || Format == PF_FloatRGBA || bForceLinearGamma)
+	if (Format == PF_FloatRGB || Format == PF_FloatRGBA || bForceLinearGamma)
 	{
 		return 1.0f;
 	}
diff --git a/Engine/Source/Runtime/Engine/Private/TextureRenderTargetCube.cpp b/Engine/Source/Runtime/Engine/Private/TextureRenderTargetCube.cpp
--- a/Engine/Source/Runtime/Engine/Private/TextureRenderTargetCube.cpp
+++ b/Engine/Source/Runtime/Engine/Private/TextureRenderTargetCube.cpp
@@ -288,6 +288,11 @@ void FTextureRenderTargetCubeResource::InitRHI(FRHICommandListBase& RHICmdList)
 			}
 		}
 
+		if (Owner->bTargetArraySlicesIndependently)
+		{
+			TexCreateFlags |= ETextureCreateFlags::TargetArraySlicesIndependently;
+		}
+
 		{
 			const FRHITextureCreateDesc Desc =
 				FRHITextureCreateDesc::CreateCube(TEXT("FTextureRenderTargetCubeResource"))
@@ -357,7 +362,7 @@ void FTextureRenderTargetCubeResource::ReleaseRHI()
 
 /**
  * Updates (resolves) the render target texture.
- * Optionally clears each face of the render target to green.
+ * Optionally clears each face of the render target to the clear color.
  * This is only called by the rendering thread.
  */
 void FTextureRenderTargetCubeResource::UpdateDeferredResource(FRHICommandListImmediate& RHICmdList, bool bClearRenderTarget/*=true*/)
diff --git a/Engine/Source/Runtime/Engine/Private/TextureRenderTargetVolume.cpp b/Engine/Source/Runtime/Engine/Private/TextureRenderTargetVolume.cpp
--- a/Engine/Source/Runtime/Engine/Private/TextureRenderTargetVolume.cpp
+++ b/Engine/Source/Runtime/Engine/Private/TextureRenderTargetVolume.cpp
@@ -296,7 +296,7 @@ void FTextureRenderTargetVolumeResource::ReleaseRHI()
 
 /**
  * Updates (resolves) the render target texture.
- * Optionally clears each face of the render target to green.
+ * Optionally clears each face of the render target to the clear color.
  * This is only called by the rendering thread.
  */
 void FTextureRenderTargetVolumeResource::UpdateDeferredResource(FRHICommandListImmediate& RHICmdList, bool bClearRenderTarget/*=true*/)
