Fixed validation code of WorldDataLayers to properly support instanced partitioned worlds.

diff --git a/Engine/Source/Runtime/Engine/Private/WorldPartition/DataLayer/WorldDataLayers.cpp b/Engine/Source/Runtime/Engine/Private/WorldPartition/DataLayer/WorldDataLayers.cpp
--- a/Engine/Source/Runtime/Engine/Private/WorldPartition/DataLayer/WorldDataLayers.cpp
+++ b/Engine/Source/Runtime/Engine/Private/WorldPartition/DataLayer/WorldDataLayers.cpp
@@ -1314,9 +1314,12 @@ bool AWorldDataLayers::CanReferenceDataLayerAsset(const UDataLayerAsset* InDataL
 	// Validate that AWorldDataLayers and its EDL Asset are part of the same plugin
 	else if (InDataLayerAsset && (RootExternalDataLayerAsset == InDataLayerAsset))
 	{
+		const UPackage* Package = GetPackage();
+		const FString PackageName = !Package->GetLoadedPath().IsEmpty() ? Package->GetLoadedPath().GetPackageName() : Package->GetName();
+
 		IPluginManager& PluginManager = IPluginManager::Get();
 		const FString AssetMountPoint = FPackageName::GetPackageMountPoint(InDataLayerAsset->GetPackage()->GetName()).ToString();
-		const FString ThisMountPoint = FPackageName::GetPackageMountPoint(GetPackage()->GetName()).ToString();
+		const FString ThisMountPoint = FPackageName::GetPackageMountPoint(PackageName).ToString();
 		TSharedPtr<IPlugin> AssetPlugin = PluginManager.FindPluginFromPath(AssetMountPoint);
 		TSharedPtr<IPlugin> ThisPlugin = PluginManager.FindPluginFromPath(ThisMountPoint);
 		if (AssetPlugin != ThisPlugin)
