[Audio]
 -  Run game-thread task to a safe place in the frame using the ticker

diff --git a/Engine/Source/Runtime/Engine/Private/AudioThread.cpp b/Engine/Source/Runtime/Engine/Private/AudioThread.cpp
--- a/Engine/Source/Runtime/Engine/Private/AudioThread.cpp
+++ b/Engine/Source/Runtime/Engine/Private/AudioThread.cpp
@@ -7,6 +7,7 @@
 #include "AudioThread.h"
 #include "Audio.h"
 #include "Async/Async.h"
+#include "Containers/Ticker.h"
 #include "ProfilingDebugging/CsvProfiler.h"
 #include "Tasks/Pipe.h"
 
@@ -369,17 +370,16 @@ void FAudioThread::RunCommandOnGameThread(TUniqueFunction<void()> InFunction, co
 	if (IsAudioThreadRunning())
 	{
 		check(IsInAudioThread());
-		FFunctionGraphTask::CreateAndDispatchWhenReady(
+		ExecuteOnGameThread(
+			TEXT("FAudioThread::RunCommandOnGameThread"),
 			[Function = MoveTemp(InFunction), InStatId]()
 			{
 				CSV_SCOPED_TIMING_STAT_EXCLUSIVE(Audio);
 				QUICK_SCOPE_CYCLE_COUNTER(STAT_AudioThread_RunCommandOnGameThread);
 				FScopeCycleCounter ScopeCycleCounter(InStatId);
 				Function();
-			},
-			TStatId(),
-			nullptr,
-			ENamedThreads::GameThread);
+			}
+		);
 	}
 	else
 	{
