[World Partition] Fallback to generic cells sorting comparison when comparing cells of different types.
Fix ASAN heap-buffer-overflow error.

diff --git a/Engine/Source/Runtime/Engine/Private/WorldPartition/WorldPartitionRuntimeCellDataSpatialHash.cpp b/Engine/Source/Runtime/Engine/Private/WorldPartition/WorldPartitionRuntimeCellDataSpatialHash.cpp
--- a/Engine/Source/Runtime/Engine/Private/WorldPartition/WorldPartitionRuntimeCellDataSpatialHash.cpp
+++ b/Engine/Source/Runtime/Engine/Private/WorldPartition/WorldPartitionRuntimeCellDataSpatialHash.cpp
@@ -146,47 +146,50 @@ void UWorldPartitionRuntimeCellDataSpatialHash::MergeStreamingSourceInfo() const
 
 int32 UWorldPartitionRuntimeCellDataSpatialHash::SortCompare(const UWorldPartitionRuntimeCellData* InOther) const
 {
-	int32 Result = (int32)CachedMinSourcePriority - (int32)InOther->CachedMinSourcePriority;
-
-	if (Result == 0)
+	if (const UWorldPartitionRuntimeCellDataSpatialHash* Other = Cast<UWorldPartitionRuntimeCellDataSpatialHash>(InOther))
 	{
-		const UWorldPartitionRuntimeCellDataSpatialHash* Other = (UWorldPartitionRuntimeCellDataSpatialHash*)InOther;
-
-		// By default, now compare cell's extent instead of its grid level since we compare cells across multiple WPs/grids (higher value is higher prio)
-		Result = GRuntimeSpatialHashSortUsingCellExtent ? int32(Other->Extent - Extent) : (InOther->HierarchicalLevel - HierarchicalLevel);
+		int32 Result = (int32)CachedMinSourcePriority - (int32)Other->CachedMinSourcePriority;
 
 		if (Result == 0)
 		{
-			if (GRuntimeSpatialHashSortUsingCellPriority)
-			{
-				// Cell priority (lower value is higher prio)
-				Result = Priority - Other->Priority;
-			}
+			// By default, now compare cell's extent instead of its grid level since we compare cells across multiple WPs/grids (higher value is higher prio)
+			Result = GRuntimeSpatialHashSortUsingCellExtent ? int32(Other->Extent - Extent) : (Other->HierarchicalLevel - HierarchicalLevel);
 
 			if (Result == 0)
 			{
-				// Closest distance (lower value is higher prio)
-				const double Diff = CachedSourceSortingDistance - Other->CachedSourceSortingDistance;
-				if (FMath::IsNearlyZero(Diff))
+				if (GRuntimeSpatialHashSortUsingCellPriority)
 				{
-					const double RawDistanceDiff = CachedMinSquareDistanceToSource - Other->CachedMinSquareDistanceToSource;
-					Result = RawDistanceDiff < 0 ? -1 : (RawDistanceDiff > 0.f ? 1 : 0);
+					// Cell priority (lower value is higher prio)
+					Result = Priority - Other->Priority;
 				}
-				else
+
+				if (Result == 0)
 				{
-					Result = Diff < 0.f ? -1 : (Diff > 0.f ? 1 : 0);
+					// Closest distance (lower value is higher prio)
+					const double Diff = CachedSourceSortingDistance - Other->CachedSourceSortingDistance;
+					if (FMath::IsNearlyZero(Diff))
+					{
+						const double RawDistanceDiff = CachedMinSquareDistanceToSource - Other->CachedMinSquareDistanceToSource;
+						Result = RawDistanceDiff < 0 ? -1 : (RawDistanceDiff > 0.f ? 1 : 0);
+					}
+					else
+					{
+						Result = Diff < 0.f ? -1 : (Diff > 0.f ? 1 : 0);
+					}
 				}
 			}
 		}
-	}
 
-	if (!GRuntimeSpatialHashSortUsingCellPriority && Result == 0)
-	{
-		// Cell priority (lower value is higher prio)
-		Result = Priority - InOther->Priority;
+		if (!GRuntimeSpatialHashSortUsingCellPriority && Result == 0)
+		{
+			// Cell priority (lower value is higher prio)
+			Result = Priority - Other->Priority;
+		}
+
+		return Result;
 	}
 
-	return Result;
+	return Super::SortCompare(InOther);
 }
 
 FBox UWorldPartitionRuntimeCellDataSpatialHash::GetCellBounds() const
