AssetManager: Suppress the PrimaryAssetId of generated maps; they conflict with the primaryassetid of the map's parent map.
Add functions to PackageName.h to detect if it is a generated map

diff --git a/Engine/Source/Runtime/CoreUObject/Private/Misc/PackageName.cpp b/Engine/Source/Runtime/CoreUObject/Private/Misc/PackageName.cpp
--- a/Engine/Source/Runtime/CoreUObject/Private/Misc/PackageName.cpp
+++ b/Engine/Source/Runtime/CoreUObject/Private/Misc/PackageName.cpp
@@ -2507,6 +2507,34 @@ bool FPackageName::IsMapPackageExtension(const TCHAR* Ext)
 	}
 }
 
+const TCHAR* FPackageName::GetGeneratedPackageSubPath()
+{
+	return TEXT("_Generated_");
+}
+
+bool FPackageName::IsUnderGeneratedPackageSubPath(FStringView FileOrLongPackagePath)
+{
+	FStringView GeneratedSubDir(GetGeneratedPackageSubPath());
+	int32 Index = FileOrLongPackagePath.Find(GeneratedSubDir, ESearchCase::IgnoreCase);
+	if (Index < 0)
+	{
+		return false;
+	}
+	if (Index == 0 || !FPathViews::IsSeparator(FileOrLongPackagePath[Index - 1]))
+	{
+		// "_Generated_/..." or ".../Prefix_Generated_/...", not the /_Generated_/ subfolder we're looking for.
+		return false;
+	}
+	if (Index + GeneratedSubDir.Len() < FileOrLongPackagePath.Len()
+		&& !FPathViews::IsSeparator(FileOrLongPackagePath[Index + GeneratedSubDir.Len()]))
+	{
+		// ".../_Generated_Suffix/...", not the /_Generated_/ subfolder we're looking for.
+		return false;
+	}
+	return true;
+}
+
+
 bool FPackageName::FindPackagesInDirectory( TArray<FString>& OutPackages, const FString& RootDir )
 {
 	// Keep track if any package has been found. Can't rely only on OutPackages.Num() > 0 as it may not be empty.
diff --git a/Engine/Source/Runtime/CoreUObject/Public/Misc/PackageName.h b/Engine/Source/Runtime/CoreUObject/Public/Misc/PackageName.h
--- a/Engine/Source/Runtime/CoreUObject/Public/Misc/PackageName.h
+++ b/Engine/Source/Runtime/CoreUObject/Public/Misc/PackageName.h
@@ -584,6 +584,15 @@ public:
 		return Filename.EndsWith(AssetPackageExtension) || Filename.EndsWith(MapPackageExtension);
 	}
 
+	/** Return the text string used to mark the _Generated_ directory of packages created by CookPackageSplitters. */
+	static COREUOBJECT_API const TCHAR* GetGeneratedPackageSubPath();
+
+	/**
+	 * Return whether a given packagename or file path is generated based on its name (we evaluate whether it
+	 * is in a GetGeneratedPackageSubPath directory).
+	 */
+	static COREUOBJECT_API bool IsUnderGeneratedPackageSubPath(FStringView FileOrLongPackagePath);
+
 	/**
 	 * This will recurse over a directory structure looking for packages.
 	 * 
diff --git a/Engine/Source/Runtime/Engine/Classes/Engine/AssetManager.h b/Engine/Source/Runtime/Engine/Classes/Engine/AssetManager.h
--- a/Engine/Source/Runtime/Engine/Classes/Engine/AssetManager.h
+++ b/Engine/Source/Runtime/Engine/Classes/Engine/AssetManager.h
@@ -70,6 +70,12 @@ public:
 	/** Type representing a packaging chunk, this is a virtual type that is never loaded off disk */
 	static ENGINE_API const FPrimaryAssetType PackageChunkType;
 
+	/**
+	 * Asset Type used in return values to indicate that a specific asset of a given PrimaryAssetType should be suppressed
+	 * and not considered a PrimaryAsset.
+	 */
+	static ENGINE_API const FPrimaryAssetType SuppressionType;
+
 	/** Virtual path $AssetSearchRoots, replaced with all roots including defaults like /Game */
 	static ENGINE_API const FString AssetSearchRootsVirtualPath;
 
diff --git a/Engine/Source/Runtime/Engine/Private/AssetManager.cpp b/Engine/Source/Runtime/Engine/Private/AssetManager.cpp
--- a/Engine/Source/Runtime/Engine/Private/AssetManager.cpp
+++ b/Engine/Source/Runtime/Engine/Private/AssetManager.cpp
@@ -321,6 +321,7 @@ struct FCompiledAssetManagerSearchRules : FAssetManagerSearchRules
 const FPrimaryAssetType UAssetManager::MapType = FName(TEXT("Map"));
 const FPrimaryAssetType UAssetManager::PrimaryAssetLabelType = FName(TEXT("PrimaryAssetLabel"));
 const FPrimaryAssetType UAssetManager::PackageChunkType = FName(TEXT("PackageChunk"));
+const FPrimaryAssetType UAssetManager::SuppressionType = FName(TEXT("__Suppressed__"));
 const FString UAssetManager::AssetSearchRootsVirtualPath = TEXT("$AssetSearchRoots");
 const FString UAssetManager::DynamicSearchRootsVirtualPath = TEXT("$DynamicSearchRoots");
 FSimpleMulticastDelegate UAssetManager::OnCompletedInitialScanDelegate;
@@ -980,11 +981,17 @@ int32 UAssetManager::ScanPathsForPrimaryAssets(FPrimaryAssetType PrimaryAssetTyp
 		FPrimaryAssetId PrimaryAssetId = ExtractPrimaryAssetIdFromData(Data, PrimaryAssetType);
 
 		// Remove invalid or wrong type assets
-		if (!PrimaryAssetId.IsValid() || PrimaryAssetId.PrimaryAssetType != PrimaryAssetType)
+		if (!PrimaryAssetId.IsValid() || PrimaryAssetId.PrimaryAssetType == UAssetManager::SuppressionType ||
+			PrimaryAssetId.PrimaryAssetType != PrimaryAssetType)
 		{
-			if (!PrimaryAssetId.IsValid())
+			if (PrimaryAssetId.PrimaryAssetType == UAssetManager::SuppressionType)
 			{
-				UE_LOG(LogAssetManager, Warning, TEXT("Ignoring primary asset %s - PrimaryAssetType %s - invalid primary asset ID"), *Data.AssetName.ToString(), *PrimaryAssetType.ToString());
+				// ExtractPrimaryAssetIdFromData decided to suppress this Asset from being a PrimaryAsset;
+				// silently accept the suppression
+			}
+			else if (!PrimaryAssetId.IsValid())
+			{
+				UE_LOG(LogAssetManager, Warning, TEXT("Ignoring primary asset %s - PrimaryAssetType %s - invalid primary asset ID"), *Data.GetObjectPathString(), *PrimaryAssetType.ToString());
 			}
 			else
 			{
@@ -1628,21 +1635,40 @@ FPrimaryAssetId UAssetManager::ExtractPrimaryAssetIdFromData(const FAssetData& A
 {
 	FPrimaryAssetId FoundId = AssetData.GetPrimaryAssetId();
 
-	if (!FoundId.IsValid() && bShouldGuessTypeAndName && SuggestedType.IsValid())
+	if (!FoundId.IsValid() && SuggestedType.IsValid())
 	{
-		const TSharedRef<FPrimaryAssetTypeData>* FoundType = AssetTypeMap.Find(SuggestedType);
+		FName SuggestedAssetName = AssetData.AssetName;
+		if (SuggestedType == MapType)
+		{
+			// UE-216073: Temporary solution to turn off WorldPartition _Generated_ streaming cells packages
+			// from being PrimaryAssets, since doing so causes duplicate PrimaryAssetId with their owning map.
+			// TODO: Fix this properly by adding a class-specific hook on the PrimaryAssetType's baseclass to make
+			// the decision for what the InAssetName should be for the asset's PrimaryAssetId.
+			if (FPackageName::IsUnderGeneratedPackageSubPath(WriteToString<256>(AssetData.PackageName)))
+			{
+				return FPrimaryAssetId(UAssetManager::SuppressionType, AssetData.PackageName);
+			}
+
+			// Maps use the full package name
+			SuggestedAssetName = AssetData.PackageName;
+		}
 
-		if (ensure(FoundType))
+		if (bShouldGuessTypeAndName)
 		{
-			// If asset at this path is already known about return that
-			FPrimaryAssetId OldID = GetPrimaryAssetIdForPath(GetAssetPathForData(AssetData));
+			const TSharedRef<FPrimaryAssetTypeData>* FoundType = AssetTypeMap.Find(SuggestedType);
 
-			if (OldID.IsValid())
+			if (ensure(FoundType))
 			{
-				return OldID;
-			}
+				// If asset at this path is already known about return that
+				FPrimaryAssetId OldID = GetPrimaryAssetIdForPath(GetAssetPathForData(AssetData));
+
+				if (OldID.IsValid())
+				{
+					return OldID;
+				}
 
-			return FPrimaryAssetId(SuggestedType, SuggestedType == MapType ? AssetData.PackageName : AssetData.AssetName);
+				return FPrimaryAssetId(SuggestedType, SuggestedAssetName);
+			}
 		}
 	}
 
