Add MeshPaintVirtualTexture object type.
This is an implementation of UTexture2D that is always a streaming virtual texture with specific tile settings from the renderer project settings.
It allocates a virtual texture on creation (instead of relying on a he material do the allocation), and tracks the allocated virtual textures.
It can provide a 2 dword descriptor giving the location of the virtual texture in the shared virtual texture pool.
It can provide a description of the shared virtual texture pool for binding to the view uniform buffer.

diff --git a/Engine/Source/Editor/VirtualTexturingEditor/Private/MeshPaintVirtualTextureThumbnailRenderer.cpp b/Engine/Source/Editor/VirtualTexturingEditor/Private/MeshPaintVirtualTextureThumbnailRenderer.cpp
new file mode 100644
--- /dev/null
+++ b/Engine/Source/Editor/VirtualTexturingEditor/Private/MeshPaintVirtualTextureThumbnailRenderer.cpp
@@ -0,0 +1,19 @@
+// Copyright Epic Games, Inc. All Rights Reserved.
+
+#include "MeshPaintVirtualTextureThumbnailRenderer.h"
+#include "VT/MeshPaintVirtualTexture.h"
+
+UMeshPaintVirtualTextureThumbnailRenderer::UMeshPaintVirtualTextureThumbnailRenderer(const FObjectInitializer& ObjectInitializer)
+	: Super(ObjectInitializer)
+{
+}
+
+bool UMeshPaintVirtualTextureThumbnailRenderer::CanVisualizeAsset(UObject* Object)
+{
+	return UTextureThumbnailRenderer::CanVisualizeAsset(Object);
+}
+
+void UMeshPaintVirtualTextureThumbnailRenderer::Draw(UObject* Object, int32 X, int32 Y, uint32 Width, uint32 Height, FRenderTarget* RenderTarget, FCanvas* Canvas, bool bAdditionalViewFamily)
+{
+	UTextureThumbnailRenderer::Draw(Object, X, Y, Width, Height, RenderTarget, Canvas, bAdditionalViewFamily);
+}
diff --git a/Engine/Source/Editor/VirtualTexturingEditor/Private/MeshPaintVirtualTextureThumbnailRenderer.h b/Engine/Source/Editor/VirtualTexturingEditor/Private/MeshPaintVirtualTextureThumbnailRenderer.h
new file mode 100644
--- /dev/null
+++ b/Engine/Source/Editor/VirtualTexturingEditor/Private/MeshPaintVirtualTextureThumbnailRenderer.h
@@ -0,0 +1,17 @@
+// Copyright Epic Games, Inc. All Rights Reserved.
+
+#pragma once
+
+#include "ThumbnailRendering/TextureThumbnailRenderer.h"
+#include "MeshPaintVirtualTextureThumbnailRenderer.generated.h"
+
+UCLASS(MinimalAPI)
+class UMeshPaintVirtualTextureThumbnailRenderer : public UTextureThumbnailRenderer
+{
+	GENERATED_UCLASS_BODY()
+
+	//~ Begin UThumbnailRenderer Interface.
+	virtual bool CanVisualizeAsset(UObject* Object) override;
+	virtual void Draw(UObject* Object, int32 X, int32 Y, uint32 Width, uint32 Height, FRenderTarget*, FCanvas* Canvas, bool bAdditionalViewFamily) override;
+	//~ EndUThumbnailRenderer Interface.
+};
diff --git a/Engine/Source/Editor/VirtualTexturingEditor/Private/VirtualTexturingEditorModule.cpp b/Engine/Source/Editor/VirtualTexturingEditor/Private/VirtualTexturingEditorModule.cpp
--- a/Engine/Source/Editor/VirtualTexturingEditor/Private/VirtualTexturingEditorModule.cpp
+++ b/Engine/Source/Editor/VirtualTexturingEditor/Private/VirtualTexturingEditorModule.cpp
@@ -3,6 +3,7 @@
 #include "VirtualTexturingEditorModule.h"
 
 #include "IPlacementModeModule.h"
+#include "MeshPaintVirtualTextureThumbnailRenderer.h"
 #include "PropertyEditorModule.h"
 #include "RuntimeVirtualTextureAssetTypeActions.h"
 #include "RuntimeVirtualTextureBuildStreamingMips.h"
@@ -12,6 +13,7 @@
 #include "ThumbnailRendering/ThumbnailManager.h"
 #include "VirtualTextureBuilderAssetTypeActions.h"
 #include "VirtualTextureBuilderThumbnailRenderer.h"
+#include "VT/MeshPaintVirtualTexture.h"
 #include "VT/RuntimeVirtualTexture.h"
 #include "VT/RuntimeVirtualTextureVolume.h"
 #include "VT/VirtualTextureBuilder.h"
@@ -59,6 +61,7 @@ void FVirtualTexturingEditorModule::StartupModule()
 
 	UThumbnailManager::Get().RegisterCustomRenderer(URuntimeVirtualTexture::StaticClass(), URuntimeVirtualTextureThumbnailRenderer::StaticClass());
 	UThumbnailManager::Get().RegisterCustomRenderer(UVirtualTextureBuilder::StaticClass(), UVirtualTextureBuilderThumbnailRenderer::StaticClass());
+	UThumbnailManager::Get().RegisterCustomRenderer(UMeshPaintVirtualTexture::StaticClass(), UMeshPaintVirtualTextureThumbnailRenderer::StaticClass());
 }
 
 void FVirtualTexturingEditorModule::ShutdownModule()
diff --git a/Engine/Source/Runtime/Engine/Classes/Engine/RendererSettings.h b/Engine/Source/Runtime/Engine/Classes/Engine/RendererSettings.h
--- a/Engine/Source/Runtime/Engine/Classes/Engine/RendererSettings.h
+++ b/Engine/Source/Runtime/Engine/Classes/Engine/RendererSettings.h
@@ -433,7 +433,7 @@ class URendererSettings : public UDeveloperSettings
 	UPROPERTY(config, EditAnywhere, Category = VirtualTextures, meta = (
 		EditCondition = "bVirtualTextures",
 		ConsoleVariable = "r.VT.TileBorderSize", DisplayName = "Tile border size",
-		ToolTip = "Size in pixels for virtual texture tile borders, will be rounded to next power-of-2. Larger borders allow higher degree of anisotropic filtering, but uses more disk/cache memory. Changing this setting requires restarting the editor.",
+		ToolTip = "Size in pixels for virtual texture tile borders, will be rounded to next multiple-of-2. Larger borders allow higher degree of anisotropic filtering, but uses more disk/cache memory. Changing this setting requires restarting the editor.",
 		ConfigRestartRequired = true))
 	uint32 VirtualTextureTileBorderSize;
 
@@ -443,6 +443,20 @@ class URendererSettings : public UDeveloperSettings
 		ToolTip = "Lower factor will increase virtual texture feedback resolution which increases CPU/GPU overhead, but may decrease streaming latency, especially if materials use many virtual textures."))
 	uint32 VirtualTextureFeedbackFactor;
 
+	UPROPERTY(config, EditAnywhere, Category = VirtualTextures, meta = (
+		EditCondition = "bVirtualTextures",
+		ConsoleVariable = "r.MeshPaintVirtualTexture.TileSize", DisplayName = "Mesh paint tile size",
+		ToolTip = "Size in pixels for mesh paint virtual texture tiles. Will be rounded to next power-of-2. Changing this setting requires restarting the editor.",
+		ConfigRestartRequired = true))
+	uint32 MeshPaintVirtualTextureTileSize;
+
+	UPROPERTY(config, EditAnywhere, Category = VirtualTextures, meta = (
+		EditCondition = "bVirtualTextures",
+		ConsoleVariable = "r.MeshPaintVirtualTexture.TileBorderSize", DisplayName = "Mesh paint tile border size",
+		ToolTip = "Size in pixels for virtual texture tile borders. Will be rounded next to multiple-of-2. Changing this setting requires restarting the editor.",
+		ConfigRestartRequired = true))
+	uint32 MeshPaintVirtualTextureTileBorderSize;
+
 	UPROPERTY(config, EditAnywhere, Category = "Runtime Virtual Textures|Enable Material Types", meta = (
 		EditCondition = "bVirtualTextures",
 		ConsoleVariable = "r.vt.rvt.EnableBaseColor", DisplayName = "Enable 'Base Color'",
diff --git a/Engine/Source/Runtime/Engine/Classes/VT/MeshPaintVirtualTexture.h b/Engine/Source/Runtime/Engine/Classes/VT/MeshPaintVirtualTexture.h
new file mode 100644
--- /dev/null
+++ b/Engine/Source/Runtime/Engine/Classes/VT/MeshPaintVirtualTexture.h
@@ -0,0 +1,58 @@
+// Copyright Epic Games, Inc. All Rights Reserved.
+
+#pragma once
+
+#include "Engine/Texture2D.h"
+#include "RHIFwd.h"
+
+#include "MeshPaintVirtualTexture.generated.h"
+
+/**
+ * Mesh paint virtual texture asset.
+ * This is a virtual texture that which will be owned by a mesh component to store the mesh painting on that component.
+ * All mesh paint virtual textures will be stored using a shared virtual texture page table and physical space.
+ * This shared space means that all mesh paint virtual textures can be accessed in a "bindless" way using a small descriptor.
+ */
+UCLASS(ClassGroup = Rendering, hidecategories = Texture, hidecategories = Compression, hidecategories = Adjustments, hidecategories = Compositing, MinimalAPI)
+class UMeshPaintVirtualTexture : public UTexture2D
+{
+	GENERATED_UCLASS_BODY()
+
+	/** Weak referernce to the owning primitive component. */
+	UPROPERTY()
+	TWeakObjectPtr<UPrimitiveComponent> OwningComponent;
+
+	//~ Begin UTexture Interface.
+	virtual void GetVirtualTextureBuildSettings(FVirtualTextureBuildSettings& OutSettings) const override;
+	virtual void UpdateResource() override;
+
+#if WITH_EDITOR
+	virtual void BeginCacheForCookedPlatformData(const ITargetPlatform* TargetPlatform) override;
+	virtual bool IsCachedCookedPlatformDataLoaded(const ITargetPlatform* TargetPlatform) override;
+	virtual void ClearCachedCookedPlatformData(const ITargetPlatform* TargetPlatform) override;
+#endif
+	//~ End UTexture Interface.
+};
+
+namespace MeshPaintVirtualTexture
+{
+	/** Returns per platform support of mesh paint virtual textures. */
+	ENGINE_API bool IsSupported(EShaderPlatform InShaderPlatform);
+
+	/** 
+	 * Get the 2 dword texture descriptor from texture resource. 
+	 * Will return a null descriptor if the texture resource is not from a UMeshPaintVirtualTexture.
+	 */
+	FUintVector2 GetTextureDescriptor(FTextureResource* InTextureResource);
+
+	/** Scene view parameters that describe the virtual texture space shared by all UMeshPaintVirtualTexture objects. */
+	struct FUniformParams
+	{
+		FTextureRHIRef PageTableTexture = nullptr;
+		FTextureRHIRef PhysicalTexture = nullptr;
+		FUintVector4 PackedUniform = FUintVector4(0, 0, 0, 0);
+	};
+	
+	/** Get the global scene view parameters shared by all UMeshPaintVirtualTexture objects. */
+	ENGINE_API FUniformParams GetUniformParams();
+}
diff --git a/Engine/Source/Runtime/Engine/Private/VT/MeshPaintVirtualTexture.cpp b/Engine/Source/Runtime/Engine/Private/VT/MeshPaintVirtualTexture.cpp
new file mode 100644
--- /dev/null
+++ b/Engine/Source/Runtime/Engine/Private/VT/MeshPaintVirtualTexture.cpp
@@ -0,0 +1,268 @@
+// Copyright Epic Games, Inc. All Rights Reserved.
+
+#include "VT/MeshPaintVirtualTexture.h"
+
+#include "Components/PrimitiveComponent.h"
+#include "EngineModule.h"
+#include "RendererInterface.h"
+#include "RenderUtils.h"
+#include "ShaderPlatformCachedIniValue.h"
+#include "TextureResource.h"
+#include "VirtualTexturing.h"
+#include "VT/VirtualTextureBuildSettings.h"
+
+static TAutoConsoleVariable<bool> CVarMeshPaintVirtualTextureSupport(
+	TEXT("r.MeshPaintVirtualTexture.Support"),
+	true,
+	TEXT("Build time support mesh painting with virtual textures"),
+	ECVF_RenderThreadSafe | ECVF_ReadOnly);
+
+static TAutoConsoleVariable<bool> CVarMeshPaintVirtualTextureEnable(
+	TEXT("r.MeshPaintVirtualTexture.Enable"),
+	true,
+	TEXT("Run time enable mesh painting with virtual textures"),
+	ECVF_RenderThreadSafe);
+
+static TAutoConsoleVariable<int32> CVarMeshPaintVirtualTextureTileSize(
+	TEXT("r.MeshPaintVirtualTexture.TileSize"),
+	32,
+	TEXT("Virtual texture tile size for mesh paint textures"),
+	ECVF_RenderThreadSafe | ECVF_ReadOnly);
+
+static TAutoConsoleVariable<int32> CVarMeshPaintVirtualTextureTileBorderSize(
+	TEXT("r.MeshPaintVirtualTexture.TileBorderSize"),
+	2,
+	TEXT("Virtual texture tile border size for mesh paint textures"),
+	ECVF_RenderThreadSafe | ECVF_ReadOnly);
+
+
+namespace MeshPaintVirtualTexture
+{
+	bool IsSupported(EShaderPlatform InShaderPlatform)
+	{
+		static FShaderPlatformCachedIniValue<bool> CPlatformVarMeshPaintVirtualTextureSupport(CVarMeshPaintVirtualTextureSupport.AsVariable());
+		return CPlatformVarMeshPaintVirtualTextureSupport.Get(InShaderPlatform) && UseVirtualTexturing(InShaderPlatform);
+	}
+
+	static bool IsEnabled()
+	{
+		return CVarMeshPaintVirtualTextureEnable.GetValueOnAnyThread();
+	}
+
+	static uint32 GetTileSize()
+	{
+		return FVirtualTextureBuildSettings::ClampAndAlignTileSize(CVarMeshPaintVirtualTextureTileSize.GetValueOnAnyThread());
+	}
+
+	static uint32 GetTileBorderSize()
+	{
+		return FVirtualTextureBuildSettings::ClampAndAlignTileBorderSize(CVarMeshPaintVirtualTextureTileBorderSize.GetValueOnAnyThread());
+	}
+
+	/** 
+	 * Fill out the scene view params from an allocated VT. 
+	 * We expect the result to be constant for all allocated VTs (so that they can share one set of data in the view).
+	 * Note that there are valid cases when it will change over time.
+	 * For example, when the virtual texture pools are resized, all VTs reallocate and change to a new value here.
+	 */
+	static void GetSceneViewParams(IAllocatedVirtualTexture* InAllocatedVT, FUniformParams& OutParams)
+	{
+		OutParams.PageTableTexture = InAllocatedVT->GetPageTableTexture(0);
+		OutParams.PhysicalTexture = InAllocatedVT->GetPhysicalTexture(0);
+
+		// Packed uniform layout should match shader unpacking in VTUniform_Unpack().
+		// This packing code is copied from FAllocatedVirtualTexture::GetPackedUniform()
+
+		const uint32 vPageSize = InAllocatedVT->GetVirtualTileSize();
+		const uint32 PageBorderSize = InAllocatedVT->GetTileBorderSize();
+		const float RcpPhysicalTextureSize = 1.0f / float(InAllocatedVT->GetPhysicalTextureSize(0));
+		const uint32 pPageSize = vPageSize + PageBorderSize * 2u;
+
+		OutParams.PackedUniform.X = ~0u; // todo fallback color
+		OutParams.PackedUniform.Y = FMath::AsUInt((float)vPageSize * RcpPhysicalTextureSize);
+		OutParams.PackedUniform.Z = FMath::AsUInt((float)PageBorderSize * RcpPhysicalTextureSize);
+		
+		const bool bPageTableExtraBits = InAllocatedVT->GetPageTableFormat() == EVTPageTableFormat::UInt32;
+		const float PackedSignBit = bPageTableExtraBits ? 1.f : -1.f;
+		OutParams.PackedUniform.W = FMath::AsUInt((float)pPageSize * RcpPhysicalTextureSize * PackedSignBit);
+	}
+
+	/** A global set to track allocated mesh paint virtual textures. */
+	TSet<IAllocatedVirtualTexture*> AllocatedVTs;
+	/** The global scene view params cached from the last allocated VT. */
+	FUniformParams Params;
+	/** Mutex for the global state. The state is all accessed only on the render thread timeline, but the Params are accessed from render worker threads. */
+	UE::FRecursiveMutex Mutex;
+
+	/** Add an allocated VT to our global allocated set. */
+	static void AddAllocatedVT(IAllocatedVirtualTexture* InAllocatedVT)
+	{
+		UE::TScopeLock Lock(Mutex);
+		
+		bool bAlreadyInSet = false;
+		AllocatedVTs.Add(InAllocatedVT, &bAlreadyInSet);
+		// This may fire in future if we allow components to share the same virtual texture.
+		// If that happens we could change to store in a map against a ref count.
+		ensure(!bAlreadyInSet);
+
+		// Update the cached view params.
+		GetSceneViewParams(InAllocatedVT, Params);
+	}
+
+	/** Remove an allocated VT from our global allocated set. */
+	static void RemoveAllocatedVT(const FVirtualTextureProducerHandle& InHandle, void* InBaton)
+	{
+		UE::TScopeLock Lock(Mutex);
+
+		const int32 Removed = AllocatedVTs.Remove((IAllocatedVirtualTexture*)InBaton);
+		ensure(Removed == 1);
+
+		// Clear the cached view arams if we have removed all the allocated VTs.
+		if (AllocatedVTs.Num() == 0)
+		{
+			Params = {};
+		}
+	}
+
+	/** 
+	 * Call this on texture resource creation.
+	 * This will acquire the virtual texture and store to our global set.
+	 */
+	static void AcquireAllocatedVT(FTextureResource* Resource)
+	{
+		FVirtualTexture2DResource* VTResource = Resource != nullptr ? Resource->GetVirtualTexture2DResource() : nullptr;
+		if (VTResource)
+		{
+			ENQUEUE_RENDER_COMMAND(AcquireVT)([VTResource](FRHICommandListImmediate& RHICmdList)
+			{
+				if (IAllocatedVirtualTexture* AllocatedVT = VTResource->AcquireAllocatedVT())
+				{
+					// Add virtual texture to the global set.
+					AddAllocatedVT(AllocatedVT);
+					// Queue on-destruction callback for removal from the global set.
+					GetRendererModule().AddVirtualTextureProducerDestroyedCallback(AllocatedVT->GetProducerHandle(0), &RemoveAllocatedVT, AllocatedVT);
+				}
+			});
+		}
+	}
+
+	FUniformParams GetUniformParams()
+	{
+		if (MeshPaintVirtualTexture::IsEnabled())
+		{
+			UE::TScopeLock Lock(Mutex);
+			return Params;
+		}
+		return {};
+	}
+
+	FUintVector2 GetTextureDescriptor(FTextureResource* InTextureResource)
+	{
+		FUintVector2 Descriptor(0, 0);
+
+		if (InTextureResource == nullptr)
+		{
+			return Descriptor;
+		}
+
+		if (!MeshPaintVirtualTexture::IsEnabled())
+		{
+			return Descriptor;
+		}
+		
+		FVirtualTexture2DResource* VTTextureResource = InTextureResource->GetVirtualTexture2DResource();
+		if (VTTextureResource == nullptr)
+		{
+			return Descriptor;
+		}
+
+		IAllocatedVirtualTexture* AllocatedVT = VTTextureResource->GetAllocatedVT();
+		if (AllocatedVT == nullptr)
+		{
+			return Descriptor;
+		}
+
+		const uint32 vPageX = AllocatedVT->GetVirtualPageX();
+		const uint32 vPageY = AllocatedVT->GetVirtualPageY();
+		const uint32 WidthInPages = AllocatedVT->GetWidthInTiles();
+		const uint32 HeightInPages = AllocatedVT->GetHeightInTiles();
+		const uint32 vPageTableMipBias = FMath::FloorLog2(AllocatedVT->GetVirtualTileSize());
+		const uint32 SpaceID = AllocatedVT->GetSpaceID();
+		const uint32 MaxLevel = AllocatedVT->GetMaxLevel();
+
+		// Descriptor layout should match shader unpacking in VTPageTableUniform_Unpack().
+		Descriptor.X = vPageX | (vPageY << 12) | (vPageTableMipBias << 24) | (SpaceID << 28);
+		Descriptor.Y = WidthInPages | (HeightInPages << 12) | (MaxLevel << 24);
+		
+		return Descriptor;
+	}
+}
+
+
+UMeshPaintVirtualTexture::UMeshPaintVirtualTexture(const FObjectInitializer& ObjectInitializer)
+	: UTexture2D(ObjectInitializer)
+{
+	VirtualTextureStreaming = true;
+}
+
+void UMeshPaintVirtualTexture::GetVirtualTextureBuildSettings(FVirtualTextureBuildSettings& OutSettings) const
+{
+	// Use the specific tile size for mesh painting textures.
+	// This is typically different from the default virtual texture tile size. That's good because it keeps the mesh painting virtual texture pools isolated.
+	// It's also better to have a small tile size since texture painting is typically low frequency and needs relatively small textures. Small tiles also give improved mipping and less wastage at distance.
+	OutSettings.TileSize = MeshPaintVirtualTexture::GetTileSize();
+	OutSettings.TileBorderSize = MeshPaintVirtualTexture::GetTileBorderSize();
+}
+
+void UMeshPaintVirtualTexture::UpdateResource()
+{
+	Super::UpdateResource();
+
+	// We get here on virtual texture pool recreation, and on texture compilation in editor.
+	// In those cases we need to reacquire the virtual texture, and notify our component.
+	MeshPaintVirtualTexture::AcquireAllocatedVT(GetResource());
+
+	// We assume a 1-1 mapping of component and texture here.
+	// If in future we want to share a painted texture across components then we will need a way to track the set of components to dirty.
+	if (UPrimitiveComponent* PrimitiveComponent = OwningComponent.Get())
+	{
+		PrimitiveComponent->MarkRenderStateDirty();
+	}
+}
+
+#if WITH_EDITOR
+
+void UMeshPaintVirtualTexture::BeginCacheForCookedPlatformData(const ITargetPlatform* TargetPlatform)
+{
+	// Even though we skip the cook of this object for non VT platforms in URuntimeVirtualTexture::Serialize()
+	// we still load the object at cook time and kick off the DDC build. This will trigger an error in the texture DDC code.
+	// Either we need to make the DDC code more robust for non VT platforms or we can skip the process here...
+	if (!UseVirtualTexturing(GMaxRHIShaderPlatform, TargetPlatform) || !MeshPaintVirtualTexture::IsSupported(GMaxRHIShaderPlatform))
+	{
+		return;
+	}
+
+	Super::BeginCacheForCookedPlatformData(TargetPlatform);
+}
+
+bool UMeshPaintVirtualTexture::IsCachedCookedPlatformDataLoaded(const ITargetPlatform* TargetPlatform)
+{
+	if (!UseVirtualTexturing(GMaxRHIShaderPlatform, TargetPlatform) || !MeshPaintVirtualTexture::IsSupported(GMaxRHIShaderPlatform))
+	{
+		return true;
+	}
+
+	return Super::IsCachedCookedPlatformDataLoaded(TargetPlatform);
+}
+
+void UMeshPaintVirtualTexture::ClearCachedCookedPlatformData(const ITargetPlatform* TargetPlatform)
+{
+	if (!UseVirtualTexturing(GMaxRHIShaderPlatform, TargetPlatform) || !MeshPaintVirtualTexture::IsSupported(GMaxRHIShaderPlatform))
+	{
+		return;
+	}
+
+	Super::ClearCachedCookedPlatformData(TargetPlatform);
+}
+
+#endif
