Add a material expression to access the mesh paint texture descriptor from a primitive and coerce it to a new mesh paint virtual texture type that can be sampled with a standard texture sample node.

diff --git a/Engine/Shaders/Private/VirtualTextureCommon.ush b/Engine/Shaders/Private/VirtualTextureCommon.ush
--- a/Engine/Shaders/Private/VirtualTextureCommon.ush
+++ b/Engine/Shaders/Private/VirtualTextureCommon.ush
@@ -192,6 +192,22 @@ VTPageTableUniform VTPageTableUniform_Unpack(uint4 PackedPageTableUniform0, uint
 	return result;
 }
 
+VTPageTableUniform VTPageTableUniform_Unpack(uint2 PackedPageTableUniform)
+{
+	VTPageTableUniform Result;
+	Result.UVScale = 1.f;
+	Result.SizeInPages.x = PackedPageTableUniform.y & 0xfff;
+	Result.SizeInPages.y = (PackedPageTableUniform.y >> 12) & 0xfff;
+	Result.MaxAnisoLog2 = 0;
+	Result.XOffsetInPages = PackedPageTableUniform.x & 0xfff;
+	Result.YOffsetInPages = (PackedPageTableUniform.x >> 12) & 0xfff;
+	Result.vPageTableMipBias = (PackedPageTableUniform.x >> 24) & 0xf;
+	Result.MaxLevel = (PackedPageTableUniform.y >> 24) & 0xf;
+	Result.AdaptiveLevelBias = 0;
+	Result.ShiftedPageTableID = PackedPageTableUniform.x & 0xf0000000;
+	return Result;
+}
+
 /** Structure carrying page table read result. Is passed as inout parameter and partially filled by several functions. */
 struct VTPageTableResult
 {
diff --git a/Engine/Source/Runtime/Core/Private/UObject/FortniteMainBranchObjectVersion.cpp b/Engine/Source/Runtime/Core/Private/UObject/FortniteMainBranchObjectVersion.cpp
--- a/Engine/Source/Runtime/Core/Private/UObject/FortniteMainBranchObjectVersion.cpp
+++ b/Engine/Source/Runtime/Core/Private/UObject/FortniteMainBranchObjectVersion.cpp
@@ -14,6 +14,6 @@ TMap<FGuid, FGuid> FFortniteMainBranchObjectVersion::GetSystemGuids()
 	SystemGuids.Add(DevGuids.Niagara_LatestScriptCompileVersion, FGuid("E486231DA967C64B95D4E73131056FA0"));
 	SystemGuids.Add(DevGuids.SkeletalMeshDerivedDataVersion, FGuid("B9AEAA3EE7AC4FCFB3860E8F3E06DD75"));
 	SystemGuids.Add(DevGuids.STATICMESH_DERIVEDDATA_VER, FGuid("A3C8F19042174D05A6AEB6E5D9A58E44"));
-	SystemGuids.Add(DevGuids.MaterialTranslationDDCVersion, FGuid("BE7C92C6DC7D430B822D815C9D75906D"));
+	SystemGuids.Add(DevGuids.MaterialTranslationDDCVersion, FGuid("A80C07846638406DA411360524FEC16A"));
 	return SystemGuids;
 }
diff --git a/Engine/Source/Runtime/Engine/Private/Materials/HLSLMaterialTranslator.cpp b/Engine/Source/Runtime/Engine/Private/Materials/HLSLMaterialTranslator.cpp
--- a/Engine/Source/Runtime/Engine/Private/Materials/HLSLMaterialTranslator.cpp
+++ b/Engine/Source/Runtime/Engine/Private/Materials/HLSLMaterialTranslator.cpp
@@ -235,9 +235,10 @@ struct FHLSLMaterialTranslator::FEnvironmentDefines
 {
 	enum EVirtualPageType
 	{
-		TABLE0 = 0,
-		TABLE1 = 1,
-		TABLE_INDIRECTION = 2
+		TABLE_GLOBAL = 1,
+		TABLE0 = 2,
+		TABLE1 = 4,
+		TABLE_INDIRECTION = 8,
 	};
 
 	bool bNeedsParticlePosition;
@@ -2640,20 +2641,37 @@ void FHLSLMaterialTranslator::GetMaterialEnvironment(EShaderPlatform InPlatform,
 	OutEnvironment.SetDefine(TEXT("NUM_VIRTUALTEXTURE_SAMPLES"), EnvironmentDefines->NumVirtualTextureSamples);
 	OutEnvironment.SetDefine(TEXT("MATERIAL_VIRTUALTEXTURE_FEEDBACK"), EnvironmentDefines->bMaterialVirtualTextureFeedback);
 
-	// Setup defines to map each VT stack to either 1 or 2 page table textures, depending on how many layers it uses
 	for (int i = 0; i < EnvironmentDefines->VirtualPageTypes.Num(); ++i)
 	{
-		const FMaterialVirtualTextureStack& Stack = MaterialCompilationOutput.UniformExpressionSet.VTStacks[i];
-		FString PageTableValue = FString::Printf(TEXT("Material.VirtualTexturePageTable0_%d"), i);
-		if (EnvironmentDefines->VirtualPageTypes[i] & FEnvironmentDefines::TABLE1)
+		if (EnvironmentDefines->VirtualPageTypes[i] & FEnvironmentDefines::TABLE_GLOBAL)
 		{
-			PageTableValue += FString::Printf(TEXT(", Material.VirtualTexturePageTable1_%d"), i);
+			OutEnvironment.SetDefine(
+				*FString::Printf(TEXT("VIRTUALTEXTURE_PAGETABLE_%d"), i), 
+				TEXT("Scene.MeshPaint.PageTableTexture"));
+
+			OutEnvironment.SetDefine(
+				*FString::Printf(TEXT("VIRTUALTEXTURE_PAGETABLE_UNIFORM_%d"), i), 
+				TEXT("GetPrimitiveData(Parameters).MeshPaintTextureDescriptor"));
 		}
-		if (EnvironmentDefines->VirtualPageTypes[i] & FEnvironmentDefines::TABLE_INDIRECTION)
+		else
 		{
-			PageTableValue += FString::Printf(TEXT(", Material.VirtualTexturePageTableIndirection_%d"), i);
+			// Setup page table defines to map each VT stack to either 1 or 2 page table textures, depending on how many layers it uses
+			const FMaterialVirtualTextureStack& Stack = MaterialCompilationOutput.UniformExpressionSet.VTStacks[i];
+			FString PageTableValue = FString::Printf(TEXT("Material.VirtualTexturePageTable0_%d"), i);
+			if (EnvironmentDefines->VirtualPageTypes[i] & FEnvironmentDefines::TABLE1)
+			{
+				PageTableValue += FString::Printf(TEXT(", Material.VirtualTexturePageTable1_%d"), i);
+			}
+			if (EnvironmentDefines->VirtualPageTypes[i] & FEnvironmentDefines::TABLE_INDIRECTION)
+			{
+				PageTableValue += FString::Printf(TEXT(", Material.VirtualTexturePageTableIndirection_%d"), i);
+			}
+			OutEnvironment.SetDefine(*FString::Printf(TEXT("VIRTUALTEXTURE_PAGETABLE_%d"), i), *PageTableValue);
+
+			// Setup page table uniform defines.
+			FString PageTableUniformValue = FString::Printf(TEXT("Material.VTPackedPageTableUniform[%d*2], Material.VTPackedPageTableUniform[%d*2+1]"), i, i);
+			OutEnvironment.SetDefine(*FString::Printf(TEXT("VIRTUALTEXTURE_PAGETABLE_UNIFORM_%d"), i), *PageTableUniformValue);
 		}
-		OutEnvironment.SetDefine(*FString::Printf(TEXT("VIRTUALTEXTURE_PAGETABLE_%d"), i), *PageTableValue);
 	}
 
 	for (int32 CollectionIndex = 0; CollectionIndex < EnvironmentDefines->ParameterCollections.Num(); CollectionIndex++)
@@ -6720,13 +6738,18 @@ uint32 FHLSLMaterialTranslator::AcquireVTStackIndex(
 	const FString& UV_Value,
 	const FString& UV_Ddx,
 	const FString& UV_Ddy,
-	bool bAdaptive, bool bGenerateFeedback)
+	bool bAdaptive, bool bGenerateFeedback, bool bMeshPaint)
 {
-	const uint64 CoordinatHash = GetParameterHash(CoordinateIndex);
+	const uint64 CoordinateHash = GetParameterHash(CoordinateIndex);
 	const uint64 MipValue0Hash = GetParameterHash(MipValue0Index);
 	const uint64 MipValue1Hash = GetParameterHash(MipValue1Index);
 
-	uint64 Hash = CityHash128to64({ CurrentScopeID, CoordinatHash });
+	uint32 Flags = 0;
+	Flags |= bAdaptive ? (1 << 0) : 0;
+	Flags |= bGenerateFeedback ? (1 << 1) : 0;
+	Flags |= bMeshPaint ? (1 << 2) : 0;
+
+	uint64 Hash = CityHash128to64({ CurrentScopeID, CoordinateHash });
 	Hash = CityHash128to64({ Hash, MipValue0Hash });
 	Hash = CityHash128to64({ Hash, MipValue1Hash });
 	Hash = CityHash128to64({ Hash, (uint64)MipValueMode });
@@ -6734,8 +6757,7 @@ uint32 FHLSLMaterialTranslator::AcquireVTStackIndex(
 	Hash = CityHash128to64({ Hash, (uint64)AddressV });
 	Hash = CityHash128to64({ Hash, (uint64)(AspectRatio * 1000.0f) });
 	Hash = CityHash128to64({ Hash, (uint64)PreallocatedStackTextureIndex });
-	Hash = CityHash128to64({ Hash, (uint64)(bAdaptive ? 1 : 0) });
-	Hash = CityHash128to64({ Hash, (uint64)(bGenerateFeedback ? 1 : 0) });
+	Hash = CityHash128to64({ Hash, (uint64)Flags });
 
 	// First check to see if we have an existing VTStack that matches this key, that can still fit another layer
 	for (int32 Index = VTStackHash.First(Hash); VTStackHash.IsValid(Index); Index = VTStackHash.Next(Index))
@@ -6744,7 +6766,7 @@ uint32 FHLSLMaterialTranslator::AcquireVTStackIndex(
 		const FMaterialVTStackEntry& Entry = VTStacks[Index];
 		if (!Stack.AreLayersFull() &&
 			Entry.ScopeID == CurrentScopeID &&
-			Entry.CoordinateHash == CoordinatHash &&
+			Entry.CoordinateHash == CoordinateHash &&
 			Entry.MipValue0Hash == MipValue0Hash &&
 			Entry.MipValue1Hash == MipValue1Hash &&
 			Entry.MipValueMode == MipValueMode &&
@@ -6753,7 +6775,8 @@ uint32 FHLSLMaterialTranslator::AcquireVTStackIndex(
 			Entry.AspectRatio == AspectRatio &&
 			Entry.PreallocatedStackTextureIndex == PreallocatedStackTextureIndex &&
 			Entry.bAdaptive == bAdaptive &&
-			Entry.bGenerateFeedback == bGenerateFeedback)
+			Entry.bGenerateFeedback == bGenerateFeedback &&
+			Entry.bMeshPaint == bMeshPaint)
 		{
 			return Index;
 		}
@@ -6764,7 +6787,7 @@ uint32 FHLSLMaterialTranslator::AcquireVTStackIndex(
 	VTStackHash.Add(Hash, StackIndex);
 	FMaterialVTStackEntry& Entry = VTStacks[StackIndex];
 	Entry.ScopeID = CurrentScopeID;
-	Entry.CoordinateHash = CoordinatHash;
+	Entry.CoordinateHash = CoordinateHash;
 	Entry.MipValue0Hash = MipValue0Hash;
 	Entry.MipValue1Hash = MipValue1Hash;
 	Entry.MipValueMode = MipValueMode;
@@ -6777,6 +6800,7 @@ uint32 FHLSLMaterialTranslator::AcquireVTStackIndex(
 	Entry.PreallocatedStackTextureIndex = PreallocatedStackTextureIndex;
 	Entry.bAdaptive = bAdaptive;
 	Entry.bGenerateFeedback = bGenerateFeedback;
+	Entry.bMeshPaint = bMeshPaint;
 
 	MaterialCompilationOutput.UniformExpressionSet.VTStacks.Add(FMaterialVirtualTextureStack(PreallocatedStackTextureIndex));
 
@@ -6802,12 +6826,12 @@ uint32 FHLSLMaterialTranslator::AcquireVTStackIndex(
 		FString SampleCodeFinite = FString::Printf(TEXT(
 			"%s("
 			"VIRTUALTEXTURE_PAGETABLE_%d, "
-			"VTPageTableUniform_Unpack(Material.VTPackedPageTableUniform[%d*2], Material.VTPackedPageTableUniform[%d*2+1]), "
+			"VTPageTableUniform_Unpack(VIRTUALTEXTURE_PAGETABLE_UNIFORM_%d), "
 			"%s, %s, %s, "
 			"0, Parameters.SvPosition.xy, "
 			"%dU + LIGHTMAP_VT_ENABLED, Parameters.VirtualTextureFeedback)"),
 			*BaseFunctionName,
-			StackIndex, StackIndex, StackIndex,
+			StackIndex, StackIndex,
 			*CoerceParameter(CoordinateIndex, MCT_Float2), GetVTAddressMode(AddressU), GetVTAddressMode(AddressV),
 			StackIndex);
 
@@ -6816,12 +6840,12 @@ uint32 FHLSLMaterialTranslator::AcquireVTStackIndex(
 			FString SampleCodeAnalytic = FString::Printf(TEXT(
 				"%sGrad("
 				"VIRTUALTEXTURE_PAGETABLE_%d, "
-				"VTPageTableUniform_Unpack(Material.VTPackedPageTableUniform[%d*2], Material.VTPackedPageTableUniform[%d*2+1]), "
+				"VTPageTableUniform_Unpack(VIRTUALTEXTURE_PAGETABLE_UNIFORM_%d), "
 				"%s, %s, %s, "
 				"%s, %s, Parameters.SvPosition.xy, "
 				"%dU + LIGHTMAP_VT_ENABLED, Parameters.VirtualTextureFeedback)"),
 				*BaseFunctionName,
-				StackIndex, StackIndex, StackIndex,
+				StackIndex, StackIndex,
 				*UV_Value, GetVTAddressMode(AddressU), GetVTAddressMode(AddressV), 
 				*UV_Ddx, *UV_Ddy,
 				StackIndex);
@@ -6839,12 +6863,12 @@ uint32 FHLSLMaterialTranslator::AcquireVTStackIndex(
 		FString SampleCodeFinite = FString::Printf(TEXT(
 			"%s("
 			"VIRTUALTEXTURE_PAGETABLE_%d, "
-			"VTPageTableUniform_Unpack(Material.VTPackedPageTableUniform[%d*2], Material.VTPackedPageTableUniform[%d*2+1]), "
+			"VTPageTableUniform_Unpack(VIRTUALTEXTURE_PAGETABLE_UNIFORM_%d), "
 			"%s, %s, %s, "
 			"%s, Parameters.SvPosition.xy, "
 			"%dU + LIGHTMAP_VT_ENABLED, Parameters.VirtualTextureFeedback)"),
 			*BaseFunctionName,
-			StackIndex, StackIndex, StackIndex, 
+			StackIndex, StackIndex,
 			*CoerceParameter(CoordinateIndex, MCT_Float2), GetVTAddressMode(AddressU), GetVTAddressMode(AddressV), 
 			*CoerceParameter(MipValue0Index, MCT_Float1),
 			StackIndex);
@@ -6853,12 +6877,12 @@ uint32 FHLSLMaterialTranslator::AcquireVTStackIndex(
 			FString SampleCodeAnalytic = FString::Printf(TEXT(
 				"%sGrad("
 				"VIRTUALTEXTURE_PAGETABLE_%d, "
-				"VTPageTableUniform_Unpack(Material.VTPackedPageTableUniform[%d*2], Material.VTPackedPageTableUniform[%d*2+1]), "
+				"VTPageTableUniform_Unpack(VIRTUALTEXTURE_PAGETABLE_UNIFORM_%d), "
 				"%s, %s, %s, "
 				"%s, %s, Parameters.SvPosition.xy, "
 				"%dU + LIGHTMAP_VT_ENABLED, Parameters.VirtualTextureFeedback)"),
 				*BaseFunctionName,
-				StackIndex, StackIndex, StackIndex,
+				StackIndex, StackIndex, 
 				*UV_Value, GetVTAddressMode(AddressU), GetVTAddressMode(AddressV),
 				*UV_Ddx, *UV_Ddy,
 				StackIndex);
@@ -6874,12 +6898,12 @@ uint32 FHLSLMaterialTranslator::AcquireVTStackIndex(
 		Entry.CodeIndex = AddCodeChunk(MCT_VTPageTableResult, TEXT(
 			"%sLevel("
 			"VIRTUALTEXTURE_PAGETABLE_%d, " 
-			"VTPageTableUniform_Unpack(Material.VTPackedPageTableUniform[%d*2], Material.VTPackedPageTableUniform[%d*2+1]), "
+			"VTPageTableUniform_Unpack(VIRTUALTEXTURE_PAGETABLE_UNIFORM_%d), "
 			"%s, %s, %s, "
 			"%s"
 			"%s)"),
 			*BaseFunctionName,
-			StackIndex, StackIndex, StackIndex,
+			StackIndex, StackIndex, 
 			*CoerceParameter(CoordinateIndex, MCT_Float2), GetVTAddressMode(AddressU), GetVTAddressMode(AddressV), 
 			*CoerceParameter(MipValue0Index, MCT_Float1),
 			*FeedbackParameter);
@@ -6888,12 +6912,12 @@ uint32 FHLSLMaterialTranslator::AcquireVTStackIndex(
 		Entry.CodeIndex = AddCodeChunk(MCT_VTPageTableResult, TEXT(
 			"%sGrad("
 			"VIRTUALTEXTURE_PAGETABLE_%d, "
-			"VTPageTableUniform_Unpack(Material.VTPackedPageTableUniform[%d*2], Material.VTPackedPageTableUniform[%d*2+1]), "
+			"VTPageTableUniform_Unpack(VIRTUALTEXTURE_PAGETABLE_UNIFORM_%d), "
 			"%s, %s, %s, "
 			"%s, %s, Parameters.SvPosition.xy, "
 			"%dU + LIGHTMAP_VT_ENABLED, Parameters.VirtualTextureFeedback)"),
 			*BaseFunctionName,
-			StackIndex, StackIndex, StackIndex, 
+			StackIndex, StackIndex, 
 			*CoerceParameter(CoordinateIndex, MCT_Float2), GetVTAddressMode(AddressU), GetVTAddressMode(AddressV),
 			*CoerceParameter(MipValue0Index, MCT_Float2), *CoerceParameter(MipValue1Index, MCT_Float2),
 			StackIndex);
@@ -7004,6 +7028,9 @@ int32 FHLSLMaterialTranslator::TextureSample(
 	const bool bFromCollection = (TextureType & MCT_TextureCollection) != 0;
 	TextureType = EMaterialValueType(TextureType & ~MCT_TextureCollection);
 
+	const bool bMeshPaint = (TextureType & MCT_TextureMeshPaint) != 0;
+	TextureType = EMaterialValueType(TextureType & ~MCT_TextureMeshPaint);
+
 	if (ShaderFrequency != SF_Pixel && MipValueMode == TMVM_MipBias)
 	{
 		Errorf(TEXT("MipBias is only supported in the pixel shader"));
@@ -7095,6 +7122,14 @@ int32 FHLSLMaterialTranslator::TextureSample(
 	{
 		TextureName = CoerceParameter(TextureIndex, MCT_TextureExternal);
 	}
+	else if (bMeshPaint)
+	{
+		check(bVirtualTexture);
+
+		TextureName = TEXT("Scene.MeshPaint.PhysicalTexture, View.SharedBilinearClampedSampler");
+
+		NumVtSamples++;
+	}
 	else if (bVirtualTexture)
 	{
 		// Note, this does not really do anything (by design) other than adding it to the UniformExpressionSet
@@ -7138,7 +7173,7 @@ int32 FHLSLMaterialTranslator::TextureSample(
 
 	// Won't be able to get the texture, if this is an external texture sample
 	const UTexture* Texture = nullptr;
-	if (TextureType != MCT_TextureExternal && !bFromCollection)
+	if (TextureType != MCT_TextureExternal && !bFromCollection && !bMeshPaint)
 	{
 		FMaterialUniformExpression* Expression = (*CurrentScopeChunks)[TextureIndex].UniformExpression;
 		const FMaterialUniformExpressionTexture* TextureExpression = Expression ? Expression->GetTextureUniformExpression() : nullptr;
@@ -7426,41 +7461,53 @@ int32 FHLSLMaterialTranslator::TextureSample(
 		int32 VTStackIndex = INDEX_NONE;
 		int32 VTLayerIndex = INDEX_NONE;
 		int32 VTPageTableIndex = INDEX_NONE;
-
-		check(VirtualTextureIndex >= 0);
+		FString VTPackedUniformName;
 
 		// Only support GPU feedback from pixel shader
 		//todo[vt]: Support feedback from other shader types
 		const bool bGenerateFeedback = EnableFeedback && ShaderFrequency == SF_Pixel;
 
-		VTLayerIndex = UniformTextureExpressions[(uint32)EMaterialTextureParameterType::Virtual][VirtualTextureIndex]->GetTextureLayerIndex();
-		if (VTLayerIndex != INDEX_NONE)
+		if (bMeshPaint)
 		{
-			// The layer index in the virtual texture stack is already known
-			// Create a page table sample for each new combination of virtual texture and sample parameters
-			VTStackIndex = AcquireVTStackIndex(MipValueMode, StaticAddressX, StaticAddressY, 1.0f, NonLWCCoordinateIndex, MipValue0Index, MipValue1Index, TextureReferenceIndex, UV_Value, UV_Ddx, UV_Ddy, AdaptiveVirtualTexture, bGenerateFeedback);
-			VTPageTableIndex = UniformTextureExpressions[(uint32)EMaterialTextureParameterType::Virtual][VirtualTextureIndex]->GetPageTableLayerIndex();
+			VTStackIndex = AcquireVTStackIndex(MipValueMode, StaticAddressX, StaticAddressY, 1.0f, CoordinateIndex, MipValue0Index, MipValue1Index, INDEX_NONE, UV_Value, UV_Ddx, UV_Ddy, AdaptiveVirtualTexture, bGenerateFeedback, bMeshPaint);
+			VTPageTableIndex = 0;
+			VTPackedUniformName = TEXT("Scene.MeshPaint.PackedUniform");
 		}
-		else
+		else 
 		{
-			// Textures can only be combined in a VT stack if they have the same aspect ratio
-			// This also means that any texture parameters set in material instances for VTs must match the aspect ratio of the texture in the parent material
-			// (Otherwise could potentially break stacks)
-			check(Texture);
+			check(VirtualTextureIndex >= 0);
+			VTPackedUniformName = FString::Printf(TEXT("Material.VTPackedUniform[%d]"), VirtualTextureIndex);
 
-			// Using Source size because we care about the aspect ratio of each block (each block of multi-block texture must have same aspect ratio)
-			// We can still combine multi-block textures of different block aspect ratios, as long as each block has the same ratio
-			// This is because we only need to overlay VT pages from within a given block
-			const float TextureAspectRatio = (float)Texture->Source.GetSizeX() / (float)Texture->Source.GetSizeY();
+			VTLayerIndex = UniformTextureExpressions[(uint32)EMaterialTextureParameterType::Virtual][VirtualTextureIndex]->GetTextureLayerIndex();
 
-			// Create a page table sample for each new set of sample parameters
-			VTStackIndex = AcquireVTStackIndex(MipValueMode, StaticAddressX, StaticAddressY, TextureAspectRatio, NonLWCCoordinateIndex, MipValue0Index, MipValue1Index, INDEX_NONE, UV_Value, UV_Ddx, UV_Ddy, AdaptiveVirtualTexture, bGenerateFeedback);
-			// Allocate a layer in the virtual texture stack for this physical sample
-			VTLayerIndex = MaterialCompilationOutput.UniformExpressionSet.VTStacks[VTStackIndex].AddLayer();
-			VTPageTableIndex = VTLayerIndex;
-		}
+			if (VTLayerIndex != INDEX_NONE)
+			{
+				// The layer index in the virtual texture stack is already known
+				// Create a page table sample for each new combination of virtual texture and sample parameters
+				VTStackIndex = AcquireVTStackIndex(MipValueMode, StaticAddressX, StaticAddressY, 1.0f, NonLWCCoordinateIndex, MipValue0Index, MipValue1Index, TextureReferenceIndex, UV_Value, UV_Ddx, UV_Ddy, AdaptiveVirtualTexture, bGenerateFeedback, bMeshPaint);
+				VTPageTableIndex = UniformTextureExpressions[(uint32)EMaterialTextureParameterType::Virtual][VirtualTextureIndex]->GetPageTableLayerIndex();
+			}
+			else
+			{
+				// Textures can only be combined in a VT stack if they have the same aspect ratio
+				// This also means that any texture parameters set in material instances for VTs must match the aspect ratio of the texture in the parent material
+				// (Otherwise could potentially break stacks)
+				check(Texture);
+
+				// Using Source size because we care about the aspect ratio of each block (each block of multi-block texture must have same aspect ratio)
+				// We can still combine multi-block textures of different block aspect ratios, as long as each block has the same ratio
+				// This is because we only need to overlay VT pages from within a given block
+				const float TextureAspectRatio = (float)Texture->Source.GetSizeX() / (float)Texture->Source.GetSizeY();
 
-		MaterialCompilationOutput.UniformExpressionSet.VTStacks[VTStackIndex].SetLayer(VTLayerIndex, VirtualTextureIndex);
+				// Create a page table sample for each new set of sample parameters
+				VTStackIndex = AcquireVTStackIndex(MipValueMode, StaticAddressX, StaticAddressY, TextureAspectRatio, NonLWCCoordinateIndex, MipValue0Index, MipValue1Index, INDEX_NONE, UV_Value, UV_Ddx, UV_Ddy, AdaptiveVirtualTexture, bGenerateFeedback, bMeshPaint);
+				// Allocate a layer in the virtual texture stack for this physical sample
+				VTLayerIndex = MaterialCompilationOutput.UniformExpressionSet.VTStacks[VTStackIndex].AddLayer();
+				VTPageTableIndex = VTLayerIndex;
+			}
+		
+			MaterialCompilationOutput.UniformExpressionSet.VTStacks[VTStackIndex].SetLayer(VTLayerIndex, VirtualTextureIndex);
+		}
 
 		// VT MipValueMode logic (most of work for VT case is in page table lookup)
 		if (MipValueMode == TMVM_MipLevel)
@@ -7472,7 +7519,7 @@ int32 FHLSLMaterialTranslator::TextureSample(
 		const FString VTPageTableResult_Finite = GetParameterCode(VTStackEntry.CodeIndex);
 
 		// 'Texture name/sampler', 'PageTableResult', 'LayerIndex', 'PackedUniform'
-		FString SampleCodeFinite = FString::Printf(TEXT("%s(%s, %s, %d, VTUniform_Unpack(Material.VTPackedUniform[%d]))"), *TextureTypeName, *TextureName, *VTPageTableResult_Finite, VTPageTableIndex, VirtualTextureIndex);
+		FString SampleCodeFinite = FString::Printf(TEXT("%s(%s, %s, %d, VTUniform_Unpack(%s))"), *TextureTypeName, *TextureName, *VTPageTableResult_Finite, VTPageTableIndex, *VTPackedUniformName);
 		SampleCodeFinite = ApplySamplerType(SampleCodeFinite, SamplerType);
 
 		if (IsAnalyticDerivEnabled() && IsDerivativeValid(UvDerivativeStatus))
@@ -7483,11 +7530,11 @@ int32 FHLSLMaterialTranslator::TextureSample(
 
 			if (SamplerDebugSupported(TextureType, bVirtualTexture, bDecal))
 			{
-				SampleCodeAnalytic = FString::Printf(TEXT("Debug%s(%s, %s, %d, VTUniform_Unpack(Material.VTPackedUniform[%d]), %s)"), *TextureTypeName, *TextureName, *VTPageTableResult_Analytic, VTPageTableIndex, VirtualTextureIndex, *UV_Scale);
+				SampleCodeAnalytic = FString::Printf(TEXT("Debug%s(%s, %s, %d, VTUniform_Unpack(%s), %s)"), *TextureTypeName, *TextureName, *VTPageTableResult_Analytic, VTPageTableIndex, *VTPackedUniformName, *UV_Scale);
 			}
 			else
 			{
-				SampleCodeAnalytic = FString::Printf(TEXT("%s(%s, %s, %d, VTUniform_Unpack(Material.VTPackedUniform[%d]))"), *TextureTypeName, *TextureName, *VTPageTableResult_Analytic, VTPageTableIndex, VirtualTextureIndex);
+				SampleCodeAnalytic = FString::Printf(TEXT("%s(%s, %s, %d, VTUniform_Unpack(%s))"), *TextureTypeName, *TextureName, *VTPageTableResult_Analytic, VTPageTableIndex, *VTPackedUniformName);
 			}
 
 			SampleCodeAnalytic = ApplySamplerType(SampleCodeAnalytic, SamplerType);
@@ -8716,6 +8763,11 @@ int32 FHLSLMaterialTranslator::VertexColor()
 	}
 }
 
+int32 FHLSLMaterialTranslator::MeshPaintTextureDescriptor()
+{
+	return GetPrimitiveProperty(EMaterialValueType(MCT_TextureVirtual | MCT_TextureMeshPaint), TEXT("MeshPaintTextureDescriptor"), TEXT("MeshPaintTextureDescriptor"));
+}
+
 int32 FHLSLMaterialTranslator::PreSkinnedPosition()
 {
 	if (ShaderFrequency != SF_Vertex)
@@ -14961,6 +15013,10 @@ void FHLSLMaterialTranslator::PrepareEnvironmentDefines()
 		{
 			EnvironmentDefines->VirtualPageTypes[i] |= FEnvironmentDefines::TABLE_INDIRECTION;
 		}
+		if (VTStacks[i].bMeshPaint)
+		{
+			EnvironmentDefines->VirtualPageTypes[i] |= FEnvironmentDefines::TABLE_GLOBAL;
+		}
 	}
 
 	// Set all the shading models for this material here 
diff --git a/Engine/Source/Runtime/Engine/Private/Materials/HLSLMaterialTranslator.h b/Engine/Source/Runtime/Engine/Private/Materials/HLSLMaterialTranslator.h
--- a/Engine/Source/Runtime/Engine/Private/Materials/HLSLMaterialTranslator.h
+++ b/Engine/Source/Runtime/Engine/Private/Materials/HLSLMaterialTranslator.h
@@ -190,6 +190,7 @@ struct FMaterialVTStackEntry
 	int32 PreallocatedStackTextureIndex;
 	bool bAdaptive;
 	bool bGenerateFeedback;
+	bool bMeshPaint;
 	float AspectRatio;
 
 	int32 CodeIndex;
@@ -1024,7 +1025,7 @@ protected:
 		const FString& UV_Value,
 		const FString& UV_Ddx,
 		const FString& UV_Ddy,
-		bool bAdaptive, bool bGenerateFeedback);
+		bool bAdaptive, bool bGenerateFeedback, bool bMeshPaint);
 
 	virtual int32 TextureSample(
 		int32 TextureIndex,
@@ -1117,6 +1118,7 @@ protected:
 	virtual int32 StaticTerrainLayerWeight(FName ParameterName, int32 Default, bool bTextureArray = false) override;
 
 	virtual int32 VertexColor() override;
+	virtual int32 MeshPaintTextureDescriptor() override;
 
 	virtual int32 PreSkinnedPosition() override;
 	virtual int32 PreSkinnedNormal() override;
diff --git a/Engine/Source/Runtime/Engine/Private/Materials/MaterialExpressions.cpp b/Engine/Source/Runtime/Engine/Private/Materials/MaterialExpressions.cpp
--- a/Engine/Source/Runtime/Engine/Private/Materials/MaterialExpressions.cpp
+++ b/Engine/Source/Runtime/Engine/Private/Materials/MaterialExpressions.cpp
@@ -124,6 +124,7 @@
 #include "Materials/MaterialExpressionIf.h"
 #include "Materials/MaterialExpressionInverseLinearInterpolate.h"
 #include "Materials/MaterialExpressionLightmapUVs.h"
+#include "Materials/MaterialExpressionMeshPaintTextureObject.h"
 #include "Materials/MaterialExpressionPrecomputedAOMask.h"
 #include "Materials/MaterialExpressionLightmassReplace.h"
 #include "Materials/MaterialExpressionLightVector.h"
@@ -2854,7 +2855,7 @@ int32 UMaterialExpressionTextureSample::Compile(class FMaterialCompiler* Compile
 			return CoordinateIndex;
 		};
 
-		if (TextureType & MCT_TextureCollection)
+		if (TextureType & (MCT_TextureCollection | MCT_TextureMeshPaint))
 		{
 			// There's no UTexture object to get here
 
@@ -3926,6 +3927,47 @@ void UMaterialExpressionVirtualTextureFeatureSwitch::GetCaption(TArray<FString>&
 
 #endif // WITH_EDITOR
 
+UMaterialExpressionMeshPaintTextureObject::UMaterialExpressionMeshPaintTextureObject(const FObjectInitializer& ObjectInitializer)
+	: Super(ObjectInitializer)
+{
+	// Structure to hold one-time initialization
+	struct FConstructorStatics
+	{
+		FText NAME_MeshPaintTexture;
+		FConstructorStatics()
+			: NAME_MeshPaintTexture(LOCTEXT("MeshPaintTexture", "MeshPaintTexture"))
+		{
+		}
+	};
+	static FConstructorStatics ConstructorStatics;
+
+#if WITH_EDITORONLY_DATA
+	MenuCategories.Add(ConstructorStatics.NAME_MeshPaintTexture);
+#endif
+
+	Outputs.Reset();
+	Outputs.Add(FExpressionOutput(TEXT("")));
+}
+
+#if WITH_EDITOR
+
+uint32 UMaterialExpressionMeshPaintTextureObject::GetOutputType(int32 OutputIndex)
+{
+	return MCT_TextureVirtual | MCT_TextureMeshPaint;
+}
+
+int32 UMaterialExpressionMeshPaintTextureObject::Compile(class FMaterialCompiler* Compiler, int32 OutputIndex)
+{
+	return Compiler->MeshPaintTextureDescriptor();
+}
+
+#endif // WITH_EDITOR
+
+void UMaterialExpressionMeshPaintTextureObject::GetCaption(TArray<FString>& OutCaptions) const
+{
+	OutCaptions.Add(TEXT("MeshPaintTextureObject"));
+}
+
 UMaterialExpressionAdd::UMaterialExpressionAdd(const FObjectInitializer& ObjectInitializer)
 	: Super(ObjectInitializer)
 {
diff --git a/Engine/Source/Runtime/Engine/Public/MaterialCompiler.h b/Engine/Source/Runtime/Engine/Public/MaterialCompiler.h
--- a/Engine/Source/Runtime/Engine/Public/MaterialCompiler.h
+++ b/Engine/Source/Runtime/Engine/Public/MaterialCompiler.h
@@ -451,6 +451,7 @@ public:
 	virtual int32 StaticTerrainLayerWeight(FName ParameterName,int32 Default, bool bTextureArray = false) = 0;
 
 	virtual int32 VertexColor() = 0;
+	virtual int32 MeshPaintTextureDescriptor() = 0;
 
 	virtual int32 PreSkinnedPosition() = 0;
 	virtual int32 PreSkinnedNormal() = 0;
@@ -1000,6 +1001,7 @@ public:
 	virtual int32 StaticTerrainLayerWeight(FName ParameterName, int32 Default, bool bTextureArray = false) override { return Compiler->StaticTerrainLayerWeight(ParameterName, Default, bTextureArray); }
 
 	virtual int32 VertexColor() override { return Compiler->VertexColor(); }
+	virtual int32 MeshPaintTextureDescriptor() override { return Compiler->MeshPaintTextureDescriptor(); }
 
 	virtual int32 PreSkinnedPosition() override { return Compiler->PreSkinnedPosition(); }
 	virtual int32 PreSkinnedNormal() override { return Compiler->PreSkinnedNormal(); }
diff --git a/Engine/Source/Runtime/Engine/Public/MaterialValueType.h b/Engine/Source/Runtime/Engine/Public/MaterialValueType.h
--- a/Engine/Source/Runtime/Engine/Public/MaterialValueType.h
+++ b/Engine/Source/Runtime/Engine/Public/MaterialValueType.h
@@ -63,6 +63,7 @@ enum EMaterialValueType : int
 	MCT_UInt4                 = 1 << 28,
 
 	MCT_TextureCollection     = 1 << 29,
+	MCT_TextureMeshPaint      = 1 << 30,
 
 	/** MCT_SparseVolumeTexture is intentionally not (yet) included here because it differs a lot from the other texture types and may not be supported/appropriate for all MCT_Texture use cases. */
 	MCT_Texture = MCT_Texture2D | MCT_TextureCube | MCT_Texture2DArray | MCT_TextureCubeArray | MCT_VolumeTexture | MCT_TextureExternal | MCT_TextureVirtual,
diff --git a/Engine/Source/Runtime/Engine/Public/Materials/MaterialExpressionMeshPaintTextureObject.h b/Engine/Source/Runtime/Engine/Public/Materials/MaterialExpressionMeshPaintTextureObject.h
new file mode 100644
--- /dev/null
+++ b/Engine/Source/Runtime/Engine/Public/Materials/MaterialExpressionMeshPaintTextureObject.h
@@ -0,0 +1,25 @@
+// Copyright Epic Games, Inc. All Rights Reserved.
+
+#pragma once
+
+#include "Materials/MaterialExpression.h"
+
+#include "MaterialExpressionMeshPaintTextureObject.generated.h"
+
+UCLASS(collapsecategories, hidecategories=Object, MinimalAPI)
+class UMaterialExpressionMeshPaintTextureObject : public UMaterialExpression
+{
+	GENERATED_UCLASS_BODY()
+
+protected:
+	/** Initialize the output pins. */
+	ENGINE_API void InitOutputs();
+
+	//~ Begin UMaterialExpression Interface
+#if WITH_EDITOR
+	ENGINE_API virtual uint32 GetOutputType(int32 OutputIndex) override;
+	ENGINE_API virtual int32 Compile(class FMaterialCompiler* Compiler, int32 OutputIndex) override;
+#endif
+	void GetCaption(TArray<FString>& OutCaptions) const;
+	//~ End UMaterialExpression Interface
+};
