Fix build issue on mac non-unity.
#rnx

diff --git a/Engine/Source/Runtime/Engine/Classes/VT/MeshPaintVirtualTexture.h b/Engine/Source/Runtime/Engine/Classes/VT/MeshPaintVirtualTexture.h
--- a/Engine/Source/Runtime/Engine/Classes/VT/MeshPaintVirtualTexture.h
+++ b/Engine/Source/Runtime/Engine/Classes/VT/MeshPaintVirtualTexture.h
@@ -48,8 +48,8 @@ namespace MeshPaintVirtualTexture
 	/** Scene view parameters that describe the virtual texture space shared by all UMeshPaintVirtualTexture objects. */
 	struct FUniformParams
 	{
-		FTextureRHIRef PageTableTexture = nullptr;
-		FTextureRHIRef PhysicalTexture = nullptr;
+		FTextureRHIRef PageTableTexture;
+		FTextureRHIRef PhysicalTexture;
 		FUintVector4 PackedUniform = FUintVector4(0, 0, 0, 0);
 	};
 	
