Post Process:  Fix for performance regression, caused by material instance blending being broken by not comparing base material pointers when checking if instances are the same.  Also made it so UserSceneTexture outputs on a node automatically disables blending, fixed an issue where blending would fail if an unrelated non-blendable node was in the array, and cleaned up the blending code to be more readable.  Also added constructor initialization of some blendable override related parameters.

#rnx

diff --git a/Engine/Source/Editor/UnrealEd/Private/PreviewMaterial.cpp b/Engine/Source/Editor/UnrealEd/Private/PreviewMaterial.cpp
--- a/Engine/Source/Editor/UnrealEd/Private/PreviewMaterial.cpp
+++ b/Engine/Source/Editor/UnrealEd/Private/PreviewMaterial.cpp
@@ -477,6 +477,12 @@ UMaterialEditorInstanceConstant::UMaterialEditorInstanceConstant(const FObjectIn
 
 	// Default to override with nothing on MIC (don't inherit parent setting).
 	bNaniteOverride = true;
+
+	PostProcessOverrides.bIsOverrideable = false;
+	PostProcessOverrides.bOverrideBlendableLocation = false;
+	PostProcessOverrides.bOverrideBlendablePriority = false;
+	PostProcessOverrides.BlendableLocationOverride = BL_SceneColorAfterTonemapping;
+	PostProcessOverrides.BlendablePriorityOverride = 0;
 }
 
 void UMaterialEditorInstanceConstant::PostEditChangeProperty(FPropertyChangedEvent& PropertyChangedEvent)
diff --git a/Engine/Source/Runtime/Engine/Private/Materials/Material.cpp b/Engine/Source/Runtime/Engine/Private/Materials/Material.cpp
--- a/Engine/Source/Runtime/Engine/Private/Materials/Material.cpp
+++ b/Engine/Source/Runtime/Engine/Private/Materials/Material.cpp
@@ -1929,7 +1929,7 @@ void UMaterial::GetDependentFunctions(TArray<UMaterialFunctionInterface*>& Depen
 }
 #endif // WITH_EDITORONLY_DATA
 
-extern FPostProcessMaterialNode* IteratePostProcessMaterialNodes(const FFinalPostProcessSettings& Dest, const UMaterialInterface* Material, const UMaterial* Base, FBlendableEntry*& Iterator);
+extern FPostProcessMaterialNode* FindExistingBlendablePostProcessNode(const FFinalPostProcessSettings& Dest, const UMaterialInterface* Material, const UMaterial* Base);
 
 void UMaterialInterface::OverrideBlendableSettings(class FSceneView& View, float Weight) const
 {
@@ -1946,9 +1946,10 @@ void UMaterialInterface::OverrideBlendableSettings(class FSceneView& View, float
 		return;
 	}
 
-	FBlendableEntry* Iterator = 0;
+	// Materials that write to UserSceneTexture outputs are automatically non-blendable
+	bool bIsBlendable = Base->bIsBlendable && GetUserSceneTextureOutput(Base) == NAME_None;
 
-	FPostProcessMaterialNode* DestNode = IteratePostProcessMaterialNodes(Dest, this, Base, Iterator);
+	FPostProcessMaterialNode* DestNode = bIsBlendable ? FindExistingBlendablePostProcessNode(Dest, this, Base) : nullptr;
 
 	// is this the first one of this material?
 	if(!DestNode)
@@ -1963,7 +1964,7 @@ void UMaterialInterface::OverrideBlendableSettings(class FSceneView& View, float
 
 			InitialMID->CopyScalarAndVectorParameters(*SourceData, View.FeatureLevel);
 
-			FPostProcessMaterialNode InitialNode(InitialMID, GetBlendableLocation(Base), GetBlendablePriority(Base), Base->bIsBlendable);
+			FPostProcessMaterialNode InitialNode(InitialMID, GetBlendableLocation(Base), GetBlendablePriority(Base), bIsBlendable);
 
 			// no blending needed on this one
 			FPostProcessMaterialNode* InitialDestNode = Dest.BlendableManager.PushBlendableData(1.0f, InitialNode);
diff --git a/Engine/Source/Runtime/Engine/Private/Materials/MaterialInstance.cpp b/Engine/Source/Runtime/Engine/Private/Materials/MaterialInstance.cpp
--- a/Engine/Source/Runtime/Engine/Private/Materials/MaterialInstance.cpp
+++ b/Engine/Source/Runtime/Engine/Private/Materials/MaterialInstance.cpp
@@ -1201,25 +1201,6 @@ bool UMaterialInstance::GetUserSceneTextureOverride(FName& InOutName) const
 	}
 }
 
-FName UMaterialInstance::GetUserSceneTextureOutput(const UMaterial* Base) const
-{
-	FName Result = NAME_None;
-
-	// Replacing tonemapper can't override output.
-	if (Base->BlendableLocation != BL_ReplacingTonemapper)
-	{
-		// UserSceneTexture output overrides are stored under key "NAME_None".  We store them in the override lookup to save space
-		// in the base structure, by avoiding a separate field just for the output override.
-		if (!GetUserSceneTextureOverride(Result) && Base)
-		{
-			// If no override was found, get the result from the base material
-			Result = FName(Base->UserSceneTexture);
-		}
-	}
-	return Result;
-
-}
-
 EBlendableLocation UMaterialInstance::GetBlendableLocation(const UMaterial* Base) const
 {
 	check(Base);
@@ -4612,32 +4593,23 @@ void UMaterialInstance::GetResourceSizeEx(FResourceSizeEx& CumulativeResourceSiz
 	}
 }
 
-FPostProcessMaterialNode* IteratePostProcessMaterialNodes(const FFinalPostProcessSettings& Dest, const UMaterialInterface* Material, const UMaterial* Base, FBlendableEntry*& Iterator)
+FPostProcessMaterialNode* FindExistingBlendablePostProcessNode(const FFinalPostProcessSettings& Dest, const UMaterialInterface* Material, const UMaterial* Base)
 {
 	EBlendableLocation Location = Material->GetBlendableLocation(Base);
 	int32 Priority = Material->GetBlendablePriority(Base);
 
-	for (;;)
-	{
-		FPostProcessMaterialNode* DataPtr = Dest.BlendableManager.IterateBlendables<FPostProcessMaterialNode>(Iterator);
-
-		if (!DataPtr)
-		{
-			// end reached
-			return 0;
-		}
-
-		// Do not consider materials that are set as not blendable
-		if (!DataPtr->GetIsBlendable())
-		{
-			return 0;
-		}
+	FBlendableEntry* Iterator = nullptr;
 
-		if(DataPtr->GetLocation() == Location && DataPtr->GetPriority() == Priority && DataPtr->GetMaterialInterface()->GetMaterial() == Material)
+	for (FPostProcessMaterialNode* DataPtr = Dest.BlendableManager.IterateBlendables<FPostProcessMaterialNode>(Iterator); DataPtr; DataPtr = Dest.BlendableManager.IterateBlendables<FPostProcessMaterialNode>(Iterator))
+	{
+		// Only consider materials that are set as blendable
+		if (DataPtr->GetIsBlendable() && DataPtr->GetLocation() == Location && DataPtr->GetPriority() == Priority && DataPtr->GetMaterialInterface()->GetMaterial() == Base)
 		{
 			return DataPtr;
 		}
 	}
+
+	return nullptr;
 }
 
 void UMaterialInstance::AllMaterialsCacheResourceShadersForRendering(bool bUpdateProgressDialog, bool bCacheAllRemainingShaders)
diff --git a/Engine/Source/Runtime/Engine/Private/Materials/MaterialInterface.cpp b/Engine/Source/Runtime/Engine/Private/Materials/MaterialInterface.cpp
--- a/Engine/Source/Runtime/Engine/Private/Materials/MaterialInterface.cpp
+++ b/Engine/Source/Runtime/Engine/Private/Materials/MaterialInterface.cpp
@@ -1222,6 +1222,25 @@ bool UMaterialInterface::GetUserSceneTextureOverride(FName& InOutName) const
 	return false;
 }
 
+FName UMaterialInterface::GetUserSceneTextureOutput(const UMaterial* Base) const
+{
+	FName Result = NAME_None;
+
+	// Replacing tonemapper can't override output.
+	if (Base->BlendableLocation != BL_ReplacingTonemapper)
+	{
+		// UserSceneTexture output overrides are stored under key "NAME_None".  We store them in the override lookup to save space
+		// in the base structure, by avoiding a separate field just for the output override.
+		if (!GetUserSceneTextureOverride(Result) && Base)
+		{
+			// If no override was found, get the result from the base material
+			Result = FName(Base->UserSceneTexture);
+		}
+	}
+	return Result;
+
+}
+
 // Certain implementations support overrides to the base material setting, default behavior is to just return the value from the base material
 EBlendableLocation UMaterialInterface::GetBlendableLocation(const UMaterial* Base) const
 {
diff --git a/Engine/Source/Runtime/Engine/Public/Materials/MaterialInstance.h b/Engine/Source/Runtime/Engine/Public/Materials/MaterialInstance.h
--- a/Engine/Source/Runtime/Engine/Public/Materials/MaterialInstance.h
+++ b/Engine/Source/Runtime/Engine/Public/Materials/MaterialInstance.h
@@ -491,10 +491,10 @@ struct FUserSceneTextureOverride
 
 struct FPostProcessBlendableOverrides
 {
-	bool bOverrideBlendableLocation;
-	bool bOverrideBlendablePriority;
-	TEnumAsByte<EBlendableLocation> BlendableLocationOverride;
-	int32 BlendablePriorityOverride;
+	bool bOverrideBlendableLocation = false;
+	bool bOverrideBlendablePriority = false;
+	TEnumAsByte<EBlendableLocation> BlendableLocationOverride = BL_SceneColorAfterTonemapping;
+	int32 BlendablePriorityOverride = 0;
 };
 
 template<class T>
@@ -871,7 +871,6 @@ public:
 	virtual ENGINE_API void RecacheUniformExpressions(bool bRecreateUniformBuffer) const override;
 	virtual ENGINE_API bool GetRefractionSettings(float& OutBiasValue) const override;
 	virtual ENGINE_API bool GetUserSceneTextureOverride(FName& InOutName) const override;
-	ENGINE_API FName GetUserSceneTextureOutput(const UMaterial* Base) const;
 	virtual ENGINE_API EBlendableLocation GetBlendableLocation(const UMaterial* Base) const override;
 	virtual ENGINE_API int32 GetBlendablePriority(const UMaterial* Base) const override;
 
diff --git a/Engine/Source/Runtime/Engine/Public/Materials/MaterialInterface.h b/Engine/Source/Runtime/Engine/Public/Materials/MaterialInterface.h
--- a/Engine/Source/Runtime/Engine/Public/Materials/MaterialInterface.h
+++ b/Engine/Source/Runtime/Engine/Public/Materials/MaterialInterface.h
@@ -913,6 +913,7 @@ public:
 	ENGINE_API bool GetFontParameterValue(const FHashedMaterialParameterInfo& ParameterInfo,class UFont*& OutFontValue, int32& OutFontPage, bool bOveriddenOnly = false) const;
 	ENGINE_API virtual bool GetRefractionSettings(float& OutBiasValue) const;
 	ENGINE_API virtual bool GetUserSceneTextureOverride(FName& InOutValue) const;
+	ENGINE_API FName GetUserSceneTextureOutput(const UMaterial* Base) const;
 	ENGINE_API virtual EBlendableLocation GetBlendableLocation(const UMaterial* Base) const;
 	ENGINE_API virtual int32 GetBlendablePriority(const UMaterial* Base) const;
 
