Fix a subtle bug with trace filter under AutoRTFM, which uncovered a second subtle bug in Verse runtime errors.

The bug was that we could do a trace in the closed, then do a trace in the open, then abort, and that abort would wipe out the changes made in the open. To fix this I've changed how the trace filter itself is performed, added some tests to check the behaviour, and ship it. But this uncovered another subtle bug with the Verse runtime errors - basically if we tried to raise a runtime error during an on-commit or on-abort we'd incorrectly enter an `AutoRTFM::Close` when not in an active transactional nest. To detect this can I've introduced `AutoRTFM::IsCommittingOrAborting` to fix that bug, and added testing.

I've also upgraded our `check`'s in the new transactional locks to be `ensure`'s - we really want to know if these cases trigger.

Lastly I've pulled in Michael Nicolella's simplification for the runtime error stuff to let us avoid the `AutoRTFM::Close` call entirely. Nice!

diff --git a/Engine/Source/Programs/AutoRTFMTests/Private/API.cpp b/Engine/Source/Programs/AutoRTFMTests/Private/API.cpp
--- a/Engine/Source/Programs/AutoRTFMTests/Private/API.cpp
+++ b/Engine/Source/Programs/AutoRTFMTests/Private/API.cpp
@@ -427,6 +427,8 @@ TEST_CASE("API.IsTransactional")
 
     bool InTransaction = false;
     bool InOpenNest = false;
+	bool InAbort = true;
+	bool InCommit = true;
 
     AutoRTFM::Commit([&]
     {
@@ -436,10 +438,27 @@ TEST_CASE("API.IsTransactional")
         {
             InOpenNest = AutoRTFM::IsTransactional();
         });
+
+		AutoRTFM::Transact([&]
+			{
+				AutoRTFM::OnAbort([&]
+					{
+						InAbort = AutoRTFM::IsTransactional();
+					});
+
+				AutoRTFM::AbortTransaction();
+			});
+
+		AutoRTFM::OnCommit([&]
+			{
+				InCommit = AutoRTFM::IsTransactional();
+			});
     });
 
     REQUIRE(true == InTransaction);
     REQUIRE(true == InOpenNest);
+	REQUIRE(false == InAbort);
+	REQUIRE(false == InCommit);
 }
 
 TEST_CASE("API.IsClosed")
@@ -450,11 +469,28 @@ TEST_CASE("API.IsClosed")
     bool InTransaction = false;
     bool InOpenNest = true;
     bool InClosedNestInOpenNest = false;
+	bool InAbort = true;
+	bool InCommit = true;
 
     AutoRTFM::Commit([&]
     {
         InTransaction = AutoRTFM::IsClosed();
 
+		AutoRTFM::Transact([&]
+			{
+				AutoRTFM::OnAbort([&]
+					{
+						InAbort = AutoRTFM::IsClosed();
+					});
+
+				AutoRTFM::AbortTransaction();
+			});
+
+		AutoRTFM::OnCommit([&]
+			{
+				InCommit = AutoRTFM::IsClosed();
+			});
+
         AutoRTFM::Open([&]
         {
             InOpenNest = AutoRTFM::IsClosed();
@@ -469,6 +505,56 @@ TEST_CASE("API.IsClosed")
     REQUIRE(true == InTransaction);
     REQUIRE(false == InOpenNest);
     REQUIRE(true == InClosedNestInOpenNest);
+	REQUIRE(false == InAbort);
+	REQUIRE(false == InCommit);
+}
+
+TEST_CASE("API.IsCommittingOrAborting")
+{
+	REQUIRE(false == AutoRTFM::IsCommittingOrAborting());
+
+	// Set to the opposite of what we expect at the end of function.
+	bool InTransaction = true;
+	bool InOpenNest = true;
+	bool InClosedNestInOpenNest = true;
+	bool InAbort = false;
+	bool InCommit = false;
+
+	AutoRTFM::Commit([&]
+		{
+			InTransaction = AutoRTFM::IsCommittingOrAborting();
+
+			AutoRTFM::Transact([&]
+				{
+					AutoRTFM::OnAbort([&]
+						{
+							InAbort = AutoRTFM::IsCommittingOrAborting();
+						});
+
+					AutoRTFM::AbortTransaction();
+				});
+
+			AutoRTFM::OnCommit([&]
+				{
+					InCommit = AutoRTFM::IsCommittingOrAborting();
+				});
+
+			AutoRTFM::Open([&]
+				{
+					InOpenNest = AutoRTFM::IsCommittingOrAborting();
+
+					REQUIRE(AutoRTFM::EContextStatus::OnTrack == AutoRTFM::Close([&]
+						{
+							InClosedNestInOpenNest = AutoRTFM::IsCommittingOrAborting();
+						}));
+				});
+		});
+
+	REQUIRE(false == InTransaction);
+	REQUIRE(false == InOpenNest);
+	REQUIRE(false == InClosedNestInOpenNest);
+	REQUIRE(true == InAbort);
+	REQUIRE(true == InCommit);
 }
 
 TEST_CASE("API.Transact")
diff --git a/Engine/Source/Runtime/Core/Private/AutoRTFM/API.cpp b/Engine/Source/Runtime/Core/Private/AutoRTFM/API.cpp
--- a/Engine/Source/Runtime/Core/Private/AutoRTFM/API.cpp
+++ b/Engine/Source/Runtime/Core/Private/AutoRTFM/API.cpp
@@ -214,7 +214,7 @@ extern "C" UE_AUTORTFM_AUTORTFM("RTFM_autortfm_is_transactional") bool autortfm_
 {
 	if (ForTheRuntime::IsAutoRTFMRuntimeEnabled())
 	{
-		return FContext::Get()->IsTransactional();
+		return FContext::IsTransactional();
 	}
 
 	return false;
@@ -225,6 +225,16 @@ extern "C" UE_AUTORTFM_AUTORTFM("RTFM_autortfm_is_closed") bool autortfm_is_clos
     return false;
 }
 
+extern "C" bool autortfm_is_committing_or_aborting()
+{
+	if (ForTheRuntime::IsAutoRTFMRuntimeEnabled())
+	{
+		return FContext::IsCommittingOrAborting();
+	}
+
+	return false;
+}
+
 // First Part - the API exposed outside transactions.
 extern "C" UE_AUTORTFM_AUTORTFM("RTFM_autortfm_transact") autortfm_result autortfm_transact(void (*UninstrumentedWork)(void*), void (*InstrumentedWork)(void*), void* Arg)
 {
diff --git a/Engine/Source/Runtime/Core/Private/AutoRTFM/Context.cpp b/Engine/Source/Runtime/Core/Private/AutoRTFM/Context.cpp
--- a/Engine/Source/Runtime/Core/Private/AutoRTFM/Context.cpp
+++ b/Engine/Source/Runtime/Core/Private/AutoRTFM/Context.cpp
@@ -99,14 +99,26 @@ bool FContext::IsTransactional()
         return false;
     }
 
-    if ((Context->GetStatus() != EContextStatus::Idle) && (Context->GetStatus() != EContextStatus::Committing))
-    {
-        return true;
-    }
-    else
-    {
-        return false;
-    }
+    return Context->GetStatus() == EContextStatus::OnTrack;
+}
+
+bool FContext::IsCommittingOrAborting()
+{
+	FContext* Context = TryGet();
+
+	if (!Context)
+	{
+		return false;
+	}
+
+	switch (Context->GetStatus())
+	{
+	default:
+		return true;
+	case EContextStatus::Idle:
+	case EContextStatus::OnTrack:
+		return false;
+	}
 }
 
 bool FContext::StartTransaction()
diff --git a/Engine/Source/Runtime/Core/Private/AutoRTFM/Context.h b/Engine/Source/Runtime/Core/Private/AutoRTFM/Context.h
--- a/Engine/Source/Runtime/Core/Private/AutoRTFM/Context.h
+++ b/Engine/Source/Runtime/Core/Private/AutoRTFM/Context.h
@@ -18,6 +18,7 @@ public:
     static FContext* TryGet();
     static FContext* Get();
     static bool IsTransactional();
+	static bool IsCommittingOrAborting();
     
     // This is public API
     ETransactionResult Transact(void (*InstrumentedFunction)(void*), void* Arg);
diff --git a/Engine/Source/Runtime/Core/Public/AutoRTFM/AutoRTFM.h b/Engine/Source/Runtime/Core/Public/AutoRTFM/AutoRTFM.h
--- a/Engine/Source/Runtime/Core/Public/AutoRTFM/AutoRTFM.h
+++ b/Engine/Source/Runtime/Core/Public/AutoRTFM/AutoRTFM.h
@@ -135,6 +135,15 @@ UE_AUTORTFM_FORCEINLINE bool autortfm_is_closed(void)
 }
 #endif
 
+#if UE_AUTORTFM
+UE_AUTORTFM_API bool autortfm_is_committing_or_aborting(void);
+#else
+UE_AUTORTFM_FORCEINLINE bool autortfm_is_committing_or_aborting(void)
+{
+	return false;
+}
+#endif
+
 #if UE_AUTORTFM
 UE_AUTORTFM_API autortfm_result autortfm_transact(void (*UninstrumentedWork)(void*), void (*InstrumentedWork)(void*), void* Arg);
 #else
@@ -449,24 +458,21 @@ auto AutoRTFMLookupInstrumentedFunctorInvoker(const TFunctor& Functor) -> void(*
 }
 #endif
 
-// Tells if we are currently running in a transaction. This will return true in an
-// open nest (see Open). This function is handled specially in the compiler, it
-// will be constant folded as 'true' in closed code, or
-// preserved as a function call in open code
+// Tells if we are currently running in a transaction. This will return true in an open nest
+// (see `Open`). This function is handled specially in the compiler and will be constant folded
+// as true in closed code, or preserved as a function call in open code.
 UE_AUTORTFM_FORCEINLINE bool IsTransactional() { return autortfm_is_transactional(); }
 
-// Tells if we are currently running in the closed nest of a transaction. By
-// default, transactional code is in a closed nest; the only way to be in an open
-// nest is to request it via Open.
-//
-// The advantages of this function over IsTransactional are:
-//
-// - It's faster, the compiler will constant-fold this.
-// - Usually, if you are doing special things for transactions, it's to work around
-//   the transactional openation in a closed nest. So, it's often more correct
-//   to test IsClosed than IsTransactional.
+// Tells if we are currently running in the closed nest of a transaction. By default,
+// transactional code is in a closed nest; the only way to be in an open nest is to request it
+// via `Open`. This function is handled specially in the compiler and will be constant folded
+// as true in closed code, and false in open code.
 UE_AUTORTFM_FORCEINLINE bool IsClosed() { return autortfm_is_closed(); }
 
+// Tells us if we are currently committing or aborting a transaction. This will return true
+// in an on-abort or on-commit.
+UE_AUTORTFM_FORCEINLINE bool IsCommittingOrAborting() { return autortfm_is_committing_or_aborting(); }
+
 // Run the functor in a transaction. Memory writes and other side effects get instrumented
 // and will be reversed if the transaction aborts.
 // 
diff --git a/Engine/Source/Runtime/Core/Public/Misc/TransactionallySafeCriticalSection.h b/Engine/Source/Runtime/Core/Public/Misc/TransactionallySafeCriticalSection.h
--- a/Engine/Source/Runtime/Core/Public/Misc/TransactionallySafeCriticalSection.h
+++ b/Engine/Source/Runtime/Core/Public/Misc/TransactionallySafeCriticalSection.h
@@ -26,7 +26,7 @@ struct FTransactionallySafeCriticalSectionDefinition final
 {
 	void Lock()
 	{
-		if (AutoRTFM::IsTransactional())
+		if (AutoRTFM::IsTransactional() || AutoRTFM::IsCommittingOrAborting())
 		{
 			AutoRTFM::Open([&]
 				{
@@ -42,7 +42,7 @@ struct FTransactionallySafeCriticalSectionDefinition final
 
 			AutoRTFM::OnAbort([this]
 				{
-					check(0 != TransactionalLockCount);
+					ensure(0 != TransactionalLockCount);
 					TransactionalLockCount -= 1;
 
 					if (0 == TransactionalLockCount)
@@ -54,17 +54,17 @@ struct FTransactionallySafeCriticalSectionDefinition final
 		else
 		{
 			CriticalSection.Lock();
-			check(0 == TransactionalLockCount);
+			ensure(0 == TransactionalLockCount);
 		}
 	}
 
 	void Unlock()
 	{
-		if (AutoRTFM::IsTransactional())
+		if (AutoRTFM::IsTransactional() || AutoRTFM::IsCommittingOrAborting())
 		{
 			AutoRTFM::OnCommit([this]
 				{
-					check(0 != TransactionalLockCount);
+					ensure(0 != TransactionalLockCount);
 					TransactionalLockCount -= 1;
 
 					if (0 == TransactionalLockCount)
@@ -75,7 +75,7 @@ struct FTransactionallySafeCriticalSectionDefinition final
 		}
 		else
 		{
-			check(0 == TransactionalLockCount);
+			ensure(0 == TransactionalLockCount);
 			CriticalSection.Unlock();
 		}
 	}
diff --git a/Engine/Source/Runtime/Engine/Private/Tests/AutoRTFM/AutoRTFMTraceFilterTests.cpp b/Engine/Source/Runtime/Engine/Private/Tests/AutoRTFM/AutoRTFMTraceFilterTests.cpp
--- a/Engine/Source/Runtime/Engine/Private/Tests/AutoRTFM/AutoRTFMTraceFilterTests.cpp
+++ b/Engine/Source/Runtime/Engine/Private/Tests/AutoRTFM/AutoRTFMTraceFilterTests.cpp
@@ -73,6 +73,57 @@ bool FAutoRTFMTraceFilterTests::RunTest(const FString & Parameters)
 			});
 
 		TEST_CHECK_TRUE(FTraceFilter::IsObjectTraceable(Object));
+
+		FTraceFilter::SetObjectIsTraceable(Object, false);
+		TEST_CHECK_TRUE(!FTraceFilter::IsObjectTraceable(Object));
+
+		UAutoRTFMTestObject* Other = NewObject<UAutoRTFMTestObject>();
+		UAutoRTFMTestObject* Another = NewObject<UAutoRTFMTestObject>();
+
+		FTraceFilter::SetObjectIsTraceable(Other, false);
+		TEST_CHECK_TRUE(!FTraceFilter::IsObjectTraceable(Other));
+
+		Result = AutoRTFM::Transact([&]
+			{
+				AutoRTFM::OnAbort([&]
+					{
+						FTraceFilter::SetObjectIsTraceable(Other, true);
+					});
+
+				FTraceFilter::SetObjectIsTraceable(Object, true);
+
+				AutoRTFM::OnAbort([&]
+					{
+						FTraceFilter::SetObjectIsTraceable(Another, true);
+					});
+
+				AutoRTFM::AbortTransaction();
+			});
+
+		TEST_CHECK_TRUE(AutoRTFM::ETransactionResult::AbortedByRequest == Result);
+		TEST_CHECK_TRUE(!FTraceFilter::IsObjectTraceable(Object));
+		TEST_CHECK_TRUE(FTraceFilter::IsObjectTraceable(Other));
+		TEST_CHECK_TRUE(FTraceFilter::IsObjectTraceable(Another));
+
+		Result = AutoRTFM::Transact([&]
+			{
+				AutoRTFM::OnCommit([&]
+					{
+						FTraceFilter::SetObjectIsTraceable(Other, false);
+					});
+
+				FTraceFilter::SetObjectIsTraceable(Object, true);
+
+				AutoRTFM::OnCommit([&]
+					{
+						FTraceFilter::SetObjectIsTraceable(Another, false);
+					});
+			});
+
+		TEST_CHECK_TRUE(AutoRTFM::ETransactionResult::Committed == Result);
+		TEST_CHECK_TRUE(FTraceFilter::IsObjectTraceable(Object));
+		TEST_CHECK_TRUE(!FTraceFilter::IsObjectTraceable(Other));
+		TEST_CHECK_TRUE(!FTraceFilter::IsObjectTraceable(Another));
 	}
 
 	{
diff --git a/Engine/Source/Runtime/Engine/Private/TraceFilter.cpp b/Engine/Source/Runtime/Engine/Private/TraceFilter.cpp
--- a/Engine/Source/Runtime/Engine/Private/TraceFilter.cpp
+++ b/Engine/Source/Runtime/Engine/Private/TraceFilter.cpp
@@ -100,10 +100,7 @@ public:
 		if (bHadElements && AnnotationMap.Num() == 0)
 		{
 			// we are removing the last one, so if we are auto removing or verifying removal, unregister now
-			AutoRTFM::OnCommit([this]
-				{
-					GUObjectArray.RemoveUObjectDeleteListener(this);
-				});
+			GUObjectArray.RemoveUObjectDeleteListener(this);
 		}
 	}
 
@@ -280,13 +277,20 @@ FAutoConsoleCommand FlushFilterStateCommand(TEXT("TraceFilter.FlushState"), TEXT
 template<bool bForceThreadSafe>
 bool FTraceFilter::IsObjectTraceable(const UObject* InObject)
 {
-	if constexpr (!bForceThreadSafe)
-	{
-		check(GObjectFilterAnnotations.IsLocked());
-	}
+	bool bResult = false;
+
+	AutoRTFM::Open([&]
+		{
+			if constexpr (!bForceThreadSafe)
+			{
+				check(GObjectFilterAnnotations.IsLocked());
+			}
 
-	// Object not found in the AnnotationMap means that it is at the default value, which is bIsTraceable == true
-	return GObjectFilterAnnotations.GetAnnotationMap().Find(InObject) == nullptr;
+			// Object not found in the AnnotationMap means that it is at the default value, which is bIsTraceable == true
+			bResult = GObjectFilterAnnotations.GetAnnotationMap().Find(InObject) == nullptr;
+		});
+
+	return bResult;
 }
 
 template bool FTraceFilter::IsObjectTraceable</*bForceThreadSafe = */ false>(const UObject* InObject);
@@ -295,33 +299,41 @@ template bool FTraceFilter::IsObjectTraceable</*bForceThreadSafe = */ true>(cons
 template<bool bForceThreadSafe>
 void FTraceFilter::SetObjectIsTraceable(const UObject* InObject, bool bIsTraceable)
 {
-	ensure(InObject);
-
-	if constexpr (bForceThreadSafe)
+	AutoRTFM::Open([&]
 	{
-		FTraceFilterObjectAnnotation Annotation;
-		Annotation.bIsTraceable = bIsTraceable;
-		GObjectFilterAnnotations.AddAnnotation(InObject, Annotation);
+		ensure(InObject);
 
-		if (bIsTraceable)
-		{
-			TRACE_OBJECT(InObject);
-		}
-	}
-	else
-	{
-		check(GObjectFilterAnnotations.IsLocked());
-		TMap<const UObjectBase*, FTraceFilterObjectAnnotation>& AnnotationMap = GObjectFilterAnnotations.GetAnnotationMap();
-		if (!bIsTraceable)
+		if constexpr (bForceThreadSafe)
 		{
-			AnnotationMap.FindOrAdd(InObject).bIsTraceable = false;
+			FTraceFilterObjectAnnotation Annotation;
+			Annotation.bIsTraceable = bIsTraceable;
+			GObjectFilterAnnotations.AddAnnotation(InObject, Annotation);
+
+			if (bIsTraceable)
+			{
+				TRACE_OBJECT(InObject);
+			}
 		}
 		else
 		{
-			AnnotationMap.Remove(InObject);
-			TRACE_OBJECT(InObject);
+			check(GObjectFilterAnnotations.IsLocked());
+			TMap<const UObjectBase*, FTraceFilterObjectAnnotation>& AnnotationMap = GObjectFilterAnnotations.GetAnnotationMap();
+			if (!bIsTraceable)
+			{
+				AnnotationMap.FindOrAdd(InObject).bIsTraceable = false;
+			}
+			else
+			{
+				AnnotationMap.Remove(InObject);
+				TRACE_OBJECT(InObject);
+			}
 		}
-	}
+	});
+
+AutoRTFM::OnAbort([InObject, bIsTraceable]
+	{
+		SetObjectIsTraceable(InObject, !bIsTraceable);
+	});
 }
 
 template void FTraceFilter::SetObjectIsTraceable</*bForceThreadSafe = */ true>(const UObject* InObject, bool bIsTraceable);
@@ -330,19 +342,7 @@ template void FTraceFilter::SetObjectIsTraceable</*bForceThreadSafe = */ false>(
 template<bool bForceThreadSafe>
 void FTraceFilter::MarkObjectTraceable(const UObject* InObject)
 {
-	ensure(InObject);
-
-	if constexpr (bForceThreadSafe)
-	{
-		FTraceFilterObjectAnnotation Annotation;
-		Annotation.bIsTraceable = true;
-		GObjectFilterAnnotations.AddAnnotation(InObject, Annotation);
-	}
-	else
-	{
-		check(GObjectFilterAnnotations.IsLocked());
-		SetObjectIsTraceable(InObject, true);
-	}
+	SetObjectIsTraceable<bForceThreadSafe>(InObject, true);
 }
 
 template void FTraceFilter::MarkObjectTraceable</*bForceThreadSafe = */ true>(const UObject* InObject);
@@ -372,3 +372,4 @@ void FTraceFilter::Unlock()
 }
 
 #endif // TRACE_FILTERING_ENABLED
+
diff --git a/Engine/Source/Runtime/Experimental/Chaos/Public/Framework/Threading.h b/Engine/Source/Runtime/Experimental/Chaos/Public/Framework/Threading.h
--- a/Engine/Source/Runtime/Experimental/Chaos/Public/Framework/Threading.h
+++ b/Engine/Source/Runtime/Experimental/Chaos/Public/Framework/Threading.h
@@ -975,7 +975,7 @@ FORCEINLINE void EnsureIsInGameThreadContext()
 	{
 		void ReadLock()
 		{
-			if (AutoRTFM::IsTransactional())
+			if (AutoRTFM::IsTransactional() || AutoRTFM::IsCommittingOrAborting())
 			{
 				// Transactionally pessimise ReadLock -> WriteLock.
 				WriteLock();
@@ -983,27 +983,27 @@ FORCEINLINE void EnsureIsInGameThreadContext()
 			else
 			{
 				Lock.ReadLock();
-				check(0 == TransactionalLockCount);
+				ensure(0 == TransactionalLockCount);
 			}
 		}
 
 		void ReadUnlock()
 		{
-			if (AutoRTFM::IsTransactional())
+			if (AutoRTFM::IsTransactional() || AutoRTFM::IsCommittingOrAborting())
 			{
 				// Transactionally pessimise ReadUnlock -> WriteUnlock.
 				WriteUnlock();
 			}
 			else
 			{
-				check(0 == TransactionalLockCount);
+				ensure(0 == TransactionalLockCount);
 				Lock.ReadUnlock();
 			}
 		}
 
 		void WriteLock()
 		{
-			if (AutoRTFM::IsTransactional())
+			if (AutoRTFM::IsTransactional() || AutoRTFM::IsCommittingOrAborting())
 			{
 				AutoRTFM::Open([&]
 					{
@@ -1019,7 +1019,7 @@ FORCEINLINE void EnsureIsInGameThreadContext()
 
 				AutoRTFM::OnAbort([this]
 					{
-						check(0 != TransactionalLockCount);
+						ensure(0 != TransactionalLockCount);
 						TransactionalLockCount -= 1;
 
 						if (0 == TransactionalLockCount)
@@ -1031,17 +1031,17 @@ FORCEINLINE void EnsureIsInGameThreadContext()
 			else
 			{
 				Lock.WriteLock();
-				check(0 == TransactionalLockCount);
+				ensure(0 == TransactionalLockCount);
 			}
 		}
 
 		void WriteUnlock()
 		{
-			if (AutoRTFM::IsTransactional())
+			if (AutoRTFM::IsTransactional() || AutoRTFM::IsCommittingOrAborting())
 			{
 				AutoRTFM::OnCommit([this]
 					{
-						check(0 != TransactionalLockCount);
+						ensure(0 != TransactionalLockCount);
 						TransactionalLockCount -= 1;
 
 						if (0 == TransactionalLockCount)
@@ -1052,7 +1052,7 @@ FORCEINLINE void EnsureIsInGameThreadContext()
 			}
 			else
 			{
-				check(0 == TransactionalLockCount);
+				ensure(0 == TransactionalLockCount);
 				Lock.WriteUnlock();
 			}
 		}
