[Material Editor] New Material Translator push.

- Added constant folding to binary operators and branches.
- Added value construction API for constructing an arithmetic value from another (e.g. float3 from a float). Not fully complete yet.
- Added Subscript value for vector and matrices subscription by component index.

diff --git a/Engine/Source/Runtime/Engine/Private/Materials/MaterialExpressionsIR.cpp b/Engine/Source/Runtime/Engine/Private/Materials/MaterialExpressionsIR.cpp
--- a/Engine/Source/Runtime/Engine/Private/Materials/MaterialExpressionsIR.cpp
+++ b/Engine/Source/Runtime/Engine/Private/Materials/MaterialExpressionsIR.cpp
@@ -308,8 +308,8 @@ static void BuildBinaryArithmeticOperator(
 		: IR::FArithmeticType::GetScalar(IR::SK_Float);
 
 	// Convert operand values to determined result type. 
-	LhsValue = Emitter.TryEmitConvert(LhsValue, ResultType);
-	RhsValue = Emitter.TryEmitConvert(RhsValue, ResultType);
+	LhsValue = Emitter.TryEmitConstruct(ResultType, LhsValue);
+	RhsValue = Emitter.TryEmitConstruct(ResultType, RhsValue);
 
 	if (Emitter.IsInvalid())
 	{
@@ -373,8 +373,8 @@ void UMaterialExpressionIf::Build(IR::FEmitter& Emitter)
 	}
 
 	// Convert both conditional argument values to the common type. This is expected to work as the common type could be found.
-	AValue = Emitter.TryEmitConvert(AValue, ConditionArgsType);
-	BValue = Emitter.TryEmitConvert(BValue, ConditionArgsType);
+	AValue = Emitter.TryEmitConstruct(ConditionArgsType, AValue);
+	BValue = Emitter.TryEmitConstruct(ConditionArgsType, BValue);
 
 	// Now determine the output type by taking the common arithmetic type between result values.
 	IR::FArithmeticTypePtr OutputType = Emitter.TryGetCommonArithmeticType(AGreaterThanBValue->Type->AsArithmetic(), AEqualsBValue->Type->AsArithmetic());
@@ -386,9 +386,9 @@ void UMaterialExpressionIf::Build(IR::FEmitter& Emitter)
 	}
 
 	// Convert result values to the common result type.
-	AGreaterThanBValue = Emitter.TryEmitConvert(AGreaterThanBValue, OutputType);
-	AEqualsBValue = Emitter.TryEmitConvert(AEqualsBValue, OutputType);
-	ALessThanBValue = Emitter.TryEmitConvert(ALessThanBValue, OutputType);
+	AGreaterThanBValue = Emitter.TryEmitConstruct(OutputType, AGreaterThanBValue);
+	AEqualsBValue = Emitter.TryEmitConstruct(OutputType, AEqualsBValue);
+	ALessThanBValue = Emitter.TryEmitConstruct(OutputType, ALessThanBValue);
 
 	if (Emitter.IsInvalid())
 	{
diff --git a/Engine/Source/Runtime/Engine/Private/Materials/MaterialIR.cpp b/Engine/Source/Runtime/Engine/Private/Materials/MaterialIR.cpp
--- a/Engine/Source/Runtime/Engine/Private/Materials/MaterialIR.cpp
+++ b/Engine/Source/Runtime/Engine/Private/Materials/MaterialIR.cpp
@@ -10,7 +10,7 @@ namespace UE::MIR
 
 FInstruction* FValue::AsInstruction()
 {
-	return (Kind >= VK_InstructionBegin && Kind < VK_InstructionEnd) ? static_cast<FInstruction*>(this) : nullptr;
+	return this && (Kind >= VK_InstructionBegin && Kind < VK_InstructionEnd) ? static_cast<FInstruction*>(this) : nullptr;
 }
 
 const FInstruction* FValue::AsInstruction() const
@@ -38,7 +38,7 @@ uint32 FValue::GetSizeInBytes() const
 {
 	switch (Kind)
 	{
-		case VK_ScalarConstant: return sizeof(FScalarConstant);
+		case VK_Constant: return sizeof(FConstant);
 		case VK_SetMaterialOutput: return sizeof(FSetMaterialOutput);
 		case VK_BinaryOperator: return sizeof(FBinaryOperator);
 		case VK_Dimensional: return sizeof(FDimensional) + sizeof(FValue*) * static_cast<const FDimensional*>(this)->GetComponents().Num();
@@ -51,7 +51,7 @@ TArrayView<FValue*> FValue::GetUses()
 {
 	switch (Kind)
 	{
-		case VK_ScalarConstant:
+		case VK_Constant:
 		case VK_Dimensional:
 			return {};
 
@@ -79,11 +79,11 @@ TArrayView<FValue*> FValue::GetUses()
 
 TArrayView<FValue* const> FDimensional::GetComponents() const
 {
-	TArrayView<FValue*> Components = const_cast<FDimensional*>(this)->GetMutableComponents();
+	TArrayView<FValue*> Components = const_cast<FDimensional*>(this)->GetComponents();
 	return { Components.GetData(), Components.Num() };
 }
 
-TArrayView<FValue*> FDimensional::GetMutableComponents()
+TArrayView<FValue*> FDimensional::GetComponents()
 {
 	FArithmeticTypePtr ArithmeticType = Type->AsArithmetic();
 	check(ArithmeticType);
@@ -92,6 +92,18 @@ TArrayView<FValue*> FDimensional::GetMutableComponents()
 	return { Ptr, ArithmeticType->NumRows };
 }
 
+bool FDimensional::AreComponentsConstant() const
+{
+	for (FValue const* Component : GetComponents())
+	{
+		if (!Component->As<FConstant>())
+		{
+			return false;
+		}
+	}
+	return true;
+}
+
 bool FInstruction::GetInnerBlock(int32 Index, FValue*& OutArg, FBlock*& OutBlock)
 {
 	switch (Kind)
@@ -130,6 +142,16 @@ bool FInstruction::GetInnerBlock(int32 Index, FValue*& OutArg, FBlock*& OutBlock
 	return false;
 }
 
+bool IsArithmeticOperator(EBinaryOperator Op)
+{
+	return Op >= BO_Add && Op <= BO_Divide;
+}
+
+bool IsComparisonOperator(EBinaryOperator Op)
+{
+	return Op >= BO_Greater && Op <= BO_NotEquals;
+}
+
 } // namespace UE::MIR
 
 #endif // #if WITH_EDITOR
diff --git a/Engine/Source/Runtime/Engine/Private/Materials/MaterialIREmitter.cpp b/Engine/Source/Runtime/Engine/Private/Materials/MaterialIREmitter.cpp
--- a/Engine/Source/Runtime/Engine/Private/Materials/MaterialIREmitter.cpp
+++ b/Engine/Source/Runtime/Engine/Private/Materials/MaterialIREmitter.cpp
@@ -12,58 +12,51 @@
 namespace UE::MIR
 {
 
-template <typename T>
-static T InitValue(FTypePtr InType)
+const TCHAR* VectorComponentToString(EVectorComponent Component)
 {
-	T Value{};
-	Value.Kind = T::TypeKind;
-	Value.Type = InType;
-	return Value;
+	static const TCHAR* Strings[] = { TEXT("x"), TEXT("y"), TEXT("z"), TEXT("w") };
+	return Strings[(int)Component];
 }
 
-struct FEmitter::FPrivate : FEmitter {
-	
-	// Creates a new `FDimensional` value of specified `Type` and returns it.
-	static FDimensional* NewDimensionalValue(FEmitter* Emitter, FArithmeticTypePtr Type)
-	{
-		check(!Type->IsScalar());
-		int Dimensions = Type->NumRows * Type->NumColumns;
-		FDimensional* Value = (FDimensional*)FMemory::Malloc(sizeof(FDimensional) + sizeof(FValue*) * Dimensions);
-		Value->Kind = VK_Dimensional;
-		Value->Type = Type;
-		return Value;
-	}
+FSwizzleMask::FSwizzleMask(EVectorComponent X)
+: NumComponents{ 1 }
+{
+	Components[0] = X;
+}
 
-	// Searches for an existing value in module that matches specified `Prototype`.
-	// If none found, it creates a new value as a copy of the prototype, adds it to
-	// the module then returns it.
-	template <typename TValueType>
-	static FValue* EmitPrototype(FEmitter* Emitter, const TValueType& Prototype)
-	{
-		if (FValue* Existing = FindValue(Emitter, &Prototype))
-		{
-			return Existing;
-		}
+FSwizzleMask::FSwizzleMask(EVectorComponent X, EVectorComponent Y)
+: NumComponents{ 4 }
+{
+	Components[0] = X;
+	Components[1] = Y;
+}
 
-		TValueType* Value = new TValueType{ Prototype };
-		PushNewValue(Emitter, Value);
-		return Value;
-	}
+FSwizzleMask::FSwizzleMask(EVectorComponent X, EVectorComponent Y, EVectorComponent Z)
+: NumComponents{ 4 }
+{
+	Components[0] = X;
+	Components[1] = Y;
+	Components[2] = Z;
+}
 
-	// Emits specified newly created `Value`. If the exact value already exists,
-	// specified one is *destroyed* and existing one is returned instead.
-	static FValue* EmitNew(FEmitter* Emitter, FValue* Value)
-	{
-		if (FValue* Existing = FindValue(Emitter, Value))
-		{
-			delete Value;
-			return Existing;
-		}
+FSwizzleMask::FSwizzleMask(EVectorComponent X, EVectorComponent Y, EVectorComponent Z, EVectorComponent W)
+: NumComponents{ 4 }
+{
 
-		PushNewValue(Emitter, Value);
-		return Value;
-	}
+	Components[0] = X;
+	Components[1] = Y;
+	Components[2] = Z;
+	Components[3] = W;
+}
 
+void FSwizzleMask::Append(EVectorComponent Component)
+{
+	check(NumComponents < 4);
+	Components[NumComponents++] = Component;
+}
+
+struct FEmitter::FPrivate : FEmitter
+{
 	// Looks for an existing value in the module that matches `Prototype` and returns it if found.
 	static FValue* FindValue(FEmitter* Emitter, const FValue* Prototype)
 	{
@@ -71,6 +64,7 @@ struct FEmitter::FPrivate : FEmitter {
 
 		for (FValue* CurrValue : Emitter->Module->Values)
 		{
+	
 			if (CurrValue->Equals(Prototype))
 			{
 				return CurrValue;
@@ -86,11 +80,81 @@ struct FEmitter::FPrivate : FEmitter {
 	}
 };
 
+template <typename T>
+static T InitValue(FTypePtr InType)
+{
+	T Value{};
+	Value.Kind = T::TypeKind;
+	Value.Type = InType;
+	return Value;
+}
+
+template <typename TValueType>
+static TValueType* NewValue(FEmitter* Emitter, FTypePtr Type)
+{
+	TValueType* Value = (TValueType*)new TValueType{};
+	Value->Kind = TValueType::TypeKind;
+	Value->Type = Type;
+	FEmitter::FPrivate::PushNewValue(Value);
+	return Value;
+}
+
+// Creates a new `FDimensional` value of specified `Type` and returns it.
+static FDimensional* NewDimensionalValue(FEmitter* Emitter, FArithmeticTypePtr Type)
+{
+	check(!Type->IsScalar());
+	int Dimensions = Type->NumRows * Type->NumColumns;
+	void* Alloc = FMemory::Malloc(sizeof(FDimensional) + sizeof(FValue*) * Dimensions);
+	FDimensional* Value = new (Alloc) FDimensional{};
+	Value->Kind = VK_Dimensional;
+	Value->Type = Type;
+	return Value;
+}
+
+// Searches for an existing value in module that matches specified `Prototype`.
+// If none found, it creates a new value as a copy of the prototype, adds it to
+// the module then returns it.
+template <typename TValueType>
+static FValue* EmitPrototype(FEmitter* Emitter, const TValueType& Prototype)
+{
+	if (FValue* Existing = FEmitter::FPrivate::FindValue(Emitter, &Prototype))
+	{
+		return Existing;
+	}
+
+	TValueType* Value = new TValueType{ Prototype };
+	FEmitter::FPrivate::PushNewValue(Emitter, Value);
+	return Value;
+}
+
+// Emits specified newly created `Value`. If the exact value already exists,
+// specified one is *destroyed* and existing one is returned instead.
+static FValue* EmitNew(FEmitter* Emitter, FValue* Value)
+{
+	if (FValue* Existing = FEmitter::FPrivate::FindValue(Emitter, Value))
+	{
+		delete Value;
+		return Existing;
+	}
+
+	FEmitter::FPrivate::PushNewValue(Emitter, Value);
+	return Value;
+}
+
 FEmitter::FEmitter(FMaterialIRModuleBuilder* InBuilder, UMaterial* InMaterial, FMaterialIRModule* InModule)
 {
 	Builder = InBuilder;
 	Material = InMaterial;
 	Module = InModule;
+
+	// Create and reference the true/false constants.
+	FConstant Temp = InitValue<FConstant>(FArithmeticType::GetBool1());
+
+	Temp.Boolean = true;
+	ConstantTrue = EmitPrototype(this, Temp);
+
+	Temp.Boolean = false;
+	ConstantFalse = EmitPrototype(this, Temp);
 }
 
 FValue* FEmitter::Get(const FExpressionInput* Input)
@@ -110,7 +174,7 @@ FEmitter& FEmitter::DefaultToFloatZero(const FExpressionInput* Input)
 	return DefaultTo(Input, 0.0f);
 }
 
-FEmitter& FEmitter::DefaultTo(const FExpressionInput* Input, float Float)
+FEmitter& FEmitter::DefaultTo(const FExpressionInput* Input, TFloat Float)
 {
 	if (!Input->IsConnected())
 	{
@@ -203,18 +267,33 @@ FValue* FEmitter::EmitConstantScalarZero(EScalarKind Kind)
 {
 	switch (Kind)
 	{
-		case SK_Bool: UE_MIR_TODO();
+		case SK_Bool: return EmitConstantFalse();
 		case SK_Int: return EmitConstantInt1(0);
 		case SK_Float: return EmitConstantFloat1(0.0f);
 		default: UE_MIR_UNREACHABLE();
 	}
 }
 
-FValue* FEmitter::EmitConstantFloat1(float InX)
+FValue* FEmitter::EmitConstantTrue()
 {
-	FScalarConstant Scalar = InitValue<FScalarConstant>(FArithmeticType::GetScalar(SK_Float));
+	return ConstantTrue;
+}
+
+FValue* FEmitter::EmitConstantFalse()
+{
+	return ConstantFalse;
+}
+
+FValue* FEmitter::EmitConstantBool1(bool InX)
+{
+	return InX ? EmitConstantTrue() : EmitConstantFalse();
+}
+
+FValue* FEmitter::EmitConstantFloat1(TFloat InX)
+{
+	FConstant Scalar = InitValue<FConstant>(FArithmeticType::GetScalar(SK_Float));
 	Scalar.Float = InX;
-	return FPrivate::EmitPrototype(this, Scalar);
+	return EmitPrototype(this, Scalar);
 }
 
 FValue* FEmitter::EmitConstantFloat2(const FVector2f& InValue)
@@ -241,11 +320,11 @@ FValue* FEmitter::EmitConstantFloat4(const FVector4f& InValue)
 	return EmitVector4(X, Y, Z, W);
 }
 
-FValue* FEmitter::EmitConstantInt1(int InX)
+FValue* FEmitter::EmitConstantInt1(TInteger InX)
 {
-	FScalarConstant Scalar = InitValue<FScalarConstant>(FArithmeticType::GetScalar(SK_Int));
+	FConstant Scalar = InitValue<FConstant>(FArithmeticType::GetScalar(SK_Int));
 	Scalar.Integer = InX;
-	return FPrivate::EmitPrototype(this, Scalar);
+	return EmitPrototype(this, Scalar);
 }
 
 FValue* FEmitter::EmitConstantInt2(const FIntVector2& InValue)
@@ -278,11 +357,11 @@ FValue* FEmitter::EmitVector2(FValue* InX, FValue* InY)
 	check(InX->Type == InY->Type);
 
 	TDimensional<2> Vector = InitValue<TDimensional<2>>(FArithmeticType::GetVector(InX->Type->AsArithmetic()->ScalarKind, 2));
-	TArrayView<FValue*> Components = Vector.GetMutableComponents();
+	TArrayView<FValue*> Components = Vector.GetComponents();
 	Components[0] = InX;
 	Components[1] = InY;
 
-	return FPrivate::EmitPrototype(this, Vector);
+	return EmitPrototype(this, Vector);
 }
 
 FValue* FEmitter::EmitVector3(FValue* InX, FValue* InY, FValue* InZ)
@@ -292,12 +371,12 @@ FValue* FEmitter::EmitVector3(FValue* InX, FValue* InY, FValue* InZ)
 	check(InY->Type == InZ->Type);
 
 	TDimensional<3> Vector = InitValue<TDimensional<3>>(FArithmeticType::GetVector(InX->Type->AsArithmetic()->ScalarKind, 3));
-	TArrayView<FValue*> Components = Vector.GetMutableComponents();
+	TArrayView<FValue*> Components = Vector.GetComponents();
 	Components[0] = InX;
 	Components[1] = InY;
 	Components[2] = InZ;
 
-	return FPrivate::EmitPrototype(this, Vector);
+	return EmitPrototype(this, Vector);
 }
 
 FValue* FEmitter::EmitVector4(FValue* InX, FValue* InY, FValue* InZ, FValue* InW)
@@ -308,31 +387,100 @@ FValue* FEmitter::EmitVector4(FValue* InX, FValue* InY, FValue* InZ, FValue* InW
 	check(InZ->Type == InW->Type);
 
 	TDimensional<4> Vector = InitValue<TDimensional<4>>(FArithmeticType::GetVector(InX->Type->AsArithmetic()->ScalarKind, 4));
-	TArrayView<FValue*> Components = Vector.GetMutableComponents();
+	TArrayView<FValue*> Components = Vector.GetComponents();
 	Components[0] = InX;
 	Components[1] = InY;
 	Components[2] = InZ;
 	Components[3] = InW;
 
-	return FPrivate::EmitPrototype(this, Vector);
+	return EmitPrototype(this, Vector);
+}
+
+FValue* FEmitter::EmitSubscript(FValue* Value, int Index)
+{
+	FArithmeticTypePtr ArithmeticType = Value->Type->AsArithmetic();
+	check(ArithmeticType);
+
+	// Getting first component and Value is already a scalar, just return itself.
+	if (Index == 0 && Value->Type->AsScalar())
+	{
+		return Value;
+	}
+
+	// Getting first component and Value is already a scalar, just return itself.
+	if (FDimensional* DimensionalValue = Value->As<FDimensional>())
+	{
+		check(Index < DimensionalValue->GetComponents().Num());
+		return DimensionalValue->GetComponents()[Index];
+	}
+	
+	// We can't resolve it at compile time: emit subscript value.
+	FSubscript Prototype = InitValue<FSubscript>(ArithmeticType->ToScalar());
+	Prototype.Arg = Value;
+	Prototype.Index = Index;
+
+	return EmitPrototype(this, Prototype);
 }
 
-FValue* FEmitter::EmitArithmetic(FArithmeticTypePtr Type, FValue* Scalar)
+FValue* FEmitter::TryEmitSwizzle(FValue* Value, FSwizzleMask Mask)
 {
-	check(Scalar->Type->AsScalar());
-	if (Type->IsScalar())
+	// At least one component must have been specified.
+	check(Mask.NumComponents > 0);
+
+	// We can only swizzle on non-matrix arithmetic types.
+	FArithmeticTypePtr ArithmeticType = Value->Type->AsVector();
+	if (!ArithmeticType || ArithmeticType->IsMatrix())
+	{
+		Errorf(TEXT("Cannot swizzle a `%s` value."), Value->Type->GetSpelling().GetData());
+		return nullptr;
+	}
+
+	// Make sure each component in the mask fits the number of components in Value.
+	for (EVectorComponent Component : Mask)
+	{
+		if ((int)Component >= ArithmeticType->NumRows)
+		{
+			Errorf(TEXT("Value of type `%s` has no component `%s`."), ArithmeticType->Spelling.GetData(), VectorComponentToString(Component));
+			return nullptr;
+		}
+	}
+
+	// If the requested number of components is the same as Value and the order in which the components
+	// are specified in the mask is sequential (e.g. x, y, z) then this is a no op, simply return Value as is.
+	if (Mask.NumComponents == ArithmeticType->GetNumComponents())
+	{
+		bool InOrder = true;
+		for (int i = 0; i < Mask.NumComponents; ++i)
+		{
+			if (Mask.Components[i] != (EVectorComponent)i)
+			{
+				InOrder = false;
+				break;
+			}
+		}
+
+		if (InOrder)
+		{
+			return Value;
+		}
+	}
+	
+	// If only one component is requested, we can use EmitSubscript() to return the single component.
+	if (Mask.NumComponents == 1)
 	{
-		check(Type == Scalar->Type);
-		return Scalar;
+		return EmitSubscript(Value, (int)Mask.Components[0]);
 	}
+	
+	// Make the result vector type.
+	FArithmeticTypePtr ResultType = FArithmeticType::GetVector(ArithmeticType->ScalarKind, Mask.NumComponents);
+	FDimensional* Result = NewDimensionalValue(this, ResultType);
 
-	FDimensional* Value = FPrivate::NewDimensionalValue(this, Type);
-	for (FValue*& Component : Value->GetMutableComponents())
+	for (int i = 0; i < Mask.NumComponents; ++i)
 	{
-		Component = Scalar;
+		Result->GetComponents()[i] = EmitSubscript(Value, (int)Mask.Components[i]);
 	}
 
-	return FPrivate::EmitNew(this, Value);
+	return Result;
 }
 
 FSetMaterialOutput* FEmitter::EmitSetMaterialOutput(EMaterialProperty InProperty, FValue* InArgValue)
@@ -348,70 +496,172 @@ FSetMaterialOutput* FEmitter::EmitSetMaterialOutput(EMaterialProperty InProperty
 	return Instr;
 }
 
-static FTypePtr GetBinaryOperatorResultType(EBinaryOperator Operator, FTypePtr LhsType, FTypePtr RhsType)
+template <typename T>
+static T FoldScalarArithmeticOperator(EBinaryOperator Operator, T Lhs, T Rhs)
 {
 	switch (Operator)
 	{
-		case BO_Add:
-		case BO_Subtract:
-		case BO_Multiply:
-		case BO_Divide:
-			check(LhsType == RhsType);
-			return LhsType;
-		
-		case BO_Greater:
-		case BO_Lower:
-		case BO_Equals:
-			return FArithmeticType::GetBool();
+		case BO_Add: return Lhs + Rhs;
+		case BO_Subtract: return Lhs - Rhs;
+		case BO_Multiply: return Lhs * Rhs;
+		case BO_Divide: return Lhs / Rhs;
+		default: UE_MIR_UNREACHABLE();
+	}
+}
 
-		default:
-			UE_MIR_UNREACHABLE();
+template <typename T>
+static bool FoldComparisonOperatorScalar(EBinaryOperator Operator, T Lhs, T Rhs)
+{
+	switch (Operator)
+	{
+		case BO_Greater: return Lhs > Rhs;
+		case BO_GreaterOrEquals: return Lhs >= Rhs;
+		case BO_Lower: return Lhs < Rhs;
+		case BO_LowerOrEquals: return Lhs <= Rhs;
+		case BO_Equals: return Lhs == Rhs;
+		case BO_NotEquals: return Lhs != Rhs;
+		default: UE_MIR_UNREACHABLE();
+	}
+}
+
+static FValue* FoldBinaryOperatorScalar(FEmitter* Emitter, EBinaryOperator Operator, const FConstant* Lhs, const FConstant* Rhs)
+{
+	FArithmeticTypePtr ArithType = Lhs->Type->AsArithmetic();
+
+	if (IsArithmeticOperator(Operator))
+	{
+		switch (ArithType->ScalarKind)
+		{
+			case SK_Int:
+			{
+				int Result = FoldScalarArithmeticOperator<TInteger>(Operator, Lhs->Integer, Rhs->Integer);
+				return Emitter->EmitConstantInt1(Result);
+			}
+
+			case SK_Float:
+			{
+				TFloat Result = FoldScalarArithmeticOperator<TFloat>(Operator, Lhs->Float, Rhs->Float);
+				return Emitter->EmitConstantFloat1(Result);
+			}
+
+			default:
+				UE_MIR_UNREACHABLE();
+		}
+	}
+	else if (IsComparisonOperator(Operator))
+	{
+		bool Result;
+		switch (ArithType->ScalarKind)
+		{
+			case SK_Int:
+				Result = FoldComparisonOperatorScalar<TInteger>(Operator, Lhs->Integer, Rhs->Integer);
+				break;
+
+			case SK_Float:
+				Result = FoldComparisonOperatorScalar<TFloat>(Operator, Lhs->Float, Rhs->Float);
+				break;
+
+			default:
+				UE_MIR_UNREACHABLE();
+		}
+		return Emitter->EmitConstantBool1(Result);
+	}
+	else
+	{
+		UE_MIR_UNREACHABLE();
 	}
 }
 
 FValue* FEmitter::EmitBinaryOperator(EBinaryOperator Operator, FValue* Lhs, FValue* Rhs)
 {
+	// Argument value types must always match.
 	check(Lhs->Type == Rhs->Type);
 
-	/* todo: operation folding */
+	// Get the operands arithmetic type.
+	FArithmeticTypePtr ArithType = Lhs->Type->AsArithmetic();
+	check(ArithType);
+
+	// Try converting operands to scalar constants. If they're both so, fold the binary
+	// operator right away and return the result.
+	const FConstant* ScalarLhs = Lhs->As<FConstant>();
+	const FConstant* ScalarRhs = Rhs->As<FConstant>();
+
+	if (ScalarLhs && ScalarRhs)
+	{
+		if (FValue* Value = FoldBinaryOperatorScalar(this, Operator, ScalarLhs, ScalarRhs))
+		{
+			return Value;
+		}
+	}
+
+	// Determine the result type. If the operator is arithmetic, the result type will be the same
+	// as the operands type. Otherwise it will have the same number of components but bool.
+	FArithmeticTypePtr ResultType = IsArithmeticOperator(Operator)
+		? ArithType
+		: FArithmeticType::Get(SK_Bool, ArithType->NumRows, ArithType->NumColumns);
+
+	// Now check that both operands are dimensional. If so, emit a new dimensional value with each component
+	// being the result of the binary operation applied to the corresponding components of the operands.
+	const FDimensional* DimensionalLhs = Lhs->As<FDimensional>();
+	const FDimensional* DimensionalRhs = Rhs->As<FDimensional>();
+
+	if (DimensionalLhs && DimensionalRhs)
+	{
+		FDimensional* Result = NewDimensionalValue(this, ResultType);
+
+		auto ResultComponents = Result->GetComponents();
+		auto LhsComponents = DimensionalLhs->GetComponents();
+		auto RhsComponents = DimensionalRhs->GetComponents();
 
-	FTypePtr ResultType = GetBinaryOperatorResultType(Operator, Lhs->Type, Rhs->Type);
+		for (int i = 0; i < ResultType->GetNumComponents(); ++i)
+		{
+			ResultComponents[i] = EmitBinaryOperator(Operator, LhsComponents[i], RhsComponents[i]);
+		}
+
+		return Result;
+	}
+
+	// One (or both) of the operands is runtime only. Emit the runtime binary operation instruction.
 	FBinaryOperator Proto = InitValue<FBinaryOperator>(ResultType);
 	Proto.Operator = Operator;
 	Proto.LhsArg= Lhs;
 	Proto.RhsArg = Rhs;
 
-	return FPrivate::EmitPrototype(this, Proto);
+	return EmitPrototype(this, Proto);
 }
 
 FValue* FEmitter::EmitBranch(FValue* Condition, FValue* True, FValue* False)
 {
-	/* todo: operation folding */
+	// Condition must be of type bool
+	check(Condition->Type->IsBool1());
 
+	// If the condition is a scalar constant, then simply evaluate the result now.
+	if (const FConstant* ConstCondition = Condition->As<FConstant>())
+	{
+		return ConstCondition->Boolean ? True : False;
+	}
+
+	// If the condition is not static, the types of the true and false operands must match.
+	// The resulting type will be that of the branch instruction created.
 	check(True->Type == False->Type);
 
+	// Create the branch instruction.
 	FBranch Proto = InitValue<FBranch>(True->Type);
 	Proto.ConditionArg = Condition;
 	Proto.TrueArg = True;
 	Proto.FalseArg = False;
 
-	return FPrivate::EmitPrototype(this, Proto);
+	return EmitPrototype(this, Proto);
 }
 
-static FValue* ConvertScalarValue(FEmitter* Emitter, FValue* ScalarValue, EScalarKind ScalarKind, EScalarKind TargetKind)
+static FValue* CastConstant(FEmitter* Emitter, FConstant* Constant, EScalarKind ConstantScalarKind, EScalarKind TargetKind)
 {
-	if (ScalarKind == TargetKind)
+	if (ConstantScalarKind == TargetKind)
 	{
-		return ScalarValue;
+		return Constant;
 	}
 
-	auto ScalarConstant = ScalarValue->As<FScalarConstant>();
-	if (!ScalarConstant)
-	{
-		UE_MIR_UNREACHABLE(); // todo
-	}
-
-	switch (ScalarKind)
+	switch (ConstantScalarKind)
 	{
 		case SK_Bool:
 		case SK_Int:
@@ -419,8 +669,8 @@ static FValue* ConvertScalarValue(FEmitter* Emitter, FValue* ScalarValue, EScala
 			switch (TargetKind)
 			{
 				case SK_Bool: UE_MIR_TODO();
-				case SK_Int: return Emitter->EmitConstantInt1(ScalarConstant->Integer);
-				case SK_Float: return Emitter->EmitConstantFloat1((float)ScalarConstant->Integer);
+				case SK_Int: return Emitter->EmitConstantInt1(Constant->Integer);
+				case SK_Float: return Emitter->EmitConstantFloat1((TFloat)Constant->Integer);
 				default: UE_MIR_UNREACHABLE();
 			}
 		}
@@ -430,7 +680,7 @@ static FValue* ConvertScalarValue(FEmitter* Emitter, FValue* ScalarValue, EScala
 			switch (TargetKind)
 			{
 				case SK_Bool: UE_MIR_TODO();
-				case SK_Int: return Emitter->EmitConstantInt1((int)ScalarConstant->Float);
+				case SK_Int: return Emitter->EmitConstantInt1((int)Constant->Float);
 				default: UE_MIR_UNREACHABLE();
 			}
 		}
@@ -441,48 +691,143 @@ static FValue* ConvertScalarValue(FEmitter* Emitter, FValue* ScalarValue, EScala
 	UE_MIR_UNREACHABLE();
 }
 
-static FValue* ConvertArithmeticValue(FEmitter* Emitter, FValue* Value, FArithmeticTypePtr TargetArithmeticType)
+static FValue* ConstructArithmeticValue(FEmitter* Emitter, FArithmeticTypePtr TargetArithmeticType, FValue* Initializer)
 {
-	FArithmeticTypePtr ValueArithmeticType = static_cast<FArithmeticTypePtr>(Value->Type);
-	if (ValueArithmeticType->IsScalar())
+	FArithmeticTypePtr InitializerArithmeticType = Initializer->Type->AsArithmetic();
+	if (!InitializerArithmeticType)
 	{
-		FValue* ConvertedScalar = ConvertScalarValue(Emitter, Value, ValueArithmeticType->ScalarKind, TargetArithmeticType->ScalarKind);
-		return Emitter->EmitArithmetic(TargetArithmeticType, ConvertedScalar);
+		Emitter->Errorf(TEXT("Cannot construct a '%s' from non arithmetic type '%s'."), TargetArithmeticType->Spelling.GetData(), Initializer->Type->GetSpelling().GetData());
+		return nullptr;
 	}
-	else
+
+	// Construct a scalar from another scalar.
+	if (TargetArithmeticType->IsScalar() && InitializerArithmeticType->IsScalar())
 	{
-		UE_MIR_TODO();
+		// Construct the scalar from a constant.
+		if (FConstant* ConstantInitializer = Initializer->As<FConstant>())
+		{
+			return CastConstant(Emitter, ConstantInitializer, InitializerArithmeticType->ScalarKind, TargetArithmeticType->ScalarKind);
+		}
+		else
+		{
+			/* todo: emit conversion unary operator */
+			UE_MIR_TODO();
+		}
 	}
-}
 
-FValue* FEmitter::TryEmitConvert(FValue* Value, FTypePtr TargetType)
-{
-	// If types already match, nothing else to do, simply return the same value.
-	FTypePtr ValueType = Value->Type;
-	if (ValueType == TargetType)
+	// Construct a vector or matrix from a scalar. E.g. float4(3.14f)
+	if (!TargetArithmeticType->IsScalar() && InitializerArithmeticType->IsScalar())
 	{
-		return Value;
+		// Create the result dimensional value.
+		FDimensional* Result = NewDimensionalValue(Emitter, TargetArithmeticType);
+
+		// Create a dimensional and initialize each of its components to the conversion
+		// of initializer value to the single component type.
+		FValue* Component = Emitter->TryEmitConstruct(TargetArithmeticType->ToScalar(), Initializer);
+
+		// Initialize all result components to the same scalar.
+		for (int i = 0; i < TargetArithmeticType->GetNumComponents(); ++i)
+		{
+			Result->GetComponents()[i] = Component;
+		}
+		
+		return EmitNew(Emitter, Result);
+	}
+	
+	// Construct a vector from another vector. If constructed vector is larger, initialize
+	// remaining components to zero. If it's smaller, truncate initializer vector and only use
+	// the necessary components.
+	if (TargetArithmeticType->IsVector() && InitializerArithmeticType->IsVector())
+	{
+		int TargetNumComponents = TargetArithmeticType->GetNumComponents();
+		int InitializerNumComponents = InitializerArithmeticType->GetNumComponents();
+
+		int MinNumComponents = FMath::Min(TargetNumComponents, InitializerNumComponents);
+		int MaxNumComponents = FMath::Max(TargetNumComponents, InitializerNumComponents);
+
+		// Create the result dimensional value.
+		FDimensional* Result = NewDimensionalValue(Emitter, TargetArithmeticType);
+
+		// Determine the result component type (scalar).
+		FArithmeticTypePtr ResultComponentType = TargetArithmeticType->ToScalar();
+
+		// For iterating over the components of the result dimensional value.
+		int Index = 0;
+
+		// Convert components from the initializer vector.
+		for (; Index < MinNumComponents; ++Index)
+		{
+			Result->GetComponents()[Index] = Emitter->TryEmitConstruct(ResultComponentType, Emitter->EmitSubscript(Initializer, Index));
+		}
+
+		// Initialize remaining result dimensional components to zero.
+		for (; Index < MaxNumComponents; ++Index)
+		{
+			Result->GetComponents()[Index] = Emitter->EmitConstantScalarZero(ResultComponentType->ScalarKind);
+		}
+
+		return EmitNew(Emitter, Result);
 	}
+	
+	// The two arithmetic types are identical matrices that differ only by their scalar type.
+	if (TargetArithmeticType->NumRows 	 == InitializerArithmeticType->NumRows &&
+		TargetArithmeticType->NumColumns == InitializerArithmeticType->NumColumns)
+	{
+		check(TargetArithmeticType->IsMatrix());
+
+		// 
+		if (FDimensional* DimensionalInitializer = Initializer->As<FDimensional>())
+		{
+			// Create the result dimensional value.
+			FDimensional* Result = NewDimensionalValue(Emitter, TargetArithmeticType);
+			
+			// Determine the result component type (scalar).
+			FArithmeticTypePtr ResultComponentType = TargetArithmeticType->ToScalar();
+
+			// Convert components from the initializer vector.
+			for (int Index = 0, Num = TargetArithmeticType->GetNumComponents(); Index < Num; ++Index)
+			{
+				Result->GetComponents()[Index] = Emitter->TryEmitConstruct(ResultComponentType, DimensionalInitializer->GetComponents()[Index]);
+			}
 
-	FValue* ResultValue{};
+			return EmitNew(Emitter, Result);
+		}
+		else
+		{
+			// Initializer is an unknown value, construct target value casting initializer.
+			FCast Prototype = InitValue<FCast>(TargetArithmeticType);
+			Prototype.Arg = Initializer;
+			return EmitPrototype(Emitter, Prototype);
+		}
+	}
+	
+	// Initializer value cannot be used to construct this arithmetic type.
+	return nullptr;
+}
 
-	// Converting from an arithmetic value to another arithmetic type.
-	FArithmeticTypePtr ValueArithmeticType = Value->Type->AsArithmetic();
-	FArithmeticTypePtr TargetArithmeticType = TargetType->AsArithmetic();
-	if (ValueArithmeticType && TargetArithmeticType)
+FValue* FEmitter::TryEmitConstruct(FTypePtr Type, FValue* Initializer)
+{
+	// If target type matches initializer's, simply return the same value.
+	FTypePtr InitializerType = Initializer->Type;
+	if (InitializerType == Type)
 	{
-		ResultValue = ConvertArithmeticValue(this, Value, TargetArithmeticType);
+		return Initializer;
 	}
+	
+	FValue* Result{};
 
-	/* ... */
+	if (FArithmeticTypePtr ArithmeticType = Type->AsArithmetic())
+	{
+		Result = ConstructArithmeticValue(this, ArithmeticType, Initializer);
+	}
 
 	// No other legal conversions applicable. Report error if we haven't converted the value.
-	if (!ResultValue)
+	if (!Result)
 	{
-		Errorf(TEXT("Cannot convert value of type '%s' to '%s'"), TEXT(""), TEXT(""));
+		Errorf(TEXT("Cannot construct a '%s' from a '%s'."), TEXT(""), TEXT(""));
 	}
-
-	return ResultValue;
+	
+	return Result;
 }
 
 FArithmeticTypePtr FEmitter::TryGetCommonArithmeticType(FArithmeticTypePtr A, FArithmeticTypePtr B)
diff --git a/Engine/Source/Runtime/Engine/Private/Materials/MaterialIRToHLSLTranslator.cpp b/Engine/Source/Runtime/Engine/Private/Materials/MaterialIRToHLSLTranslator.cpp
--- a/Engine/Source/Runtime/Engine/Private/Materials/MaterialIRToHLSLTranslator.cpp
+++ b/Engine/Source/Runtime/Engine/Private/Materials/MaterialIRToHLSLTranslator.cpp
@@ -15,6 +15,8 @@
 #include "Materials/MaterialExpressionVolumetricAdvancedMaterialOutput.h"
 #include "RenderUtils.h"
 
+#include <inttypes.h>
+
 #if WITH_EDITOR
 
 namespace IR = UE::MIR;
@@ -50,24 +52,12 @@ struct FHLSLPrinter
 		return *this;
 	}
 
-	FHLSLPrinter& operator<<(int32 Value)
+	FHLSLPrinter& operator<<(int Value)
 	{
 		Buffer.Appendf(TEXT("%d"), Value);
 		return *this;
 	}
 
-	FHLSLPrinter& operator<<(uint32 Value)
-	{
-		Buffer.Appendf(TEXT("%u"), Value);
-		return *this;
-	}
-
-	FHLSLPrinter& operator<<(float Value)
-	{
-		Buffer.Appendf(TEXT("%.5ff"), Value);
-		return *this;
-	}
-
     FHLSLPrinter& operator<<(ENoOp)
 	{
 		return *this;
@@ -272,6 +262,26 @@ struct FTranslator : FMaterialIRToHLSLTranslation
 	{
 		switch (Instr->Kind)
 		{
+			case IR::VK_Dimensional:
+			{
+				const IR::FDimensional* Vector = static_cast<const IR::FDimensional*>(Instr);
+
+				IR::FArithmeticTypePtr ArithType = Vector->Type->AsArithmetic();
+				check(ArithType && ArithType->IsVector());
+
+				Printer << ScalarKindToString(ArithType->ScalarKind) << ArithType->NumRows << TEXT("(");
+
+				Printer.BeginList();
+				for (IR::FValue* Component : Vector->GetComponents())
+				{
+					Printer.PrintListSeparator();
+					LowerValue(Component);
+				}
+
+				Printer << TEXT(")");
+				break;
+			}
+
 			case IR::VK_SetMaterialOutput:
 			{
 				auto Output = static_cast<const IR::FSetMaterialOutput*>(Instr);
@@ -279,7 +289,7 @@ struct FTranslator : FMaterialIRToHLSLTranslation
 				// Special case MP_SubsurfaceColor as the actual property is a combination of the color and the profile but we don't want to expose the profile
 				const FString& PropertyName = (Output->Property == MP_SubsurfaceColor) ? "Subsurface" : FMaterialAttributeDefinitionMap::GetAttributeName(Output->Property);
 
-				Printer << TEXT("PixelMaterialInputs.") << PropertyName << TEXT(" = ") << InlineValue(Output->Arg);
+				Printer << TEXT("PixelMaterialInputs.") << PropertyName << TEXT(" = ") << LowerValue(Output->Arg);
 
 				break;
 			}
@@ -288,7 +298,7 @@ struct FTranslator : FMaterialIRToHLSLTranslation
 			{
 				auto BinaryOperator = static_cast<const IR::FBinaryOperator*>(Instr);
 
-				InlineValue(BinaryOperator->LhsArg);
+				LowerValue(BinaryOperator->LhsArg);
 
 				const TCHAR* OpString;
 				switch (BinaryOperator->Operator)
@@ -304,7 +314,7 @@ struct FTranslator : FMaterialIRToHLSLTranslation
 				}
 				Printer << OpString;
 
-				InlineValue(BinaryOperator->RhsArg);
+				LowerValue(BinaryOperator->RhsArg);
 
 				break;
 			}
@@ -315,24 +325,41 @@ struct FTranslator : FMaterialIRToHLSLTranslation
 
 				if (IsFoldable(Branch))
 				{
-					Printer << InlineValue(Branch->ConditionArg)
-					<< TEXT(" ? ") << InlineValue(Branch->TrueArg)
-					<< TEXT(" : ") << InlineValue(Branch->FalseArg);
+					Printer << LowerValue(Branch->ConditionArg)
+						<< TEXT(" ? ") << LowerValue(Branch->TrueArg)
+						<< TEXT(" : ") << LowerValue(Branch->FalseArg);
 				}
 				else
 				{
 					Printer << EndOfStatement;
-					Printer << TEXT("if (") << InlineValue(Branch->ConditionArg) << TEXT(")") << NewLine << OpenBrace;
+					Printer << TEXT("if (") << LowerValue(Branch->ConditionArg) << TEXT(")") << NewLine << OpenBrace;
 					Printer << LowerBlock(Branch->TrueBlock);
-					Printer << LocalToIdentifier[Instr] << " = " << InlineValue(Branch->TrueArg) << EndOfStatement;
+					Printer << LocalToIdentifier[Instr] << " = " << LowerValue(Branch->TrueArg) << EndOfStatement;
 					Printer << CloseBrace << NewLine;
 					Printer << TEXT("else") << NewLine << OpenBrace;
 					Printer << LowerBlock(Branch->FalseBlock);
-					Printer << LocalToIdentifier[Instr] << " = " << InlineValue(Branch->FalseArg) << EndOfStatement;
+					Printer << LocalToIdentifier[Instr] << " = " << LowerValue(Branch->FalseArg) << EndOfStatement;
 					Printer << CloseBrace;
 				}
 				break;
 			}
+			
+			case IR::VK_Subscript:
+			{
+				const IR::FSubscript* Subscript = static_cast<const IR::FSubscript*>(Instr);
+
+				LowerValue(Subscript->Arg);
+
+				if (IR::FArithmeticTypePtr ArgArithmeticType = Subscript->Arg->Type->AsVector())
+				{
+					const TCHAR* ComponentsStr[] = { TEXT(".x"), TEXT(".y"), TEXT(".z"), TEXT(".w") };
+					check(Subscript->Index <= ArgArithmeticType->GetNumComponents());
+
+					Printer << ComponentsStr[Subscript->Index];
+				}
+
+				break;
+			}
 
 			default:
 				UE_MIR_UNREACHABLE();
@@ -341,50 +368,45 @@ struct FTranslator : FMaterialIRToHLSLTranslation
 		return NoOp;
 	}
 
-	ENoOp InlineValue(const IR::FValue* InValue)
+	ENoOp LowerValue(const IR::FValue* InValue)
 	{
 		if (const IR::FInstruction* Instr = InValue->AsInstruction())
 		{
 			if (Instr->NumUsers <= 1 && IsFoldable(Instr))
 			{
-                LowerInstruction(Instr);
-            }
-            else
-			{
-                Printer << LocalToIdentifier[Instr];
-            }
-		}
-		else if (const IR::FScalarConstant* Scalar = InValue->As<IR::FScalarConstant>())
-		{
-			IR::FArithmeticTypePtr ArithType = Scalar->Type->AsArithmetic();
-			check(ArithType && ArithType->IsScalar());
-
-			switch (ArithType->ScalarKind)
+				Printer << LowerInstruction(Instr);
+			}
+			else
 			{
-				case IR::SK_Bool: Printer << Scalar->Boolean; break;
-				case IR::SK_Int: Printer << Scalar->Integer; break;
-				case IR::SK_Float: Printer << Scalar->Float; break;
+				Printer << LocalToIdentifier[Instr];
 			}
+
+			return NoOp;
 		}
-		else if (const IR::FDimensional* Vector = InValue->As<IR::FDimensional>())
+
+		switch (InValue->Kind)
 		{
-			IR::FArithmeticTypePtr ArithType = Vector->Type->AsArithmetic();
-			check(ArithType && ArithType->IsVector());
+			case IR::VK_Constant:
+			{
+				const IR::FConstant* Scalar = static_cast<const IR::FConstant*>(InValue);
+				
+				IR::FArithmeticTypePtr ArithType = Scalar->Type->AsArithmetic();
+				check(ArithType && ArithType->IsScalar());
 
-			Printer << ScalarKindToString(ArithType->ScalarKind) << ArithType->NumRows << TEXT("(");
+				switch (ArithType->ScalarKind)
+				{
+					case IR::SK_Bool:  Printer.Buffer.Append(Scalar->Boolean ? TEXT("true") : TEXT("false")); break;
+					case IR::SK_Int:   Printer.Buffer.Appendf(TEXT("%") PRId64, Scalar->Integer); break;
+					case IR::SK_Float: Printer.Buffer.Appendf(TEXT("%.5ff"), Scalar->Float); break;
+				}
 
-			Printer.BeginList();
-			for (IR::FValue* Component : Vector->GetComponents())
-			{
-				Printer.PrintListSeparator();
-				InlineValue(Component);
+				break;
 			}
 
-			Printer << TEXT(")");
-		}
-		else
-		{
-			UE_MIR_UNREACHABLE();
+			default:
+			{
+				UE_MIR_UNREACHABLE();
+			}
 		}
 
 		return NoOp;
diff --git a/Engine/Source/Runtime/Engine/Private/Materials/MaterialIRTypes.cpp b/Engine/Source/Runtime/Engine/Private/Materials/MaterialIRTypes.cpp
--- a/Engine/Source/Runtime/Engine/Private/Materials/MaterialIRTypes.cpp
+++ b/Engine/Source/Runtime/Engine/Private/Materials/MaterialIRTypes.cpp
@@ -66,6 +66,11 @@ FStringView FType::GetSpelling() const
 	UE_MIR_UNREACHABLE();
 }
 
+bool FType::IsBool1() const
+{
+	return this == FArithmeticType::GetBool1();
+}
+
 FArithmeticTypePtr FType::AsArithmetic() const
 {
 	return Kind == TK_Arithmetic ? static_cast<FArithmeticTypePtr>(this) : nullptr; 
@@ -100,17 +105,17 @@ const TCHAR* ScalarKindToString(EScalarKind Kind)
 	}
 }
 
-FArithmeticTypePtr FArithmeticType::GetBool()
+FArithmeticTypePtr FArithmeticType::GetBool1()
 {
 	return GetScalar(SK_Bool);
 }
 
-FArithmeticTypePtr FArithmeticType::GetInt()
+FArithmeticTypePtr FArithmeticType::GetInt1()
 {
 	return GetScalar(SK_Int);
 }
 
-FArithmeticTypePtr FArithmeticType::GetFloat()
+FArithmeticTypePtr FArithmeticType::GetFloat1()
 {
 	return GetScalar(SK_Float);
 }
@@ -195,6 +200,11 @@ FArithmeticTypePtr FArithmeticType::Get(EScalarKind InScalarKind, int NumRows, i
 	return &Types[Index];
 }
 
+FArithmeticTypePtr FArithmeticType::ToScalar() const
+{
+	return FArithmeticType::GetScalar(ScalarKind);
+}
+
 } // namespace UE::MIR
 
 #endif // #if WITH_EDITOR
diff --git a/Engine/Source/Runtime/Engine/Public/Materials/MaterialIR.h b/Engine/Source/Runtime/Engine/Public/Materials/MaterialIR.h
--- a/Engine/Source/Runtime/Engine/Public/Materials/MaterialIR.h
+++ b/Engine/Source/Runtime/Engine/Public/Materials/MaterialIR.h
@@ -12,17 +12,20 @@ namespace UE::MIR {
 enum EValueKind
 {
 	/* Values */
-	VK_ScalarConstant,
-	VK_Dimensional,
+
+	VK_Constant,
 
 	/* Instructions */
-	VK_InstructionBegin,
 
-	VK_SetMaterialOutput = VK_InstructionBegin,
+	VK_Dimensional,
+	VK_SetMaterialOutput,
 	VK_BinaryOperator,
 	VK_Branch,
+	VK_Subscript,
+	VK_Cast,
 
 	VK_InstructionEnd,
+	VK_InstructionBegin = VK_Dimensional,
 };
 
 /* Values */
@@ -38,6 +41,9 @@ struct FValue
 	bool Equals(const FValue* Other) const;
 	uint32 GetSizeInBytes() const;
 	TArrayView<FValue*> GetUses();
+	
+	template <typename T>
+	T* As() { return this && IsA(T::TypeKind) ? static_cast<T*>(this) : nullptr; }
 
 	template <typename T>
 	const T* As() const { return this && IsA(T::TypeKind) ? static_cast<const T*>(this) : nullptr; }
@@ -49,30 +55,49 @@ struct TValue : FValue
 	static constexpr EValueKind TypeKind = TTypeKind;
 };
 
-struct FScalarConstant : TValue<VK_ScalarConstant>
+using TInteger = int64_t;
+using TFloat = double;
+
+struct FConstant : TValue<VK_Constant>
 {
 	union
 	{
-		bool  Boolean;
-		int   Integer;
-		float Float;
+		bool  		Boolean;
+		TInteger	Integer;
+		TFloat 	Float;
 	};
-};
 
-struct FDimensional : TValue<VK_Dimensional>
-{
-	TArrayView<FValue* const> GetComponents() const;
-	TArrayView<FValue*> GetMutableComponents();
+	template <typename T>
+	T Get() const
+	{
+		if constexpr (std::is_same_v<T, bool>)
+		{
+			return Boolean;
+		}
+		else if constexpr (std::is_integral_v<T>)
+		{
+			return Integer;
+		}
+		else if constexpr (std::is_floating_point_v<T>)
+		{
+			return Float;
+		}
+		else
+		{
+			check(false && "unexpected type T.");
+		}
+	}
 };
 
-template <int TDimension>
-struct TDimensional : FDimensional
+/* Instructions */
+
+struct FBlock
 {
-	FValue* Components[TDimension];
+	FBlock* Parent{};
+	FInstruction* Instructions{};
+	int32 Level{};
 };
 
-/* Instructions */
-
 enum EInstructionFlags
 {
 	IF_None = 0,
@@ -81,7 +106,7 @@ enum EInstructionFlags
 
 struct FInstruction : FValue
 {
-	EInstructionFlags Flags = IF_None;
+	EInstructionFlags Flags{};
 	FInstruction* Next{};
 	FBlock* Block{};
 	uint32 NumUsers{};
@@ -91,19 +116,31 @@ struct FInstruction : FValue
 	bool GetInnerBlock(int32 Index, FValue*& OutArg, FBlock*& OutBlock); 
 };
 
-struct FBlock
-{
-	FBlock* Parent{};
-	FInstruction* Instructions{};
-	int32 Level{};
-};
-
 template <EValueKind TTypeKind>
 struct TInstruction : FInstruction
 {
 	static constexpr EValueKind TypeKind = TTypeKind;
 };
 
+struct FDimensional : TInstruction<VK_Dimensional>
+{
+	// Returns the constant array of component values. 
+	TArrayView<FValue* const> GetComponents() const;
+
+	// Returns the mutable array of component values. 
+	TArrayView<FValue*> GetComponents();
+
+	// Returns whether all components are constant.
+	bool AreComponentsConstant() const;
+
+};
+
+template <int TDimension>
+struct TDimensional : FDimensional
+{
+	FValue* Components[TDimension];
+};
+
 struct FSetMaterialOutput : TInstruction<VK_SetMaterialOutput>
 {
 	EMaterialProperty Property;
@@ -113,15 +150,25 @@ struct FSetMaterialOutput : TInstruction<VK_SetMaterialOutput>
 enum EBinaryOperator
 {
 	BO_Invalid,
+
+	/* Arithmetic */
 	BO_Add,
 	BO_Subtract,
 	BO_Multiply,
 	BO_Divide,
+
+	/* Comparison */
 	BO_Greater,
+	BO_GreaterOrEquals,
 	BO_Lower,
+	BO_LowerOrEquals,
 	BO_Equals,
+	BO_NotEquals,
 };
 
+bool IsArithmeticOperator(EBinaryOperator Op);
+bool IsComparisonOperator(EBinaryOperator Op);
+
 struct FBinaryOperator : TInstruction<VK_BinaryOperator>
 {
 	EBinaryOperator Operator = BO_Invalid;
@@ -138,5 +185,16 @@ struct FBranch : TInstruction<VK_Branch>
 	FBlock FalseBlock{};
 };
 
+struct FSubscript : TInstruction<VK_Subscript>
+{
+	FValue* Arg;
+	int Index;
+};
+
+struct FCast : TInstruction<VK_Cast>
+{
+	FValue* Arg{};
+};
+
 } // namespace UE::MIR
 #endif
diff --git a/Engine/Source/Runtime/Engine/Public/Materials/MaterialIREmitter.h b/Engine/Source/Runtime/Engine/Public/Materials/MaterialIREmitter.h
--- a/Engine/Source/Runtime/Engine/Public/Materials/MaterialIREmitter.h
+++ b/Engine/Source/Runtime/Engine/Public/Materials/MaterialIREmitter.h
@@ -9,6 +9,31 @@
 
 namespace UE::MIR {
 
+enum class EVectorComponent : uint8_t
+{
+	X, Y, Z, W, 
+};
+
+// Return the lower case string representation of specified component (e.g. "x")
+const TCHAR* VectorComponentToString(EVectorComponent);
+
+//
+struct FSwizzleMask
+{
+	EVectorComponent Components[4];
+	int NumComponents{};
+
+	FSwizzleMask() {}
+	FSwizzleMask(EVectorComponent X);
+	FSwizzleMask(EVectorComponent X, EVectorComponent Y);
+	FSwizzleMask(EVectorComponent X, EVectorComponent Y, EVectorComponent Z);
+	FSwizzleMask(EVectorComponent X, EVectorComponent Y, EVectorComponent Z, EVectorComponent W);
+
+	void Append(EVectorComponent Component);
+	const EVectorComponent* begin() const { return Components; }
+	const EVectorComponent* end() const { return Components + NumComponents; }
+};
+
 //
 class FEmitter
 {
@@ -30,7 +55,7 @@ public:
 	FEmitter& DefaultToFloatZero(const FExpressionInput* Input);
 
 	//
-	FEmitter& DefaultTo(const FExpressionInput* Input, float Float);
+	FEmitter& DefaultTo(const FExpressionInput* Input, TFloat Float);
 
 	// It gets the value flowing into it and checks that its type is float scalar.
 	FValue* TryGetFloat(const FExpressionInput* Input);
@@ -62,11 +87,14 @@ public:
 
 	FValue* EmitConstantFromShaderValue(const UE::Shader::FValue& InValue);
 	FValue* EmitConstantScalarZero(EScalarKind Kind);
-	FValue* EmitConstantFloat1(float InX);
+	FValue* EmitConstantTrue();
+	FValue* EmitConstantFalse();
+	FValue* EmitConstantBool1(bool InX);
+	FValue* EmitConstantFloat1(TFloat InX);
 	FValue* EmitConstantFloat2(const FVector2f& InValue);
 	FValue* EmitConstantFloat3(const FVector3f& InValue);
 	FValue* EmitConstantFloat4(const FVector4f& InValue);
-	FValue* EmitConstantInt1(int InX);
+	FValue* EmitConstantInt1(TInteger InX);
 	FValue* EmitConstantInt2(const FIntVector2& InValue);
 	FValue* EmitConstantInt3(const FIntVector3& InValue);
 	FValue* EmitConstantInt4(const FIntVector4& InValue);
@@ -76,14 +104,15 @@ public:
 
 	/* Other Values */
 
-	FValue* EmitArithmetic(FArithmeticTypePtr Type, FValue* Scalar);
+	FValue* EmitSubscript(FValue* Value, int ComponentIndex);
+	FValue* TryEmitSwizzle(FValue* Value, FSwizzleMask Mask);
 
 	/* Instructions */
 
 	FSetMaterialOutput* EmitSetMaterialOutput(EMaterialProperty InProperty, FValue* InArgValue);
 	FValue* EmitBinaryOperator(EBinaryOperator Operator, FValue* Lhs, FValue* Rhs);
 	FValue* EmitBranch(FValue* Condition, FValue* True, FValue* False);
-	FValue* TryEmitConvert(FValue* Value, FTypePtr TargetType);
+	FValue* TryEmitConstruct(FTypePtr Type, FValue* Initializer);
 
 	/* Types */
 
@@ -100,14 +129,17 @@ public:
 
 	void Error(FString Message);
 
+	struct FPrivate;
+
 private:
 	UMaterial* Material{};
 	FMaterialIRModule* Module{};
 	UMaterialExpression* Expression{};
 	FMaterialIRModuleBuilder* Builder{};
 	bool bHasExprBuildError = false;
+	FValue* ConstantTrue;
+	FValue* ConstantFalse;
 
-	struct FPrivate;
 	friend FMaterialIRModuleBuilder;
 };
 
diff --git a/Engine/Source/Runtime/Engine/Public/Materials/MaterialIRTypes.h b/Engine/Source/Runtime/Engine/Public/Materials/MaterialIRTypes.h
--- a/Engine/Source/Runtime/Engine/Public/Materials/MaterialIRTypes.h
+++ b/Engine/Source/Runtime/Engine/Public/Materials/MaterialIRTypes.h
@@ -19,16 +19,32 @@ const TCHAR* TypeKindToString(ETypeKind Kind);
 
 struct FType
 {
-	ETypeKind 	Kind;
+	// Identifies what derived type this is.
+	ETypeKind Kind;
 
+	// Returns the type matching specified UE::Shader::FType.
 	static FTypePtr FromShaderType(const UE::Shader::FType& InShaderType);
+	
+	// Returns the `void` type.
 	static FTypePtr GetVoid();
 
-	FStringView GetSpelling() const;
+	// Returns whether this type is a `bool` scalar.
+	bool IsBool1() const;
+
+	// Returns this type upcast this type to ArithmeticType if it is one. Otherwise it returns nullptr.
 	FArithmeticTypePtr AsArithmetic() const;
+
+	// Returns this type upcast this type to ArithmeticType if it's a scalar. Otherwise it returns nullptr.
 	FArithmeticTypePtr AsScalar() const;
+
+	// Returns this type upcast this type to ArithmeticType if it's a vector. Otherwise it returns nullptr.
 	FArithmeticTypePtr AsVector() const;
+
+	// Returns this type upcast this type to ArithmeticType if it's a matrix. Otherwise it returns nullptr.
 	FArithmeticTypePtr AsMatrix() const;
+
+	// Returns the this type name spelling (e.g. float4x4).
+	FStringView GetSpelling() const;
 };
 
 enum EScalarKind
@@ -45,17 +61,20 @@ struct FArithmeticType : FType
 	int NumRows;
 	int NumColumns;
 
-	static FArithmeticTypePtr GetBool();
-	static FArithmeticTypePtr GetInt();
-	static FArithmeticTypePtr GetFloat();
+	static FArithmeticTypePtr GetBool1();
+	static FArithmeticTypePtr GetInt1();
+	static FArithmeticTypePtr GetFloat1();
+
 	static FArithmeticTypePtr GetScalar(EScalarKind InScalarKind);
 	static FArithmeticTypePtr GetVector(EScalarKind InScalarKind, int NumRows);
 	static FArithmeticTypePtr GetMatrix(EScalarKind InScalarKind, int NumColumns, int NumRows);
 	static FArithmeticTypePtr Get(EScalarKind InScalarKind, int NumRows, int NumColumns);
 
-	bool IsScalar() const { return NumRows == 1 && NumColumns == 1; }
+	int  GetNumComponents() const { return NumRows * NumColumns; }
+	bool IsScalar() const { return GetNumComponents() == 1; }
 	bool IsVector() const { return NumRows > 1 && NumColumns == 1; }
 	bool IsMatrix() const { return NumRows > 1 && NumColumns > 1; }
+	FArithmeticTypePtr ToScalar() const;
 };
 
 } // namespace UE::MIR
