When detaching actors from their external package for cooking, don't rename child actors as part of renaming other objects outered to the actor package.

diff --git a/Engine/Source/Runtime/Engine/Private/Level.cpp b/Engine/Source/Runtime/Engine/Private/Level.cpp
--- a/Engine/Source/Runtime/Engine/Private/Level.cpp
+++ b/Engine/Source/Runtime/Engine/Private/Level.cpp
@@ -4322,7 +4322,7 @@ void ULevel::DetachAttachAllActorsPackages(bool bReattach)
 					GetObjectsWithPackage(ActorExternalPackage, Objects, bIncludeNestedSubobjects, RF_NoFlags, EInternalObjectFlags::Garbage);
 					for (UObject* Object : Objects)
 					{
-						if (Object != Actor && Object->GetFName() != NAME_PackageMetaData)
+						if (!Cast<AActor>(Object) && Object->GetFName() != NAME_PackageMetaData)
 						{
 							// Move objects in the destination level package
 							Object->Rename(nullptr, LevelPackage);
