Fixed incorrect motion vectors when using translated worldspace in WPO and moving the camera. The current PreViewTranslation was applied to the last-frame transforms, where PrevPreViewTranslation should be used.

diff --git a/Engine/Source/Runtime/Core/Private/UObject/FortniteMainBranchObjectVersion.cpp b/Engine/Source/Runtime/Core/Private/UObject/FortniteMainBranchObjectVersion.cpp
--- a/Engine/Source/Runtime/Core/Private/UObject/FortniteMainBranchObjectVersion.cpp
+++ b/Engine/Source/Runtime/Core/Private/UObject/FortniteMainBranchObjectVersion.cpp
@@ -14,6 +14,6 @@ TMap<FGuid, FGuid> FFortniteMainBranchObjectVersion::GetSystemGuids()
 	SystemGuids.Add(DevGuids.Niagara_LatestScriptCompileVersion, FGuid("E486231DA967C64B95D4E73131056FA0"));
 	SystemGuids.Add(DevGuids.SkeletalMeshDerivedDataVersion, FGuid("B9AEAA3EE7AC4FCFB3860E8F3E06DD75"));
 	SystemGuids.Add(DevGuids.STATICMESH_DERIVEDDATA_VER, FGuid("A3C8F19042174D05A6AEB6E5D9A58E44"));
-	SystemGuids.Add(DevGuids.MaterialTranslationDDCVersion, FGuid("A80C07846638406DA411360524FEC16A"));
+	SystemGuids.Add(DevGuids.MaterialTranslationDDCVersion, FGuid("5F69E03D3B204DA38D8CF564CB24F23F"));
 	return SystemGuids;
 }
diff --git a/Engine/Source/Runtime/Engine/Private/Materials/HLSLMaterialTranslator.cpp b/Engine/Source/Runtime/Engine/Private/Materials/HLSLMaterialTranslator.cpp
--- a/Engine/Source/Runtime/Engine/Private/Materials/HLSLMaterialTranslator.cpp
+++ b/Engine/Source/Runtime/Engine/Private/Materials/HLSLMaterialTranslator.cpp
@@ -10614,11 +10614,18 @@ static FString MultiplyMatrix(const TCHAR* Vector, const TCHAR* Matrix, int AWCo
 	}
 }
 
-static FString MultiplyTranslatedMatrix(const TCHAR* Vector, const TCHAR* MatrixPreTranslation, int AWComponent)
+static FString MultiplyTranslatedMatrix(const TCHAR* Vector, const TCHAR* MatrixPreTranslation, int AWComponent, bool bCompilingPreviousFrame)
 {
 	if (AWComponent)
 	{
-		return FString::Printf(TEXT("mul(MaterialFloat4(%s, 1.0f), DFFastToTranslatedWorld(%s, ResolvedView.PreViewTranslation)).xyz"), Vector, MatrixPreTranslation);
+		if (bCompilingPreviousFrame)
+		{
+			return FString::Printf(TEXT("mul(MaterialFloat4(%s, 1.0f), DFFastToTranslatedWorld(%s, ResolvedView.PrevPreViewTranslation)).xyz"), Vector, MatrixPreTranslation);
+		}
+		else
+		{
+			return FString::Printf(TEXT("mul(MaterialFloat4(%s, 1.0f), DFFastToTranslatedWorld(%s, ResolvedView.PreViewTranslation)).xyz"), Vector, MatrixPreTranslation);
+		}
 	}
 	else
 	{
@@ -10728,8 +10735,8 @@ int32 FHLSLMaterialTranslator::TransformBase(EMaterialCommonBasis SourceCoordBas
 			{
 				if (bIsPositionTransform)
 				{
-					CodeStr = MultiplyTranslatedMatrix(TEXT("<A>"), TEXT("Get<PREV>LocalToWorldDF(Parameters)"), AWComponent);
-					CodeDerivStr = MultiplyTranslatedMatrix(TEXT("<A>"), TEXT("Get<PREV>LocalToWorldDF(Parameters)"), 0);
+					CodeStr = MultiplyTranslatedMatrix(TEXT("<A>"), TEXT("Get<PREV>LocalToWorldDF(Parameters)"), AWComponent, bCompilingPreviousFrame);
+					CodeDerivStr = MultiplyTranslatedMatrix(TEXT("<A>"), TEXT("Get<PREV>LocalToWorldDF(Parameters)"), 0, bCompilingPreviousFrame);
 				}
 			}
 			// else use MCB_World as intermediary basis
@@ -10778,19 +10785,19 @@ int32 FHLSLMaterialTranslator::TransformBase(EMaterialCommonBasis SourceCoordBas
 				// TODO: inconsistent with TransformLocal<TO>World with instancing
 				// We have explicit options for "local" and "instance" spaces, but then GetLocalToWorld returns instance space, while GetWorldToLocal always returns primitive space. 
 				// It's inconsistent, but replacing either will break existing materials.
-				CodeStr = MultiplyTranslatedMatrix(TEXT("<A>"), TEXT("Get<PREV>WorldToLocalDF(Parameters)"), AWComponent);
-				CodeDerivStr = MultiplyTranslatedMatrix(TEXT("<A>"), TEXT("Get<PREV>WorldToLocalDF(Parameters)"), 0);
+				CodeStr = MultiplyTranslatedMatrix(TEXT("<A>"), TEXT("Get<PREV>WorldToLocalDF(Parameters)"), AWComponent, bCompilingPreviousFrame);
+				CodeDerivStr = MultiplyTranslatedMatrix(TEXT("<A>"), TEXT("Get<PREV>WorldToLocalDF(Parameters)"), 0, bCompilingPreviousFrame);
 			}
 			else if (DestCoordBasis == MCB_MeshParticle)
 			{
-				CodeStr = MultiplyTranslatedMatrix(TEXT("<A>"), TEXT("Parameters.Particle.WorldToParticle"), AWComponent);
-				CodeDerivStr = MultiplyTranslatedMatrix(TEXT("<A>"), TEXT("Parameters.Particle.WorldToParticle"), 0);
+				CodeStr = MultiplyTranslatedMatrix(TEXT("<A>"), TEXT("Parameters.Particle.WorldToParticle"), AWComponent, bCompilingPreviousFrame);
+				CodeDerivStr = MultiplyTranslatedMatrix(TEXT("<A>"), TEXT("Parameters.Particle.WorldToParticle"), 0, bCompilingPreviousFrame);
 				bUsesParticleWorldToLocal = true;
 			}
 			else if (DestCoordBasis == MCB_Instance)
 			{
-				CodeStr = MultiplyTranslatedMatrix(TEXT("<A>"), TEXT("GetWorldToInstanceDF(Parameters)"), AWComponent);
-				CodeDerivStr = MultiplyTranslatedMatrix(TEXT("<A>"), TEXT("GetWorldToInstanceDF(Parameters)"), 0);
+				CodeStr = MultiplyTranslatedMatrix(TEXT("<A>"), TEXT("GetWorldToInstanceDF(Parameters)"), AWComponent, bCompilingPreviousFrame);
+				CodeDerivStr = MultiplyTranslatedMatrix(TEXT("<A>"), TEXT("GetWorldToInstanceDF(Parameters)"), 0, bCompilingPreviousFrame);
 				bUsesInstanceWorldToLocalPS |= ShaderFrequency == SF_Pixel;
 			}
 			// else use MCB_World as intermediary basis
@@ -10884,8 +10891,8 @@ int32 FHLSLMaterialTranslator::TransformBase(EMaterialCommonBasis SourceCoordBas
 			}
 			else if (DestCoordBasis == MCB_TranslatedWorld)
 			{
-				CodeStr = MultiplyTranslatedMatrix(TEXT("<A>"), TEXT("Parameters.Particle.ParticleToWorld"), AWComponent);
-				CodeDerivStr = MultiplyTranslatedMatrix(TEXT("<A>"), TEXT("Parameters.Particle.ParticleToWorld"), 0);
+				CodeStr = MultiplyTranslatedMatrix(TEXT("<A>"), TEXT("Parameters.Particle.ParticleToWorld"), AWComponent, bCompilingPreviousFrame);
+				CodeDerivStr = MultiplyTranslatedMatrix(TEXT("<A>"), TEXT("Parameters.Particle.ParticleToWorld"), 0, bCompilingPreviousFrame);
 				bUsesParticleLocalToWorld = true;
 			}
 			// use World as an intermediary base
@@ -10902,8 +10909,8 @@ int32 FHLSLMaterialTranslator::TransformBase(EMaterialCommonBasis SourceCoordBas
 			}
 			else if (DestCoordBasis == MCB_TranslatedWorld)
 			{
-				CodeStr = MultiplyTranslatedMatrix(TEXT("<A>"), TEXT("Get<PREV>InstanceToWorldDF(Parameters)"), AWComponent);
-				CodeDerivStr = MultiplyTranslatedMatrix(TEXT("<A>"), TEXT("Get<PREV>InstanceToWorldDF(Parameters)"), 0);
+				CodeStr = MultiplyTranslatedMatrix(TEXT("<A>"), TEXT("Get<PREV>InstanceToWorldDF(Parameters)"), AWComponent, bCompilingPreviousFrame);
+				CodeDerivStr = MultiplyTranslatedMatrix(TEXT("<A>"), TEXT("Get<PREV>InstanceToWorldDF(Parameters)"), 0, bCompilingPreviousFrame);
 				bUsesInstanceLocalToWorldPS |= ShaderFrequency == SF_Pixel;
 			}
 			// use World as an intermediary base
