Fix the physics issue when creating physics state in a transaction.

#rnx

diff --git a/Engine/Source/Programs/AutoRTFMEngineTests/Classes/AutoRTFMTestEngine.h b/Engine/Source/Programs/AutoRTFMEngineTests/Classes/AutoRTFMTestEngine.h
new file mode 100644
--- /dev/null
+++ b/Engine/Source/Programs/AutoRTFMEngineTests/Classes/AutoRTFMTestEngine.h
@@ -0,0 +1,17 @@
+// Copyright Epic Games, Inc. All Rights Reserved.
+
+#pragma once
+
+#include "Engine.h"
+#include "AutoRTFMTestEngine.generated.h"
+
+UCLASS()
+class UAutoRTFMTestEngine : public UEngine
+{
+    GENERATED_BODY()
+
+public:
+    int Value = 42;
+
+	void Tick(float, bool) override {}
+};
diff --git a/Engine/Source/Programs/AutoRTFMEngineTests/Private/AutoRTFMEngineTests.cpp b/Engine/Source/Programs/AutoRTFMEngineTests/Private/AutoRTFMEngineTests.cpp
--- a/Engine/Source/Programs/AutoRTFMEngineTests/Private/AutoRTFMEngineTests.cpp
+++ b/Engine/Source/Programs/AutoRTFMEngineTests/Private/AutoRTFMEngineTests.cpp
@@ -7,6 +7,7 @@
 #include "NullTestRunner.h"
 #include "Materials/Material.h"
 #include "AutoRTFM/AutoRTFM.h"
+#include "AutoRTFMTestEngine.h"
 
 #if WITH_AUTOMATION_WORKER
 namespace UE::AutoRTFM
@@ -57,6 +58,9 @@ INT32_MAIN_INT32_ARGC_TCHAR_ARGV()
 	PreInit();
 	LoadModules();
 	PostInit();
+	
+	GEngine = NewObject<UAutoRTFMTestEngine>(GetTransientPackage());
+	GEngine->DefaultPhysMaterial = NewObject<UPhysicalMaterial>();
 
 	UE_LOG(LogAutoRTFMEngineTests, Display, TEXT("AutoRTFMEngineTests"));
 
diff --git a/Engine/Source/Runtime/Core/Public/AutoRTFM/AutoRTFM.h b/Engine/Source/Runtime/Core/Public/AutoRTFM/AutoRTFM.h
--- a/Engine/Source/Runtime/Core/Public/AutoRTFM/AutoRTFM.h
+++ b/Engine/Source/Runtime/Core/Public/AutoRTFM/AutoRTFM.h
@@ -965,10 +965,10 @@ namespace AutoRTFM::Private
 #endif
 
 // Older macros where the code is passed as a macro argument. These should be phased out as they make debugging more difficult
-#define UE_AUTORTFM_OPEN_IMPL(...) AutoRTFM::Open([&]() { __VA_ARGS__ })
-#define UE_AUTORTFM_ONABORT_IMPL(...) UE_AUTORTFM_BEGIN_DISABLE_WARNINGS AutoRTFM::OnAbort([=]() { __VA_ARGS__ }) UE_AUTORTFM_END_DISABLE_WARNINGS
-#define UE_AUTORTFM_ONCOMMIT_IMPL(...) UE_AUTORTFM_BEGIN_DISABLE_WARNINGS AutoRTFM::OnCommit([=]() { __VA_ARGS__ }) UE_AUTORTFM_END_DISABLE_WARNINGS
-#define UE_AUTORTFM_TRANSACT_IMPL(...) AutoRTFM::Transact([&]() { __VA_ARGS__ })
+#define UE_AUTORTFM_OPEN_IMPL(...) ::AutoRTFM::Open([&]() { __VA_ARGS__ })
+#define UE_AUTORTFM_ONABORT_IMPL(...) UE_AUTORTFM_BEGIN_DISABLE_WARNINGS ::AutoRTFM::OnAbort([=]() { __VA_ARGS__ }) UE_AUTORTFM_END_DISABLE_WARNINGS
+#define UE_AUTORTFM_ONCOMMIT_IMPL(...) UE_AUTORTFM_BEGIN_DISABLE_WARNINGS ::AutoRTFM::OnCommit([=]() { __VA_ARGS__ }) UE_AUTORTFM_END_DISABLE_WARNINGS
+#define UE_AUTORTFM_TRANSACT_IMPL(...) ::AutoRTFM::Transact([&]() { __VA_ARGS__ })
 
 #define UE_AUTORTFM_OPEN_IMPL2 AutoRTFM::Private::FOpenHelper{} + [&]()
 #define UE_AUTORTFM_ONABORT_IMPL2(...) AutoRTFM::Private::FOnAbortHelper{} + [__VA_ARGS__]()
diff --git a/Engine/Source/Runtime/CoreUObject/Public/UObject/DynamicallyTypedValue.h b/Engine/Source/Runtime/CoreUObject/Public/UObject/DynamicallyTypedValue.h
--- a/Engine/Source/Runtime/CoreUObject/Public/UObject/DynamicallyTypedValue.h
+++ b/Engine/Source/Runtime/CoreUObject/Public/UObject/DynamicallyTypedValue.h
@@ -10,6 +10,7 @@
 #include "Templates/UnrealTemplate.h"
 #include "UObject/Object.h"
 #include "UObject/UObjectGlobals.h"
+#include "AutoRTFM/AutoRTFM.h"
 
 class FReferenceCollector;
 
diff --git a/Engine/Source/Runtime/Engine/Private/PhysicsEngine/PhysDrawing.cpp b/Engine/Source/Runtime/Engine/Private/PhysicsEngine/PhysDrawing.cpp
--- a/Engine/Source/Runtime/Engine/Private/PhysicsEngine/PhysDrawing.cpp
+++ b/Engine/Source/Runtime/Engine/Private/PhysicsEngine/PhysDrawing.cpp
@@ -1236,41 +1236,58 @@ void FKAggregateGeom::GetAggGeom(const FTransform& Transform, const FColor Color
 /** Release the RenderInfo (if its there) and safely clean up any resources. Can be called from any thread. */
 void FKAggregateGeom::FreeRenderInfo()
 {
-	// See if we have rendering resources to free
-	auto RenderInfo = RenderInfoPtr.exchange(nullptr, std::memory_order_acq_rel);
+	// See if we have rendering resources to free. We do this in the open (non-transactionally) because of
+	// the threaded nature of this operation.
+	FKConvexGeomRenderInfo* RenderInfo = nullptr;
+
+	AutoRTFM::Open([&]
+		{
+			RenderInfo = RenderInfoPtr.exchange(nullptr, std::memory_order_acq_rel);
+		});
+
 	if (RenderInfo)
 	{
-		// Should always have these if RenderInfo exists
-		check(RenderInfo->VertexBuffers);
-		check(RenderInfo->IndexBuffer);
-
-		// Fire off a render command to free these resources
-		ENQUEUE_RENDER_COMMAND(FKAggregateGeomFreeRenderInfo)(
-			[RenderInfoToRelease = RenderInfo](FRHICommandList& RHICmdList)
+		AutoRTFM::OnCommit([this, RenderInfo]
 			{
-				RenderInfoToRelease->VertexBuffers->ColorVertexBuffer.ReleaseResource();
-				RenderInfoToRelease->VertexBuffers->StaticMeshVertexBuffer.ReleaseResource();
-				RenderInfoToRelease->VertexBuffers->PositionVertexBuffer.ReleaseResource();
-				RenderInfoToRelease->IndexBuffer->ReleaseResource();
+				// Should always have these if RenderInfo exists
+				check(nullptr != RenderInfo);
+				check(RenderInfo->VertexBuffers);
+				check(RenderInfo->IndexBuffer);
+
+				// Fire off a render command to free these resources
+				ENQUEUE_RENDER_COMMAND(FKAggregateGeomFreeRenderInfo)(
+					[RenderInfoToRelease = RenderInfo](FRHICommandList& RHICmdList)
+					{
+						RenderInfoToRelease->VertexBuffers->ColorVertexBuffer.ReleaseResource();
+						RenderInfoToRelease->VertexBuffers->StaticMeshVertexBuffer.ReleaseResource();
+						RenderInfoToRelease->VertexBuffers->PositionVertexBuffer.ReleaseResource();
+						RenderInfoToRelease->IndexBuffer->ReleaseResource();
 
-				// May not exist if no geometry was available
-				if (RenderInfoToRelease->CollisionVertexFactory != nullptr)
-				{
-					RenderInfoToRelease->CollisionVertexFactory->ReleaseResource();
-				}
+						// May not exist if no geometry was available
+						if (RenderInfoToRelease->CollisionVertexFactory != nullptr)
+						{
+							RenderInfoToRelease->CollisionVertexFactory->ReleaseResource();
+						}
 
-				// Free memory.
-				delete RenderInfoToRelease->VertexBuffers;
-				delete RenderInfoToRelease->IndexBuffer;
+						// Free memory.
+						delete RenderInfoToRelease->VertexBuffers;
+						delete RenderInfoToRelease->IndexBuffer;
 
-				if (RenderInfoToRelease->CollisionVertexFactory != nullptr)
-				{
-					delete RenderInfoToRelease->CollisionVertexFactory;
-				}
+						if (RenderInfoToRelease->CollisionVertexFactory != nullptr)
+						{
+							delete RenderInfoToRelease->CollisionVertexFactory;
+						}
 
-				delete RenderInfoToRelease;
-			}
-		);
+						delete RenderInfoToRelease;
+					});
+			});
+
+		AutoRTFM::OnAbort([this, RenderInfo]
+			{
+				check(nullptr != RenderInfo);
+				FKConvexGeomRenderInfo* OldRenderInfo = RenderInfoPtr.exchange(RenderInfo, std::memory_order_acq_rel);
+				ensure(nullptr == OldRenderInfo);
+			});
 	}
 }
 
diff --git a/Engine/Source/Runtime/Engine/Private/Tests/AutoRTFM/AutoRTFMActorComponentTests.cpp b/Engine/Source/Runtime/Engine/Private/Tests/AutoRTFM/AutoRTFMActorComponentTests.cpp
--- a/Engine/Source/Runtime/Engine/Private/Tests/AutoRTFM/AutoRTFMActorComponentTests.cpp
+++ b/Engine/Source/Runtime/Engine/Private/Tests/AutoRTFM/AutoRTFMActorComponentTests.cpp
@@ -4,6 +4,7 @@
 #include "Misc/App.h"
 #include "AutoRTFM/AutoRTFM.h"
 #include "AutoRTFMTestActor.h"
+#include "AutoRTFMTestBodySetup.h"
 #include "AutoRTFMTestLevel.h"
 #include "AutoRTFMTestPrimitiveComponent.h"
 #include "PhysicsProxy/SingleParticlePhysicsProxy.h"
@@ -102,6 +103,30 @@ bool FAutoRTFMActorComponentTests::RunTest(const FString & Parameters)
 	TEST_CHECK_TRUE(AutoRTFM::ETransactionResult::Committed == Result);
 	TEST_CHECK_TRUE(!Component->IsRegistered());
 
+	// This test requires us to have a fresh body instance so that it has to be created during the register.
+	Component->BodyInstance = FBodyInstance();
+
+	// And we need a valid body setup so that there is shapes created too.
+	Component->BodySetup = NewObject<UAutoRTFMTestBodySetup>();
+	Component->BodySetup->AggGeom.SphereElems.Add(FKSphereElem(1.0f));
+
+	Result = AutoRTFM::Transact([&]
+		{
+			Component->RegisterComponentWithWorld(World);
+			AutoRTFM::AbortTransaction();
+		});
+
+	TEST_CHECK_TRUE(AutoRTFM::ETransactionResult::AbortedByRequest == Result);
+	TEST_CHECK_TRUE(!Component->IsRegistered());
+
+	Result = AutoRTFM::Transact([&]
+		{
+			Component->RegisterComponentWithWorld(World);
+		});
+
+	TEST_CHECK_TRUE(AutoRTFM::ETransactionResult::Committed == Result);
+	TEST_CHECK_TRUE(Component->IsRegistered());
+
 	return true;
 }
 
diff --git a/Engine/Source/Runtime/Engine/Private/Tests/AutoRTFM/AutoRTFMTestBodySetup.h b/Engine/Source/Runtime/Engine/Private/Tests/AutoRTFM/AutoRTFMTestBodySetup.h
new file mode 100644
--- /dev/null
+++ b/Engine/Source/Runtime/Engine/Private/Tests/AutoRTFM/AutoRTFMTestBodySetup.h
@@ -0,0 +1,15 @@
+// Copyright Epic Games, Inc. All Rights Reserved.
+
+#pragma once
+
+#include "PhysicsEngine/BodySetup.h"
+#include "AutoRTFMTestBodySetup.generated.h"
+
+UCLASS()
+class UAutoRTFMTestBodySetup : public UBodySetup
+{
+    GENERATED_BODY()
+
+public:
+    int Value = 42;
+};
diff --git a/Engine/Source/Runtime/Engine/Private/Tests/AutoRTFM/AutoRTFMTestPrimitiveComponent.h b/Engine/Source/Runtime/Engine/Private/Tests/AutoRTFM/AutoRTFMTestPrimitiveComponent.h
--- a/Engine/Source/Runtime/Engine/Private/Tests/AutoRTFM/AutoRTFMTestPrimitiveComponent.h
+++ b/Engine/Source/Runtime/Engine/Private/Tests/AutoRTFM/AutoRTFMTestPrimitiveComponent.h
@@ -12,4 +12,8 @@ class UAutoRTFMTestPrimitiveComponent : public UPrimitiveComponent
 
 public:
     int Value = 42;
+
+	UBodySetup* BodySetup = nullptr;
+
+	UBodySetup* GetBodySetup() override { return BodySetup; }
 };
diff --git a/Engine/Source/Runtime/Experimental/Chaos/Public/Chaos/PBDRigidsEvolution.h b/Engine/Source/Runtime/Experimental/Chaos/Public/Chaos/PBDRigidsEvolution.h
--- a/Engine/Source/Runtime/Experimental/Chaos/Public/Chaos/PBDRigidsEvolution.h
+++ b/Engine/Source/Runtime/Experimental/Chaos/Public/Chaos/PBDRigidsEvolution.h
@@ -774,8 +774,16 @@ public:
 
 	FUniqueIdx GenerateUniqueIdx()
 	{
-		//NOTE: this should be thread safe since evolution has already been initialized on GT
-		return Particles.GetUniqueIndices().GenerateUniqueIdx();
+		// NOTE: this should be thread safe since evolution has already been initialized on GT
+		FUniqueIdx Result;
+		
+		// This uses thread safe atomics, so we have to do it in the open.
+		UE_AUTORTFM_OPEN({ Result = Particles.GetUniqueIndices().GenerateUniqueIdx(); });
+
+		// But if we abort, we need to release the now unused index.
+		AutoRTFM::OnAbort([this, Result] { this->ReleaseUniqueIdx(Result); });
+
+		return Result;
 	}
 
 	void ReleaseUniqueIdx(FUniqueIdx UniqueIdx)
