World Partition - Avoid calls to CanChangeIsSpatiallyLoadedFlag() in AActor::Serialize()

diff --git a/Engine/Source/Runtime/Engine/Private/Actor.cpp b/Engine/Source/Runtime/Engine/Private/Actor.cpp
--- a/Engine/Source/Runtime/Engine/Private/Actor.cpp
+++ b/Engine/Source/Runtime/Engine/Private/Actor.cpp
@@ -913,9 +913,12 @@ void AActor::Serialize(FArchive& Ar)
 			ActorGuid = FGuid::NewGuid();
 		}
 
-		if (!CanChangeIsSpatiallyLoadedFlag() && (Ar.CustomVer(FUE5ReleaseStreamObjectVersion::GUID) < FUE5ReleaseStreamObjectVersion::ActorGridPlacementDeprecateDefaultValueFixup))
+		if (Ar.CustomVer(FUE5ReleaseStreamObjectVersion::GUID) < FUE5ReleaseStreamObjectVersion::ActorGridPlacementDeprecateDefaultValueFixup)
 		{
-			bIsSpatiallyLoaded = GetClass()->GetDefaultObject<AActor>()->bIsSpatiallyLoaded;
+			if (!CanChangeIsSpatiallyLoadedFlag())
+			{
+				bIsSpatiallyLoaded = GetClass()->GetDefaultObject<AActor>()->bIsSpatiallyLoaded;
+			}
 		}
 
 		if (Ar.CustomVer(FFortniteMainBranchObjectVersion::GUID) < FFortniteMainBranchObjectVersion::WorldPartitionActorDescSerializeContentBundleGuid)
