CSV Tracking for:
- Average Jitter (milliseconds)
- Packet Loss Percentage (In/Out)
#tests Tested in a local test build. Gathered local csv stats and generated local client PRS reports

diff --git a/Engine/Config/BaseEngine.ini b/Engine/Config/BaseEngine.ini
--- a/Engine/Config/BaseEngine.ini
+++ b/Engine/Config/BaseEngine.ini
@@ -3630,6 +3630,9 @@ IntervalSeconds=1
 +Listeners=(MetricName=Ping, Class=/Script/Engine.NetworkMetricsCSV_Replication)
 +Listeners=(MetricName=InRate, Class=/Script/Engine.NetworkMetricsCSV_Replication)
 +Listeners=(MetricName=OutRate, Class=/Script/Engine.NetworkMetricsCSV_Replication)
++Listeners=(MetricName=AverageJitterInMS, Class=/Script/Engine.NetworkMetricsCSV_Replication)
++Listeners=(MetricName=InPacketsLost, Class=/Script/Engine.NetworkMetricsCSV_Replication)
++Listeners=(MetricName=OutPacketsLost, Class=/Script/Engine.NetworkMetricsCSV_Replication)
 +Listeners=(MetricName=NumConnections, Class=/Script/Engine.NetworkMetricsPerfCounters)
 +Listeners=(MetricName=AvgPing, Class=/Script/Engine.NetworkMetricsPerfCounters)
 +Listeners=(MetricName=InBunches, Class=/Script/Engine.NetworkMetricsPerfCounters)
diff --git a/Engine/Source/Runtime/Engine/Private/Net/NetworkMetricsDefs.cpp b/Engine/Source/Runtime/Engine/Private/Net/NetworkMetricsDefs.cpp
--- a/Engine/Source/Runtime/Engine/Private/Net/NetworkMetricsDefs.cpp
+++ b/Engine/Source/Runtime/Engine/Private/Net/NetworkMetricsDefs.cpp
@@ -33,6 +33,7 @@ const FName OutRateClientMax("OutRateClientMax");
 const FName OutRateClientMin("OutRateClientMin");
 const FName InPacketsLost("InPacketsLost");
 const FName OutPacketsLost("OutPacketsLost");
+const FName AverageJitterInMS("AverageJitterInMS");
 const FName InBunches("InBunches");
 const FName OutBunches("OutBunches");
 const FName Ping("Ping");
diff --git a/Engine/Source/Runtime/Engine/Private/NetDriver.cpp b/Engine/Source/Runtime/Engine/Private/NetDriver.cpp
--- a/Engine/Source/Runtime/Engine/Private/NetDriver.cpp
+++ b/Engine/Source/Runtime/Engine/Private/NetDriver.cpp
@@ -1684,6 +1684,9 @@ void UNetDriver::SetupNetworkMetrics()
 	GetMetrics()->CreateInt(UE::Net::Metric::InPacketsLost, 0);
 	GetMetrics()->CreateInt(UE::Net::Metric::OutPacketsLost, 0);
 
+	// (Client only) Current jitter for this connection in milliseconds
+	GetMetrics()->CreateFloat(UE::Net::Metric::AverageJitterInMS, 0);
+
 	// The number of incoming bunches per second across all connections.
 	GetMetrics()->CreateInt(UE::Net::Metric::InBunches, 0);
 
@@ -8120,6 +8123,7 @@ void UNetDriver::UpdateNetworkStats()
 				if (ServerConnection != nullptr)
 				{
 					NumOpenChannels = ServerConnection->OpenChannels.Num();
+					GetMetrics()->SetFloat(UE::Net::Metric::AverageJitterInMS, ServerConnection->GetAverageJitterInMS());
 				}
 
 				for (int32 i = 0; i < ClientConnections.Num(); i++)
diff --git a/Engine/Source/Runtime/Engine/Public/Net/NetworkMetricsDefs.h b/Engine/Source/Runtime/Engine/Public/Net/NetworkMetricsDefs.h
--- a/Engine/Source/Runtime/Engine/Public/Net/NetworkMetricsDefs.h
+++ b/Engine/Source/Runtime/Engine/Public/Net/NetworkMetricsDefs.h
@@ -35,6 +35,7 @@ extern ENGINE_API const FName OutRateClientMax;
 extern ENGINE_API const FName OutRateClientMin;
 extern ENGINE_API const FName InPacketsLost;
 extern ENGINE_API const FName OutPacketsLost;
+extern ENGINE_API const FName AverageJitterInMS;
 extern ENGINE_API const FName InBunches;
 extern ENGINE_API const FName OutBunches;
 extern ENGINE_API const FName Ping;
