Correctly calculate the numebr of elements in texcoord buffers

#tests ab tested using a replay

diff --git a/Engine/Source/Runtime/Engine/Private/Rendering/StaticMeshVertexBuffer.cpp b/Engine/Source/Runtime/Engine/Private/Rendering/StaticMeshVertexBuffer.cpp
--- a/Engine/Source/Runtime/Engine/Private/Rendering/StaticMeshVertexBuffer.cpp
+++ b/Engine/Source/Runtime/Engine/Private/Rendering/StaticMeshVertexBuffer.cpp
@@ -274,7 +274,7 @@ FBufferRHIRef FStaticMeshVertexBuffer::CreateTangentsRHIBuffer(FRHICommandListBa
 
 FBufferRHIRef FStaticMeshVertexBuffer::CreateTexCoordRHIBuffer(FRHICommandListBase& RHICmdList)
 {
-	return CreateRHIBuffer(RHICmdList, TexcoordData, GetNumTexCoords(), BUF_Static | BUF_ShaderResource, TEXT("TexCoordRHIBuffer"));
+	return CreateRHIBuffer(RHICmdList, TexcoordData, GetNumVertices() * GetNumTexCoords(), BUF_Static | BUF_ShaderResource, TEXT("TexCoordRHIBuffer"));
 }
 
 void FStaticMeshVertexBuffer::InitRHIForStreaming(
