Fix another physics bug under transactions with respect to welding.

diff --git a/Engine/Source/Runtime/Engine/Private/PrimitiveComponentPhysics.cpp b/Engine/Source/Runtime/Engine/Private/PrimitiveComponentPhysics.cpp
--- a/Engine/Source/Runtime/Engine/Private/PrimitiveComponentPhysics.cpp
+++ b/Engine/Source/Runtime/Engine/Private/PrimitiveComponentPhysics.cpp
@@ -969,7 +969,18 @@ bool UPrimitiveComponent::WeldToImplementation(USceneComponent * InParent, FName
 			//if root is kinematic simply set child to be kinematic and we're done
 			if ((RootComponent->IsSimulatingPhysics(SocketName) == false) && (bWeldToKinematicParent == false))
 			{
-				FPlatformAtomics::InterlockedExchangePtr((void**)&BI->WeldParent, nullptr);
+				void* OriginalWeldParent = nullptr;
+
+				// Because this uses atomics we need to run it non-transactionally.
+				UE_AUTORTFM_OPEN({ OriginalWeldParent = FPlatformAtomics::InterlockedExchangePtr((void**)&BI->WeldParent, nullptr); });
+
+				// But remember that on abort we need to fudge the pointer back into the place we took it from.
+				AutoRTFM::OnAbort([BI, OriginalWeldParent]
+					{
+						void* Old = FPlatformAtomics::InterlockedExchangePtr((void**)&BI->WeldParent, OriginalWeldParent);
+						ensure(nullptr == Old);
+					});
+				
 				SetSimulatePhysics(false);
 				return false;	//return false because we need to continue with regular body initialization
 			}
@@ -1025,8 +1036,20 @@ void UPrimitiveComponent::UnWeldFromParent()
 			bool bRootIsBeingDeleted = !IsValidChecked(RootComponent) || RootComponent->IsUnreachable();
 			const FBodyInstance* PrevWeldParent = NewRootBI->WeldParent;
 			RootBI->UnWeld(NewRootBI);
-			
-			FPlatformAtomics::InterlockedExchangePtr((void**)&NewRootBI->WeldParent, nullptr);
+
+			{
+				void* OriginalWeldParent = nullptr;
+
+				// Because this uses atomics we need to run it non-transactionally.
+				UE_AUTORTFM_OPEN({ OriginalWeldParent = FPlatformAtomics::InterlockedExchangePtr((void**)&NewRootBI->WeldParent, nullptr); });
+
+				// But remember that on abort we need to fudge the pointer back into the place we took it from.
+				AutoRTFM::OnAbort([NewRootBI, OriginalWeldParent]
+					{
+						void* Old = FPlatformAtomics::InterlockedExchangePtr((void**)&NewRootBI->WeldParent, OriginalWeldParent);
+						ensure(nullptr == Old);
+					});
+			}
 
 			bool bHasBodySetup = GetBodySetup() != nullptr;
 
@@ -1061,7 +1084,19 @@ void UPrimitiveComponent::UnWeldFromParent()
 					}
 
 					//At this point, NewRootBI must be kinematic because it's being unwelded.
-					FPlatformAtomics::InterlockedExchangePtr((void**)&ChildBI->WeldParent, nullptr); //null because we are currently kinematic
+
+					void* OriginalWeldParent = nullptr;
+
+					// Because this uses atomics we need to run it non-transactionally.
+					//null because we are currently kinematic
+					UE_AUTORTFM_OPEN({ OriginalWeldParent = FPlatformAtomics::InterlockedExchangePtr((void**)&ChildBI->WeldParent, nullptr); });
+
+					// But remember that on abort we need to fudge the pointer back into the place we took it from.
+					AutoRTFM::OnAbort([ChildBI, OriginalWeldParent]
+						{
+							void* Old = FPlatformAtomics::InterlockedExchangePtr((void**)&ChildBI->WeldParent, OriginalWeldParent);
+							ensure(nullptr == Old);
+						});
 				}
 			}
 
diff --git a/Engine/Source/Runtime/Engine/Private/Tests/AutoRTFM/AutoRTFMActorComponentTests.cpp b/Engine/Source/Runtime/Engine/Private/Tests/AutoRTFM/AutoRTFMActorComponentTests.cpp
--- a/Engine/Source/Runtime/Engine/Private/Tests/AutoRTFM/AutoRTFMActorComponentTests.cpp
+++ b/Engine/Source/Runtime/Engine/Private/Tests/AutoRTFM/AutoRTFMActorComponentTests.cpp
@@ -127,6 +127,78 @@ bool FAutoRTFMActorComponentTests::RunTest(const FString & Parameters)
 	TEST_CHECK_TRUE(AutoRTFM::ETransactionResult::Committed == Result);
 	TEST_CHECK_TRUE(Component->IsRegistered());
 
+	Component->UnregisterComponent();
+	TEST_CHECK_TRUE(!Component->IsRegistered());
+
+	FBodyInstance SomeInstance;
+
+	// This test requires us to have a fresh body instance so that it has to be created during the register.
+	Component->BodyInstance = FBodyInstance();
+	Component->BodyInstance.bSimulatePhysics = 1;
+	Component->BodyInstance.WeldParent = &SomeInstance;
+	TEST_CHECK_TRUE(Component->IsWelded());
+
+	UAutoRTFMTestBodySetup* BodySetup = NewObject<UAutoRTFMTestBodySetup>();
+	BodySetup->AggGeom.SphereElems.Add(FKSphereElem(1.0f));
+
+	Component->BodyInstance.BodySetup = BodySetup;
+
+	UAutoRTFMTestPrimitiveComponent* Parent0 = NewObject<UAutoRTFMTestPrimitiveComponent>(Actor);
+	UAutoRTFMTestPrimitiveComponent* Parent1 = NewObject<UAutoRTFMTestPrimitiveComponent>(Actor);
+
+	Result = AutoRTFM::Transact([&]
+		{
+			Component->WeldTo(Parent0);
+			AutoRTFM::AbortTransaction();
+		});
+
+	TEST_CHECK_TRUE(AutoRTFM::ETransactionResult::AbortedByRequest == Result);
+	TEST_CHECK_TRUE(Component->IsWelded());
+	TEST_CHECK_TRUE(&SomeInstance == Component->BodyInstance.WeldParent);
+
+	Result = AutoRTFM::Transact([&]
+		{
+			Component->WeldTo(Parent0);
+		});
+
+	TEST_CHECK_TRUE(AutoRTFM::ETransactionResult::Committed == Result);
+	TEST_CHECK_TRUE(!Component->IsWelded());
+	TEST_CHECK_TRUE(nullptr == Component->BodyInstance.WeldParent);
+
+	Result = AutoRTFM::Transact([&]
+		{
+			Component->WeldTo(Parent1);
+			AutoRTFM::AbortTransaction();
+		});
+
+	TEST_CHECK_TRUE(AutoRTFM::ETransactionResult::AbortedByRequest == Result);
+	TEST_CHECK_TRUE(!Component->IsWelded());
+
+	Result = AutoRTFM::Transact([&]
+		{
+			Component->WeldTo(Parent1);
+		});
+
+	TEST_CHECK_TRUE(AutoRTFM::ETransactionResult::Committed == Result);
+	TEST_CHECK_TRUE(!Component->IsWelded());
+
+	Result = AutoRTFM::Transact([&]
+		{
+			Component->UnWeldFromParent();
+			AutoRTFM::AbortTransaction();
+		});
+
+	TEST_CHECK_TRUE(AutoRTFM::ETransactionResult::AbortedByRequest == Result);
+	TEST_CHECK_TRUE(!Component->IsWelded());
+
+	Result = AutoRTFM::Transact([&]
+		{
+			Component->UnWeldFromParent();
+		});
+
+	TEST_CHECK_TRUE(AutoRTFM::ETransactionResult::Committed == Result);
+	TEST_CHECK_TRUE(!Component->IsWelded());
+
 	return true;
 }
 
