Stop NetGUID messages when using recording replays

diff --git a/Engine/Source/Runtime/Engine/Private/PackageMapClient.cpp b/Engine/Source/Runtime/Engine/Private/PackageMapClient.cpp
--- a/Engine/Source/Runtime/Engine/Private/PackageMapClient.cpp
+++ b/Engine/Source/Runtime/Engine/Private/PackageMapClient.cpp
@@ -2867,9 +2867,18 @@ bool FNetGUIDCache::SupportsObject( const UObject* Object, const TWeakObjectPtr<
 		return true;
 	}
 
-	UE_LOG( LogNetPackageMap, Warning, TEXT( "FNetGUIDCache::SupportsObject: %s NOT Supported." ), *Object->GetFullName() );
-	//UE_LOG( LogNetPackageMap, Warning, TEXT( "   %s"), *DebugContextString );
-
+	// Do not display warning when recording replays 
+	// IMPORTANT : this is a workaround until until UE-169326 is properly fixed
+	bool bDisplayWarning = true;
+	if (UWorld* World = Object->GetWorld())
+	{
+		bDisplayWarning = !World->IsRecordingClientReplay();
+	}
+	if (bDisplayWarning)
+	{
+		UE_LOG(LogNetPackageMap, Warning, TEXT("FNetGUIDCache::SupportsObject: %s NOT Supported."), *Object->GetFullName());
+		//UE_LOG( LogNetPackageMap, Warning, TEXT( "   %s"), *DebugContextString );
+	}
 	return false;
 }
 
