Some changes to Trace services and controller api

* Guard for invalid message bus when creating endpoint.
* Allow ping request to target all services rather than specific instance. This will allow using TraceController independent of session manager.
* Added way of creating a TraceController using a custom message bus.
* Renamed `SendDiscoveryPing` -> `SendDiscoveryRequest` and `OnConnectRequest` -> `OnDiscoveryPing` to be more consistent with other method names.

diff --git a/Engine/Source/Runtime/Engine/Private/TraceService.cpp b/Engine/Source/Runtime/Engine/Private/TraceService.cpp
--- a/Engine/Source/Runtime/Engine/Private/TraceService.cpp
+++ b/Engine/Source/Runtime/Engine/Private/TraceService.cpp
@@ -13,13 +13,14 @@ class FTraceServiceImpl
 {
 public:
 	FTraceServiceImpl();
+	FTraceServiceImpl(const TSharedPtr<IMessageBus>&);
 	virtual ~FTraceServiceImpl() {};
 	
 private:
 	void OnStatusPing(const FTraceControlStatusPing& Message, const TSharedRef<IMessageContext>& Context);
 	void OnChannelsPing(const FTraceControlChannelsPing& Message, const TSharedRef<IMessageContext>& Context);
 	void OnSettingsPing(const FTraceControlSettingsPing& Message, const TSharedRef<IMessageContext, ESPMode::ThreadSafe>& Context);
-	void OnConnectRequest(const FTraceControlDiscoveryPing& Message, const TSharedRef<IMessageContext>& Context);
+	void OnDiscoveryPing(const FTraceControlDiscoveryPing& Message, const TSharedRef<IMessageContext>& Context);
 	void OnStop(const FTraceControlStop& Message, const TSharedRef<IMessageContext>& Context);
 	void OnSend(const FTraceControlSend& Message, const TSharedRef<IMessageContext>& Context);
 	void OnChannelSet(const FTraceControlChannelsSet& Message, const TSharedRef<IMessageContext>& Context);
@@ -41,48 +42,55 @@ private:
 
 
 FTraceServiceImpl::FTraceServiceImpl()
+	: FTraceServiceImpl(IMessagingModule::Get().GetDefaultBus())
+{
+}
+
+FTraceServiceImpl::FTraceServiceImpl(const TSharedPtr<IMessageBus>& InBus)
 {
 	SessionId = FApp::GetSessionId();
 	InstanceId = FApp::GetInstanceId();
-	
-	MessageEndpoint = FMessageEndpoint::Builder("FTraceService")
-		.Handling<FTraceControlDiscoveryPing>(this, &FTraceServiceImpl::OnConnectRequest)
-		.Handling<FTraceControlChannelsSet>(this, &FTraceServiceImpl::OnChannelSet)
-		.Handling<FTraceControlStop>(this, &FTraceServiceImpl::OnStop)
-		.Handling<FTraceControlSend>(this, &FTraceServiceImpl::OnSend)
-		.Handling<FTraceControlFile>(this, &FTraceServiceImpl::OnFile)
-		.Handling<FTraceControlSnapshotSend>(this, &FTraceServiceImpl::OnSnapshotSend)
-		.Handling<FTraceControlSnapshotFile>(this, &FTraceServiceImpl::OnSnapshotFile)
-		.Handling<FTraceControlPause>(this, &FTraceServiceImpl::OnPause)
-		.Handling<FTraceControlResume>(this, &FTraceServiceImpl::OnResume)
-		.Handling<FTraceControlBookmark>(this, &FTraceServiceImpl::OnBookmark)
-		.Handling<FTraceControlScreenshot>(this, &FTraceServiceImpl::OnScreenshot)
-		.Handling<FTraceControlSetStatNamedEvents>(this, &FTraceServiceImpl::OnSetStatNamedEvents)
-		.Handling<FTraceControlStatusPing>(this, &FTraceServiceImpl::OnStatusPing)
-		.Handling<FTraceControlSettingsPing>(this, &FTraceServiceImpl::OnSettingsPing)
-		.Handling<FTraceControlChannelsPing>(this, &FTraceServiceImpl::OnChannelsPing);
-
-	if (!MessageEndpoint.IsValid())
+
+	if (InBus.IsValid())
 	{
-		 return;
-	}
+		MessageEndpoint = FMessageEndpoint::Builder("FTraceService", InBus.ToSharedRef())
+			.Handling<FTraceControlDiscoveryPing>(this, &FTraceServiceImpl::OnDiscoveryPing)
+			.Handling<FTraceControlChannelsSet>(this, &FTraceServiceImpl::OnChannelSet)
+			.Handling<FTraceControlStop>(this, &FTraceServiceImpl::OnStop)
+			.Handling<FTraceControlSend>(this, &FTraceServiceImpl::OnSend)
+			.Handling<FTraceControlFile>(this, &FTraceServiceImpl::OnFile)
+			.Handling<FTraceControlSnapshotSend>(this, &FTraceServiceImpl::OnSnapshotSend)
+			.Handling<FTraceControlSnapshotFile>(this, &FTraceServiceImpl::OnSnapshotFile)
+			.Handling<FTraceControlPause>(this, &FTraceServiceImpl::OnPause)
+			.Handling<FTraceControlResume>(this, &FTraceServiceImpl::OnResume)
+			.Handling<FTraceControlBookmark>(this, &FTraceServiceImpl::OnBookmark)
+			.Handling<FTraceControlScreenshot>(this, &FTraceServiceImpl::OnScreenshot)
+			.Handling<FTraceControlSetStatNamedEvents>(this, &FTraceServiceImpl::OnSetStatNamedEvents)
+			.Handling<FTraceControlStatusPing>(this, &FTraceServiceImpl::OnStatusPing)
+			.Handling<FTraceControlSettingsPing>(this, &FTraceServiceImpl::OnSettingsPing)
+			.Handling<FTraceControlChannelsPing>(this, &FTraceServiceImpl::OnChannelsPing);
 
-	MessageEndpoint->Subscribe<FTraceControlStatusPing>();
-	MessageEndpoint->Subscribe<FTraceControlSettingsPing>();
-	MessageEndpoint->Subscribe<FTraceControlDiscoveryPing>();
-	MessageEndpoint->Subscribe<FTraceControlChannelsPing>();
-	MessageEndpoint->Subscribe<FTraceControlStop>();
-	MessageEndpoint->Subscribe<FTraceControlSend>();
-	MessageEndpoint->Subscribe<FTraceControlChannelsSet>();
-	MessageEndpoint->Subscribe<FTraceControlFile>();
-	MessageEndpoint->Subscribe<FTraceControlSnapshotSend>();
-	MessageEndpoint->Subscribe<FTraceControlSnapshotSend>();
-	MessageEndpoint->Subscribe<FTraceControlPause>();
-	MessageEndpoint->Subscribe<FTraceControlResume>();
-	MessageEndpoint->Subscribe<FTraceControlBookmark>();
-	MessageEndpoint->Subscribe<FTraceControlScreenshot>();
-	MessageEndpoint->Subscribe<FTraceControlSetStatNamedEvents>();
-	
+		if (!MessageEndpoint.IsValid())
+		{
+			return;
+		}
+
+		MessageEndpoint->Subscribe<FTraceControlStatusPing>();
+		MessageEndpoint->Subscribe<FTraceControlSettingsPing>();
+		MessageEndpoint->Subscribe<FTraceControlDiscoveryPing>();
+		MessageEndpoint->Subscribe<FTraceControlChannelsPing>();
+		MessageEndpoint->Subscribe<FTraceControlStop>();
+		MessageEndpoint->Subscribe<FTraceControlSend>();
+		MessageEndpoint->Subscribe<FTraceControlChannelsSet>();
+		MessageEndpoint->Subscribe<FTraceControlFile>();
+		MessageEndpoint->Subscribe<FTraceControlSnapshotSend>();
+		MessageEndpoint->Subscribe<FTraceControlSnapshotSend>();
+		MessageEndpoint->Subscribe<FTraceControlPause>();
+		MessageEndpoint->Subscribe<FTraceControlResume>();
+		MessageEndpoint->Subscribe<FTraceControlBookmark>();
+		MessageEndpoint->Subscribe<FTraceControlScreenshot>();
+		MessageEndpoint->Subscribe<FTraceControlSetStatNamedEvents>();
+	}
 }
 
 void FTraceServiceImpl::FillTraceStatusMessage(FTraceControlStatus* Message)
@@ -252,13 +260,14 @@ void FTraceServiceImpl::OnSettingsPing(const FTraceControlSettingsPing& Message,
 	MessageEndpoint->Send(Response, Context->GetSender());
 }
 
-void FTraceServiceImpl::OnConnectRequest(const FTraceControlDiscoveryPing& Message, const TSharedRef<IMessageContext>& Context)
+void FTraceServiceImpl::OnDiscoveryPing(const FTraceControlDiscoveryPing& Message, const TSharedRef<IMessageContext>& Context)
 {
-	if (Message.InstanceId == FApp::GetInstanceId() && Message.SessionId == FApp::GetSessionId())
+	if ((!Message.SessionId.IsValid() && !Message.InstanceId.IsValid()) || (Message.InstanceId == FApp::GetInstanceId() && Message.SessionId == FApp::GetSessionId()))
 	{
 		const auto Response = FMessageEndpoint::MakeMessage<FTraceControlDiscovery>();
 		Response->SessionId = FApp::GetSessionId();
 		Response->InstanceId = FApp::GetInstanceId();
+
 		FillTraceStatusMessage(Response);
 
 		MessageEndpoint->Send(
@@ -274,3 +283,7 @@ FTraceService::FTraceService()
 	Impl = MakePimpl<FTraceServiceImpl>();
 }
 
+FTraceService::FTraceService(TSharedPtr<IMessageBus> InBus)
+{
+	Impl = MakePimpl<FTraceServiceImpl>(InBus);
+}
diff --git a/Engine/Source/Runtime/Engine/Public/TraceService.h b/Engine/Source/Runtime/Engine/Public/TraceService.h
--- a/Engine/Source/Runtime/Engine/Public/TraceService.h
+++ b/Engine/Source/Runtime/Engine/Public/TraceService.h
@@ -1,13 +1,16 @@
 // Copyright Epic Games, Inc. All Rights Reserved.
 #pragma once
 
-#include <Templates/PimplPtr.h>
+#include "Templates/PimplPtr.h"
+#include "Templates/SharedPointer.h"
+
+class IMessageBus;
 
 class FTraceService
 {
 public:
 	ENGINE_API FTraceService();
-	virtual ~FTraceService() {}
+	ENGINE_API FTraceService(TSharedPtr<IMessageBus> Bus);
 	
 private:
 	TPimplPtr<class FTraceServiceImpl> Impl;
diff --git a/Engine/Source/Runtime/SessionServices/Private/TraceController.cpp b/Engine/Source/Runtime/SessionServices/Private/TraceController.cpp
--- a/Engine/Source/Runtime/SessionServices/Private/TraceController.cpp
+++ b/Engine/Source/Runtime/SessionServices/Private/TraceController.cpp
@@ -39,6 +39,18 @@ FTraceController::~FTraceController()
 	}
 }
 
+void FTraceController::SendDiscoveryRequest(const FGuid& SessionId, const FGuid& InstanceId) const
+{
+	const auto Message = FMessageEndpoint::MakeMessage<FTraceControlDiscoveryPing>();
+	Message->SessionId = SessionId;
+	Message->InstanceId = InstanceId;
+	MessageEndpoint->Publish<FTraceControlDiscoveryPing>(Message);
+}
+void FTraceController::SendDiscoveryRequest()
+{
+	const auto Message = FMessageEndpoint::MakeMessage<FTraceControlDiscoveryPing>();
+	MessageEndpoint->Publish<FTraceControlDiscoveryPing>(Message);
+}
 
 void FTraceController::SendStatusUpdateRequest()
 {
@@ -253,7 +265,7 @@ void FTraceController::OnInstanceSelectionChanged(const TSharedPtr<ISessionInsta
 	}
 	else
 	{
-		SendDiscoveryPing(InstanceInfo);
+		SendDiscoveryRequest(InstanceInfo->GetOwnerSession()->GetSessionId(), InstanceInfo->GetInstanceId());
 	}
 
 	SessionSelectionChangedEvent.Broadcast();
@@ -308,26 +320,20 @@ bool FTraceController::RediscoverSelectedSession() const
 	{
 		if (!InstanceToAddress.Contains(SelectedInstance->GetInstanceId()))
 		{
-			SendDiscoveryPing(SelectedInstance);
+			SendDiscoveryRequest(SelectedInstance->GetOwnerSession()->GetSessionId(), SelectedInstance->GetInstanceId());
 			return true;
 		}
 	}
 	return false;
 }
 
-void FTraceController::SendDiscoveryPing(const TSharedPtr<ISessionInstanceInfo>& Instance) const
-{
-	// We haven't discovered this instance yet, send a discovery message specifically
-	// to that instance.
-	const auto Message = FMessageEndpoint::MakeMessage<FTraceControlDiscoveryPing>();
-	Message->SessionId = Instance->GetOwnerSession()->GetSessionId();
-	Message->InstanceId = Instance->GetInstanceId();
-	MessageEndpoint->Publish<FTraceControlDiscoveryPing>(Message);
-}
-
 FTraceController::FTracingInstance::FTracingInstance(const TSharedRef<IMessageBus>& InMessageBus, FMessageAddress InService)
 	: Status()
 	, Commands(InMessageBus, InService)
 {
 }
 
+TSharedPtr<ITraceController> ITraceController::Create(TSharedPtr<IMessageBus>& InBus)
+{
+	return MakeShareable(new FTraceController(InBus.ToSharedRef()));
+}
\ No newline at end of file
diff --git a/Engine/Source/Runtime/SessionServices/Private/TraceController.h b/Engine/Source/Runtime/SessionServices/Private/TraceController.h
--- a/Engine/Source/Runtime/SessionServices/Private/TraceController.h
+++ b/Engine/Source/Runtime/SessionServices/Private/TraceController.h
@@ -26,6 +26,8 @@ public:
 	virtual ~FTraceController() override;
 
 private:
+	virtual void SendDiscoveryRequest(const FGuid& SessionId, const FGuid& InstanceId) const override;
+	virtual void SendDiscoveryRequest() override;
 	virtual void SendStatusUpdateRequest() override;
 	virtual void SendChannelUpdateRequest() override;
 	virtual void SendSettingsUpdateRequest() override;
@@ -70,9 +72,6 @@ private:
 	*/
 	bool RediscoverSelectedSession() const;
 
-	/* Send a discovery ping to a specific instance rather than broadcast a request. */
-	void SendDiscoveryPing(const TSharedPtr<ISessionInstanceInfo>& Instance) const;
-
 private:
 	
 	struct FTracingInstance
diff --git a/Engine/Source/Runtime/SessionServices/Public/ITraceController.h b/Engine/Source/Runtime/SessionServices/Public/ITraceController.h
--- a/Engine/Source/Runtime/SessionServices/Public/ITraceController.h
+++ b/Engine/Source/Runtime/SessionServices/Public/ITraceController.h
@@ -125,6 +125,18 @@ public:
 	
 	virtual ~ITraceController() = default;
 
+	/**
+	 * Request to discover a specific instance
+	 * @param SessionId Session id of the instance
+	 * @param InstanceId Instance id
+	 */
+	virtual void SendDiscoveryRequest(const FGuid& SessionId, const FGuid& InstanceId) const = 0;
+
+	/**
+	 * Request to discover all instances
+	 */
+	virtual void SendDiscoveryRequest() = 0;
+
 	/**
 	 * Request update of the status from all sessions and instances.
 	 */
@@ -193,4 +205,11 @@ public:
 	 * Returns the number of selected sessions.
 	 */
 	virtual uint32 GetNumSelectedInstances() = 0;
-};
\ No newline at end of file
+
+	/**
+	 * Create an instance using a custom message bus
+	 * @param MessageBus Message bus the controller will use
+	 */
+	SESSIONSERVICES_API static TSharedPtr<ITraceController> Create(TSharedPtr<class IMessageBus>& MessageBus);
+};
+
