Fix the check to see if a world partition runtime level should have a primary asset to work correctly for unloaded external packages.

diff --git a/Engine/Source/Runtime/Engine/Private/World.cpp b/Engine/Source/Runtime/Engine/Private/World.cpp
--- a/Engine/Source/Runtime/Engine/Private/World.cpp
+++ b/Engine/Source/Runtime/Engine/Private/World.cpp
@@ -9388,10 +9388,10 @@ bool UWorld::ResolveSubobject(const TCHAR* SubObjectPath, UObject*& OutObject, b
 FPrimaryAssetId UWorld::GetPrimaryAssetId() const
 {
 	UPackage* Package = GetOutermost();
-	const IWorldPartitionCell* WorldPartitionCell = PersistentLevel ? PersistentLevel->GetWorldPartitionRuntimeCell() : nullptr;
+	const bool bIsWorldPartitionRuntime = PersistentLevel ? PersistentLevel->IsWorldPartitionRuntimeCell() : false;
 
 	// PIE and world partition runtime levels are temporary and do not represent a primary asset
-	if (!Package->HasAnyPackageFlags(PKG_PlayInEditor) && !WorldPartitionCell)
+	if (!Package->HasAnyPackageFlags(PKG_PlayInEditor) && !bIsWorldPartitionRuntime)
 	{
 		// Return Map:/path/to/map
 		return FPrimaryAssetId(UAssetManager::MapType, Package->GetFName());
