IncrementalCook: AnimSequence: Add a determinismhelper to AnimSequence to display the DDC key when there is diff in iterativevalidate.
#rnx

diff --git a/Engine/Source/Runtime/Engine/Classes/Animation/AnimSequence.h b/Engine/Source/Runtime/Engine/Classes/Animation/AnimSequence.h
--- a/Engine/Source/Runtime/Engine/Classes/Animation/AnimSequence.h
+++ b/Engine/Source/Runtime/Engine/Classes/Animation/AnimSequence.h
@@ -937,6 +937,7 @@ public:
 	friend class UAnimationBlueprintLibrary;
 	friend class UAnimBoneCompressionSettings;
 	friend class FCustomAttributeCustomization;
+	friend class FAnimSequenceDeterminismHelper;
 	friend class FAnimSequenceTestBase;
 	friend struct UE::Anim::Compression::FScopedCompressionGuard;
 	friend class FAnimDataControllerTestBase;
diff --git a/Engine/Source/Runtime/Engine/Private/Animation/AnimSequence.cpp b/Engine/Source/Runtime/Engine/Private/Animation/AnimSequence.cpp
--- a/Engine/Source/Runtime/Engine/Private/Animation/AnimSequence.cpp
+++ b/Engine/Source/Runtime/Engine/Private/Animation/AnimSequence.cpp
@@ -58,7 +58,9 @@ LLM_DEFINE_TAG(SequenceData);
 #include "Animation/AnimData/IAnimationDataModel.h"
 #include "Animation/BuiltInAttributeTypes.h"
 #include "AssetRegistry/AssetRegistryModule.h"
+#include "Cooker/CookDeterminismHelper.h"
 #include "ProfilingDebugging/CookStats.h"
+#include "Serialization/CompactBinaryWriter.h"
 #include "Serialization/MemoryHasher.h"
 #include "Misc/DataValidation.h"
 #endif // WITH_EDITOR
@@ -911,6 +913,32 @@ void UAnimSequence::GetPreloadDependencies(TArray<UObject*>& OutDeps)
 	}
 }
 
+#if WITH_EDITOR
+class FAnimSequenceDeterminismHelper : public UE::Cook::IDeterminismHelper
+{
+public:
+	FAnimSequenceDeterminismHelper(UAnimSequence* InAnimSequence)
+		: AnimSequence(InAnimSequence)
+	{
+	}
+
+	virtual void ConstructDiagnostics(UE::Cook::IDeterminismConstructDiagnosticsContext& Context) override
+	{
+		const FIoHash KeyHash = AnimSequence->CreateDerivedDataKeyHash(Context.GetTargetPlatform());
+
+		FCbWriter Writer;
+		Writer.BeginObject();
+		Writer << "DDCKey" << WriteToString<64>(KeyHash);
+		Writer.EndObject();
+
+		Context.AddDiagnostic("UAnimSequence", Writer.Save());
+	}
+
+private:
+	UAnimSequence* AnimSequence;
+};
+#endif
+
 void UAnimSequence::PreSave(const class ITargetPlatform* TargetPlatform)
 {
 	PRAGMA_DISABLE_DEPRECATION_WARNINGS;
@@ -937,6 +965,11 @@ void UAnimSequence::PreSave(FObjectPreSaveContext ObjectSaveContext)
 	{
 		UpdateRetargetSourceAsset();
 	}
+
+	if (ObjectSaveContext.IsDeterminismDebug())
+	{
+		ObjectSaveContext.RegisterDeterminismHelper(new FAnimSequenceDeterminismHelper(this));
+	}
 #endif
 
 	Super::PreSave(ObjectSaveContext);
