[AudioThread]
 - Add a CVar to restore legacy behavior for commands sent to the game thread in case of unforeseen problems

#rnx

diff --git a/Engine/Source/Runtime/Engine/Private/AudioThread.cpp b/Engine/Source/Runtime/Engine/Private/AudioThread.cpp
--- a/Engine/Source/Runtime/Engine/Private/AudioThread.cpp
+++ b/Engine/Source/Runtime/Engine/Private/AudioThread.cpp
@@ -50,6 +50,13 @@ FAutoConsoleVariableRef  CVarAudioCommandFenceWaitTimeMs(
 	TEXT("Sets number of ms for fence wait"), 
 	ECVF_Default);
 
+static bool GAudioThreadUseSafeRunCommandOnGameThread = true;
+FAutoConsoleVariableRef CVarAudioThreadUseSafeRunCommandOnGameThread(
+	TEXT("AudioThread.UseSafeRunCommandOnGameThread"),
+	GAudioThreadUseSafeRunCommandOnGameThread,
+	TEXT("When active, limit commands sent to the game thread to run at a safe place inside a tick"),
+	ECVF_Default);
+
 struct FAudioThreadInteractor
 {
 	static void UseAudioThreadCVarSinkFunction()
@@ -370,16 +377,35 @@ void FAudioThread::RunCommandOnGameThread(TUniqueFunction<void()> InFunction, co
 	if (IsAudioThreadRunning())
 	{
 		check(IsInAudioThread());
-		ExecuteOnGameThread(
-			TEXT("FAudioThread::RunCommandOnGameThread"),
-			[Function = MoveTemp(InFunction), InStatId]()
-			{
-				CSV_SCOPED_TIMING_STAT_EXCLUSIVE(Audio);
-				QUICK_SCOPE_CYCLE_COUNTER(STAT_AudioThread_RunCommandOnGameThread);
-				FScopeCycleCounter ScopeCycleCounter(InStatId);
-				Function();
-			}
-		);
+		if (GAudioThreadUseSafeRunCommandOnGameThread)
+		{
+			ExecuteOnGameThread(
+				TEXT("FAudioThread::RunCommandOnGameThread"),
+				[Function = MoveTemp(InFunction), InStatId]()
+				{
+					CSV_SCOPED_TIMING_STAT_EXCLUSIVE(Audio);
+					QUICK_SCOPE_CYCLE_COUNTER(STAT_AudioThread_RunCommandOnGameThread);
+					FScopeCycleCounter ScopeCycleCounter(InStatId);
+					Function();
+				}
+			);
+		}
+		else
+		{
+			// This is the legacy behavior that will run game thread tasks anywhere on the game thread.
+			// This could be inside the GC, postload, etc... which might lead to problematic behavior.
+			FFunctionGraphTask::CreateAndDispatchWhenReady(
+				[Function = MoveTemp(InFunction), InStatId]()
+				{
+					CSV_SCOPED_TIMING_STAT_EXCLUSIVE(Audio);
+					QUICK_SCOPE_CYCLE_COUNTER(STAT_AudioThread_RunCommandOnGameThread);
+					FScopeCycleCounter ScopeCycleCounter(InStatId);
+					Function();
+				},
+				TStatId(),
+				nullptr,
+				ENamedThreads::GameThread);
+		}
 	}
 	else
 	{
