[World Partition HashSet LHGrid]
- When selecting the hierarchical level of objects from their bound's maximum extent, allow objects to stay a level below if they don't exceed more that a quarter of the cell size on a single axis.
- This allow objects that would normally go a level higher to stay lower in the grid hierarchical level while respecting the rule that cells can't grow more that half a cell size in each axis.

diff --git a/Engine/Source/Runtime/Engine/Private/WorldPartition/RuntimeHashSet/RuntimePartitionLHGrid.cpp b/Engine/Source/Runtime/Engine/Private/WorldPartition/RuntimeHashSet/RuntimePartitionLHGrid.cpp
--- a/Engine/Source/Runtime/Engine/Private/WorldPartition/RuntimeHashSet/RuntimePartitionLHGrid.cpp
+++ b/Engine/Source/Runtime/Engine/Private/WorldPartition/RuntimeHashSet/RuntimePartitionLHGrid.cpp
@@ -32,11 +32,28 @@ struct FCellCoord
 		return (X == Other.X) && (Y == Other.Y) && (Z == Other.Z) && (Level == Other.Level);
 	}
 
-	static inline int32 GetLevelForBox(const FBox& InBox, int32 InCellSize)
+	static inline int32 GetLevelForBox(const FBox& InBox, int32 InCellSize, const FVector& InOrigin)
 	{
 		const FVector Extent = InBox.GetSize();
 		const FVector::FReal MaxLength = Extent.GetMax();
-		return FMath::CeilToInt32(FMath::Max<FVector::FReal>(FMath::Log2(MaxLength / InCellSize), 0));
+		const int32 Level = FMath::CeilToInt32(FMath::Max<FVector::FReal>(FMath::Log2(MaxLength / InCellSize), 0));
+
+		if (Level)
+		{
+			const FCellCoord CellCoord = GetCellCoords(InBox.GetCenter(), InCellSize, Level - 1, InOrigin);
+			const FBox CellBounds = GetCellBounds(CellCoord, InCellSize, InOrigin);
+			const FVector MaxUnderLap = FVector::Max3(CellBounds.Min - InBox.Min, InBox.Max - CellBounds.Max, FVector::Zero());
+			const FVector::FReal MaxUnderLapLength = MaxUnderLap.GetMax();
+
+			// Allow objects that slightly exceed the cell size to be placed in the lower levels. We don't want large objects to
+			// grow cells to their maximum extent, which is half the cell size on each axis, so we allow a quarter on each axis.
+			if (MaxUnderLapLength < InCellSize / 4)
+			{
+				return Level - 1;
+			}
+		}
+
+		return Level;
 	}
 
 	static inline FCellCoord GetCellCoords(const FVector& InPos, int32 InCellSize, int32 InLevel, const FVector& InOrigin)
@@ -179,7 +196,7 @@ bool URuntimePartitionLHGrid::GenerateStreaming(const FGenerateStreamingParams&
 		if (ActorSetInstance->bIsSpatiallyLoaded)
 		{
 			const FBox ActorSetInstanceBounds = FBox(ActorSetInstance->Bounds.Min * SpaceMask, ActorSetInstance->Bounds.Max * SpaceMask);
-			const int32 GridLevel = FCellCoord::GetLevelForBox(ActorSetInstanceBounds, CellSize);
+			const int32 GridLevel = FCellCoord::GetLevelForBox(ActorSetInstanceBounds, CellSize, Origin * SpaceMask);
 			const FCellCoord CellCoord = FCellCoord::GetCellCoords(ActorSetInstanceBounds.GetCenter(), CellSize, GridLevel, Origin * SpaceMask);
 			CellsActorSetInstances.FindOrAdd(CellCoord).Add(ActorSetInstance);
 		}
