[World Partition Sequencer] Properly handle the case where we are unable to set data layers from tracks due to authority.
Also changed logs to verbose to avoid logging them in shipping builds.

diff --git a/Engine/Source/Runtime/Engine/Private/WorldPartition/DataLayer/WorldDataLayers.cpp b/Engine/Source/Runtime/Engine/Private/WorldPartition/DataLayer/WorldDataLayers.cpp
--- a/Engine/Source/Runtime/Engine/Private/WorldPartition/DataLayer/WorldDataLayers.cpp
+++ b/Engine/Source/Runtime/Engine/Private/WorldPartition/DataLayer/WorldDataLayers.cpp
@@ -228,24 +228,6 @@ bool AWorldDataLayers::SetDataLayerRuntimeState(const UDataLayerInstance* InData
 
 	if (CurrentState != InState)
 	{
-		if (GetWorld()->IsGameWorld())
-		{
-			PRAGMA_DISABLE_DEPRECATION_WARNINGS
-			if (DataLayersFilterDelegate.IsBound())
-			{
-				const FString DataLayerShortName(InDataLayerInstance->GetDataLayerShortName());
-				if (!DataLayersFilterDelegate.Execute(FName(*DataLayerShortName), CurrentState, InState))
-				{
-					UE_LOG(LogWorldPartition, Verbose, TEXT("Data Layer '%s' state change was filtered out: %s -> %s"),
-						*DataLayerShortName,
-						*StaticEnum<EDataLayerRuntimeState>()->GetDisplayNameTextByValue((int64)CurrentState).ToString(),
-						*StaticEnum<EDataLayerRuntimeState>()->GetDisplayNameTextByValue((int64)InState).ToString());
-					return false;
-				}
-			}
-			PRAGMA_ENABLE_DEPRECATION_WARNINGS
-		}
-
 		TargetLoadedDataLayerNames.Remove(InDataLayerInstance->GetDataLayerFName());
 		TargetActiveDataLayerNames.Remove(InDataLayerInstance->GetDataLayerFName());
 
