Fix up some function naming/comment wording.
#rnx

diff --git a/Engine/Source/Runtime/Engine/Private/VT/MeshPaintVirtualTexture.cpp b/Engine/Source/Runtime/Engine/Private/VT/MeshPaintVirtualTexture.cpp
--- a/Engine/Source/Runtime/Engine/Private/VT/MeshPaintVirtualTexture.cpp
+++ b/Engine/Source/Runtime/Engine/Private/VT/MeshPaintVirtualTexture.cpp
@@ -60,12 +60,12 @@ namespace MeshPaintVirtualTexture
 	}
 
 	/** 
-	 * Fill out the scene view params from an allocated VT. 
-	 * We expect the result to be constant for all allocated VTs (so that they can share one set of data in the view).
-	 * Note that there are valid cases when it will change over time.
+	 * Fill out the scene uniforms from an allocated VT. 
+	 * We expect the result to be constant for all allocated VTs (so that they can share one uniform buffer).
+	 * Note that there are valid cases when it will change over time (but we always use the latest).
 	 * For example, when the virtual texture pools are resized, all VTs reallocate and change to a new value here.
 	 */
-	static void GetSceneViewParams(IAllocatedVirtualTexture* InAllocatedVT, FUniformParams& OutParams)
+	static void GetSceneUniformParams(IAllocatedVirtualTexture* InAllocatedVT, FUniformParams& OutParams)
 	{
 		OutParams.PageTableTexture = InAllocatedVT->GetPageTableTexture(0);
 		OutParams.PhysicalTexture = InAllocatedVT->GetPhysicalTexture(0);
@@ -89,7 +89,7 @@ namespace MeshPaintVirtualTexture
 
 	/** A global set to track allocated mesh paint virtual textures. */
 	TSet<IAllocatedVirtualTexture*> AllocatedVTs;
-	/** The global scene view params cached from the last allocated VT. */
+	/** The global scene uniform params cached from the last allocated VT. */
 	FUniformParams Params;
 	/** Mutex for the global state. The state is all accessed only on the render thread timeline, but the Params are accessed from render worker threads. */
 	UE::FRecursiveMutex Mutex;
@@ -105,8 +105,8 @@ namespace MeshPaintVirtualTexture
 		// If that happens we could change to store in a map against a ref count.
 		ensure(!bAlreadyInSet);
 
-		// Update the cached view params.
-		GetSceneViewParams(InAllocatedVT, Params);
+		// Update the cached uniform params.
+		GetSceneUniformParams(InAllocatedVT, Params);
 	}
 
 	/** Remove an allocated VT from our global allocated set. */
@@ -117,7 +117,7 @@ namespace MeshPaintVirtualTexture
 		const int32 Removed = AllocatedVTs.Remove((IAllocatedVirtualTexture*)InBaton);
 		ensure(Removed == 1);
 
-		// Clear the cached view arams if we have removed all the allocated VTs.
+		// Clear the cached uniform params if we have removed all the allocated VTs.
 		if (AllocatedVTs.Num() == 0)
 		{
 			Params = {};
