Add mesh paint texture to return of GetUsedTextures on static component.
#rnx

diff --git a/Engine/Source/Runtime/Engine/Classes/Components/StaticMeshComponent.h b/Engine/Source/Runtime/Engine/Classes/Components/StaticMeshComponent.h
--- a/Engine/Source/Runtime/Engine/Classes/Components/StaticMeshComponent.h
+++ b/Engine/Source/Runtime/Engine/Classes/Components/StaticMeshComponent.h
@@ -547,6 +547,7 @@ public:
 	ENGINE_API virtual bool BuildTextureStreamingDataImpl(ETextureStreamingBuildType BuildType, EMaterialQualityLevel::Type QualityLevel, ERHIFeatureLevel::Type FeatureLevel, TSet<FGuid>& DependentResources, bool& bOutSupportsBuildTextureStreamingData) override;
 	/** Get the StreaminTexture data. */
 	ENGINE_API virtual void GetStreamingRenderAssetInfo(FStreamingTextureLevelContext& LevelContext, TArray<FStreamingRenderAssetPrimitiveInfo>& OutStreamingRenderAssets) const override;
+	ENGINE_API virtual void GetUsedTextures(TArray<UTexture*>& OutTextures, EMaterialQualityLevel::Type QualityLevel) override;
 #if WITH_EDITOR
 	ENGINE_API virtual bool RemapActorTextureStreamingBuiltDataToLevel(const class UActorTextureStreamingBuildDataComponent* InActorTextureBuildData) override;
 	ENGINE_API virtual uint32 ComputeHashTextureStreamingBuiltData() const override;
diff --git a/Engine/Source/Runtime/Engine/Private/Components/StaticMeshComponent.cpp b/Engine/Source/Runtime/Engine/Private/Components/StaticMeshComponent.cpp
--- a/Engine/Source/Runtime/Engine/Private/Components/StaticMeshComponent.cpp
+++ b/Engine/Source/Runtime/Engine/Private/Components/StaticMeshComponent.cpp
@@ -1129,6 +1129,16 @@ void UStaticMeshComponent::GetStreamingRenderAssetInfo(FStreamingTextureLevelCon
 	}
 }
 
+void UStaticMeshComponent::GetUsedTextures(TArray<UTexture*>& OutTextures, EMaterialQualityLevel::Type QualityLevel)
+{
+	if (MeshPaintTexture)
+	{
+		OutTextures.AddUnique(MeshPaintTexture);
+	}
+
+	Super::GetUsedTextures(OutTextures, QualityLevel);
+}
+
 UBodySetup* UStaticMeshComponent::GetBodySetup()
 {
 	if (GetStaticMesh())
