Output FRayTracingGeometryManager RequestedSize/TotalResidentSize/TotalAlwaysResidentSize stats in CSV Profiler.

diff --git a/Engine/Source/Runtime/Engine/Private/Rendering/RayTracingGeometryManager.cpp b/Engine/Source/Runtime/Engine/Private/Rendering/RayTracingGeometryManager.cpp
--- a/Engine/Source/Runtime/Engine/Private/Rendering/RayTracingGeometryManager.cpp
+++ b/Engine/Source/Runtime/Engine/Private/Rendering/RayTracingGeometryManager.cpp
@@ -122,6 +122,8 @@ DECLARE_DWORD_ACCUMULATOR_STAT(TEXT("Pending Build Primitives"), STAT_RayTracing
 DECLARE_DWORD_ACCUMULATOR_STAT(TEXT("Pending Streaming Requests"), STAT_RayTracingPendingStreamingRequests, STATGROUP_RayTracingGeometry);
 DECLARE_DWORD_ACCUMULATOR_STAT(TEXT("In-flight Streaming Requests"), STAT_RayTracingInflightStreamingRequests, STATGROUP_RayTracingGeometry);
 
+CSV_DEFINE_CATEGORY(RayTracingGeometry, true);
+
 FRayTracingGeometryManager::FRayTracingGeometryManager()
 {
 	StreamingRequests.SetNum(GRayTracingStreamingMaxPendingRequests);
@@ -661,8 +663,9 @@ void FRayTracingGeometryManager::Tick(FRHICommandList& RHICmdList)
 #endif
 
 		SET_MEMORY_STAT(STAT_RayTracingGeometryRequestedMemory, RequestedSize);
+		CSV_CUSTOM_STAT(RayTracingGeometry, RequestedSizeMB, RequestedSize / 1024.0f / 1024.0f, ECsvCustomStatOp::Set);
 
-		uint64 ResidentGeometryMemoryPoolSize = FUnitConversion::Convert(GRayTracingResidentGeometryMemoryPoolSizeInMB, EUnit::Megabytes, EUnit::Bytes);
+		const uint64 ResidentGeometryMemoryPoolSize = FUnitConversion::Convert(GRayTracingResidentGeometryMemoryPoolSizeInMB, EUnit::Megabytes, EUnit::Bytes);
 
 		// Step 3
 		// - if making requested geometries resident will put us over budget -> evict some geometry not referenced by TLAS
@@ -777,6 +780,7 @@ void FRayTracingGeometryManager::Tick(FRHICommandList& RHICmdList)
 		}
 
 		SET_MEMORY_STAT(STAT_RayTracingGeometryRequestedMemory, TotalResidentSize);
+		CSV_CUSTOM_STAT(RayTracingGeometry, RequestedSizeMB, TotalResidentSize / 1024.0f / 1024.0f, ECsvCustomStatOp::Set);
 	}
 
 	{
@@ -804,6 +808,9 @@ void FRayTracingGeometryManager::Tick(FRHICommandList& RHICmdList)
 
 	SET_MEMORY_STAT(STAT_RayTracingGeometryResidentMemory, TotalResidentSize);
 	SET_MEMORY_STAT(STAT_RayTracingGeometryAlwaysResidentMemory, TotalAlwaysResidentSize);
+
+	CSV_CUSTOM_STAT(RayTracingGeometry, TotalResidentSizeMB, TotalResidentSize / 1024.0f / 1024.0f, ECsvCustomStatOp::Set);
+	CSV_CUSTOM_STAT(RayTracingGeometry, TotalAlwaysResidentSizeMB, TotalAlwaysResidentSize / 1024.0f / 1024.0f, ECsvCustomStatOp::Set);
 }
 
 bool FRayTracingGeometryManager::RequestRayTracingGeometryStreamIn(FRHICommandList& RHICmdList, RayTracingGeometryHandle GeometryHandle)
