Properly initializing deprecated uproperties in NavMovementComponent so we can copy over correct values to the new struct for NavMovementProperties

diff --git a/Engine/Source/Runtime/Engine/Private/Components/NavMovementComponent.cpp b/Engine/Source/Runtime/Engine/Private/Components/NavMovementComponent.cpp
--- a/Engine/Source/Runtime/Engine/Private/Components/NavMovementComponent.cpp
+++ b/Engine/Source/Runtime/Engine/Private/Components/NavMovementComponent.cpp
@@ -12,6 +12,13 @@
 //----------------------------------------------------------------------//
 UNavMovementComponent::UNavMovementComponent(const FObjectInitializer& ObjectInitializer)
 	: Super(ObjectInitializer)
+PRAGMA_DISABLE_DEPRECATION_WARNINGS
+	, FixedPathBrakingDistance(0.0f)
+	, bUpdateNavAgentWithOwnersCollision(true)
+	, bUseAccelerationForPaths(false)
+	, bUseFixedBrakingDistanceForPaths(false)
+	, bStopMovementAbortPaths(true)
+PRAGMA_ENABLE_DEPRECATION_WARNINGS
 {
 	bComponentShouldUpdatePhysicsVolume = true;
 }
@@ -65,7 +72,7 @@ void UNavMovementComponent::PostLoad()
 
 void UNavMovementComponent::UpdateNavAgent(const UObject& ObjectToUpdateFrom)
 {
-	if (!NavMovementProperties.bUpdateNavAgentWithOwnersCollision)
+	if (!ShouldUpdateNavAgentWithOwnersCollision())
 	{
 		return;
 	}
@@ -76,7 +83,7 @@ void UNavMovementComponent::UpdateNavAgent(const UObject& ObjectToUpdateFrom)
 	if (const UCapsuleComponent* CapsuleComponent = Cast<UCapsuleComponent>(&ObjectToUpdateFrom))
 	{
 		NavAgentProps.AgentRadius = CapsuleComponent->GetScaledCapsuleRadius();
-		NavAgentProps.AgentHeight = CapsuleComponent->GetScaledCapsuleHalfHeight() * 2.f;;
+		NavAgentProps.AgentHeight = CapsuleComponent->GetScaledCapsuleHalfHeight() * 2.f;
 	}
 	else if (const AActor* ObjectAsActor = Cast<AActor>(&ObjectToUpdateFrom))
 	{
