Wait for streaming before rebuilding texture derived data

#tests opened a map in editor and flied around

diff --git a/Engine/Source/Runtime/Engine/Private/TextureDerivedData.cpp b/Engine/Source/Runtime/Engine/Private/TextureDerivedData.cpp
--- a/Engine/Source/Runtime/Engine/Private/TextureDerivedData.cpp
+++ b/Engine/Source/Runtime/Engine/Private/TextureDerivedData.cpp
@@ -4048,6 +4048,9 @@ void UTexture::ForceRebuildPlatformData(uint8 InEncodeSpeedOverride /* =255 ETex
 	FTexturePlatformData** PlatformDataLinkPtr = GetRunningPlatformData();
 	if (PlatformDataLinkPtr && *PlatformDataLinkPtr && FApp::CanEverRender())
 	{
+		// Cache() will clear FTexturePlatformData::Mips which can be accessed by the streaming update
+		WaitForPendingInitOrStreaming();
+
 		FTexturePlatformData *&PlatformDataLink = *PlatformDataLinkPtr;
 		FlushRenderingCommands();
 
