[World Partition Actor Descriptors] Better logs when detecting duplicated actor descriptors.

diff --git a/Engine/Source/Runtime/Engine/Private/WorldPartition/ActorDescList.cpp b/Engine/Source/Runtime/Engine/Private/WorldPartition/ActorDescList.cpp
--- a/Engine/Source/Runtime/Engine/Private/WorldPartition/ActorDescList.cpp
+++ b/Engine/Source/Runtime/Engine/Private/WorldPartition/ActorDescList.cpp
@@ -1,7 +1,6 @@
 // Copyright Epic Games, Inc. All Rights Reserved.
 
 #include "WorldPartition/ActorDescList.h"
-#include "WorldPartition/WorldPartitionLog.h"
 
 #if WITH_EDITOR
 FWorldPartitionActorDesc* FActorDescList::AddActor(const AActor* InActor)
diff --git a/Engine/Source/Runtime/Engine/Public/WorldPartition/ActorDescList.h b/Engine/Source/Runtime/Engine/Public/WorldPartition/ActorDescList.h
--- a/Engine/Source/Runtime/Engine/Public/WorldPartition/ActorDescList.h
+++ b/Engine/Source/Runtime/Engine/Public/WorldPartition/ActorDescList.h
@@ -3,6 +3,7 @@
 
 #include "CoreMinimal.h"
 #include "GameFramework/Actor.h"
+#include "WorldPartition/WorldPartitionLog.h"
 #include "WorldPartition/WorldPartitionActorDesc.h"
 #include "WorldPartition/WorldPartitionActorDescType.h"
 
@@ -207,10 +208,17 @@ void TActorDescList<DescType>::AddActorDescriptor(DescType* ActorDesc)
 {
 	check(!bIsProxy);
 	check(ActorDesc);
-	checkf(!ActorsByGuid.Contains(ActorDesc->GetGuid()), TEXT("Duplicated actor descriptor guid '%s' detected: `%s`"), *ActorDesc->GetGuid().ToString(), *ActorDesc->GetActorNameString());
 
-	TUniquePtr<DescType>* NewActorDesc = new(ActorDescList) TUniquePtr<DescType>(ActorDesc);
-	ActorsByGuid.Add(ActorDesc->GetGuid(), NewActorDesc);
+	if (TUniquePtr<DescType>* ExistingActorDesc = ActorsByGuid.FindRef(ActorDesc->GetGuid()))
+	{
+		UE_LOG(LogWorldPartition, Error, TEXT("Duplicated actor descriptor detected:\n\tExisting: %s\n\t     New: %s"), 
+			*ExistingActorDesc->Get()->ToString(FWorldPartitionActorDesc::EToStringMode::Full), *ActorDesc->ToString(FWorldPartitionActorDesc::EToStringMode::Full));
+	}
+	else
+	{
+		TUniquePtr<DescType>* NewActorDesc = new(ActorDescList) TUniquePtr<DescType>(ActorDesc);
+		ActorsByGuid.Add(ActorDesc->GetGuid(), NewActorDesc);
+	}
 }
 
 template<class DescType>
