[World Partition Actor Desc] Add code to detect invalid custom versions from actor descriptors (missing, invalid or newer versions).

diff --git a/Engine/Source/Runtime/Engine/Private/WorldPartition/WorldPartitionActorDesc.cpp b/Engine/Source/Runtime/Engine/Private/WorldPartition/WorldPartitionActorDesc.cpp
--- a/Engine/Source/Runtime/Engine/Private/WorldPartition/WorldPartitionActorDesc.cpp
+++ b/Engine/Source/Runtime/Engine/Private/WorldPartition/WorldPartitionActorDesc.cpp
@@ -242,6 +242,33 @@ void FWorldPartitionActorDesc::Init(const FWorldPartitionActorDescInitData& Desc
 		FCustomVersionContainer CustomVersions;
 		CustomVersions.Serialize(MetadataAr);
 		MetadataAr.SetCustomVersions(CustomVersions);
+
+		TArray<FCustomVersionDifference> Diffs = FCurrentCustomVersions::Compare(CustomVersions.GetAllVersions(), *DescData.PackageName.ToString());
+		for (FCustomVersionDifference Diff : Diffs)
+		{
+			if (Diff.Type == ECustomVersionDifference::Missing)
+			{
+				UE_LOG(LogWorldPartition, Fatal, TEXT("Missing custom version for actor descriptor '%s'"), *DescData.PackageName.ToString());
+			}
+			else if (Diff.Type == ECustomVersionDifference::Invalid)
+			{
+				UE_LOG(LogWorldPartition, Fatal, TEXT("Invalid custom version for actor descriptor '%s'"), *DescData.PackageName.ToString());
+			}
+			else if (Diff.Type == ECustomVersionDifference::Newer)
+			{
+				int32 PackageVersion = -1;
+				int32 HeadCodeVersion = -1;
+				if (const FCustomVersion* PackagePtr = CustomVersions.GetVersion(Diff.Version->Key))
+				{
+					PackageVersion = PackagePtr->Version;
+				}
+				if (TOptional<FCustomVersion> CurrentPtr = FCurrentCustomVersions::Get(Diff.Version->Key))
+				{
+					HeadCodeVersion = CurrentPtr->Version;
+				}
+				UE_LOG(LogWorldPartition, Fatal, TEXT("Newer custom version for actor descriptor '%s' (file: %d, head: %d)"), *DescData.PackageName.ToString(), PackageVersion, HeadCodeVersion);
+			}
+		}
 	
 		// Serialize metadata payload
 		FActorDescArchive ActorDescAr(MetadataAr, this);
