[Data Layers] Fix invalid call to ConditionalPostLoad from AWorldDataLayers::OnDataLayerManagerInitialized causing asserts during GatherTextFromAssetsCommandlet.

diff --git a/Engine/Source/Runtime/Engine/Private/WorldPartition/DataLayer/WorldDataLayers.cpp b/Engine/Source/Runtime/Engine/Private/WorldPartition/DataLayer/WorldDataLayers.cpp
--- a/Engine/Source/Runtime/Engine/Private/WorldPartition/DataLayer/WorldDataLayers.cpp
+++ b/Engine/Source/Runtime/Engine/Private/WorldPartition/DataLayer/WorldDataLayers.cpp
@@ -1131,7 +1131,6 @@ void AWorldDataLayers::InitializeExternalPackageDataLayerInstances()
 void AWorldDataLayers::OnDataLayerManagerInitialized()
 {
 #if WITH_EDITOR
-	ConditionalPostLoad();
 	// At this point, LoadedExternalPackageDataLayerInstances are fully loaded, transfer them to the DataLayerInstances list.
 	if (IsUsingExternalPackageDataLayerInstances())
 	{
