IncrementalCook: AnimSequence: Add more feedback for determinism investigation - include the entire ddc key in the diagnostic rather than just the hash.
#rnx

diff --git a/Engine/Source/Runtime/Engine/Classes/Animation/AnimSequence.h b/Engine/Source/Runtime/Engine/Classes/Animation/AnimSequence.h
--- a/Engine/Source/Runtime/Engine/Classes/Animation/AnimSequence.h
+++ b/Engine/Source/Runtime/Engine/Classes/Animation/AnimSequence.h
@@ -906,6 +906,9 @@ protected:
 
 #if WITH_EDITOR
 	ENGINE_API FIoHash CreateDerivedDataKeyHash(const ITargetPlatform* TargetPlatform);
+private:
+	FString CreateDerivedDataKeyString(const ITargetPlatform* TargetPlatform);
+protected:
 	ENGINE_API FIoHash BeginCacheDerivedData(const ITargetPlatform* TargetPlatform);
 	ENGINE_API bool PollCacheDerivedData(const FIoHash& KeyHash) const;
 	ENGINE_API void EndCacheDerivedData(const FIoHash& KeyHash);
diff --git a/Engine/Source/Runtime/Engine/Private/Animation/AnimSequence.cpp b/Engine/Source/Runtime/Engine/Private/Animation/AnimSequence.cpp
--- a/Engine/Source/Runtime/Engine/Private/Animation/AnimSequence.cpp
+++ b/Engine/Source/Runtime/Engine/Private/Animation/AnimSequence.cpp
@@ -924,11 +924,19 @@ public:
 
 	virtual void ConstructDiagnostics(UE::Cook::IDeterminismConstructDiagnosticsContext& Context) override
 	{
-		const FIoHash KeyHash = AnimSequence->CreateDerivedDataKeyHash(Context.GetTargetPlatform());
+		FString DDCKey = AnimSequence->CreateDerivedDataKeyString(Context.GetTargetPlatform());
+		FIoHash DDCKeyHash = FIoHash::Zero;
+		if (!DDCKey.IsEmpty())
+		{
+			FMemoryHasherBlake3 Writer;
+			Writer << DDCKey;
+			DDCKeyHash = Writer.Finalize();
+		}
 
 		FCbWriter Writer;
 		Writer.BeginObject();
-		Writer << "DDCKey" << WriteToString<64>(KeyHash);
+		Writer << "DDCKeyHash" << WriteToString<64>(DDCKeyHash);
+		Writer << "DDCKey" << DDCKey;
 		Writer.EndObject();
 
 		Context.AddDiagnostic("UAnimSequence", Writer.Save());
@@ -4545,16 +4553,29 @@ void UAnimSequence::CalculateNumberOfSampledKeys()
 }
 
 FIoHash UAnimSequence::CreateDerivedDataKeyHash(const ITargetPlatform* TargetPlatform)
+{
+	FString Ret = CreateDerivedDataKeyString(TargetPlatform);
+	if (Ret.IsEmpty())
+	{
+		return FIoHash::Zero;
+	}
+
+	// New animation DDC key format; use just the hash of the complete dependencies string.
+	FMemoryHasherBlake3 Writer;
+	Writer << Ret;
+	return Writer.Finalize();
+}
+
+
+FString UAnimSequence::CreateDerivedDataKeyString(const ITargetPlatform* TargetPlatform)
 {
 	const USkeleton* CurrentSkeleton = GetSkeleton();
 	if (CurrentSkeleton == nullptr)
 	{
 		UE_LOG(LogAnimation, Warning, TEXT("Animation Compression request for %s failed, Skeleton == nullptr."), *GetName());
-		return FIoHash::Zero;
+		return FString();
 	}
 
-	// New animation DDC key format
-	FMemoryHasherBlake3 Writer;
 
 	const bool bIsValidAdditive = IsValidAdditive();
 	const char AdditiveType = bIsValidAdditive ? NibbleToTChar(AdditiveAnimType) : '0';
@@ -4622,9 +4643,7 @@ FIoHash UAnimSequence::CreateDerivedDataKeyHash(const ITargetPlatform* TargetPla
 		*UE::Anim::Compression::AnimationCompressionVersionString
 	);
 
-	Writer << Ret;
-
-	return Writer.Finalize();
+	return Ret;
 }
 
 FIoHash UAnimSequence::BeginCacheDerivedData(const ITargetPlatform* TargetPlatform)
