Fixing NavMovementComponents deprecation and copying of nav properties to correctly copy over all properties

diff --git a/Engine/Source/Runtime/Core/Public/UObject/FortniteReleaseBranchCustomObjectVersion.h b/Engine/Source/Runtime/Core/Public/UObject/FortniteReleaseBranchCustomObjectVersion.h
--- a/Engine/Source/Runtime/Core/Public/UObject/FortniteReleaseBranchCustomObjectVersion.h
+++ b/Engine/Source/Runtime/Core/Public/UObject/FortniteReleaseBranchCustomObjectVersion.h
@@ -51,6 +51,9 @@ struct FFortniteReleaseBranchCustomObjectVersion
 		// UPCGPoint custom serialization
 		PCGPointStructuredSerializer,
 
+		// Deprecation of Nav Movement Properties and moving them to a new struct
+		NavMovementComponentMovingPropertiesToStruct,
+
 		// -----<new versions can be added above this line>-------------------------------------------------
 		VersionPlusOne,
 		LatestVersion = VersionPlusOne - 1
diff --git a/Engine/Source/Runtime/Engine/Classes/GameFramework/NavMovementComponent.h b/Engine/Source/Runtime/Engine/Classes/GameFramework/NavMovementComponent.h
--- a/Engine/Source/Runtime/Engine/Classes/GameFramework/NavMovementComponent.h
+++ b/Engine/Source/Runtime/Engine/Classes/GameFramework/NavMovementComponent.h
@@ -31,26 +31,26 @@ protected:
 	/** Braking distance override used with acceleration driven path following (bUseAccelerationForPaths) */
 	UE_DEPRECATED(5.5, "FixedPathBrakingDistance is deprecated, please use NavMovementProperties.FixedPathBrakingDistance instead.")
 	UPROPERTY(meta = (DeprecatedProperty, DeprecationMessage = "FixedPathBrakingDistance is deprecated, please use NavMovementProperties.FixedPathBrakingDistance instead."))
-	float FixedPathBrakingDistance;
+	float FixedPathBrakingDistance_DEPRECATED;
 
 	/** If set to true NavAgentProps' radius and height will be updated with Owner's collision capsule size */
 	UE_DEPRECATED(5.5, "bUpdateNavAgentWithOwnersCollision is deprecated, please use NavMovementProperties.bUpdateNavAgentWithOwnersCollision instead.")
 	UPROPERTY(meta = (DeprecatedProperty, DeprecationMessage = "bUpdateNavAgentWithOwnersCollision is deprecated, please use NavMovementProperties.bUpdateNavAgentWithOwnersCollision instead."))
-	uint8 bUpdateNavAgentWithOwnersCollision:1;
+	uint8 bUpdateNavAgentWithOwnersCollision_DEPRECATED:1;
 
 	/** If set, pathfollowing will control character movement via acceleration values. If false, it will set velocities directly. */
 	UE_DEPRECATED(5.5, "bUseAccelerationForPaths is deprecated, please use NavMovementProperties.bUseAccelerationForPaths instead.")
 	UPROPERTY(GlobalConfig, meta = (DeprecatedProperty, DeprecationMessage = "bUseAccelerationForPaths is deprecated, please use NavMovementProperties.bUseAccelerationForPaths instead."))
-	uint8 bUseAccelerationForPaths : 1;
+	uint8 bUseAccelerationForPaths_DEPRECATED : 1;
 
 	/** If set, FixedPathBrakingDistance will be used for path following deceleration */
 	UE_DEPRECATED(5.5, "bUseFixedBrakingDistanceForPaths is deprecated, please use NavMovementProperties.bUseFixedBrakingDistanceForPaths instead.")
 	UPROPERTY(meta = (DeprecatedProperty, DeprecationMessage = "bUseFixedBrakingDistanceForPaths is deprecated, please use NavMovementProperties.bUseFixedBrakingDistanceForPaths instead."))
-	uint8 bUseFixedBrakingDistanceForPaths : 1;
+	uint8 bUseFixedBrakingDistanceForPaths_DEPRECATED : 1;
 
 	/** If set, StopActiveMovement call will abort current path following request */
 	UE_DEPRECATED(5.5, "bStopMovementAbortPaths is deprecated, please use NavMovementProperties.bStopMovementAbortPaths instead.")
-	uint8 bStopMovementAbortPaths:1;
+	uint8 bStopMovementAbortPaths_DEPRECATED:1;
 	
 	UPROPERTY(EditAnywhere, Category = NavMovement)
 	FNavMovementProperties NavMovementProperties;
@@ -116,7 +116,7 @@ public:
 	/** Returns based location of controlled actor */
 	ENGINE_API virtual FBasedPosition GetFeetLocationBased() const override { return GetActorFeetLocationBased(); };
 	
-	ENGINE_API virtual void PostLoad() override;
+	ENGINE_API virtual void Serialize(FArchive& Ar) override;
 
 	/** Set nav agent properties from an object */
 	ENGINE_API virtual void UpdateNavAgent(const UObject& ObjectToUpdateFrom) override;
diff --git a/Engine/Source/Runtime/Engine/Private/Components/NavMovementComponent.cpp b/Engine/Source/Runtime/Engine/Private/Components/NavMovementComponent.cpp
--- a/Engine/Source/Runtime/Engine/Private/Components/NavMovementComponent.cpp
+++ b/Engine/Source/Runtime/Engine/Private/Components/NavMovementComponent.cpp
@@ -3,6 +3,7 @@
 #include "GameFramework/NavMovementComponent.h"
 #include "AI/NavigationSystemBase.h"
 #include "Components/CapsuleComponent.h"
+#include "UObject/FortniteReleaseBranchCustomObjectVersion.h"
 
 #include UE_INLINE_GENERATED_CPP_BY_NAME(NavMovementComponent)
 
@@ -13,11 +14,11 @@
 UNavMovementComponent::UNavMovementComponent(const FObjectInitializer& ObjectInitializer)
 	: Super(ObjectInitializer)
 PRAGMA_DISABLE_DEPRECATION_WARNINGS
-	, FixedPathBrakingDistance(0.0f)
-	, bUpdateNavAgentWithOwnersCollision(true)
-	, bUseAccelerationForPaths(false)
-	, bUseFixedBrakingDistanceForPaths(false)
-	, bStopMovementAbortPaths(true)
+	, FixedPathBrakingDistance_DEPRECATED(0.0f)
+	, bUpdateNavAgentWithOwnersCollision_DEPRECATED(true)
+	, bUseAccelerationForPaths_DEPRECATED(false)
+	, bUseFixedBrakingDistanceForPaths_DEPRECATED(false)
+	, bStopMovementAbortPaths_DEPRECATED(true)
 PRAGMA_ENABLE_DEPRECATION_WARNINGS
 {
 	bComponentShouldUpdatePhysicsVolume = true;
@@ -28,45 +29,23 @@ FBasedPosition UNavMovementComponent::GetActorFeetLocationBased() const
 	return FBasedPosition(NULL, GetActorFeetLocation());
 }
 
-void UNavMovementComponent::PostLoad()
+void UNavMovementComponent::Serialize(FArchive& Ar)
 {
-	Super::PostLoad();
+	Super::Serialize(Ar);
 
-#if WITH_EDITORONLY_DATA
-	PRAGMA_DISABLE_DEPRECATION_WARNINGS
-	
-	// copying over any default values to the new struct that houses the properties
-	if (FixedPathBrakingDistance != 0) 
+#if WITH_EDITOR
+	if (Ar.IsLoading() && GetLinkerCustomVersion(FFortniteReleaseBranchCustomObjectVersion::GUID) < FFortniteReleaseBranchCustomObjectVersion::NavMovementComponentMovingPropertiesToStruct)
 	{
-		NavMovementProperties.FixedPathBrakingDistance = FixedPathBrakingDistance;
-		FixedPathBrakingDistance = 0;
-	}
-	
-	if (!bUpdateNavAgentWithOwnersCollision)
-	{
-		NavMovementProperties.bUpdateNavAgentWithOwnersCollision = bUpdateNavAgentWithOwnersCollision;
-		bUpdateNavAgentWithOwnersCollision = true;
+		PRAGMA_DISABLE_DEPRECATION_WARNINGS
+
+		NavMovementProperties.FixedPathBrakingDistance = FixedPathBrakingDistance_DEPRECATED;
+		NavMovementProperties.bUpdateNavAgentWithOwnersCollision = bUpdateNavAgentWithOwnersCollision_DEPRECATED;
+		NavMovementProperties.bUseAccelerationForPaths = bUseAccelerationForPaths_DEPRECATED;
+		NavMovementProperties.bUseFixedBrakingDistanceForPaths = bUseFixedBrakingDistanceForPaths_DEPRECATED;
+		NavMovementProperties.bStopMovementAbortPaths = bStopMovementAbortPaths_DEPRECATED;
+
+		PRAGMA_ENABLE_DEPRECATION_WARNINGS
 	}
-	
-	if (bUseAccelerationForPaths)
-    {
-    	NavMovementProperties.bUseAccelerationForPaths = bUseAccelerationForPaths;
-    	bUseAccelerationForPaths = false;
-    }
-	
-	if (bUseFixedBrakingDistanceForPaths)
-    {
-    	NavMovementProperties.bUseFixedBrakingDistanceForPaths = bUseFixedBrakingDistanceForPaths;
-    	bUseFixedBrakingDistanceForPaths = false;
-    }
-    
-    if (!bStopMovementAbortPaths)
-    {
-    	NavMovementProperties.bStopMovementAbortPaths = bStopMovementAbortPaths;
-    	bStopMovementAbortPaths = true;
-    }
-	
-	PRAGMA_ENABLE_DEPRECATION_WARNINGS
 #endif // WITH_EDITORONLY_DATA	
 }
 
