Support for per instance/component post-processing AnimBP

Post-processing AnimBPs are currently set inside the skeletal mesh asset. This means each skel mesh component will run the post-processing AnimBP the referenced skel mesh asset has set.

This change allows us to set or override the post-processing AnimBP set in the skel mesh asset so that we have control over it on an instance/component level.

diff --git a/Engine/Source/Runtime/Engine/Classes/Components/SkeletalMeshComponent.h b/Engine/Source/Runtime/Engine/Classes/Components/SkeletalMeshComponent.h
--- a/Engine/Source/Runtime/Engine/Classes/Components/SkeletalMeshComponent.h
+++ b/Engine/Source/Runtime/Engine/Classes/Components/SkeletalMeshComponent.h
@@ -358,6 +358,17 @@ public:
 	TArray<TObjectPtr<UAnimInstance>> SubInstances;
 #endif
 
+	/** Post-processing AnimBP to use for the given skeletal mesh component, overriding the one set in the skeletal mesh asset. */
+	UPROPERTY(transient)
+	TSubclassOf<UAnimInstance> OverridePostProcessAnimBP;
+
+	/**
+	 * Get the post-processing AnimBP to be used for this skeletal mesh component.
+	 * In case an override post-processing AnimBP is set, it will always return the AnimBP class of that.
+	 * Otherwise the one set in skeletal mesh asset will be returned.
+	 */
+	ENGINE_API TSubclassOf<UAnimInstance> GetPostProcessAnimBPClassToBeUsed() const;
+
 	/** An instance created from the PostPhysicsBlueprint property of the skeletal mesh we're using,
 	 *  Runs after (and receives pose from) the main anim instance.
 	 */
@@ -365,6 +376,14 @@ public:
 	TObjectPtr<UAnimInstance> PostProcessAnimInstance;
 
 public:
+	/**
+	 * Set the post-processing AnimBP to be used for this skeletal mesh component.
+	 * In case an override post-processing AnimBP is set, the one set in skeletal mesh asset will be ignored and not used.
+	 * @param ReinitAnimInstances Can be false when called e.g. from the construction script in a Blueprint. True when this is called while
+	 *							  the game is running and the anim instances need to be re-initialized.
+	 */
+	UFUNCTION(BlueprintCallable, Category = "Components|SkeletalMesh")
+	ENGINE_API void SetOverridePostProcessAnimBP(TSubclassOf<UAnimInstance> InPostProcessAnimBlueprint, bool ReinitAnimInstances = true);
 
 	/** Toggles whether the post process blueprint will run for this component */
 	UFUNCTION(BlueprintCallable, Category="Components|SkeletalMesh")
diff --git a/Engine/Source/Runtime/Engine/Private/Components/SkeletalMeshComponent.cpp b/Engine/Source/Runtime/Engine/Private/Components/SkeletalMeshComponent.cpp
--- a/Engine/Source/Runtime/Engine/Private/Components/SkeletalMeshComponent.cpp
+++ b/Engine/Source/Runtime/Engine/Private/Components/SkeletalMeshComponent.cpp
@@ -701,7 +701,7 @@ bool USkeletalMeshComponent::NeedToSpawnPostPhysicsInstance(bool bForceReinit) c
 	if(GetSkeletalMeshAsset())
 	{
 		const UClass* MainInstanceClass = *AnimClass;
-		const UClass* ClassToUse = *GetSkeletalMeshAsset()->GetPostProcessAnimBlueprint();
+		const UClass* ClassToUse = *GetPostProcessAnimBPClassToBeUsed();
 		const UClass* CurrentClass = PostProcessAnimInstance ? PostProcessAnimInstance->GetClass() : nullptr;
 
 		const IAnimClassInterface* AnimClassInterface = IAnimClassInterface::GetFromClass(ClassToUse);
@@ -1093,7 +1093,7 @@ bool USkeletalMeshComponent::InitializeAnimScriptInstance(bool bForceReinit, boo
 		}
 
 		// May need to clear out the post physics instance
-		UClass* NewMeshInstanceClass = *SkelMesh->GetPostProcessAnimBlueprint();
+		UClass* NewMeshInstanceClass = *GetPostProcessAnimBPClassToBeUsed();
 		if(!NewMeshInstanceClass || NewMeshInstanceClass == *AnimClass || (PostProcessAnimInstance && PostProcessAnimInstance->CurrentSkeleton != GetSkeletalMeshAsset()->GetSkeleton()) || !GetSkeletalMeshAsset()->GetSkeleton())
 		{
 			PostProcessAnimInstance = nullptr;
@@ -1101,7 +1101,7 @@ bool USkeletalMeshComponent::InitializeAnimScriptInstance(bool bForceReinit, boo
 
 		if(NeedToSpawnPostPhysicsInstance(bForceReinit))
 		{
-			PostProcessAnimInstance = NewObject<UAnimInstance>(this, *SkelMesh->GetPostProcessAnimBlueprint());
+			PostProcessAnimInstance = NewObject<UAnimInstance>(this, *GetPostProcessAnimBPClassToBeUsed());
 
 			if(PostProcessAnimInstance)
 			{
@@ -1125,7 +1125,7 @@ bool USkeletalMeshComponent::InitializeAnimScriptInstance(bool bForceReinit, boo
 				bInitializedPostInstance = true;
 			}
 		}
-		else if (!SkelMesh->GetPostProcessAnimBlueprint().Get())
+		else if (!GetPostProcessAnimBPClassToBeUsed())
 		{
 			PostProcessAnimInstance = nullptr;
 		}
@@ -1264,12 +1264,12 @@ void USkeletalMeshComponent::PostEditChangeProperty(FPropertyChangedEvent& Prope
 			if(PostProcessAnimInstance)
 			{
 				UClass* CurrentClass = PostProcessAnimInstance->GetClass();
-				UClass* MeshClass = SkelMesh ? *SkelMesh->GetPostProcessAnimBlueprint() : nullptr;
+				UClass* MeshClass = SkelMesh ? *GetPostProcessAnimBPClassToBeUsed() : nullptr;
 				if(CurrentClass != MeshClass)
 				{
 					if(MeshClass)
 					{
-						PostProcessAnimInstance = NewObject<UAnimInstance>(this, *SkelMesh->GetPostProcessAnimBlueprint());
+						PostProcessAnimInstance = NewObject<UAnimInstance>(this, *GetPostProcessAnimBPClassToBeUsed());
 						PostProcessAnimInstance->InitializeAnimation();
 
 						if(HasBegunPlay())
@@ -4389,13 +4389,41 @@ bool USkeletalMeshComponent::ShouldPostUpdatePostProcessInstance() const
 	return PostProcessAnimInstance && PostProcessAnimInstance->NeedsUpdate() && !bDisablePostProcessBlueprint;
 }
 
-bool USkeletalMeshComponent::ShouldEvaluatePostProcessInstance() const
+void USkeletalMeshComponent::SetOverridePostProcessAnimBP(TSubclassOf<UAnimInstance> InPostProcessAnimBlueprint, bool ReinitAnimInstances)
 {
+	OverridePostProcessAnimBP = InPostProcessAnimBlueprint;
+
+	if (ReinitAnimInstances)
+	{
+		InitializeAnimScriptInstance();
+	}
+}
+
+TSubclassOf<UAnimInstance> USkeletalMeshComponent::GetPostProcessAnimBPClassToBeUsed() const
+{
+	if (OverridePostProcessAnimBP)
+	{
+		return OverridePostProcessAnimBP;
+	}
+
 	if (USkeletalMesh* SkelMesh = GetSkeletalMeshAsset())
 	{
-		if (!SkelMesh->ShouldEvaluatePostProcessAnimBP(GetPredictedLODLevel()))
+		return SkelMesh->GetPostProcessAnimBlueprint();
+	}
+
+	return nullptr;
+}
+
+bool USkeletalMeshComponent::ShouldEvaluatePostProcessInstance() const
+{
+	if (!OverridePostProcessAnimBP)
+	{
+		if (USkeletalMesh* SkelMesh = GetSkeletalMeshAsset())
 		{
-			return false;
+			if (!SkelMesh->ShouldEvaluatePostProcessAnimBP(GetPredictedLODLevel()))
+			{
+				return false;
+			}
 		}
 	}
 
