Material Layer Substrate Prototype
- Modify the Tree Model so that top level LayerNodes also contain a First sub node representing the default material attributes.
- Modify several checks where the num blends was expected to be 1 minus the number of layers
- In the UI, build the ViewModel for the tree view so that the layers and their content is displayed bottom up
- Introduce the new GraphRebuild using the tree datastructure of the Layers

#virtualized

diff --git a/Engine/Source/Editor/MaterialEditor/Private/MaterialPropertyHelpers.cpp b/Engine/Source/Editor/MaterialEditor/Private/MaterialPropertyHelpers.cpp
--- a/Engine/Source/Editor/MaterialEditor/Private/MaterialPropertyHelpers.cpp
+++ b/Engine/Source/Editor/MaterialEditor/Private/MaterialPropertyHelpers.cpp
@@ -142,7 +142,11 @@ void FMaterialPropertyHelpers::OnMaterialLayerAssetChanged(const struct FAssetDa
 			if (InMaterialFunction->Blends[Index] != LayerFunction)
 			{
 				InMaterialFunction->Blends[Index] = LayerFunction;
+#ifdef ENABLE_MATERIAL_LAYER_PROTOTYPE
+				InMaterialFunction->UnlinkLayerFromParent(Index); // 
+#else
 				InMaterialFunction->UnlinkLayerFromParent(Index + 1); // Blend indices are offset by 1, no blend for base layer
+#endif
 			}
 			break;
 		default:
diff --git a/Engine/Source/Editor/MaterialEditor/Private/SMaterialSubstrateTree.cpp b/Engine/Source/Editor/MaterialEditor/Private/SMaterialSubstrateTree.cpp
--- a/Engine/Source/Editor/MaterialEditor/Private/SMaterialSubstrateTree.cpp
+++ b/Engine/Source/Editor/MaterialEditor/Private/SMaterialSubstrateTree.cpp
@@ -205,8 +205,63 @@ TOptional<EItemDropZone> SMaterialSubstrateTreeItem::CanAcceptDrop(const FDragDr
 {
 	TSharedPtr<FLayerDragDropOp> LayerDragDropOperation = DragDropEvent.GetOperationAs<FLayerDragDropOp>();
 	TSharedPtr<FAssetDragDropOp> AssetDragDropOperation = DragDropEvent.GetOperationAs<FAssetDragDropOp>();
-	if (LayerDragDropOperation.IsValid() || AssetDragDropOperation.IsValid())
+	if (LayerDragDropOperation.IsValid())
 	{
+		// PREVENT LAYER DRAG FOR NOW
+		//return DropZone;
+	}
+	else if (AssetDragDropOperation.IsValid())
+	{
+
+		// Identify the type of asset
+		bool HasLayerFuncAsset = false;
+		bool HasBlendFuncAsset = false;
+		for (const FAssetData& AssetData : AssetDragDropOperation->GetAssets())
+		{
+			if (AssetData.AssetClassPath.GetAssetName() == TEXT("MaterialFunctionMaterialLayer"))
+			{
+				HasLayerFuncAsset = true;
+			}
+			else if (AssetData.AssetClassPath.GetAssetName() == TEXT("MaterialFunctionMaterialLayerBlend"))
+			{
+				HasBlendFuncAsset = true;
+			}
+		}
+
+		// Drop above or below could CREATE a new layer node:
+		int32 TargetNodeId = StackParameterData->ParameterInfo.Index;
+		int32 ParentNodeId = Tree->FunctionInstance->GetNodeParent(TargetNodeId);
+		auto ChildrenNodeId = Tree->FunctionInstance->GetNodeChildren(ParentNodeId);
+		int32 SiblingIdx = -1;
+		bool FoundTarget = ChildrenNodeId.Find(TargetNodeId, SiblingIdx);
+
+		switch (DropZone)
+		{
+		case EItemDropZone::AboveItem:
+			if (!HasLayerFuncAsset) // Can only add above if drop a new layer function
+			{
+				return TOptional<EItemDropZone>();
+			}
+			break;
+
+		case EItemDropZone::BelowItem:
+			if (!HasLayerFuncAsset  // Can only add under if drop a new layer function
+				|| (SiblingIdx == 0)) // Can't add BEFORE the first layer or sub layer
+			{
+				return TOptional<EItemDropZone>();
+			}
+			break;
+		case EItemDropZone::OntoItem:
+			if (!(HasLayerFuncAsset || HasBlendFuncAsset) // Can only drop valid assets
+				|| (!HasLayerFuncAsset && HasBlendFuncAsset && (SiblingIdx == 0))) // Or just a blendfunc on the first layer
+			{
+				return TOptional<EItemDropZone>();
+			}
+			break;
+		default:
+			break;
+		}
+
 		return DropZone;
 	}
 	return TOptional<EItemDropZone>();
@@ -234,31 +289,64 @@ FReply SMaterialSubstrateTreeItem::OnLayerDrop(const FDragDropEvent& DragDropEve
 
 		if (AssetDropOp.IsValid())
 		{
-			bool bIsBackgroundItem = IsNodeKeyForBackgroundParameter(StackParameterData->NodeKey);
+			// Identify the type of asset
+			bool HasLayerFuncAsset = false;
+			bool HasBlendFuncAsset = false;
+			for (const FAssetData& AssetData : AssetDropOp->GetAssets())
+			{
+				if (AssetData.AssetClassPath.GetAssetName() == TEXT("MaterialFunctionMaterialLayer"))
+				{
+					HasLayerFuncAsset = true;
+				}
+				else if (AssetData.AssetClassPath.GetAssetName() == TEXT("MaterialFunctionMaterialLayerBlend"))
+				{
+					HasBlendFuncAsset = true;
+				}
+			}
 
 			// Drop above or below could CREATE a new layer node:
 			int32 TargetNodeId = StackParameterData->ParameterInfo.Index;
-			int32 ParentNodeId = (bIsBackgroundItem ? TargetNodeId : Tree->FunctionInstance->GetNodeParent(TargetNodeId));
+			int32 TargetNodeDepth = Tree->FunctionInstance->GetNodeDepth(TargetNodeId);
+			int32 ParentNodeId = Tree->FunctionInstance->GetNodeParent(TargetNodeId);
 			auto ChildrenNodeId = Tree->FunctionInstance->GetNodeChildren(ParentNodeId);
 			int32 SiblingIdx = -1;
 			bool FoundTarget = ChildrenNodeId.Find(TargetNodeId, SiblingIdx);
 
+			bool DidModifyTree = false;
+
 			switch (DropZone)
 			{
-			// NOTE: the Sibling Idx should be swaped when we will disoplay the items from bottom up
+			// NOTE: The drop cases Above and Below take into account the fact that the list is displayed bottom up!
 			case EItemDropZone::AboveItem:
-				TargetNodeId = Tree->FunctionInstance->AppendLayerNode(ParentNodeId, SiblingIdx);
+				check(HasLayerFuncAsset); // Only add with a valid new LayerFunc asset
+				TargetNodeId = Tree->FunctionInstance->AppendLayerNode(ParentNodeId, SiblingIdx + 1); // Above means insert after
+				DidModifyTree = true;
 				break;
 
 			case EItemDropZone::BelowItem:
-				TargetNodeId = Tree->FunctionInstance->AppendLayerNode(ParentNodeId, SiblingIdx + 1);
+				check(HasLayerFuncAsset); // Only add with a valid new LayerFunc asset
+				TargetNodeId = Tree->FunctionInstance->AppendLayerNode(ParentNodeId, SiblingIdx); // Under means insert at
+				DidModifyTree = true;
+				break;
+
+			case EItemDropZone::OntoItem:
+				// Dropping LayerFunc asset (with blendfunc too maybe) on a top level layer means a NEW sub layer is created
+				if (HasLayerFuncAsset && (TargetNodeDepth <= 1)) // Top level Layers
+				{
+					// add a new layer in this target node last on the stack
+					TargetNodeId = Tree->FunctionInstance->AppendLayerNode(TargetNodeId, -1);
+					DidModifyTree = true;
+				}
+				else
+				{
+					// Assign the new asset(s) to this particular target node
+				}
 				break;
 			default:
 				break;
 			}
 
 			// Then drop
-			bool DidModifyTree = false;
 			for (const FAssetData& AssetData : AssetDropOp->GetAssets())
 			{
 				EMaterialParameterAssociation InAssociation = EMaterialParameterAssociation::GlobalParameter;
@@ -714,9 +802,8 @@ TSharedPtr<SWidget> SMaterialSubstrateTree::CreateContextMenu()
 	{
 		FSortedParamDataPtr StackParameterData = SelectedItemsArray[0];
 		
-		bool bIsBackgroundItem = IsNodeKeyForBackgroundParameter(StackParameterData->NodeKey);
-		bool bCanAppendSubLayer = FunctionInstance->CanAppendLayerNode(StackParameterData->ParameterInfo.Index) && !bIsBackgroundItem;
-		bool bCanRemoveLayer = FunctionInstance->CanRemoveLayerNode(StackParameterData->ParameterInfo.Index) && !bIsBackgroundItem;
+		bool bCanAppendSubLayer = FunctionInstance->CanAppendLayerNode(StackParameterData->ParameterInfo.Index);
+		bool bCanRemoveLayer = FunctionInstance->CanRemoveLayerNode(StackParameterData->ParameterInfo.Index);
 	
 		if (bCanAppendSubLayer)
 		{
@@ -1006,33 +1093,25 @@ struct FRecursiveCreateWidgetsContext
 	TSharedPtr<IPropertyHandle>	BlendHandle;
 };
 
-void SMaterialSubstrateTree::RecursiveCreateWidgets(FRecursiveCreateWidgetsContext* InContext, FNodeId InNodeId, TArray<TSharedPtr<FSortedParamData>>& InParentContainer, bool GenerateChildren, bool bIsBackgroundItem)
+void SMaterialSubstrateTree::RecursiveCreateWidgets(FRecursiveCreateWidgetsContext* InContext, FNodeId InNodeId, TArray<TSharedPtr<FSortedParamData>>& InParentContainer, bool GenerateChildren)
 {
 	auto Payload = FunctionInstance->Tree.Payloads[InNodeId];
 
-
 	TSharedRef<FSortedParamData> StackProperty = MakeShared<FSortedParamData>();
 	StackProperty->StackDataType = EStackDataType::Stack;
 	StackProperty->Parameter = InContext->Parameter;
 	StackProperty->ParameterInfo.Index = InNodeId;
 	StackProperty->NodeKey = FString::FromInt(StackProperty->ParameterInfo.Index);
 
-	if (bIsBackgroundItem)
-	{
-		StackProperty->NodeKey = CreateNodeKeyForBackgroundParameter(StackProperty->NodeKey);
-	}
-
-
 	if (GenerateChildren)
 	{
-		// Create the default background sub item representing the layer and blend func stored in the layer item
-		RecursiveCreateWidgets(InContext, InNodeId, StackProperty->Children, false, true);
-
 		// Sub layers
 		auto RootChildren = FunctionInstance->GetNodeChildren(InNodeId);
 		for (int i = 0; i < RootChildren.Num(); ++i)
 		{
-			RecursiveCreateWidgets(InContext, RootChildren[i], StackProperty->Children, false, false);
+			int Index = i;
+			Index = RootChildren.Num() - 1 - i; // Reverse the order to display the layers bottom up
+			RecursiveCreateWidgets(InContext, RootChildren[Index], StackProperty->Children, false);
 		}
 	}
 
@@ -1069,7 +1148,7 @@ void SMaterialSubstrateTree::RecursiveCreateWidgets(FRecursiveCreateWidgetsConte
 
 	}
 
-	if (!bIsBackgroundItem && Payload.Blend != -1)
+	if (Payload.Blend != -1)
 	{
 		TSharedRef<FSortedParamData> ChildProperty = MakeShared<FSortedParamData>();
 		ChildProperty->StackDataType = EStackDataType::Asset;
@@ -1242,7 +1321,9 @@ void SMaterialSubstrateTree::CreateGroupsWidget()
 			auto RootChildren = FunctionInstance->GetNodeChildren(-1);
 			for (int i = 0; i < RootChildren.Num(); ++i)
 			{
-				RecursiveCreateWidgets(&Context, RootChildren[i], LayerProperties, true, false);
+				int Index = i;
+				Index = RootChildren.Num() - 1 - i; // Reverse the order to display the layers bottom up
+				RecursiveCreateWidgets(&Context, RootChildren[Index], LayerProperties, true);
 			}
 		}
 #endif // ENABLE_MATERIAL_LAYER_PROTOTYPE
@@ -1294,11 +1375,11 @@ void SMaterialSubstrateTree::UpdateThumbnailMaterial(TEnumAsByte<EMaterialParame
 		}
 		MaterialToUpdate = MaterialEditorInstance->StoredBlendPreviews[ParameterIndex];
 	}
-
+	return;
 	if (MaterialToUpdate != nullptr)
 	{
 		// Need to invert index b/c layer properties is generated in reverse order
-	/*	TArray<TSharedPtr<FSortedParamData>> AssetChildren = LayerProperties[LayerProperties.Num() - 1 - InIndex]->Children;
+		TArray<TSharedPtr<FSortedParamData>> AssetChildren = LayerProperties[LayerProperties.Num() - 1 - InIndex]->Children;
 
 		TArray<FEditorParameterGroup> ParameterGroups;
 		for (TSharedPtr<FSortedParamData> AssetChild : AssetChildren)
@@ -1327,7 +1408,6 @@ void SMaterialSubstrateTree::UpdateThumbnailMaterial(TEnumAsByte<EMaterialParame
 		}
 
 		FMaterialPropertyHelpers::TransitionAndCopyParameters(MaterialToUpdate, ParameterGroups, true);
-*/
 	}
 }
 
diff --git a/Engine/Source/Editor/MaterialEditor/Private/SMaterialSubstrateTree.h b/Engine/Source/Editor/MaterialEditor/Private/SMaterialSubstrateTree.h
--- a/Engine/Source/Editor/MaterialEditor/Private/SMaterialSubstrateTree.h
+++ b/Engine/Source/Editor/MaterialEditor/Private/SMaterialSubstrateTree.h
@@ -74,7 +74,7 @@ public:
 	/** Builds the custom parameter groups category */
 	void CreateGroupsWidget();
 	using FNodeId = int32;
-	void RecursiveCreateWidgets(struct FRecursiveCreateWidgetsContext* Context, FNodeId InNodeId, TArray<TSharedPtr<FSortedParamData>>& InParentContainer, bool GenerateChildren, bool IsBackgroundItem);
+	void RecursiveCreateWidgets(struct FRecursiveCreateWidgetsContext* Context, FNodeId InNodeId, TArray<TSharedPtr<FSortedParamData>>& InParentContainer, bool GenerateChildren);
 
 
 	SMaterialLayersFunctionsInstanceWrapper* GetWrapper() { return Wrapper; }
diff --git a/Engine/Source/Runtime/Engine/Private/Materials/MaterialExpressions.cpp b/Engine/Source/Runtime/Engine/Private/Materials/MaterialExpressions.cpp
--- a/Engine/Source/Runtime/Engine/Private/Materials/MaterialExpressions.cpp
+++ b/Engine/Source/Runtime/Engine/Private/Materials/MaterialExpressions.cpp
@@ -8112,6 +8112,152 @@ void UMaterialExpressionMaterialAttributeLayers::PostEditChangeProperty(FPropert
 #if WITH_EDITOR
 void UMaterialExpressionMaterialAttributeLayers::RebuildLayerGraph(bool bReportErrors)
 {
+#if ENABLE_MATERIAL_LAYER_PROTOTYPE
+	const FMaterialLayersFunctions* LayersTree = (ParamLayers ? ParamLayers : &DefaultLayers);
+	const TArray<UMaterialFunctionInterface*>& Layers = GetLayers();
+	const TArray<UMaterialFunctionInterface*>& Blends = GetBlends();
+	const TArray<bool>& LayerStates = GetLayerStates();
+
+	// Pre-populate callers, we maintain these transient objects to avoid
+	// heavy UObject recreation as the graphs are frequently rebuilt
+	while (LayerCallers.Num() < Layers.Num())
+	{
+		LayerCallers.Add(NewObject<UMaterialExpressionMaterialFunctionCall>(GetTransientPackage()));
+	}
+	while (BlendCallers.Num() < Blends.Num())
+	{
+		BlendCallers.Add(NewObject<UMaterialExpressionMaterialFunctionCall>(GetTransientPackage()));
+	}
+
+	// Reset graph connectivity
+	bIsLayerGraphBuilt = false;
+	NumActiveLayerCallers = 0;
+	NumActiveBlendCallers = 0;
+
+	if (ValidateLayerConfiguration(nullptr, bReportErrors))
+	{
+		// Initialize layer function callers
+		for (int32 LayerIndex = 0; LayerIndex < Layers.Num(); ++LayerIndex)
+		{
+			if (Layers[LayerIndex] && LayerStates[LayerIndex])
+			{
+				LayerCallers[LayerIndex]->MaterialFunction = Layers[LayerIndex];
+				LayerCallers[LayerIndex]->FunctionParameterInfo.Association = EMaterialParameterAssociation::LayerParameter;
+				LayerCallers[LayerIndex]->FunctionParameterInfo.Index = LayerIndex;
+
+				Layers[LayerIndex]->GetInputsAndOutputs(LayerCallers[LayerIndex]->FunctionInputs, LayerCallers[LayerIndex]->FunctionOutputs);
+				for (FFunctionExpressionOutput& FunctionOutput : LayerCallers[LayerIndex]->FunctionOutputs)
+				{
+					LayerCallers[LayerIndex]->Outputs.Add(FunctionOutput.Output);
+				}
+
+				// Optional: Single material attributes input, the base input to the stack
+				if (LayerCallers[LayerIndex]->FunctionInputs.Num() > 0)
+				{
+					if (Input.GetTracedInput().Expression)
+					{
+						LayerCallers[LayerIndex]->FunctionInputs[0].Input = Input;
+					}
+				}
+
+				// Recursively run through internal functions to allow connection of inputs/outputs
+				LayerCallers[LayerIndex]->UpdateFromFunctionResource();
+
+				++NumActiveLayerCallers;
+			}
+			else
+			{
+				// Empty entries for opaque layers
+				LayerCallers[LayerIndex]->MaterialFunction = nullptr;
+			}
+		}
+
+		for (int32 BlendIndex = 0; BlendIndex < Blends.Num(); ++BlendIndex)
+		{
+			if (Blends[BlendIndex])
+			{
+				BlendCallers[BlendIndex]->MaterialFunction = Blends[BlendIndex];
+				BlendCallers[BlendIndex]->FunctionParameterInfo.Association = EMaterialParameterAssociation::BlendParameter;
+				BlendCallers[BlendIndex]->FunctionParameterInfo.Index = BlendIndex;
+
+				Blends[BlendIndex]->GetInputsAndOutputs(BlendCallers[BlendIndex]->FunctionInputs, BlendCallers[BlendIndex]->FunctionOutputs);
+				for (FFunctionExpressionOutput& FunctionOutput : BlendCallers[BlendIndex]->FunctionOutputs)
+				{
+					BlendCallers[BlendIndex]->Outputs.Add(FunctionOutput.Output);
+				}
+
+				// Recursively run through internal functions to allow connection of inputs/ouputs
+				BlendCallers[BlendIndex]->UpdateFromFunctionResource();
+
+				++NumActiveBlendCallers;
+			}
+			else
+			{
+				// Empty entries for opaque layers
+				BlendCallers[BlendIndex]->MaterialFunction = nullptr;
+			}
+		}
+
+		// Empty out unused callers
+		for (int32 CallerIndex = Layers.Num(); CallerIndex < LayerCallers.Num(); ++CallerIndex)
+		{
+			LayerCallers[CallerIndex]->MaterialFunction = nullptr;
+		}
+
+		for (int32 CallerIndex = Blends.Num(); CallerIndex < BlendCallers.Num(); ++CallerIndex)
+		{
+			BlendCallers[CallerIndex]->MaterialFunction = nullptr;
+		}
+
+		struct TreeGraphBuilder
+		{
+			UMaterialExpressionMaterialAttributeLayers* This = nullptr;
+			const FMaterialLayersFunctions* Tree = nullptr;
+
+			UMaterialExpressionMaterialFunctionCall* AssembleBranchGraph(int32 InSourceId)
+			{
+				UMaterialExpressionMaterialFunctionCall* PreviousLayerInput = nullptr;
+
+				if (InSourceId != -1)
+				{
+					int32 BackgroundLayerFuncIdx = Tree->GetLayerFuncIndex(InSourceId);
+					PreviousLayerInput = (This->LayerCallers.IsValidIndex(BackgroundLayerFuncIdx) ? This->LayerCallers[BackgroundLayerFuncIdx] : nullptr);
+				}
+
+				auto NodeChildrenIds = Tree->GetNodeChildren(InSourceId);
+				for (auto NodeId : NodeChildrenIds)
+				{
+					UMaterialExpressionMaterialFunctionCall* CurrentLayerInput = AssembleBranchGraph(NodeId);
+
+					int32 BlendFuncIdx = Tree->GetBlendFuncIndex(NodeId);
+					UMaterialExpressionMaterialFunctionCall* CurrentBlendCaller = (This->BlendCallers.IsValidIndex(BlendFuncIdx) ? This->BlendCallers[BlendFuncIdx] : nullptr);
+
+					if (CurrentBlendCaller && CurrentBlendCaller->MaterialFunction)
+					{
+						CurrentBlendCaller->FunctionInputs[0].Input.Connect(0, PreviousLayerInput);
+						CurrentBlendCaller->FunctionInputs[1].Input.Connect(0, CurrentLayerInput);
+
+						PreviousLayerInput = CurrentBlendCaller;
+					}
+					else
+					{
+						PreviousLayerInput = ToRawPtr(CurrentLayerInput);
+					}
+				}
+
+				return PreviousLayerInput;
+			}
+		} BuildTreeGraph = { this, LayersTree };
+
+		OutputCaller = BuildTreeGraph.AssembleBranchGraph(-1);
+
+		if (OutputCaller && (OutputCaller->MaterialFunction == nullptr))
+		{	
+			OutputCaller = nullptr;
+		}
+
+#else
+
 	const TArray<UMaterialFunctionInterface*>& Layers = GetLayers();
 	const TArray<UMaterialFunctionInterface*>& Blends = GetBlends();
 	const TArray<bool>& LayerStates = GetLayerStates();
@@ -8134,6 +8280,8 @@ void UMaterialExpressionMaterialAttributeLayers::RebuildLayerGraph(bool bReportE
 
 	if (ValidateLayerConfiguration(nullptr, bReportErrors))
 	{
+
+
 		// Initialize layer function callers
 		for (int32 LayerIndex = 0; LayerIndex < Layers.Num(); ++LayerIndex)
 		{
@@ -8230,6 +8378,8 @@ void UMaterialExpressionMaterialAttributeLayers::RebuildLayerGraph(bool bReportE
 			}
 		}
 
+#endif
+
 		bIsLayerGraphBuilt = true;
 	}
 	else if (bReportErrors)
@@ -8336,15 +8486,24 @@ bool UMaterialExpressionMaterialAttributeLayers::ValidateLayerConfiguration(FMat
 		}
 		
 		// Null blends signify an opaque layer so count as valid for the sake of graph validation
+#if ENABLE_MATERIAL_LAYER_PROTOTYPE
+		if (Layers[BlendIndex] && LayerStates[BlendIndex])
+#else
 		if (Layers[BlendIndex+1] && LayerStates[BlendIndex+1])
+#endif
 		{
 			++NumActiveBlends;
 		}
 	}
 
+#if ENABLE_MATERIAL_LAYER_PROTOTYPE
+	bool bValidGraphLayout = (NumActiveLayers == 0 && NumActiveBlends == 0)		// Pass-through
+		|| (NumActiveLayers == NumActiveBlends);			    // Single layer
+#else
 	bool bValidGraphLayout = (NumActiveLayers == 0 && NumActiveBlends == 0)		// Pass-through
 		|| (NumActiveLayers == 1 && NumActiveBlends == 0)						// Single layer
 		|| (NumActiveLayers >= 2 && NumActiveBlends == NumActiveLayers - 1);	// Blend graph
+#endif
 
 	if (!bValidGraphLayout)
 	{
@@ -8449,6 +8608,14 @@ int32 UMaterialExpressionMaterialAttributeLayers::Compile(FMaterialCompiler* Com
 
 	if (ValidateLayerConfiguration(Compiler, true) && bIsLayerGraphBuilt)
 	{
+#if ENABLE_MATERIAL_LAYER_PROTOTYPE
+		if (OutputCaller)
+		{
+			Result = OutputCaller->Compile(Compiler, 0);
+		}
+		else 
+#endif //ENABLE_MATERIAL_LAYER_PROTOTYPE
+
 		if (NumActiveBlendCallers > 0 && BlendCallers[NumActiveBlendCallers-1]->MaterialFunction)
 		{
 			// Multiple blended layers
@@ -8505,6 +8672,11 @@ FSubstrateOperator* UMaterialExpressionMaterialAttributeLayers::SubstrateGenerat
 	//Behaviour mirrors the behaviour of the existing compile function
 	if (ValidateLayerConfiguration(Compiler, true) && bIsLayerGraphBuilt)
 	{
+		if (OutputCaller)
+		{
+			OutOperator = OutputCaller->SubstrateGenerateMaterialTopologyTree(Compiler, Parent, 0);
+		}
+/*
 		if (NumActiveBlendCallers > 0 && BlendCallers[NumActiveBlendCallers - 1]->MaterialFunction)
 		{
 			// Multiple blended layers
@@ -8514,7 +8686,7 @@ FSubstrateOperator* UMaterialExpressionMaterialAttributeLayers::SubstrateGenerat
 		{
 			// Single layer
 			OutOperator = LayerCallers[NumActiveLayerCallers - 1]->SubstrateGenerateMaterialTopologyTree(Compiler, Parent, 0);
-		}
+		} */
 		else if (NumActiveLayerCallers == 0)
 		{
 			// Pass-through
@@ -17916,10 +18088,12 @@ void FMaterialLayersFunctions::PostSerialize(const FArchive& Ar)
 void FMaterialLayersFunctions::AddDefaultBackgroundLayer()
 {
 #if ENABLE_MATERIAL_LAYER_PROTOTYPE
-	// This call is only valid for the very first layer created in the Layer tree.
-	check(Layers.Num() == 0 && Tree.IsEmpty());
-	Tree.AddNode({ 0, -1 }, -1);
-#endif
+	// This call is only valid for the very first layer created from the constructor.
+	check(Layers.Num() == 0 && Blends.Num() == 0 && Tree.IsEmpty());
+	AppendLayerNode(-1, -1);
+	FText LayerName = FText(LOCTEXT("Background", "Background"));
+	EditorOnly.LayerNames[0] = LayerName;
+#else
 
 	// Default to a non-blended "background" layer
 	Layers.AddDefaulted();	
@@ -17931,6 +18105,8 @@ void FMaterialLayersFunctions::AddDefaultBackgroundLayer()
 	// Default constructor assigning different guids will break FStructUtils::AttemptToFindUninitializedScriptStructMembers
 	EditorOnly.LayerGuids.Add(BackgroundGuid);
 	EditorOnly.LayerLinkStates.Add(EMaterialLayerLinkState::NotFromParent);
+
+#endif // ENABLE_MATERIAL_LAYER_PROTOTYPE
 }
 
 int32 FMaterialLayersFunctions::AppendBlendedLayer()
@@ -17959,18 +18135,61 @@ int32 FMaterialLayersFunctions::AddLayerCopy(const FMaterialLayersFunctionsRunti
 	const int32 LayerIndex = Layers.Num();
 
 	Layers.Add(Source.Layers[SourceLayerIndex]);
+
+#if ENABLE_MATERIAL_LAYER_PROTOTYPE
+	// Legacy parents pre substrate support
+	if (Source.Blends.Num() < Source.Layers.Num())
+	{
+		if (LayerIndex > 0)
+		{
+			Blends.Add(Source.Blends[SourceLayerIndex - 1]);
+		}
+		else
+		{
+			Blends.AddDefaulted();
+		}
+	}
+	// Source is substrate tree, same number of layersand blends
+	else
+	{
+		Blends.Add(Source.Blends[SourceLayerIndex]);
+	}
+#else
 	if (LayerIndex > 0)
 	{
 		Blends.Add(Source.Blends[SourceLayerIndex - 1]);
 	}
+#endif
 
 	EditorOnly.LayerStates.Add(bVisible);
 	EditorOnly.LayerNames.Add(SourceEditorOnly.LayerNames[SourceLayerIndex]);
 	EditorOnly.RestrictToLayerRelatives.Add(SourceEditorOnly.RestrictToLayerRelatives[SourceLayerIndex]);
+
+#if ENABLE_MATERIAL_LAYER_PROTOTYPE
+	// Legacy parents pre substrate support
+	if (SourceEditorOnly.RestrictToBlendRelatives.Num() < SourceEditorOnly.RestrictToLayerRelatives.Num())
+	{
+		if (LayerIndex > 0)
+		{
+			EditorOnly.RestrictToBlendRelatives.Add(SourceEditorOnly.RestrictToBlendRelatives[SourceLayerIndex - 1]);
+}
+		else
+		{
+			EditorOnly.RestrictToBlendRelatives.AddDefaulted();
+		}
+	}
+	// Source is substrate tree, same number of layersand blends
+	else
+	{
+		EditorOnly.RestrictToBlendRelatives.Add(SourceEditorOnly.RestrictToBlendRelatives[SourceLayerIndex]);
+	}
+#else
 	if (LayerIndex > 0)
 	{
 		EditorOnly.RestrictToBlendRelatives.Add(SourceEditorOnly.RestrictToBlendRelatives[SourceLayerIndex - 1]);
 	}
+#endif
+
 	EditorOnly.LayerGuids.Add(SourceEditorOnly.LayerGuids[SourceLayerIndex]);
 	EditorOnly.LayerLinkStates.Add(LinkState);
 
@@ -17986,12 +18205,20 @@ void FMaterialLayersFunctions::InsertLayerCopy(const FMaterialLayersFunctionsRun
 	check(LinkState != EMaterialLayerLinkState::Uninitialized);
 	check(LayerIndex > 0);
 	Layers.Insert(Source.Layers[SourceLayerIndex], LayerIndex);
+#if ENABLE_MATERIAL_LAYER_PROTOTYPE
+	Blends.Insert(Source.Blends[SourceLayerIndex], LayerIndex);
+#else
 	Blends.Insert(Source.Blends[SourceLayerIndex - 1], LayerIndex - 1);
+#endif
 
 	EditorOnly.LayerStates.Insert(SourceEditorOnly.LayerStates[SourceLayerIndex], LayerIndex);
 	EditorOnly.LayerNames.Insert(SourceEditorOnly.LayerNames[SourceLayerIndex], LayerIndex);
 	EditorOnly.RestrictToLayerRelatives.Insert(SourceEditorOnly.RestrictToLayerRelatives[SourceLayerIndex], LayerIndex);
+#if ENABLE_MATERIAL_LAYER_PROTOTYPE
+	EditorOnly.RestrictToBlendRelatives.Insert(SourceEditorOnly.RestrictToBlendRelatives[SourceLayerIndex], LayerIndex);
+#else
 	EditorOnly.RestrictToBlendRelatives.Insert(SourceEditorOnly.RestrictToBlendRelatives[SourceLayerIndex - 1], LayerIndex - 1);
+#endif
 	EditorOnly.LayerGuids.Insert(SourceEditorOnly.LayerGuids[SourceLayerIndex], LayerIndex);
 	EditorOnly.LayerLinkStates.Insert(LinkState, LayerIndex);
 }
@@ -18000,6 +18227,16 @@ void FMaterialLayersFunctions::RemoveBlendedLayerAt(int32 Index)
 {
 	if (Layers.IsValidIndex(Index))
 	{
+#if ENABLE_MATERIAL_LAYER_PROTOTYPE
+		check(Layers.IsValidIndex(Index) && Blends.IsValidIndex(Index));
+		Layers.RemoveAt(Index);
+		Blends.RemoveAt(Index);
+
+		check(EditorOnly.LayerStates.IsValidIndex(Index) &&
+			EditorOnly.LayerNames.IsValidIndex(Index) &&
+			EditorOnly.RestrictToLayerRelatives.IsValidIndex(Index) &&
+			EditorOnly.RestrictToBlendRelatives.IsValidIndex(Index));
+#else
 		check(Layers.IsValidIndex(Index) && Blends.IsValidIndex(Index - 1));
 		Layers.RemoveAt(Index);
 		Blends.RemoveAt(Index - 1);
@@ -18008,6 +18245,7 @@ void FMaterialLayersFunctions::RemoveBlendedLayerAt(int32 Index)
 			EditorOnly.LayerNames.IsValidIndex(Index) &&
 			EditorOnly.RestrictToLayerRelatives.IsValidIndex(Index) &&
 			EditorOnly.RestrictToBlendRelatives.IsValidIndex(Index - 1));
+#endif
 
 		if (EditorOnly.LayerLinkStates[Index] != EMaterialLayerLinkState::NotFromParent)
 		{
@@ -18020,7 +18258,11 @@ void FMaterialLayersFunctions::RemoveBlendedLayerAt(int32 Index)
 		EditorOnly.LayerStates.RemoveAt(Index);
 		EditorOnly.LayerNames.RemoveAt(Index);
 		EditorOnly.RestrictToLayerRelatives.RemoveAt(Index);
+#if ENABLE_MATERIAL_LAYER_PROTOTYPE
+		EditorOnly.RestrictToBlendRelatives.RemoveAt(Index);
+#else
 		EditorOnly.RestrictToBlendRelatives.RemoveAt(Index - 1);
+#endif
 		EditorOnly.LayerGuids.RemoveAt(Index);
 		EditorOnly.LayerLinkStates.RemoveAt(Index);
 	}
@@ -18033,11 +18275,19 @@ void FMaterialLayersFunctions::MoveBlendedLayer(int32 SrcLayerIndex, int32 DstLa
 	if (SrcLayerIndex != DstLayerIndex)
 	{
 		Layers.Swap(SrcLayerIndex, DstLayerIndex);
+#if ENABLE_MATERIAL_LAYER_PROTOTYPE
+		Blends.Swap(SrcLayerIndex, DstLayerIndex);
+#else
 		Blends.Swap(SrcLayerIndex - 1, DstLayerIndex - 1);
+#endif
 		EditorOnly.LayerStates.Swap(SrcLayerIndex, DstLayerIndex);
 		EditorOnly.LayerNames.Swap(SrcLayerIndex, DstLayerIndex);
 		EditorOnly.RestrictToLayerRelatives.Swap(SrcLayerIndex, DstLayerIndex);
+#if ENABLE_MATERIAL_LAYER_PROTOTYPE
+		EditorOnly.RestrictToBlendRelatives.Swap(SrcLayerIndex, DstLayerIndex);
+#else
 		EditorOnly.RestrictToBlendRelatives.Swap(SrcLayerIndex - 1, DstLayerIndex - 1);
+#endif
 		EditorOnly.LayerGuids.Swap(SrcLayerIndex, DstLayerIndex);
 		EditorOnly.LayerLinkStates.Swap(SrcLayerIndex, DstLayerIndex);
 	}
@@ -18133,7 +18383,11 @@ bool FMaterialLayersFunctions::MatchesParent(const FMaterialLayersFunctionsRunti
 		{
 			return false;
 		}
+#if ENABLE_MATERIAL_LAYER_PROTOTYPE
+		if (LayerIndex > 0 && Runtime.Blends[LayerIndex] != ParentRuntime.Blends[ParentLayerIndex])
+#else
 		if (LayerIndex > 0 && Runtime.Blends[LayerIndex - 1] != ParentRuntime.Blends[ParentLayerIndex - 1])
+#endif
 		{
 			return false;
 		}
@@ -18201,7 +18455,11 @@ bool FMaterialLayersFunctions::ResolveParent(const FMaterialLayersFunctionsRunti
 			{
 				// See if we match layer in parent
 				if (Runtime.Layers[LayerIndex] == ParentRuntime.Layers[ParentLayerIndex] &&
-					(LayerIndex == 0 || Runtime.Blends[LayerIndex - 1] == ParentRuntime.Blends[ParentLayerIndex - 1]))
+#if ENABLE_MATERIAL_LAYER_PROTOTYPE
+				(LayerIndex == 0 || Runtime.Blends[LayerIndex] == ParentRuntime.Blends[ParentLayerIndex]))
+#else
+				(LayerIndex == 0 || Runtime.Blends[LayerIndex - 1] == ParentRuntime.Blends[ParentLayerIndex - 1]))
+#endif
 				{
 					// Parent layer matches, so link to parent
 					ResolvedLayerIndex = ResolvedLayers.AddLayerCopy(ParentRuntime, ParentEditorOnly, ParentLayerIndex, bLayerVisible, EMaterialLayerLinkState::LinkedToParent);
@@ -18248,7 +18506,7 @@ bool FMaterialLayersFunctions::ResolveParent(const FMaterialLayersFunctionsRunti
 
 			// Update the link state, depending on if we can find this layer in the parent
 			ResolvedLayers.AddLayerCopy(Runtime, EditorOnly, LayerIndex, bLayerVisible, (ParentLayerIndex == INDEX_NONE) ? EMaterialLayerLinkState::NotFromParent : EMaterialLayerLinkState::UnlinkedFromParent);
-			ParentLayerIndices.Add(ParentLayerIndex);
+ 			ParentLayerIndices.Add(ParentLayerIndex);
 		}
 	}
 
@@ -18361,7 +18619,11 @@ void FMaterialLayersFunctions::Validate(const FMaterialLayersFunctionsRuntimeDat
 {
 	if (Runtime.Layers.Num() > 0)
 	{
+#if ENABLE_MATERIAL_LAYER_PROTOTYPE
+		check(Runtime.Blends.Num() == Runtime.Layers.Num());
+#else
 		check(Runtime.Blends.Num() == Runtime.Layers.Num() - 1);
+#endif // ENABLE_MATERIAL_LAYER_PROTOTYPE
 		check(Runtime.Layers.Num() == EditorOnly.LayerStates.Num());
 		check(Runtime.Layers.Num() == EditorOnly.LayerNames.Num());
 		check(Runtime.Layers.Num() == EditorOnly.LayerGuids.Num());
@@ -18422,10 +18684,17 @@ FLayerNodeId FMaterialLayersFunctions::AppendLayerNode(FLayerNodeId InParent, in
 	if (!CanAppendLayerNode(InParent))
 		return FMaterialLayersFunctionsTree::InvalidId;
 
+	// if new layer node's parent is root then add a L1 group layer and THEN the L2 first node
+	if (InParent == FMaterialLayersFunctionsTree::InvalidId)
+	{
+		const int32 LayerIndex = AppendBlendedLayer();
+		InParent = Tree.AddNode({ LayerIndex, LayerIndex }, InParent, InSiblingIndex);
+	}
+
 	const int32 LayerIndex = AppendBlendedLayer();
 
 	// InPayload
-	FLayerNodeId NewNodeId = Tree.AddNode({ LayerIndex, LayerIndex - 1 }, InParent, InSiblingIndex);
+	FLayerNodeId NewNodeId = Tree.AddNode({ LayerIndex, LayerIndex }, InParent, InSiblingIndex);
 
 	MLFT::Log(this, __FUNCTION__);
 
@@ -18437,10 +18706,12 @@ FLayerNodeId FMaterialLayersFunctions::AppendLayerNode(FLayerNodeId InParent, in
 
 bool FMaterialLayersFunctions::CanRemoveLayerNode(FLayerNodeId InNodeId) const
 {
-	// NOTE: Current design prevent to delete 1st default node
+	// Only if really valid
 	if (Tree.IsValidId(InNodeId))
 	{
-		if (InNodeId != 0)
+		// NOTE: Current design guarantee first sub layer in any layer, and first top layer
+		FLayerNodeId SiblingHead = Tree.GetSiblingHeadId(InNodeId);
+		if (InNodeId != SiblingHead)
 		{
 			return true;
 		}
diff --git a/Engine/Source/Runtime/Engine/Public/Materials/MaterialExpressionMaterialAttributeLayers.h b/Engine/Source/Runtime/Engine/Public/Materials/MaterialExpressionMaterialAttributeLayers.h
--- a/Engine/Source/Runtime/Engine/Public/Materials/MaterialExpressionMaterialAttributeLayers.h
+++ b/Engine/Source/Runtime/Engine/Public/Materials/MaterialExpressionMaterialAttributeLayers.h
@@ -74,6 +74,9 @@ class UMaterialExpressionMaterialAttributeLayers : public UMaterialExpression
 	UPROPERTY(Transient)
 	int32 NumActiveBlendCallers;
 
+	UPROPERTY(Transient)
+	TObjectPtr<UMaterialExpressionMaterialFunctionCall> OutputCaller;
+
 	UPROPERTY(Transient)
 	bool bIsLayerGraphBuilt;
 
diff --git a/Engine/Source/Runtime/Engine/Public/Materials/MaterialLayersFunctions.h b/Engine/Source/Runtime/Engine/Public/Materials/MaterialLayersFunctions.h
--- a/Engine/Source/Runtime/Engine/Public/Materials/MaterialLayersFunctions.h
+++ b/Engine/Source/Runtime/Engine/Public/Materials/MaterialLayersFunctions.h
@@ -518,18 +518,24 @@ struct FMaterialLayersFunctions : public FMaterialLayersFunctionsRuntimeData
 
 	void ToggleBlendedLayerVisibility(int32 Index)
 	{
+		if (Index < 0)
+			return;
 		check(EditorOnly.LayerStates.IsValidIndex(Index));
 		EditorOnly.LayerStates[Index] = !EditorOnly.LayerStates[Index];
 	}
 
 	void SetBlendedLayerVisibility(int32 Index, bool InNewVisibility)
 	{
+		if (Index < 0)
+			return;
 		check(EditorOnly.LayerStates.IsValidIndex(Index));
 		EditorOnly.LayerStates[Index] = InNewVisibility;
 	}
 
 	bool GetLayerVisibility(int32 Index) const
 	{
+		if (Index < 0)
+			return true;
 		check(EditorOnly.LayerStates.IsValidIndex(Index));
 		return EditorOnly.LayerStates[Index];
 	}
