[Material Editor] Fixed bug in StringTemplate when handling %% escape sequences.

diff --git a/Engine/Source/Runtime/Engine/Private/Materials/StringTemplate.cpp b/Engine/Source/Runtime/Engine/Private/Materials/StringTemplate.cpp
--- a/Engine/Source/Runtime/Engine/Private/Materials/StringTemplate.cpp
+++ b/Engine/Source/Runtime/Engine/Private/Materials/StringTemplate.cpp
@@ -116,11 +116,16 @@ bool FStringTemplate::Load(FString TemplateStringArg, FErrorInfo& OutErrorinfo)
 				PushChunk(true);
 				NumNamedParameters += 1;
 			}
-			else
+			else if (Accept('%')) // e.g. %%
 			{
-				Accept('%'); // e.g. %%
+				ChunkLength -= 1; // take one '%'
 				NewChunks.Add({ TEXT("%"), false });
 			}
+			else
+			{
+				OutErrorinfo.Message = TEXT("unrecognized % sequence. Are you missing an additional '%' for the modulo sequence?");
+				return false;
+			}
 		}
 		else
 		{
