Align virtual texture tile border size to multiple of two instead of power of two.
This fixes an issue where the streaming mips asset for runtime virtual texture would mismatch the original runtime virtual texture.

diff --git a/Engine/Source/Runtime/Engine/Private/Texture2D.cpp b/Engine/Source/Runtime/Engine/Private/Texture2D.cpp
--- a/Engine/Source/Runtime/Engine/Private/Texture2D.cpp
+++ b/Engine/Source/Runtime/Engine/Private/Texture2D.cpp
@@ -1691,7 +1691,7 @@ void FVirtualTexture2DResource::InitializeEditorResources(IVirtualTexture* InVir
 		}
 		else
 		{
-			const bool bCopyUnwantedBordersForAlignment = VTData->TileBorderSize <= 2 && IsBlockCompressedFormat(PixelFormat);
+			const bool bCopyUnwantedBordersForAlignment = (VTData->TileBorderSize % 4 != 0) && IsBlockCompressedFormat(PixelFormat);
 			const uint32 MipScaleFactor = (1u << MipLevel);
 			const uint32 MipWidthInTiles = FMath::DivideAndRoundUp(GetNumTilesX(), MipScaleFactor);
 			const uint32 MipHeightInTiles = FMath::DivideAndRoundUp(GetNumTilesY(), MipScaleFactor);
@@ -1726,7 +1726,7 @@ void FVirtualTexture2DResource::InitializeEditorResources(IVirtualTexture* InVir
 			Texture2DRHI = RHICreateTexture(Desc);
 
 			// We want to strip borders when compositing tiles since we're just laying out tiles in a regular texture.
-			// But if we have block compressed formats with border less than 4 then doing this will lead to an unaligned copy. We keep the small unwanted borders in that case.
+			// But if we have block compressed formats with border not divisible by 4 then doing this will lead to an unaligned copy. We keep the small unwanted borders in that case.
 			const EVTProducePageFlags ProducePageFlags = bCopyUnwantedBordersForAlignment ? EVTProducePageFlags::None : EVTProducePageFlags::SkipPageBorders;
 
 			TArray<IVirtualTextureFinalizer*> Finalizers;
diff --git a/Engine/Source/Runtime/Engine/Private/TextureDerivedData.cpp b/Engine/Source/Runtime/Engine/Private/TextureDerivedData.cpp
--- a/Engine/Source/Runtime/Engine/Private/TextureDerivedData.cpp
+++ b/Engine/Source/Runtime/Engine/Private/TextureDerivedData.cpp
@@ -1387,7 +1387,7 @@ static void GetTextureBuildSettings(
 
 		FVirtualTextureBuildSettings VirtualTextureBuildSettings;
 		Texture.GetVirtualTextureBuildSettings(VirtualTextureBuildSettings);
-		OutBuildSettings.VirtualTextureTileSize = FMath::RoundUpToPowerOfTwo(VirtualTextureBuildSettings.TileSize);
+		OutBuildSettings.VirtualTextureTileSize = FVirtualTextureBuildSettings::ClampAndAlignTileSize(VirtualTextureBuildSettings.TileSize);
 
 		// Apply any LOD group tile size bias here
 		const int32 TileSizeBias = TextureLODSettings.GetTextureLODGroup(Texture.LODGroup).VirtualTextureTileSizeBias;
@@ -1400,8 +1400,8 @@ static void GetTextureBuildSettings(
 		// 0 is a valid value for border size
 		// 1 would be OK in some cases, but breaks BC compressed formats, since it will result in physical tiles that aren't divisible by block size (4)
 		// Could allow border size of 1 for non BC compressed virtual textures, but somewhat complicated to get that correct, especially with multiple layers
-		// Doesn't seem worth the complexity for now, so clamp the size to be at least 2
-		OutBuildSettings.VirtualTextureBorderSize = (VirtualTextureBuildSettings.TileBorderSize > 0) ? FMath::RoundUpToPowerOfTwo(FMath::Max(VirtualTextureBuildSettings.TileBorderSize, 2)) : 0;
+		// Doesn't seem worth the complexity for now, so ensure we use multiple of 2
+		OutBuildSettings.VirtualTextureBorderSize = FVirtualTextureBuildSettings::ClampAndAlignTileBorderSize(VirtualTextureBuildSettings.TileBorderSize);
 	}
 	else
 	{
