[AutoRTFM] Fix build

FAutomationTestBase::GetTestFlags() has different return types between different release branches.
Use type traits to infer the return type so this compiles for both old and new signatures.

diff --git a/Engine/Source/Runtime/Engine/Private/Tests/AutoRTFM/AutoRTFMActorComponentTests.cpp b/Engine/Source/Runtime/Engine/Private/Tests/AutoRTFM/AutoRTFMActorComponentTests.cpp
--- a/Engine/Source/Runtime/Engine/Private/Tests/AutoRTFM/AutoRTFMActorComponentTests.cpp
+++ b/Engine/Source/Runtime/Engine/Private/Tests/AutoRTFM/AutoRTFMActorComponentTests.cpp
@@ -14,6 +14,8 @@
 #include "Chaos/Core.h"
 #include "Engine/World.h"
 
+#include <utility>
+
 #if WITH_DEV_AUTOMATION_TESTS
 
 namespace
@@ -59,7 +61,10 @@ public:
 		AAutoRTFMTestActor* Actor,
 		UAutoRTFMTestPrimitiveComponent* Component) = 0;
 
-	EAutomationTestFlags GetTestFlags() const
+	// GetTestFlags() changed return type between branches. Support old and new types.
+	using GetTestFlagsReturnType = decltype(std::declval<FAutomationTestBase>().GetTestFlags());
+
+	GetTestFlagsReturnType GetTestFlags() const
 	{
 		return EAutomationTestFlags::EngineFilter |
 			EAutomationTestFlags::ClientContext |
