[Material Editor] New material translator.

- Added FMaterialInsights, a new collector of material translation reflection data, to be used for accurate tracking of referenced resource, pin and connection type information  and which outputs the material writes to.
- Material connections insights wired up to pins and connection coloring.

diff --git a/Engine/Source/Editor/MaterialEditor/Private/MaterialEditor.cpp b/Engine/Source/Editor/MaterialEditor/Private/MaterialEditor.cpp
--- a/Engine/Source/Editor/MaterialEditor/Private/MaterialEditor.cpp
+++ b/Engine/Source/Editor/MaterialEditor/Private/MaterialEditor.cpp
@@ -225,7 +225,7 @@ FMatExpressionPreview::FMatExpressionPreview(UMaterialExpression* InExpression)
 	SetQualityLevelProperties(GMaxRHIFeatureLevel);
 
 	UMaterial* BaseMaterial = InExpression->Material;
-	if (BaseMaterial->IsUsingNewHLSLGenerator())
+	if (BaseMaterial->IsUsingNewHLSLGenerator() && !IsUsingNewMaterialTranslatorPrototype())
 	{
 		FMaterialCachedHLSLTree* LocalTree = new FMaterialCachedHLSLTree();
 		LocalTree->GenerateTree(BaseMaterial, nullptr, InExpression);
diff --git a/Engine/Source/Editor/UnrealEd/Private/MaterialGraph.cpp b/Engine/Source/Editor/UnrealEd/Private/MaterialGraph.cpp
--- a/Engine/Source/Editor/UnrealEd/Private/MaterialGraph.cpp
+++ b/Engine/Source/Editor/UnrealEd/Private/MaterialGraph.cpp
@@ -25,6 +25,7 @@
 #include "Materials/MaterialExpressionExecBegin.h"
 #include "Materials/MaterialExpressionExecEnd.h"
 #include "Materials/MaterialFunction.h"
+#include "Materials/MaterialInsights.h"
 
 #include "MaterialCachedHLSLTree.h"
 #include "HLSLTree/HLSLTreeEmit.h"
@@ -149,16 +150,11 @@ UMaterialGraphNode_Base* FindGraphNodeForObject(const UObject* Object)
 }
 } // namespace Private
 
-void UMaterialGraph::UpdatePinTypes()
+static void UpdatePinTypes_Old(UMaterialGraph* Graph, UMaterial* Material)
 {
 	using namespace UE::HLSLTree;
 
-	if (!Material->IsUsingNewHLSLGenerator() || Material->IsUsingNewTranslatorPrototype())
-	{
-		return;
-	}
-
-	for (UEdGraphNode* Node : Nodes)
+	for (UEdGraphNode* Node :Graph->Nodes)
 	{
 		for (UEdGraphPin* Pin : Node->Pins)
 		{
@@ -227,6 +223,58 @@ void UMaterialGraph::UpdatePinTypes()
 	}
 }
 
+void UMaterialGraph::UpdatePinTypes()
+{
+	if (Material->IsUsingNewTranslatorPrototype())
+	{
+		const FMaterialInsights* Insight = Material->MaterialInsight.Get();
+		if (!Insight)
+		{
+			return;
+		}
+
+		// Reset all pins to void category.
+		for (UEdGraphNode* Node :Nodes)
+		{
+			for (UEdGraphPin* Pin : Node->Pins)
+			{
+				if (Pin->PinType.PinCategory != UMaterialGraphSchema::PC_Exec)
+				{
+					Pin->PinType.PinCategory = UMaterialGraphSchema::PC_Void;
+				}
+			}
+		}
+		
+		for (const FMaterialInsights::FConnectionInsight& ConnectionInsight : Insight->ConnectionInsights)
+		{
+			if (UMaterialGraphNode_Base* InputNode = ::Private::FindGraphNodeForObject(ConnectionInsight.InputObject))
+			{
+				const int32 InputIndex = InputNode->GetSourceIndexForInputIndex(ConnectionInsight.InputIndex);
+				if (UEdGraphPin* InputPin = InputNode->GetInputPin(InputIndex))
+				{
+					const UE::Shader::FValueTypeDescription InputTypeDesc = UE::Shader::GetValueTypeDescription(ConnectionInsight.ValueType);
+					InputPin->PinType.PinCategory = UMaterialGraphSchema::PC_ValueType;
+					InputPin->PinType.PinSubCategory = InputTypeDesc.Name;
+				}
+			}
+
+			if (UMaterialGraphNode_Base* OutputNode = ::Private::FindGraphNodeForObject(ConnectionInsight.OutputExpression))
+			{
+				if (UEdGraphPin* OutputPin = OutputNode->GetOutputPin(ConnectionInsight.OutputIndex))
+				{
+					const UE::Shader::FValueTypeDescription InputTypeDesc = UE::Shader::GetValueTypeDescription(ConnectionInsight.ValueType);
+					OutputPin->PinType.PinCategory = UMaterialGraphSchema::PC_ValueType;
+					OutputPin->PinType.PinSubCategory = InputTypeDesc.Name;
+				}
+			}
+		}
+	}
+	else if (Material->IsUsingNewHLSLGenerator())
+	{
+		UpdatePinTypes_Old(this, Material);
+	}
+}
+
 void UMaterialGraph::RebuildGraphInternal(const TMap<UMaterialExpression*, TArray<UMaterialExpression*>>& SubgraphExpressionMap, const TMap<UMaterialExpression*, TArray<UMaterialExpressionComment*>>& SubgraphCommentMap)
 {
 	Modify();
diff --git a/Engine/Source/Runtime/Engine/Private/Materials/MaterialExpressionsIR.cpp b/Engine/Source/Runtime/Engine/Private/Materials/MaterialExpressionsIR.cpp
--- a/Engine/Source/Runtime/Engine/Private/Materials/MaterialExpressionsIR.cpp
+++ b/Engine/Source/Runtime/Engine/Private/Materials/MaterialExpressionsIR.cpp
@@ -255,7 +255,7 @@ namespace IR = UE::MIR;
 
 void UMaterialExpression::Build(IR::FEmitter& Emitter)
 {
-	Emitter.Error(TEXT("Node does not support expressions"));
+	Emitter.Error(TEXT("Unsupported material expression."));
 }
 
 void UMaterialExpressionConstant::Build(IR::FEmitter& Emitter)
diff --git a/Engine/Source/Runtime/Engine/Private/Materials/MaterialIRModuleBuilder.cpp b/Engine/Source/Runtime/Engine/Private/Materials/MaterialIRModuleBuilder.cpp
--- a/Engine/Source/Runtime/Engine/Private/Materials/MaterialIRModuleBuilder.cpp
+++ b/Engine/Source/Runtime/Engine/Private/Materials/MaterialIRModuleBuilder.cpp
@@ -14,6 +14,7 @@
 #include "MaterialExpressionIO.h"
 #include "MaterialShared.h"
 #include "Materials/MaterialExpression.h"
+#include "Materials/MaterialInsights.h"
 #include "Async/ParallelFor.h"
 
 namespace IR = UE::MIR;
@@ -21,8 +22,7 @@ namespace IR = UE::MIR;
 struct FMaterialIRModuleBuilder::FPrivate
 {
 	FMaterialIRModuleBuilder& Builder;
-	UMaterial& Material;
-	FMaterialIRTargetParams Params;
+	FMaterialIRModuleBuildParams Params;
 	FMaterialIRModule& Module;
 	IR::FEmitter& Emitter;
 	TArray<UMaterialExpression*> ExpressionAnalysisStack;
@@ -32,7 +32,7 @@ struct FMaterialIRModuleBuilder::FPrivate
 	{
 		// Prepare the array of FSetMaterialOutputInstr outputs from the material attributes inputs.
 		FMaterialInputDescription Input;
-		for (int32 Index = 0; UE::Utility::NextMaterialAttributeInput(&Material, Index, Input); ++Index)
+		for (int32 Index = 0; UE::Utility::NextMaterialAttributeInput(Params.Material, Index, Input); ++Index)
 		{
 			EMaterialProperty Property = (EMaterialProperty)Index;
 
@@ -44,7 +44,7 @@ struct FMaterialIRModuleBuilder::FPrivate
 			}
 			else if (!Input.Input->IsConnected())
 			{
-				Output->Arg = UE::Utility::CreateMaterialAttributeDefaultValue(Emitter, &Material, Property);
+				Output->Arg = UE::Utility::CreateMaterialAttributeDefaultValue(Emitter, Params.Material, Property);
 			}
 			else
 			{
@@ -109,25 +109,66 @@ struct FMaterialIRModuleBuilder::FPrivate
 			// Invoke the expression build function. This will perform semantic analysis, error reporting and
 			// emit IR values for its outputs (which will flow into connected expressions inputs).
 			Emitter.Expression->Build(Emitter);
+
+			// Populate the insight information about this expression pins.
+			if (Params.TargetInsight)
+			{
+				AddExpressionConnectionInsights(Emitter.Expression);
+			}
 		}
 	}
 
-	void Build_LinkMaterialOutputsToIncomingValues()
+	void AddExpressionConnectionInsights(UMaterialExpression* Expression)
 	{
-		FMaterialInputDescription Input;
-		for (IR::FSetMaterialOutput* Output : Module.Outputs)
+		// Update expression inputs insight.
+		for (FExpressionInputIterator It{ Expression}; It; ++It)
 		{
-			if (Output->Arg || !ensure(Material.GetExpressionInputDescription(Output->Property, Input))) {
+			if (!It->IsConnected())
+			{
 				continue;
 			}
 
-			IR::FValue** ValuePtr = Builder.OutputValues.Find(Input.Input->Expression->GetOutput(0));
-			check(ValuePtr);
+			IR::FValue** Value = Builder.InputValues.Find(It.Input);
+			PushConnectionInsight(Expression, It.Index, It->Expression, It->OutputIndex, Value ? (*Value)->Type : nullptr);
+		}
+	}
+	
+	void Build_LinkMaterialOutputsToIncomingValues()
+	{
+		for (IR::FSetMaterialOutput* Output : Module.Outputs)
+		{
+			FMaterialInputDescription Input;
+			ensure(Params.Material->GetExpressionInputDescription(Output->Property, Input));
+
+			if (!Output->Arg)
+			{
+				IR::FValue** ValuePtr = Builder.OutputValues.Find(Input.Input->Expression->GetOutput(0));
+				check(ValuePtr);
 
-			Output->Arg = *ValuePtr;
+				Builder.InputValues.Add(Input.Input, *ValuePtr);
+				Output->Arg = *ValuePtr;
+			}
+
+			if (Params.TargetInsight)
+			{
+				check(Output->Arg);
+				PushConnectionInsight(Params.Material, (int)Output->Property, Input.Input->Expression, Input.Input->OutputIndex, Output->Arg->Type);
+			}
 		}
 	}
 
+	void PushConnectionInsight(const UObject* InputObject, int InputIndex, const UMaterialExpression* OutputExpression, int OutputIndex, IR::FTypePtr Type)
+	{
+		FMaterialInsights::FConnectionInsight Insight;
+		Insight.InputObject = InputObject,
+		Insight.OutputExpression = OutputExpression,
+		Insight.InputIndex = InputIndex,
+		Insight.OutputIndex = OutputIndex,
+		Insight.ValueType = Type ? Type->ToValueType() : UE::Shader::EValueType::Any,
+		
+		Params.TargetInsight->ConnectionInsights.Push(Insight);
+	}
+
 	void Build_FinalizeValueGraph()
 	{
 		InstructionStack.Reserve(64);
@@ -212,7 +253,7 @@ struct FMaterialIRModuleBuilder::FPrivate
 			}
 		}
 	}
-	
+
     static IR::FBlock* FindCommonParentBlock(IR::FBlock* A, IR::FBlock* B)
     {
         if (A == B) {
@@ -236,14 +277,14 @@ struct FMaterialIRModuleBuilder::FPrivate
     }
 };
 
-bool FMaterialIRModuleBuilder::Build(UMaterial* InMaterial, const FMaterialIRTargetParams& TargetParams, FMaterialIRModule* TargetModule)
+bool FMaterialIRModuleBuilder::Build(const FMaterialIRModuleBuildParams& Params, FMaterialIRModule* TargetModule)
 {
 	TargetModule->Empty();
-	TargetModule->ShaderPlatform = TargetParams.ShaderPlatform;
+	TargetModule->ShaderPlatform = Params.ShaderPlatform;
 
-	IR::FEmitter Emitter{ this, InMaterial, TargetModule };
+	IR::FEmitter Emitter{ this, Params.Material, TargetModule };
 
-	FPrivate Private{ *this, *InMaterial, TargetParams, *TargetModule, Emitter };
+	FPrivate Private{ *this, Params, *TargetModule, Emitter };
 
 	Private.Build_GenerateOutputInstructions();
 	Private.Build_AnalyzeExpressionGraph();
diff --git a/Engine/Source/Runtime/Engine/Private/Materials/MaterialIRTypes.cpp b/Engine/Source/Runtime/Engine/Private/Materials/MaterialIRTypes.cpp
--- a/Engine/Source/Runtime/Engine/Private/Materials/MaterialIRTypes.cpp
+++ b/Engine/Source/Runtime/Engine/Private/Materials/MaterialIRTypes.cpp
@@ -66,6 +66,43 @@ FStringView FType::GetSpelling() const
 	UE_MIR_UNREACHABLE();
 }
 
+UE::Shader::EValueType FType::ToValueType() const
+{
+	using namespace UE::Shader;
+
+	if (FArithmeticTypePtr ArithmeticType = AsArithmetic())
+	{
+		if (ArithmeticType->IsMatrix())
+		{
+			if (ArithmeticType->NumRows == 4 && ArithmeticType->NumColumns == 4)
+			{
+				if (ArithmeticType->ScalarKind == SK_Float)
+				{
+					return EValueType::Float4x4;
+				}
+				else
+				{
+					return EValueType::Numeric4x4;
+				}
+			}
+
+			return EValueType::Any;
+		}
+
+		check(ArithmeticType->NumColumns == 1 && ArithmeticType->NumRows <= 4);
+
+		switch (ArithmeticType->ScalarKind)
+		{
+			case SK_Bool: 	return (EValueType)((int)EValueType::Bool1 + ArithmeticType->NumRows - 1); break;
+			case SK_Int: 	return (EValueType)((int)EValueType::Int1 + ArithmeticType->NumRows - 1); break;
+			case SK_Float: 	return (EValueType)((int)EValueType::Float1 + ArithmeticType->NumRows - 1); break;
+			default: UE_MIR_UNREACHABLE();
+		}
+	}
+	
+	UE_MIR_UNREACHABLE();
+}
+
 bool FType::IsBool1() const
 {
 	return this == FArithmeticType::GetBool1();
diff --git a/Engine/Source/Runtime/Engine/Private/Materials/MaterialInsights.cpp b/Engine/Source/Runtime/Engine/Private/Materials/MaterialInsights.cpp
new file mode 100644
--- /dev/null
+++ b/Engine/Source/Runtime/Engine/Private/Materials/MaterialInsights.cpp
@@ -0,0 +1,13 @@
+// Copyright Epic Games, Inc. All Rights Reserved.
+
+#include "Materials/MaterialInsights.h"
+
+FMaterialInsights::FMaterialInsights()
+{
+}
+
+void FMaterialInsights::Empty()
+{
+	ConnectionInsights.Empty();
+}
+
diff --git a/Engine/Source/Runtime/Engine/Private/Materials/MaterialInterface.cpp b/Engine/Source/Runtime/Engine/Private/Materials/MaterialInterface.cpp
--- a/Engine/Source/Runtime/Engine/Private/Materials/MaterialInterface.cpp
+++ b/Engine/Source/Runtime/Engine/Private/Materials/MaterialInterface.cpp
@@ -44,6 +44,7 @@
 #include "UObject/ArchiveCookContext.h"
 #include "UObject/Package.h"
 #include "ShaderCompiler.h"
+#include "Materials/MaterialInsights.h"
 
 #if WITH_EDITOR
 #include "ObjectCacheEventSink.h"
@@ -52,7 +53,7 @@
 
 #define LOCTEXT_NAMESPACE "MaterialInterface"
 
-/**
+/** 
  * This is used to deprecate data that has been built with older versions.
  * To regenerate the data, commands like "BUILDMATERIALTEXTURESTREAMINGDATA" can be used in the editor.
  * Ideally the data would be stored the DDC instead of the asset, but this is not yet  possible because it requires the GPU.
diff --git a/Engine/Source/Runtime/Engine/Private/Materials/MaterialShared.cpp b/Engine/Source/Runtime/Engine/Private/Materials/MaterialShared.cpp
--- a/Engine/Source/Runtime/Engine/Private/Materials/MaterialShared.cpp
+++ b/Engine/Source/Runtime/Engine/Private/Materials/MaterialShared.cpp
@@ -68,6 +68,7 @@
 #include "Materials/MaterialIRModuleBuilder.h"
 #include "Materials/MaterialIRToHLSLTranslator.h"
 #include "Materials/MaterialSourceTemplate.h"
+#include "Materials/MaterialInsights.h"
 
 #define LOCTEXT_NAMESPACE "MaterialShared"
 
@@ -3408,15 +3409,19 @@ bool FMaterial::Translate_New(const FMaterialShaderMapId& InShaderMapId,
 
 	FMaterialIRModule Module;
 
-	FMaterialIRTargetParams TargetParams{
+	GetMaterialInterface()->MaterialInsight.Reset(new FMaterialInsights);
+
+	FMaterialIRModuleBuildParams TargetParams{
+		.Material = GetMaterialInterface()->GetMaterial(),
 		.ShaderPlatform = InShaderPlatform,
 		.TargetPlatform = InTargetPlatform,
 		.StaticParameters = InStaticParameters,
+		.TargetInsight = GetMaterialInterface()->MaterialInsight.Get(),
 	};
 
 	// Build the material
 	FMaterialIRModuleBuilder Builder;
-	if (!Builder.Build(GetMaterialInterface()->GetMaterial(), TargetParams, &Module))
+	if (!Builder.Build(TargetParams, &Module))
 	{
 		for (const FMaterialIRModule::FError& Error : Module.GetErrors())
 		{
diff --git a/Engine/Source/Runtime/Engine/Public/MaterialExpressionIO.h b/Engine/Source/Runtime/Engine/Public/MaterialExpressionIO.h
--- a/Engine/Source/Runtime/Engine/Public/MaterialExpressionIO.h
+++ b/Engine/Source/Runtime/Engine/Public/MaterialExpressionIO.h
@@ -35,25 +35,24 @@ struct FExpressionInput
 	 * Material expression that this input is connected to, or NULL if not connected. 
 	 * If you want to be safe when checking against dangling Reroute nodes, please use GetTracedInput before accessing this property.
 	*/
-	class UMaterialExpression*	Expression;
+	class UMaterialExpression* Expression;
 
 	/** 
 	 * Index into Expression's outputs array that this input is connected to.
 	 * If you want to be safe when checking against dangling Reroute nodes, please use GetTracedInput before accessing this property.
 	*/
-	int32						OutputIndex;
+	int32 OutputIndex;
 
 	/** 
 	 * Optional name of the input.  
 	 * Note that this is the only member which is not derived from the output currently connected. 
 	 */
 	FName InputName;
-		
-	int32						Mask,
-								MaskR,
-								MaskG,
-								MaskB,
-								MaskA;
+	
+	/**
+	 * Input value component mask.
+	 */
+	int32 Mask, MaskR, MaskG, MaskB, MaskA;
 
 	FExpressionInput()
 		: OutputIndex(0)
diff --git a/Engine/Source/Runtime/Engine/Public/Materials/MaterialExpression.h b/Engine/Source/Runtime/Engine/Public/Materials/MaterialExpression.h
--- a/Engine/Source/Runtime/Engine/Public/Materials/MaterialExpression.h
+++ b/Engine/Source/Runtime/Engine/Public/Materials/MaterialExpression.h
@@ -363,10 +363,11 @@ class UMaterialExpression : public UObject
 
 	/** Counts the number of inputs this expression has. Default implementation has O(n) complexity. */
 	ENGINE_API virtual int32 CountInputs() const;
+
 	/**
 	 * Returns the output at specified index is valid
 	 */
-	ENGINE_API FExpressionOutput* GetOutput(int32 OutputIndex) { return &Outputs[OutputIndex]; }
+	ENGINE_API FExpressionOutput* GetOutput(int32 OutputIndex) { return Outputs.IsValidIndex(OutputIndex) ? &Outputs[OutputIndex] : nullptr; }
 
 	/**
 	 * Returns the input at index InputIndex if valid, nullptr otherwise.
@@ -696,6 +697,59 @@ struct FExpressionInputIterator
     }
 };
 
+/**
+* @brief An iterator for traversing the outputs of a UMaterialExpression.
+*
+* This struct provides a way to iterate over the outputs of a given
+* UMaterialExpression. It starts from the first output and advances
+* through subsequent outputs until all are iterated through.
+*
+* Usage:
+*
+* @code
+* UMaterialExpression* MyExpression = ...;
+* for (FExpressionOutputIterator It{ MyExpression }; It; ++It)
+* {
+*     // Process the output, e.g. It->Output->IsConnected();
+* }
+* @endcode
+*/
+struct FExpressionOutputIterator
+{
+	/** The expression whose outputs to iterate through. */
+    UMaterialExpression* Expression;
+
+	/** Current output. */
+    FExpressionOutput* Output;
+
+	/** Current output index. */
+	int Index;
+
+    FExpressionOutputIterator(UMaterialExpression* InExpression)
+    : Expression{ InExpression }
+    , Output{ InExpression->GetOutput(0) }
+    , Index{ 0 }
+    {
+    }
+
+    operator bool() const
+    {
+        return Output != nullptr;
+    }
+
+    FExpressionOutputIterator& operator++()
+    {
+		Index += 1;
+		Output = Expression->GetOutput(Index);
+        return *this;
+    }
+
+    FExpressionOutput* operator->()
+    {
+        return Output;
+    }
+};
+
 #endif // WITH_EDITOR
 
 #ifndef ENABLE_MATERIAL_LAYER_PROTOTYPE
diff --git a/Engine/Source/Runtime/Engine/Public/Materials/MaterialIRCommon.h b/Engine/Source/Runtime/Engine/Public/Materials/MaterialIRCommon.h
--- a/Engine/Source/Runtime/Engine/Public/Materials/MaterialIRCommon.h
+++ b/Engine/Source/Runtime/Engine/Public/Materials/MaterialIRCommon.h
@@ -29,6 +29,7 @@ struct FExpressionOutput;
 struct FShaderCompilerEnvironment;
 struct FStaticParameterSet;
 struct FMaterialInputDescription;
+struct FMaterialInsights;
 
 namespace UE::MIR
 {
diff --git a/Engine/Source/Runtime/Engine/Public/Materials/MaterialIRModuleBuilder.h b/Engine/Source/Runtime/Engine/Public/Materials/MaterialIRModuleBuilder.h
--- a/Engine/Source/Runtime/Engine/Public/Materials/MaterialIRModuleBuilder.h
+++ b/Engine/Source/Runtime/Engine/Public/Materials/MaterialIRModuleBuilder.h
@@ -5,18 +5,20 @@
 
 #if WITH_EDITOR
 
-struct FMaterialIRTargetParams
+struct FMaterialIRModuleBuildParams
 {
+	UMaterial* Material;
 	EShaderPlatform ShaderPlatform;
 	const ITargetPlatform* TargetPlatform;
 	const FStaticParameterSet& StaticParameters;
+	FMaterialInsights* TargetInsight{};
 };
 
 //
 class FMaterialIRModuleBuilder
 {
 public:
-	bool Build(UMaterial* InMaterial, const FMaterialIRTargetParams& TargetParams, FMaterialIRModule* TargetModule);
+	bool Build(const FMaterialIRModuleBuildParams& Params, FMaterialIRModule* TargetModule);
 
 private:
 	TMap<const FExpressionInput*, UE::MIR::FValue*> InputValues;
diff --git a/Engine/Source/Runtime/Engine/Public/Materials/MaterialIRTypes.h b/Engine/Source/Runtime/Engine/Public/Materials/MaterialIRTypes.h
--- a/Engine/Source/Runtime/Engine/Public/Materials/MaterialIRTypes.h
+++ b/Engine/Source/Runtime/Engine/Public/Materials/MaterialIRTypes.h
@@ -45,6 +45,9 @@ struct FType
 
 	// Returns the this type name spelling (e.g. float4x4).
 	FStringView GetSpelling() const;
+
+	// Converts this type to a UE::Shader::EValueType.
+	UE::Shader::EValueType ToValueType() const;
 };
 
 enum EScalarKind
diff --git a/Engine/Source/Runtime/Engine/Public/Materials/MaterialInsights.h b/Engine/Source/Runtime/Engine/Public/Materials/MaterialInsights.h
new file mode 100644
--- /dev/null
+++ b/Engine/Source/Runtime/Engine/Public/Materials/MaterialInsights.h
@@ -0,0 +1,40 @@
+// Copyright Epic Games, Inc. All Rights Reserved.
+#pragma once
+
+#include "CoreMinimal.h"
+#include "Shader/ShaderTypes.h"
+
+struct FExpressionInput;
+struct FExpressionOutput;
+class UMaterialExpression;
+
+/**
+ * @brief A structure that holds reflection information about a material.
+ * This structure is typically populated by the material translator as a side product of
+ * the translation process itself.
+ *
+ * You can use these insights for things like providing semantic colouring
+ * of the graph UI or accurately knowing what resources are referenced by the translated
+ * materials.
+ */
+struct FMaterialInsights
+{
+    /**
+     * @brief Nested structure that represents a single connection insight.
+     */
+    struct FConnectionInsight
+    {
+        const UObject* InputObject;                   ///< Pointer to the input object of the connection.
+        const UMaterialExpression* OutputExpression;  ///< Pointer to the output expression of the connection.
+        int InputIndex;                               ///< Index of the input in the connection.
+        int OutputIndex;                              ///< Index of the output in the connection.
+        UE::Shader::EValueType ValueType;             ///< Type of the value flowing through the connection.
+    };
+
+    TArray<FConnectionInsight> ConnectionInsights;    ///< Array of connection insights.
+
+    FMaterialInsights();
+
+	/// Clears all insight data.
+    void Empty();
+};
\ No newline at end of file
diff --git a/Engine/Source/Runtime/Engine/Public/Materials/MaterialInterface.h b/Engine/Source/Runtime/Engine/Public/Materials/MaterialInterface.h
--- a/Engine/Source/Runtime/Engine/Public/Materials/MaterialInterface.h
+++ b/Engine/Source/Runtime/Engine/Public/Materials/MaterialInterface.h
@@ -63,6 +63,7 @@ struct FParameterChannelNames;
 enum EShaderPlatform : uint16;
 struct FSubstrateCompilationConfig;
 class UMaterialExpressionCustomOutput;
+struct FMaterialInsights;
 
 typedef TArray<FMaterialResource*> FMaterialResourceDeferredDeletionArray;
 
@@ -1182,6 +1183,11 @@ public:
 
 	/** Filter out ShadingModels field to a shader platform settings */
 	static void FilterOutPlatformShadingModels(EShaderPlatform Platform, FMaterialShadingModelField& ShadingModels);
+	
+#if WITH_EDITOR
+	TUniquePtr<FMaterialInsights> MaterialInsight;
+#endif // WITH_EDITOR
+
 protected:
 	/** Returns a bitfield indicating which feature levels should be compiled for rendering. GMaxRHIFeatureLevel is always present */
 	ENGINE_API uint32 GetFeatureLevelsToCompileForRendering() const;
@@ -1201,6 +1207,7 @@ protected:
 #if WITH_EDITOR
 	TUniquePtr<FMaterialCachedHLSLTree> CachedHLSLTree;
 #endif // WITH_EDITOR
+
 private:
 	/**
 	 * Post loads all default materials.
