Fix for Substrate Material Layers where the optimizer would skip a call to UMaterialExpressionSubstrateSetAttributes::Compile because it would re-navigate to the target expression call. Resolved by exchanging the logic to ensure the FExpressionInput compile calls are used instead.

diff --git a/Engine/Source/Runtime/Engine/Private/Materials/MaterialExpressions.cpp b/Engine/Source/Runtime/Engine/Private/Materials/MaterialExpressions.cpp
--- a/Engine/Source/Runtime/Engine/Private/Materials/MaterialExpressions.cpp
+++ b/Engine/Source/Runtime/Engine/Private/Materials/MaterialExpressions.cpp
@@ -29940,21 +29940,16 @@ void UMaterialExpressionSubstrateSetAttributes::PostEditChangeProperty(FProperty
 int32 UMaterialExpressionSubstrateSetAttributes::Compile(class FMaterialCompiler* Compiler, int32 OutputIndex)
 {
 	int32 Ret = INDEX_NONE;
-	UMaterialExpression* Expression = nullptr;
+	FExpressionInput* Expression = nullptr;
 
-	EMaterialProperty Property = FMaterialAttributeDefinitionMap::GetProperty(Compiler->GetMaterialAttribute());
-	if (Property == MP_FrontMaterial)
+	FGuid AttributeId = Compiler->GetMaterialAttribute();
+	if (AttributeId == FMaterialAttributeDefinitionMap::GetID(MP_FrontMaterial))
 	{
-		Expression = FrontMaterial.Expression;
+		Ret = FrontMaterial.Compile(Compiler);
 	}
 	else
 	{
-		Expression = NonSubstrateAttributes.Expression;
-	}
-
-	if (Expression)
-	{
-		Ret = Expression->Compile(Compiler, OutputIndex);
+		Ret = NonSubstrateAttributes.CompileWithDefault(Compiler, AttributeId);
 	}
 
 	return Ret;
