[SkeletalMesh] Disable static relevance path in FSkeletalMeshObjectCPUSkin to fix crash due to invalid cache entries.

Since skeletal mesh skinning uses mesh caching, bone weight visualization in the editor crashes on undo operations.
This happened when CPU skinning was recreating the vertex factory which is not supported by mesh caching.
The fix is to disable the static relevance path in FSkeletalMeshObjectCPUSkin.

#rnx

diff --git a/Engine/Source/Runtime/Engine/Private/SkeletalRenderCPUSkin.cpp b/Engine/Source/Runtime/Engine/Private/SkeletalRenderCPUSkin.cpp
--- a/Engine/Source/Runtime/Engine/Private/SkeletalRenderCPUSkin.cpp
+++ b/Engine/Source/Runtime/Engine/Private/SkeletalRenderCPUSkin.cpp
@@ -73,7 +73,9 @@ FSkeletalMeshObjectCPUSkin::FSkeletalMeshObjectCPUSkin(USkinnedMeshComponent* In
 	}
 
 	InitResources(InMeshComponent);
-	bSupportsStaticRelevance = true;
+
+	// This must be disabled since mesh caching doesn't support recreating the vertex factory on demand.
+	bSupportsStaticRelevance = false;
 }
 
 
