[Shaders] Split up ShaderCompiler.cpp into multiple source files for better readability.

This change does not contain any logical code changes. ShaderCompiler.cpp has been split up into the following source files:
- ShaderCompiler.cpp - Original file continues to include the shader compiling manager. Maintaining this filename simplifies P4 history scrubbing.
- ShaderCompilerEditor.cpp - Most of the functions that are for WITH_EDITOR only.
- ShaderCompilerJobCache.cpp - Job collection and job caching logic.
- ShaderCompilerStats.cpp - Gathering and output logic for shader compilation statistics.
- ShaderCompilerThreadRunnable.cpp - Base and subclasses for shader compile worker threads (FShaderCompileThreadRunnable class implementation).
- ShaderCompilerUtil.cpp - Global shader complation utility functions and FShaderCompileUtilities class implementation.

#rnx

diff --git a/Engine/Source/Runtime/Engine/Private/ODSC/ODSCThread.cpp b/Engine/Source/Runtime/Engine/Private/ODSC/ODSCThread.cpp
--- a/Engine/Source/Runtime/Engine/Private/ODSC/ODSCThread.cpp
+++ b/Engine/Source/Runtime/Engine/Private/ODSC/ODSCThread.cpp
@@ -33,6 +33,32 @@ FODSCRequestPayload::FODSCRequestPayload(
 
 }
 
+FArchive& operator<<(FArchive& Ar, FODSCRequestPayload& Payload)
+{
+	int32 iShaderPlatform = static_cast<int32>(Payload.ShaderPlatform);
+	int32 iFeatureLevel = static_cast<int32>(Payload.FeatureLevel);
+	int32 iQualityLevel = static_cast<int32>(Payload.QualityLevel);
+
+	Ar << iShaderPlatform;
+	Ar << iFeatureLevel;
+	Ar << iQualityLevel;
+	Ar << Payload.MaterialName;
+	Ar << Payload.VertexFactoryName;
+	Ar << Payload.PipelineName;
+	Ar << Payload.ShaderTypeNames;
+	Ar << Payload.PermutationId;
+	Ar << Payload.RequestHash;
+
+	if (Ar.IsLoading())
+	{
+		Payload.ShaderPlatform = static_cast<EShaderPlatform>(iShaderPlatform);
+		Payload.FeatureLevel = static_cast<ERHIFeatureLevel::Type>(iFeatureLevel);
+		Payload.QualityLevel = static_cast<EMaterialQualityLevel::Type>(iQualityLevel);
+	}
+
+	return Ar;
+}
+
 FODSCMessageHandler::FODSCMessageHandler(EShaderPlatform InShaderPlatform, ERHIFeatureLevel::Type InFeatureLevel, EMaterialQualityLevel::Type InQualityLevel, ODSCRecompileCommand InRecompileCommandType)
 :	ShaderPlatform(InShaderPlatform),
 	FeatureLevel(InFeatureLevel),
