[Shaders] Split up ShaderCompiler.cpp into multiple source files for better readability.

diff --git a/Engine/Source/Runtime/Engine/Private/ShaderCompiler/ShaderCompiler.cpp b/Engine/Source/Runtime/Engine/Private/ShaderCompiler/ShaderCompiler.cpp
--- a/Engine/Source/Runtime/Engine/Private/ShaderCompiler/ShaderCompiler.cpp
+++ b/Engine/Source/Runtime/Engine/Private/ShaderCompiler/ShaderCompiler.cpp
@@ -5,22 +5,17 @@
 =============================================================================*/
 
 #include "ShaderCompiler.h"
+#include "ShaderCompilerPrivate.h"
 #include "AsyncCompilationHelpers.h"
 #include "AssetCompilingManager.h"
-#include "AnalyticsEventAttribute.h"
-#include "Async/ParallelFor.h"
 #include "ClearReplacementShaders.h"
 #include "ComponentRecreateRenderStateContext.h"
 #include "Components/PrimitiveComponent.h"
 #include "DataDrivenShaderPlatformInfo.h"
-#include "DerivedDataCache.h"
-#include "DerivedDataRequestOwner.h"
 #include "DistributedBuildControllerInterface.h"
 #include "EditorSupportDelegates.h"
 #include "Engine/RendererSettings.h"
 #include "Features/IModularFeatures.h"
-#include "HAL/PlatformFile.h"
-#include "HAL/PlatformFileManager.h"
 #include "Interfaces/IShaderFormat.h"
 #include "Interfaces/ITargetPlatform.h"
 #include "Interfaces/ITargetPlatformManagerModule.h"
@@ -29,15 +24,12 @@
 #include "MaterialShared.h"
 #include "Materials/Material.h"
 #include "Materials/MaterialInstance.h"
-#include "Misc/Compression.h"
 #include "Misc/ConfigCacheIni.h"
 #include "Misc/CoreDelegates.h"
-#include "Misc/EngineVersion.h"
 #include "Misc/FeedbackContext.h"
 #include "Misc/FileHelper.h"
 #include "Misc/MessageDialog.h"
 #include "Misc/ScopedSlowTask.h"
-#include "Misc/ScopeTryLock.h"
 #include "Modules/ModuleManager.h"
 #include "ObjectCacheContext.h"
 #include "ProfilingDebugging/CookStats.h"
@@ -57,8 +49,6 @@
 #include "StereoRenderUtils.h"
 #include "Tasks/Task.h"
 #include "UObject/DevObjectVersion.h"
-#include "UObject/Linker.h"
-#include "UObject/Package.h"
 #include "UObject/UObjectIterator.h"
 #include "Math/UnitConversion.h"
 #include "UnrealEngine.h"
@@ -84,12 +74,6 @@
 
 DEFINE_LOG_CATEGORY(LogShaderCompilers);
 
-// Switch to Verbose after initial testing
-#define UE_SHADERCACHE_LOG_LEVEL		VeryVerbose
-
-// whether to parallelize writing/reading task files
-#define UE_SHADERCOMPILER_FIFO_JOB_EXECUTION  1
-
 LLM_DEFINE_TAG(ShaderCompiler);
 
 static TAutoConsoleVariable<bool> CVarRecompileShadersOnSave(
@@ -99,69 +83,6 @@ static TAutoConsoleVariable<bool> CVarRecompileShadersOnSave(
 	TEXT("Default: false"),
 	ECVF_ReadOnly);
 
-int32 GShaderCompilerJobCache = 1;
-static FAutoConsoleVariableRef CVarShaderCompilerJobCache(
-	TEXT("r.ShaderCompiler.JobCache"),
-	GShaderCompilerJobCache,
-	TEXT("if != 0, shader compiler cache (based on the unpreprocessed input hash) will be disabled. By default, it is enabled."),
-	ECVF_Default
-);
-
-static TAutoConsoleVariable<bool> CVarShaderCompilerDebugValidateJobCache(
-	TEXT("r.ShaderCompiler.DebugValidateJobCache"),
-	false,
-	TEXT("Enables debug mode for job cache which will fully execute all jobs and validate that job outputs with matching input hashes match."),
-	ECVF_Default
-);
-
-int32 GShaderCompilerMaxJobCacheMemoryMB = 16LL * 1024LL;
-static FAutoConsoleVariableRef CVarShaderCompilerMaxJobCacheMemoryMB(
-	TEXT("r.ShaderCompiler.MaxJobCacheMemoryMB"),
-	GShaderCompilerMaxJobCacheMemoryMB,
-	TEXT("if != 0, shader compiler cache will be limited to this many megabytes (16GB by default). If 0, the usage will be unlimited. Minimum of this or r.ShaderCompiler.MaxJobCacheMemoryPercent applies."),
-	ECVF_Default
-);
-
-int32 GShaderCompilerMaxJobCacheMemoryPercent = 5;
-static FAutoConsoleVariableRef CVarShaderCompilerMaxJobCacheMemoryPercent(
-	TEXT("r.ShaderCompiler.MaxJobCacheMemoryPercent"),
-	GShaderCompilerMaxJobCacheMemoryPercent,
-	TEXT("if != 0, shader compiler cache will be limited to this percentage of available physical RAM (5% by default). If 0, the usage will be unlimited. Minimum of this or r.ShaderCompiler.MaxJobCacheMemoryMB applies."),
-	ECVF_Default
-);
-
-int32 GShaderCompilerJobCacheOverflowReducePercent = 80;
-static FAutoConsoleVariableRef CVarShaderCompilerJobCacheOverflowReducePercent(
-	TEXT("r.ShaderCompiler.JobCacheOverflowReducePercent"),
-	GShaderCompilerJobCacheOverflowReducePercent,
-	TEXT("When shader compiler job cache memory overflows, reduce memory to this percentage of the maximum.  Reduces overhead relative to cleaning up items one at a time when at max budget."),
-	ECVF_Default
-);
-
-static TAutoConsoleVariable<bool> CVarJobCacheDDC(
-	TEXT("r.ShaderCompiler.JobCacheDDC"),
-	true,
-	TEXT("Skips compilation of all shaders on Material and Material Instance PostLoad and relies on on-demand shader compilation to compile what is needed."),
-	ECVF_ReadOnly);
-
-static TAutoConsoleVariable<bool> CVarJobCacheDDCPolicy(
-	TEXT("r.ShaderCompiler.JobCacheDDCEnableRemotePolicy"),
-	true,
-	TEXT("If true, individual shader jobs will be cached to remote/shared DDC instances in all operation modes; if false they will only cache to DDC instances on the local machine.\n"),
-	ECVF_ReadOnly);
-
-static TAutoConsoleVariable<bool> CVarJobCacheDDCCookPolicy(
-	TEXT("r.ShaderCompiler.JobCacheDDCCookEnableRemotePolicy"),
-	true,
-	TEXT("If true, individual shader jobs will be cached to remote/shared DDC instances in all cook commandlet only; if false they will only cache to DDC instances on the local machine.\n"),
-	ECVF_ReadOnly);
-
-static TAutoConsoleVariable<bool> CVarDebugDumpWorkerInputs(
-	TEXT("r.ShaderCompiler.DebugDumpWorkerInputs"),
-	false,
-	TEXT("If true, worker input files will be saved for each individual compile job alongside other debug data (note that r.DumpShaderDebugInfo must also be enabled for this to function)"),
-	ECVF_ReadOnly);
-
 static TAutoConsoleVariable<bool> CVarDebugDumpJobInputHashes(
 	TEXT("r.ShaderCompiler.DebugDumpJobInputHashes"),
 	false,
@@ -190,34 +111,6 @@ static TAutoConsoleVariable<bool> CVarDebugDumpDetailedShaderSource(
 	TEXT("If false, or the preprocessed job cache is disabled, this will simply dump whatever source is passed to the compiler (equivalent to either 1 or 3 depending on if the IShaderFormat implementation modifies the source in the compile step."),
 	ECVF_ReadOnly);
 
-static TAutoConsoleVariable<bool> CVarCompileParallelInProcess(
-	TEXT("r.ShaderCompiler.ParallelInProcess"),
-	false,
-	TEXT("EXPERIMENTAL- If true, shader compilation will be executed in-process in parallel. Note that this will serialize if the legacy preprocessor is enabled."),
-	ECVF_ReadOnly);
-
-static TAutoConsoleVariable<bool> CVarShaderCompilerPerShaderDDCAsync(
-	TEXT("r.ShaderCompiler.PerShaderDDCAsync"),
-	true,
-	TEXT("if != 0, Per-shader DDC queries will run async, instead of in the SubmitJobs task."),
-	ECVF_Default
-);
-
-int32 GShaderCompilerPerShaderDDCGlobal = 1;
-static FAutoConsoleVariableRef CVarShaderCompilerPerShaderDDCGlobal(
-	TEXT("r.ShaderCompiler.PerShaderDDCGlobal"),
-	GShaderCompilerPerShaderDDCGlobal,
-	TEXT("if != 0, Per-shader DDC queries enabled for global and default shaders."),
-	ECVF_Default
-);
-
-static TAutoConsoleVariable<bool> CVarShaderCompilerPerShaderDDCCook(
-	TEXT("r.ShaderCompiler.PerShaderDDCCook"),
-	true,
-	TEXT("If true, per-shader DDC caching will be enabled during cooks."),
-	ECVF_Default
-);
-
 static TAutoConsoleVariable<bool> CVarDisableSourceStripping(
 	TEXT("r.ShaderCompiler.DisableSourceStripping"),
 	false,
@@ -231,53 +124,6 @@ static TAutoConsoleVariable<bool> CVarAreShaderErrorsFatal(
 	TEXT("Default: true"),
 	ECVF_RenderThreadSafe);
 
-int32 GMaxShaderStatsToLog = 5;
-static FAutoConsoleVariableRef CVarShaderCompilerMaxShaderStatsToPrint(
-	TEXT("r.ShaderCompiler.MaxShaderStatsToPrint"),
-	GMaxShaderStatsToLog,
-	TEXT("Max number of shaders FShaderJobCache stats logs out"),
-	ECVF_Default
-);
-
-bool AreShaderErrorsFatal()
-{
-	return CVarAreShaderErrorsFatal.GetValueOnAnyThread();
-}
-
-static bool IsShaderJobCacheDDCRemotePolicyEnabled()
-{
-	return CVarJobCacheDDCPolicy.GetValueOnAnyThread() || (IsRunningCookCommandlet() && CVarJobCacheDDCCookPolicy.GetValueOnAnyThread());
-}
-
-
-bool IsShaderJobCacheDDCEnabled()
-{
-#if WITH_EDITOR
-	static const bool bForceAllowShaderCompilerJobCache = FParse::Param(FCommandLine::Get(), TEXT("forceAllowShaderCompilerJobCache"));
-	static const bool bEnablePerShaderDDCCook = IsRunningCookCommandlet() && CVarShaderCompilerPerShaderDDCCook.GetValueOnAnyThread();
-	static const bool bIsNonCookCommandlet = IsRunningCommandlet() && !IsRunningCookCommandlet();
-#else
-	const bool bForceAllowShaderCompilerJobCache = false;
-	const bool bEnablePerShaderDDCCook = false;
-	const bool bIsNonCookCommandlet = false;
-#endif
-
-	// Enable remote per-shader DDC for editor, game, cooks (if cvar is set), and for other commandlets only if the force flag is set on the cmdline.
-	if ((GIsEditor || IsRunningGame() || bEnablePerShaderDDCCook) && (!bIsNonCookCommandlet || bForceAllowShaderCompilerJobCache))
-	{
-		// job cache itself must be enabled first
-		return GShaderCompilerJobCache && CVarJobCacheDDC.GetValueOnAnyThread();
-	}
-
-	return false;
-}
-
-bool IsMaterialMapDDCEnabled()
-{
-	// If we are loading individual shaders from the shader job cache for ODSC, don't attempt to load full material maps.  Always load/cache material maps in cooks.
-	return (IsShaderJobCacheDDCEnabled() == false) || IsRunningCookCommandlet();
-}
-
 int32 GShaderCompilerAllowDistributedCompilation = 1;
 static FAutoConsoleVariableRef CVarShaderCompilerAllowDistributedCompilation(
 	TEXT("r.ShaderCompiler.AllowDistributedCompilation"),
@@ -302,15 +148,6 @@ static FAutoConsoleVariableRef CVarGSShaderCheckLevel(
 	ECVF_Default
 );
 
-float GShaderCompilerTooLongIOThresholdSeconds = 0.3;
-static FAutoConsoleVariableRef CVarShaderCompilerTooLongIOThresholdSeconds(
-	TEXT("r.ShaderCompiler.TooLongIOThresholdSeconds"),
-	GShaderCompilerTooLongIOThresholdSeconds,
-	TEXT("By default, task files for SCW will be read/written sequentially, but if we ever spend more than this time (0.3s by default) doing that, we'll switch to parallel.") \
-	TEXT("We don't default to parallel writes as it increases the CPU overhead from the shader compiler."),
-	ECVF_Default
-);
-
 static TAutoConsoleVariable<bool> CVarShaderCompilerDumpDDCKeys(
 	TEXT("r.ShaderCompiler.DumpDDCKeys"),
 	false,
@@ -318,10970 +155,5212 @@ static TAutoConsoleVariable<bool> CVarShaderCompilerDumpDDCKeys(
 	ECVF_Default
 );
 
-int32 GShaderCompilerDebugDiscardCacheOutputs = 0;
-static FAutoConsoleVariableRef CVarShaderCompilerDebugDiscardCacheOutputs(
-	TEXT("r.ShaderCompiler.DebugDiscardCacheOutputs"),
-	GShaderCompilerDebugDiscardCacheOutputs,
-	TEXT("if != 0, cache outputs are discarded (not added to the output map) for debugging purposes.\nEliminates usefulness of the cache, but allows repeated triggering of the same jobs for stress testing (for example, rapid undo/redo in the Material editor)."),
-	ECVF_Default
-);
-
-int32 GShaderCompilerDebugStallSubmitJob = 0;
-static FAutoConsoleVariableRef CVarShaderCompilerDebugStallSubmitJob(
-	TEXT("r.ShaderCompiler.DebugStallSubmitJob"),
-	GShaderCompilerDebugStallSubmitJob,
-	TEXT("For debugging, a value in milliseconds to stall in SubmitJob, to help reproduce threading bugs."),
-	ECVF_Default
-);
-
-int32 GShaderCompilerDebugStallDDCQuery = 0;
-static FAutoConsoleVariableRef CVarShaderCompilerDebugStallDCCQuery(
-	TEXT("r.ShaderCompiler.DebugStallDDCQuery"),
-	GShaderCompilerDebugStallDDCQuery,
-	TEXT("For debugging, a value in milliseconds to stall in the DDC completion callback, to help reproduce threading bugs, or simulate higher latency DDC for perf testing."),
-	ECVF_Default
-);
-
 
-/** Copy of TIntrusiveLinkedListIterator, specific to FShaderCommonCompileJob */
-class FShaderCommonCompileJobIterator
+bool AreShaderErrorsFatal()
 {
-public:
-	explicit FShaderCommonCompileJobIterator(FShaderCommonCompileJob* FirstLink)
-		: CurrentLink(FirstLink)
-	{ }
-
-	/**
-	 * Advances the iterator to the next element.
-	 */
-	FORCEINLINE void Next()
-	{
-		checkSlow(CurrentLink);
-		CurrentLink = CurrentLink->NextLink;
-	}
+	return CVarAreShaderErrorsFatal.GetValueOnAnyThread();
+}
 
-	FORCEINLINE FShaderCommonCompileJobIterator& operator++()
-	{
-		Next();
-		return *this;
-	}
 
-	FORCEINLINE FShaderCommonCompileJobIterator operator++(int)
+namespace ShaderCompiler
+{
+	FString GetTargetPlatformName(const ITargetPlatform* TargetPlatform)
 	{
-		auto Tmp = *this;
-		Next();
-		return Tmp;
-	}
+		if (TargetPlatform)
+		{
+			return TargetPlatform->PlatformName();
+		}
 
-	/** conversion to "bool" returning true if the iterator is valid. */
-	FORCEINLINE explicit operator bool() const
-	{
-		return CurrentLink != nullptr;
+		return TEXT("(current)");
 	}
 
-	FORCEINLINE bool operator==(const FShaderCommonCompileJobIterator& Rhs) const { return CurrentLink == Rhs.CurrentLink; }
-	FORCEINLINE bool operator!=(const FShaderCommonCompileJobIterator& Rhs) const { return CurrentLink != Rhs.CurrentLink; }
-
-	// Accessors.
-	FORCEINLINE FShaderCommonCompileJob& operator->() const
+	bool IsRemoteCompilingAllowed()
 	{
-		checkSlow(this->CurrentLink);
-		return *(this->CurrentLink);
-	}
+		// commandline switches override the CVars
+		static bool bDisabledFromCommandline = FParse::Param(FCommandLine::Get(), TEXT("NoRemoteShaderCompile"));
+		if (bDisabledFromCommandline)
+		{
+			return false;
+		}
 
-	FORCEINLINE FShaderCommonCompileJob& operator*() const
-	{
-		checkSlow(this->CurrentLink);
-		return *(this->CurrentLink);
+		return GShaderCompilerAllowDistributedCompilation != 0;
 	}
+} // namespace ShaderCompiler
 
-private:
-	FShaderCommonCompileJob* CurrentLink;
-};
+/** Storage for the global shadar map(s) that have been replaced by new one(s), which aren't yet compiled.
+ * 
+ *	Sometimes a mesh drawing command references a pointer to global SM's memory. To nix these MDCs when we're replacing a global SM, we would just recreate the render state for all the components, but
+ *	we may need to access a global shader during such an update, creating a catch 22. So deleting the global SM and updating components is deferred until the new one is compiled. 
+ */
+FGlobalShaderMap* GGlobalShaderMap_DeferredDeleteCopy[SP_NumPlatforms] = {nullptr};
 
-/** Map element type for job cache */
-struct FShaderJobData
+#if ENABLE_COOK_STATS
+namespace GlobalShaderCookStats
 {
-	using FJobOutputHash = FBlake3Hash;
-
-	FShaderCompilerInputHash InputHash;
-
-	/** Output hash will be zero if output data has not been written yet, or can be cleared if output data has been removed */
-	FJobOutputHash OutputHash;
-
-	/** Track which code path wrote this output, for tracking down a bug */
-	bool bOutputFromDDC;
-
-	/**
-	 * In-flight job with the given input hash.  Needs to be a reference pointer to handle cancelling of jobs, where an async DDC query
-	 * (which receives a pointer to FShaderJobData) may be in-flight that still references a job that has otherwise been deleted.
-	 * Cancelled jobs will have been unlinked from the PendingSubmitJobTaskJobs list in RemoveAllPendingJobsWithId, which can be
-	 * detected in the callback, and further processing on the job skipped.
-	 */
-	FShaderCommonCompileJobPtr JobInFlight;
-
-	/** Head of a linked list of duplicate jobs */
-	FShaderCommonCompileJob* DuplicateJobsWaitList = nullptr;
-
-	bool IsEmpty() const
-	{
-		return OutputHash.IsZero() && JobInFlight == nullptr && DuplicateJobsWaitList == nullptr;
-	}
+	FCookStats::FDDCResourceUsageStats UsageStats;
+	int32 ShadersCompiled = 0;
 
-	FORCEINLINE bool HasOutput() const
-	{
-		return OutputHash.IsZero() == false;
-	}
-};
+	static FCookStatsManager::FAutoRegisterCallback RegisterCookStats([](FCookStatsManager::AddStatFuncRef AddStat)
+		{
+			UsageStats.LogStats(AddStat, TEXT("GlobalShader.Usage"), TEXT(""));
+			AddStat(TEXT("GlobalShader.Misc"), FCookStatsManager::CreateKeyValueArray(
+				TEXT("ShadersCompiled"), ShadersCompiled
+			));
+		});
+}
+#endif
 
-/** Block of map elements for job cache */
-struct FShaderJobDataBlock
+const FString& GetGlobalShaderMapDDCKey()
 {
-	static const int32 BlockSize = 512;
-	static_assert(FMath::IsPowerOfTwo(BlockSize));
+	static FString GlobalShaderMapDDCKey = GetGlobalShaderMapDDCGuid().ToString();
+	return GlobalShaderMapDDCKey;
+}
 
-	FShaderJobData Data[BlockSize];
-};
+const FGuid& GetGlobalShaderMapDDCGuid()
+{
+	static FGuid GlobalShaderMapDDCGuid = FDevSystemGuids::GetSystemGuid(FDevSystemGuids::Get().GLOBALSHADERMAP_DERIVEDDATA_VER);
+	return GlobalShaderMapDDCGuid;
+}
 
-class FShaderJobDataMap
+const FString& GetMaterialShaderMapDDCKey()
 {
-public:
-	FShaderJobDataMap()
-	{
-		// Reserve so we don't need a special case for an empty HashTable array
-		Reserve(FShaderJobDataBlock::BlockSize);
-	}
+	static FString MaterialShaderMapDDCKey = GetMaterialShaderMapDDCGuid().ToString();
+	return MaterialShaderMapDDCKey;
+}
 
-	FShaderJobData* Find(const FShaderCompilerInputHash& Key);
-	FShaderJobCacheRef FindOrAdd(const FShaderCompilerInputHash& Key);
+const FGuid& GetMaterialShaderMapDDCGuid()
+{
+	static FGuid MaterialShaderMapDDCGuid = FDevSystemGuids::GetSystemGuid(FDevSystemGuids::Get().MATERIALSHADERMAP_DERIVEDDATA_VER);
+	return MaterialShaderMapDDCGuid;
+}
 
-	FORCEINLINE int32 Num() const
-	{
-		return NumItems;
-	}
+bool ShouldDumpShaderDDCKeys()
+{
+	return CVarShaderCompilerDumpDDCKeys.GetValueOnAnyThread();
+}
 
-	FORCEINLINE FShaderJobData& operator[](int32 Index)
-	{
-		check((uint32)Index < (uint32)NumItems);
-		return DataBlocks[Index / FShaderJobDataBlock::BlockSize].Data[Index & (FShaderJobDataBlock::BlockSize - 1)];
-	}
+void DumpShaderDDCKeyToFile(const EShaderPlatform InPlatform, bool bWithEditor, const FString& FileName, const FString& DDCKey)
+{
+	// deprecated version
+	const FString SubDirectory = bWithEditor ? TEXT("Editor") : TEXT("Game");
+	const FString TempPath = FPaths::ProjectSavedDir() / TEXT("ShaderDDCKeys") / SubDirectory / LexToString(InPlatform);
+	IFileManager::Get().MakeDirectory(*TempPath, true);
 
-	FORCEINLINE const FShaderJobData& operator[](int32 Index) const
-	{
-		return (*const_cast<FShaderJobDataMap*>(this))[Index];
-	}
+	const FString TempFile = TempPath / FileName;
 
-	uint64 GetAllocatedSize() const
-	{
-		return DataBlocks.GetAllocatedSize() + HashTable.GetAllocatedSize();
-	}
+	TUniquePtr<FArchive> DumpAr(IFileManager::Get().CreateFileWriter(*TempFile));
+	// serializing the string via << produces a non-textual file because it saves string's length, too
+	DumpAr->Serialize(const_cast<TCHAR*>(*DDCKey), DDCKey.Len() * sizeof(TCHAR));
+}
 
-	void RemoveLeadingBlocks(int32 BlocksToRemove)
-	{
-		check(BlocksToRemove <= DataBlocks.Num() && BlocksToRemove > 0);
-		DataBlocks.RemoveAt(0, BlocksToRemove);
-		NumItems -= BlocksToRemove * FShaderJobDataBlock::BlockSize;
-		check(NumItems >= 0);
+void DumpShaderDDCKeyToFile(const EShaderPlatform InPlatform, bool bEditorOnly, const TCHAR* DebugGroupName, const FString& DDCKey)
+{
+	const FString FileName = FString::Printf(TEXT("DDCKey-%s.txt"), bEditorOnly ? TEXT("Editor") : TEXT("Game"));
 
-		if (NumItems == 0)
-		{
-			// If we happened to remove ALL the items, reserve again, as done in the constructor
-			Reserve(FShaderJobDataBlock::BlockSize);
-		}
-		else
-		{
-			// Otherwise, we need to rehash, as all item indices will have changed
-			ReHash(GetDesiredHashTableSize());
-		}
-	}
+	const FString TempPath = GShaderCompilingManager->GetAbsoluteShaderDebugInfoDirectory() / FGenericDataDrivenShaderPlatformInfo::GetName(InPlatform).ToString() / DebugGroupName;
+	IFileManager::Get().MakeDirectory(*TempPath, true);
 
-private:
-	void ReHash(int32 HashTableSize);
-	void Reserve(int32 NumReserve);
+	const FString TempFile = TempPath / FileName;
 
-	int32 GetDesiredHashTableSize() const
-	{
-		return FMath::RoundUpToPowerOfTwo(DataBlocks.Num() * FShaderJobDataBlock::BlockSize * 2);
-	}
+	FFileHelper::SaveStringToFile(DDCKey, *TempFile);
+}
 
-	/** An indirect array of blocks is used, so data elements never move in memory when the table grows */
-	TIndirectArray<FShaderJobDataBlock> DataBlocks;
-	int32 NumItems = 0;
 
-	/** Power of two hash table with linear probing */
-	TArray<uint32> HashTable;
-	uint32 HashTableMask = 0;
-};
+static float GRegularWorkerTimeToLive = 20.0f;
+static float GBuildWorkerTimeToLive = 600.0f;
 
 
-struct FShaderJobCacheStoredOutput
+template<class EnumType>
+constexpr auto& CastEnumToUnderlyingTypeReference(EnumType& Type)
 {
-private:
-	/** How many times this output is referenced by the cached jobs */
-	int32 NumReferences = 0;
-
-public:
-
-	/** How many times this output has been returned as a cached result, no matter the input hash */
-	int32 NumHits = 0;
-
-	/** Canned output */
-	FSharedBuffer JobOutput;
+	static_assert(TIsEnum<EnumType>::Value, "");
+	using UnderType = __underlying_type(EnumType);
+	return reinterpret_cast<UnderType&>(Type);
+}
 
-	/** Path to where the cached debug info is stored. */
-	FString CachedDebugInfoPath;
+// Set to 1 to debug ShaderCompileWorker.exe. Set a breakpoint in LaunchWorker() to get the cmd-line.
+#define DEBUG_SHADERCOMPILEWORKER 0
 
-	/** Similar to FRefCountBase AddRef, but not atomic */
-	int32 AddRef()
-	{
-		++NumReferences;
+// Default value comes from bPromptToRetryFailedShaderCompiles in BaseEngine.ini
+// This is set as a global variable to allow changing in the debugger even in release
+// For example if there are a lot of content shader compile errors you want to skip over without relaunching
+bool GRetryShaderCompilation = true;
 
-		return NumReferences;
-	}
+static FShaderCompilingManager::EDumpShaderDebugInfo GDumpShaderDebugInfo = FShaderCompilingManager::EDumpShaderDebugInfo::Never;
+static FAutoConsoleVariableRef CVarDumpShaderDebugInfo(
+	TEXT("r.DumpShaderDebugInfo"),
+	CastEnumToUnderlyingTypeReference(GDumpShaderDebugInfo),
+	TEXT("Dumps debug info for compiled shaders to GameName/Saved/ShaderDebugInfo\n")
+	TEXT("When set to 1, debug info is dumped for all compiled shader\n")
+	TEXT("When set to 2, it is restricted to shaders with compilation errors\n")
+	TEXT("When set to 3, it is restricted to shaders with compilation errors or warnings\n")
+	TEXT("The debug info is platform dependent, but usually includes a preprocessed version of the shader source.\n")
+	TEXT("Global shaders automatically dump debug info if r.ShaderDevelopmentMode is enabled, this cvar is not necessary.\n")
+	TEXT("On iOS, if the PowerVR graphics SDK is installed to the default path, the PowerVR shader compiler will be called and errors will be reported during the cook.")
+	);
 
-	int32 GetNumReferences() const
-	{
-		return NumReferences;
-	}
+static int32 GDumpShaderDebugInfoShort = 0;
+static FAutoConsoleVariableRef CVarDumpShaderDebugShortNames(
+	TEXT("r.DumpShaderDebugShortNames"),
+	GDumpShaderDebugInfoShort,
+	TEXT("Only valid when r.DumpShaderDebugInfo > 0.\n")
+	TEXT("When set to 1, will shorten names factory and shader type folder names to avoid issues with long paths.")
+	);
 
-	/** Similar to FRefCountBase Release, but not atomic */
-	int32 Release()
-	{
-		checkf(NumReferences >= 0, TEXT("Attempting to release shader job cache output that was already released"));
+static int32 GDumpShaderDebugInfoBindless = 0;
+static FAutoConsoleVariableRef CVarDumpShaderDebugBindlessNames(
+	TEXT("r.DumpShaderDebugBindlessNames"),
+	GDumpShaderDebugInfoBindless,
+	TEXT("Only valid when r.DumpShaderDebugInfo > 0.\n")
+	TEXT("When set to 1, will add bindless folder names.")
+	);
 
-		--NumReferences;
-
-		const int32 RemainingNumReferences = NumReferences;
+static int32 GDumpShaderDebugInfoSCWCommandLine = 0;
+static FAutoConsoleVariableRef CVarDumpShaderDebugSCWCommandLine(
+	TEXT("r.DumpShaderDebugWorkerCommandLine"),
+	GDumpShaderDebugInfoSCWCommandLine,
+	TEXT("Only valid when r.DumpShaderDebugInfo > 0.\n")
+	TEXT("When set to 1, it will generate a file that can be used with ShaderCompileWorker's -directcompile.")
+	);
 
-		if (RemainingNumReferences == 0)
-		{
-			delete this;
-		}
+static int32 GShaderMapCompilationTimeout = 2 * 60 * 60;	// anything below an hour can hit a false positive
+static FAutoConsoleVariableRef CVarShaderMapCompilationTimeout(
+	TEXT("r.ShaderCompiler.ShadermapCompilationTimeout"),
+	GShaderMapCompilationTimeout,
+	TEXT("Maximum number of seconds a single shadermap (which can be comprised of multiple jobs) can be compiled after being considered hung.")
+);
 
-		return RemainingNumReferences;
-	}
+static int32 GCrashOnHungShaderMaps = 0;
+static FAutoConsoleVariableRef CVarCrashOnHungShaderMaps(
+	TEXT("r.ShaderCompiler.CrashOnHungShaderMaps"),
+	GCrashOnHungShaderMaps,
+	TEXT("If set to 1, the shader compiler will crash on hung shadermaps.")
+);
 
-	uint64 GetAllocatedSize() const
-	{
-		return static_cast<uint64>(JobOutput.GetSize() + sizeof(*this));
-	}
-};
+static int32 GForceAllCoresForShaderCompiling = 0;
+static FAutoConsoleVariableRef CVarForceAllCoresForShaderCompiling(
+	TEXT("r.ForceAllCoresForShaderCompiling"),
+	GForceAllCoresForShaderCompiling,
+	TEXT("When set to 1, it will ignore INI settings and launch as many ShaderCompileWorker instances as cores are available.\n")
+	TEXT("Improves shader throughput but for big projects it can make the machine run OOM")
+);
 
+static TAutoConsoleVariable<int32> CVarShadersSymbols(
+	TEXT("r.Shaders.Symbols"),
+	0,
+	TEXT("Enables debugging of shaders in platform specific graphics debuggers. This will generate and write shader symbols.\n")
+	TEXT("This enables the behavior of both r.Shaders.GenerateSymbols and r.Shaders.WriteSymbols.\n")
+	TEXT("Enables shader debugging features that require shaders to be recompiled. This compiles shaders with symbols and also includes extra runtime information like shader names. When using graphical debuggers it can be useful to enable this on startup.\n")
+	TEXT("This setting can be overriden in any Engine.ini under the [ShaderCompiler] section."),
+	ECVF_ReadOnly);
 
-/**
- * Class that provides a lock striped hash table of jobs, to reduce lock contention when adding or removing jobs
- */
-class FShaderCompilerJobTable
-{
-public:
-	static const int32 NUM_STRIPE_BITS = 6;
-	static const int32 NUM_STRIPES = 1 << NUM_STRIPE_BITS;
+static TAutoConsoleVariable<int32> CVarShadersSymbolsInfo(
+	TEXT("r.Shaders.SymbolsInfo"),
+	0,
+	TEXT("In lieu of a full set of platform shader PDBs, save out a slimmer ShaderSymbols.Info which contains shader platform hashes and shader debug info.\n")
+	TEXT("An option for when it is not practical to save PDBs for shaders all the time.\n")
+	TEXT("This setting can be overriden in any Engine.ini under the [ShaderCompiler] section."),
+	ECVF_ReadOnly);
 
-	/** We want to use the high bits of the hash for the stripe index, as it won't have influence on the hash table index within the stripe */
-	static const int32 STRIPE_SHIFT = 32 - NUM_STRIPE_BITS;
+static TAutoConsoleVariable<int32> CVarShadersGenerateSymbols(
+	TEXT("r.Shaders.GenerateSymbols"),
+	0,
+	TEXT("Enables generation of data for shader debugging when compiling shaders. This explicitly does not write any shader symbols to disk.\n")
+	TEXT("This setting can be overriden in any Engine.ini under the [ShaderCompiler] section."),
+	ECVF_ReadOnly);
 
-	template<typename JobType, typename KeyType>
-	FShaderCommonCompileJobPtr PrepareJob(uint32 InId, const KeyType& InKey, EShaderCompileJobPriority InPriority, bool& bOutNewJob)
-	{
-		const uint32 Hash = InKey.MakeHash(InId);
-		FLockStripeData& Stripe = GetStripe(JobType::Type, Hash);
+static TAutoConsoleVariable<int32> CVarShadersWriteSymbols(
+	TEXT("r.Shaders.WriteSymbols"),
+	0,
+	TEXT("Enables writing shader symbols to disk for platforms that support that. This explicitly does not enable generation of shader symbols.\n")
+	TEXT("This setting can be overriden in any Engine.ini under the [ShaderCompiler] section."),
+	ECVF_ReadOnly);
 
-		FWriteScopeLock Locker(Stripe.StripeLock);
+static TAutoConsoleVariable<FString> CVarShadersSymbolPathOverride(
+	TEXT("r.Shaders.SymbolPathOverride"),
+	"",
+	TEXT("Override output location of shader symbols. If the path contains the text '{Platform}', that will be replaced with the shader platform string.\n")
+	TEXT("Empty: use default location Saved/ShaderSymbols/{Platform}\n")
+	TEXT("This setting can be overriden in any Engine.ini under the [ShaderCompiler] section."),
+	ECVF_ReadOnly);
 
-		JobType* ResultJob = InternalFindJob<JobType>(Hash, InId, InKey);
-		bOutNewJob = false;
+static TAutoConsoleVariable<int32> CVarAllowUniqueDebugInfo(
+	TEXT("r.Shaders.AllowUniqueSymbols"),
+	0,
+	TEXT("When enabled, this tells supported shader compilers to generate symbols based on source files.\n")
+	TEXT("Enabling this can cause a drastic increase in the number of symbol files, enable only if absolutely necessary.\n")
+	TEXT("This setting can be overriden in any Engine.ini under the [ShaderCompiler] section."),
+	ECVF_ReadOnly);
 
-		if (ResultJob == nullptr)
-		{
-			ResultJob = new JobType(Hash, InId, InPriority, InKey);
-			InternalAddJob(ResultJob);
-			bOutNewJob = true;
-		}
+static TAutoConsoleVariable<int32> CVarShadersWriteSymbolsZip(
+	TEXT("r.Shaders.WriteSymbols.Zip"),
+	0,
+	TEXT(" 0: Export as loose files.\n")
+	TEXT(" 1: Export as an uncompressed archive.\n"),
+	ECVF_ReadOnly);
 
-		return ResultJob;
-	}
+static TAutoConsoleVariable<int32> CVarShadersEnableExtraData(
+	TEXT("r.Shaders.ExtraData"),
+	0,
+	TEXT("Enables generation of extra shader data that can be used at runtime. This includes shader names and other platform specific data.\n")
+	TEXT("This can add bloat to compiled shaders and can prevent shaders from being deduplicated.\n")
+	TEXT("This setting can be overriden in any Engine.ini under the [ShaderCompiler] section."),
+	ECVF_ReadOnly);
 
-	// PrepareJob creates a job with the given key if it's unique, while this adds an existing job, typically one that is cloned from another job
-	void AddExistingJob(FShaderCommonCompileJob* InJob)
-	{
-		FLockStripeData& Stripe = GetStripe(InJob->Type, InJob->Hash);
+static TAutoConsoleVariable<int32> CVarOptimizeShaders(
+	TEXT("r.Shaders.Optimize"),
+	1,
+	TEXT("Whether to optimize shaders.  When using graphical debuggers like Nsight it can be useful to disable this on startup.\n")
+	TEXT("This setting can be overriden in any Engine.ini under the [ShaderCompiler] section."),
+	ECVF_ReadOnly);
 
-		FWriteScopeLock Locker(Stripe.StripeLock);
-		InternalAddJob(InJob);
-	}
+static TAutoConsoleVariable<int32> CVarShaderFastMath(
+	TEXT("r.Shaders.FastMath"),
+	1,
+	TEXT("Whether to use fast-math optimisations in shaders."),
+	ECVF_ReadOnly);
 
-	void RemoveJob(FShaderCommonCompileJob* InJob)
-	{
-		FLockStripeData& Stripe = GetStripe(InJob->Type, InJob->Hash);
+static TAutoConsoleVariable<int32> CVarShaderZeroInitialise(
+	TEXT("r.Shaders.ZeroInitialise"),
+	1,
+	TEXT("Whether to enforce zero initialise local variables of primitive type in shaders. Defaults to 1 (enabled). Not all shader languages can omit zero initialisation."),
+	ECVF_ReadOnly);
 
-		FWriteScopeLock Locker(Stripe.StripeLock);
+static TAutoConsoleVariable<int32> CVarShaderBoundsChecking(
+	TEXT("r.Shaders.BoundsChecking"),
+	1,
+	TEXT("Whether to enforce bounds-checking & flush-to-zero/ignore for buffer reads & writes in shaders. Defaults to 1 (enabled). Not all shader languages can omit bounds checking."),
+	ECVF_ReadOnly);
 
-		const int32 JobIndex = InJob->JobIndex;
+static TAutoConsoleVariable<int32> CVarShaderWarningsAsErrors(
+	TEXT("r.Shaders.WarningsAsErrors"),
+	0,
+	TEXT("Whether to treat warnings as errors when compiling shaders. (0: disabled (default), 1: global shaders only, 2: all shaders)). This setting may be ignored on older platforms."),
+	ECVF_ReadOnly);
 
-		check(JobIndex != INDEX_NONE);
-		check(Stripe.Jobs[JobIndex] == InJob);
-		check(InJob->PendingPriority == EShaderCompileJobPriority::None);
-		InJob->JobIndex = INDEX_NONE;
+static TAutoConsoleVariable<int32> CVarShaderFlowControl(
+	TEXT("r.Shaders.FlowControlMode"),
+	0,
+	TEXT("Specifies whether the shader compiler should preserve or unroll flow-control in shader code.\n")
+	TEXT("This is primarily a debugging aid and will override any per-shader or per-material settings if not left at the default value (0).\n")
+	TEXT("\t0: Off (Default) - Entirely at the discretion of the platform compiler or the specific shader/material.\n")
+	TEXT("\t1: Prefer - Attempt to preserve flow-control.\n")
+	TEXT("\t2: Avoid - Attempt to unroll and flatten flow-control.\n"),
+	ECVF_ReadOnly);
 
-		Stripe.JobHash.Remove(InJob->Hash, JobIndex);
-		Stripe.FreeIndices.Add(JobIndex);
-		Stripe.Jobs[JobIndex].SafeRelease();
-	}
+static TAutoConsoleVariable<int32> CVarD3DCheckedForTypedUAVs(
+	TEXT("r.D3D.CheckedForTypedUAVs"),
+	1,
+	TEXT("Whether to disallow usage of typed UAV loads, as they are unavailable in Windows 7 D3D 11.0.\n")
+	TEXT(" 0: Allow usage of typed UAV loads.\n")
+	TEXT(" 1: Disallow usage of typed UAV loads. (default)"),
+	ECVF_ReadOnly);
 
-private:
-	template<typename JobType, typename KeyType>
-	JobType* InternalFindJob(uint32 InJobHash, uint32 InJobId, const KeyType& InKey) const
-	{
-		const FLockStripeData& Stripe = GetStripe(JobType::Type, InJobHash);
+static TAutoConsoleVariable<int32> CVarD3DForceDXC(
+	TEXT("r.D3D.ForceDXC"),
+	0,
+	TEXT("Forces DirectX Shader Compiler (DXC) to be used for all D3D shaders. Shaders compiled with this option are only compatible with D3D12.\n")
+	TEXT(" 0: Disable (default)\n")
+	TEXT(" 1: Force new compiler for all shaders"),
+	ECVF_ReadOnly);
 
-		uint32 CurrentPriorityIndex = 0u;
-		int32 CurrentIndex = INDEX_NONE;
-		for (int32 Index = Stripe.JobHash.First(InJobHash); Stripe.JobHash.IsValid(Index); Index = Stripe.JobHash.Next(Index))
-		{
-			const FShaderCommonCompileJob* Job = Stripe.Jobs[Index].GetReference();
-			check(Job->Type == JobType::Type);
+static TAutoConsoleVariable<int32> CVarWarpCulling(
+	TEXT("r.WarpCulling"),
+	0,
+	TEXT("Enable Warp Culling optimization for platforms that support it.\n")
+	TEXT(" 0: Disable (default)\n")
+	TEXT(" 1: Enable"),
+	ECVF_ReadOnly);
 
-			// We find the job that matches the key with the highest priority
-			if (Job->Id == InJobId &&
-				(uint32)Job->Priority >= CurrentPriorityIndex &&
-				static_cast<const JobType*>(Job)->Key == InKey)
-			{
-				CurrentPriorityIndex = (uint32)Job->Priority;
-				CurrentIndex = Index;
-			}
-		}
+static TAutoConsoleVariable<int32> CVarCullBeforeFetch(
+	TEXT("r.CullBeforeFetch"),
+	0,
+	TEXT("Enable Cull-Before-Fetch optimization for platforms that support it.\n")
+	TEXT(" 0: Disable (default)\n")
+	TEXT(" 1: Enable"),
+	ECVF_ReadOnly);
 
-		return CurrentIndex != INDEX_NONE ? static_cast<JobType*>(Stripe.Jobs[CurrentIndex].GetReference()) : nullptr;
-	}
+ENGINE_API int32 GCreateShadersOnLoad = 0;
+static FAutoConsoleVariableRef CVarCreateShadersOnLoad(
+	TEXT("r.CreateShadersOnLoad"),
+	GCreateShadersOnLoad,
+	TEXT("Whether to create shaders on load, which can reduce hitching, but use more memory.  Otherwise they will be created as needed.")
+);
 
-	void InternalAddJob(FShaderCommonCompileJob* InJob)
-	{
-		FLockStripeData& Stripe = GetStripe(InJob->Type, InJob->Hash);
+static TAutoConsoleVariable<FString> CVarShaderOverrideDebugDir(
+	TEXT("r.OverrideShaderDebugDir"),
+	"",
+	TEXT("Override output location of shader debug files\n")
+	TEXT("Empty: use default location Saved\\ShaderDebugInfo.\n"),
+	ECVF_ReadOnly);
 
-		int32 JobIndex = INDEX_NONE;
-		if (Stripe.FreeIndices.Num() > 0)
-		{
-			JobIndex = Stripe.FreeIndices.Pop(EAllowShrinking::No);
-			check(!Stripe.Jobs[JobIndex].IsValid());
-			Stripe.Jobs[JobIndex] = InJob;
-		}
-		else
-		{
-			JobIndex = Stripe.Jobs.Add(InJob);
-		}
+static TAutoConsoleVariable<int32> CVarShadersValidation(
+	TEXT("r.Shaders.Validation"),
+	1,
+	TEXT("Enabled shader compiler validation warnings and errors."),
+	ECVF_ReadOnly);
 
-		check(Stripe.Jobs[JobIndex].IsValid());
-		Stripe.JobHash.Add(InJob->Hash, JobIndex);
+static TAutoConsoleVariable<int32> CVarShadersRemoveDeadCode(
+	TEXT("r.Shaders.RemoveDeadCode"),
+	1,
+	TEXT("Run a preprocessing step that removes unreferenced code before compiling shaders.\n")
+	TEXT("This can improve the compilation speed for shaders which include many large utility headers.\n")
+	TEXT("\t0: Keep all input source code.\n")
+	TEXT("\t1: Remove unreferenced code before compilation (Default)\n"),
+	ECVF_ReadOnly);
 
-		check(InJob->Priority != EShaderCompileJobPriority::None);
-		check(InJob->PendingPriority == EShaderCompileJobPriority::None);
-		check(InJob->JobIndex == INDEX_NONE);
-		InJob->JobIndex = JobIndex;
+namespace ShaderCompiler
+{
+	bool IsDumpShaderDebugInfoAlwaysEnabled()
+	{
+		return GDumpShaderDebugInfo != FShaderCompilingManager::EDumpShaderDebugInfo::Always;
 	}
+} // namespace ShaderCompiler
 
-	struct FLockStripeData
-	{
-		TArray<FShaderCommonCompileJobPtr> Jobs;
-		TArray<int32> FreeIndices;
-		FHashTable JobHash;
-		FRWLock StripeLock;
-	};
+#if ENABLE_COOK_STATS
+namespace ShaderCompilerCookStats
+{
+	static double BlockingTimeSec = 0.0;
+	static double GlobalBeginCompileShaderTimeSec = 0.0;
+	static int32 GlobalBeginCompileShaderCalls = 0;
+	static double ProcessAsyncResultsTimeSec = 0.0;
+	double AsyncCompileTimeSec = 0.0;
 
-	FLockStripeData Stripes[NumShaderCompileJobTypes][NUM_STRIPES];
+	static FCookStatsManager::FAutoRegisterCallback RegisterCookStats([](FCookStatsManager::AddStatFuncRef AddStat)
+		{
+			AddStat(TEXT("ShaderCompiler"), FCookStatsManager::CreateKeyValueArray(
+				TEXT("BlockingTimeSec"), BlockingTimeSec,
+				TEXT("AsyncCompileTimeSec"), AsyncCompileTimeSec,
+				TEXT("GlobalBeginCompileShaderTimeSec"), GlobalBeginCompileShaderTimeSec,
+				TEXT("GlobalBeginCompileShaderCalls"), GlobalBeginCompileShaderCalls,
+				TEXT("ProcessAsyncResultsTimeSec"), ProcessAsyncResultsTimeSec
+			));
+		});
+}
+#endif
 
-	FORCEINLINE FLockStripeData& GetStripe(EShaderCompileJobType JobType, uint32 Hash)
+#if WITH_EDITOR
+static bool CheckSingleJob(const FShaderCompileJob& SingleJob, TArray<FString>& OutErrors)
+{
+	if (SingleJob.bSucceeded)
 	{
-		checkf((uint8)JobType < (uint8)NumShaderCompileJobTypes, TEXT("Out of range JobType index %u"), (uint8)JobType);
-		return Stripes[(uint8)JobType][Hash >> STRIPE_SHIFT];
+		checkf(SingleJob.Output.ShaderCode.GetShaderCodeSize() > 0, TEXT("Abnormal shader code size for a succeded job: %d bytes"), SingleJob.Output.ShaderCode.GetShaderCodeSize());
 	}
-	FORCEINLINE const FLockStripeData& GetStripe(EShaderCompileJobType JobType, uint32 Hash) const
+
+	bool bSucceeded = SingleJob.bSucceeded;
+
+	if (SingleJob.Key.ShaderType)
 	{
-		checkf((uint8)JobType < (uint8)NumShaderCompileJobTypes, TEXT("Out of range JobType index %u"), (uint8)JobType);
-		return Stripes[(uint8)JobType][Hash >> STRIPE_SHIFT];
+		// Allow the shader validation to fail the compile if it sees any parameters bound that aren't supported.
+		const bool bValidationResult = SingleJob.Key.ShaderType->ValidateCompiledResult(
+			(EShaderPlatform)SingleJob.Input.Target.Platform,
+			SingleJob.Output.ParameterMap,
+			OutErrors);
+		bSucceeded = bValidationResult && bSucceeded;
 	}
-};
-
-/** Private implementation class for FShaderCompileJobCollection */
-class FShaderJobCache
-{
-public:
-	FShaderJobCache(FCriticalSection& InCompileQueueSection);
-	~FShaderJobCache();
 
-	// Returns job pointer for new job, otherwise returns NULL
-	template <typename JobType, typename KeyType>
-	JobType* PrepareJob(uint32 InId, const KeyType& InKey, EShaderCompileJobPriority InPriority)
+	if (SingleJob.Key.VFType)
 	{
-		bool bNewJob;
-		FShaderCommonCompileJobPtr Result = JobTable.PrepareJob<JobType>(InId, InKey, InPriority, bNewJob);
+		const int32 OriginalNumErrors = OutErrors.Num();
 
-		if (bNewJob)
-		{
-			// If it's a new job, return it -- it's OK to cast the ref-counted pointer to a raw pointer, because JobTable
-			// itself has a reference to the job, and a newly added job hasn't been submitted yet, so it can't make a
-			// round trip through the pipeline and be released until that happens.
-			return (JobType*)Result.GetReference();
-		}
-		else if (InPriority > Result->Priority)
+		// Allow the vertex factory to fail the compile if it sees any parameters bound that aren't supported
+		SingleJob.Key.VFType->ValidateCompiledResult((EShaderPlatform)SingleJob.Input.Target.Platform, SingleJob.Output.ParameterMap, OutErrors);
+
+		if (OutErrors.Num() > OriginalNumErrors)
 		{
-			// Or if the priority changed, update that
-			InternalSetPriority(Result, InPriority);
+			bSucceeded = false;
 		}
-
-		return nullptr;
-	}
-
-	void RemoveJob(FShaderCommonCompileJob* InJob)
-	{
-		JobTable.RemoveJob(InJob);
 	}
 
-	int32 RemoveAllPendingJobsWithId(uint32 InId);
-
-	void SubmitJob(FShaderCommonCompileJob* Job);
-	void SubmitJobs(const TArray<FShaderCommonCompileJobPtr>& InJobs);
+	return bSucceeded;
+};
+#endif // WITH_EDITOR
 
-	/** This is an entry point for all jobs that have finished the compilation (whether real or cached). Can be called from multiple threads. Returns mutex stall time. */
-	double ProcessFinishedJob(FShaderCommonCompileJob* FinishedJob, bool bCompilationSkipped);
+/** Information tracked for each shader compile worker process instance. */
+struct FShaderCompileWorkerInfo
+{
+	/** Process handle of the worker app once launched.  Invalid handle means no process. */
+	FProcHandle WorkerProcess;
 
-	/** Adds the job to cache. */
-	void AddToCacheAndProcessPending(FShaderCommonCompileJob* FinishedJob);
+	/** Tracks whether tasks have been issued to the worker. */
+	bool bIssuedTasksToWorker;	
 
-	/** Populates caching stats in the given compiler stats struct. */
-	void GetStats(FShaderCompilerStats& OutStats) const;
+	/** Whether the worker has been launched for this set of tasks. */
+	bool bLaunchedWorker;
 
-	int32 GetNumPendingJobs(EShaderCompileJobPriority InPriority) const;
+	/** Tracks whether all tasks issued to the worker have been received. */
+	bool bComplete;
 
-	int32 GetNumOutstandingJobs() const;
+	/** Whether this worker is available for new jobs. It will be false when shutting down the worker. */
+	bool bAvailable; 
 
-	int32 GetNumPendingJobs() const;
+	/** Time at which the worker started the most recent batch of tasks. */
+	double StartTime;
 
-	int32 GetPendingJobs(EShaderCompilerWorkerType InWorkerType, EShaderCompileJobPriority InPriority, int32 MinNumJobs, int32 MaxNumJobs, TArray<FShaderCommonCompileJobPtr>& OutJobs);
+	/** Time at which the worker ended the most recent batch of tasks. */
+	double FinishTime = 0.0;
 
-private:
-	using FJobOutputHash = FBlake3Hash;
-	using FJobCachedOutput = FSharedBuffer;
-	using FStoredOutput = FShaderJobCacheStoredOutput;
+	/** Jobs that this worker is responsible for compiling. */
+	TArray<FShaderCommonCompileJobPtr> QueuedJobs;
 
-	// cannot allow managing this from outside as the caching logic is not exposed
-	inline int32 InternalSubtractNumOutstandingJobs(int32 Value)
+	FShaderCompileWorkerInfo() :
+		bIssuedTasksToWorker(false),		
+		bLaunchedWorker(false),
+		bComplete(false),
+		bAvailable(true),
+		StartTime(0)
 	{
-		const int32 PrevNumOutstandingJobs = NumOutstandingJobs.Subtract(Value);
-		check(PrevNumOutstandingJobs >= Value);
-		return PrevNumOutstandingJobs - Value;
 	}
 
-	void InternalSetPriority(FShaderCommonCompileJob* Job, EShaderCompileJobPriority InPriority);
-
-	/** Looks for or adds an entry for the given hash in the cache.  Returns cached output if it exists, or may initialize DDC request if one has been issued. */
-	FShaderJobCacheRef FindOrAdd(const FShaderCompilerInputHash& Hash, EShaderCompileJobPriority JobPriority, const bool bCheckDDC, TPimplPtr<UE::DerivedData::FRequestOwner>& InoutRequestOwner, FJobCachedOutput*& OutCachedOutput);
-
-	/** Find an existing item in the cache. */
-	FShaderJobData* Find(const FShaderCompilerInputHash& Hash);
-
-	/** Add a reference to a duplicate job (to the DuplicateJobs array) */
-	void AddDuplicateJob(FShaderCommonCompileJob* DuplicateJob);
-
-	/** Remove a reference to a duplicate job (from the DuplicateJobs array)  */
-	void RemoveDuplicateJob(FShaderCommonCompileJob* DuplicateJob);
-
-	/** Adds a job output to the cache */
-	void AddJobOutput(FShaderJobData& JobData, const FShaderCommonCompileJob* FinishedJob, const FShaderCompilerInputHash& Hash, const FJobCachedOutput& Contents, int32 InitialHitCount, const bool bAddToDDC);
-
-	/** Returns memory used by the cache*/
-	uint64 GetAllocatedMemory() const;
-
-	/** Compute memory used by the cache from scratch.  Should match GetAllocatedMemory() if CurrentlyAllocateMemory is being properly updated (useful for validation). */
-	uint64 ComputeAllocatedMemory() const;
-
-	/** Calculates current memory budget, in bytes */
-	uint64 GetCurrentMemoryBudget() const;
-
-	/** Cleans up oldest outputs to fit in the given memory budget */
-	void CullOutputsToMemoryBudget(uint64 TargetBudgetBytes);
-
-	/** Copied from TLinkedListBase::Unlink */
-	FORCEINLINE static void Unlink(FShaderCommonCompileJob& Job)
+	// warning: not virtual
+	~FShaderCompileWorkerInfo()
 	{
-		if (Job.NextLink)
-		{
-			Job.NextLink->PrevLink = Job.PrevLink;
-		}
-		if (Job.PrevLink)
+		if(WorkerProcess.IsValid())
 		{
-			*Job.PrevLink = Job.NextLink;
+			FPlatformProcess::TerminateProc(WorkerProcess);
+			FPlatformProcess::CloseProc(WorkerProcess);
 		}
-		// Make it safe to call Unlink again.
-		Job.NextLink = nullptr;
-		Job.PrevLink = nullptr;
 	}
+};
 
-	/**
-	 * Similar to TLinkedListBase::Unlink, but updates a Tail pointer if the Tail is unlinked.  The tail must
-	 * originally be initialized as Tail = &Head.
-	 */
-	FORCEINLINE void UnlinkWithTail(FShaderCommonCompileJob& Job, FShaderCommonCompileJob**& Tail)
-	{
-		// Update tail if we are removing that element
-		if (Tail == &Job.NextLink)
-		{
-			Tail = Job.PrevLink;
-		}
-		Unlink(Job);
-	}
+FShaderCompilingManager* GShaderCompilingManager = nullptr;
 
-	/** Copied from TLinkedListBase::LinkHead */
-	FORCEINLINE void LinkHead(FShaderCommonCompileJob& Job, FShaderCommonCompileJob*& Head)
+bool FShaderCompilingManager::AllTargetPlatformSupportsRemoteShaderCompiling()
+{
+	// no compiling support
+	if (!AllowShaderCompiling())
 	{
-		if (Head != NULL)
-		{
-			Head->PrevLink = &Job.NextLink;
-		}
-
-		Job.NextLink = Head;
-		Job.PrevLink = &Head;
-		Head = &Job;
+		return false;
 	}
 
-	/** Copied from TLinkedListBase::LinkAfter */
-	FORCEINLINE void LinkAfter(FShaderCommonCompileJob& Job, FShaderCommonCompileJob* After)
+	ITargetPlatformManagerModule* TPM = GetTargetPlatformManager();	
+	if (!TPM)
 	{
-		checkSlow(After != NULL);
-		check(Job.NextLink == nullptr && Job.PrevLink == nullptr);
-
-		Job.PrevLink = &After->NextLink;
-		Job.NextLink = *Job.PrevLink;
-		*Job.PrevLink = (FShaderCommonCompileJob*)&Job;
-
-		if (Job.NextLink != NULL)
-		{
-			Job.NextLink->PrevLink = &Job.NextLink;
-		}
+		return false;
 	}
-
-	/**
-	 * Similar to TLinkedListBase::LinkHead, but uses atomic operations to allow multiple producer threads to add to the linked list
-	 * without needing synchronization ("wait free").  Note that synchronization is required for other operations on the list, such
-	 * as traversal or removal, as the list isn't in a fully valid state mid operation (Head always points to the latest item
-	 * inserted, but it may not yet be linked with the rest of the items).  Synchronization is accomplished by using a read lock
-	 * (which multiple threads can hold) for atomic insertion operations, and a write lock for all other operations.
-	 */
-	FORCEINLINE void LinkHeadAtomic(FShaderCommonCompileJob& Job, FShaderCommonCompileJob*& Head)
+	
+	const TArray<ITargetPlatform*>& Platforms = TPM->GetActiveTargetPlatforms();
+	for (int32 Index = 0; Index < Platforms.Num(); Index++)
 	{
-		check(Job.NextLink == nullptr && Job.PrevLink == nullptr);
-
-		// It's important that PrevLink is set before the InterlockedExchange, as a subsequent Head pointer exchange could write
-		// another item and need to update PrevLink for this item before this function completes.
-		Job.PrevLink = &Head;
-
-		FShaderCommonCompileJob* OldHead = (FShaderCommonCompileJob*)FPlatformAtomics::InterlockedExchange((PTRINT*)&Head, (PTRINT)&Job);
-		if (OldHead != nullptr)
+		if (!Platforms[Index]->CanSupportRemoteShaderCompile())
 		{
-			OldHead->PrevLink = &Job.NextLink;
+			return false;
 		}
-		Job.NextLink = OldHead;
-	}
-
-	/**
-	 * Variation that links a job at the tail of the list.  The tail must originally be initialized as Tail = &Head.  Similar to
-	 * LinkHeadAtomic above (see more detailed comments there), a read lock is required for this operation.
-	 */
-	FORCEINLINE static void LinkTailAtomic(FShaderCommonCompileJob& Job, FShaderCommonCompileJob**& Tail)
-	{
-		check(Job.NextLink == nullptr && Job.PrevLink == nullptr);
-
-		FShaderCommonCompileJob** OldTail = (FShaderCommonCompileJob**)FPlatformAtomics::InterlockedExchange((PTRINT*)&Tail, (PTRINT)&Job.NextLink);
-		Job.PrevLink = OldTail;
-
-		// Update previous tail's next pointer (or OldTail may be pointing at Head if list was empty)
-		*OldTail = (FShaderCommonCompileJob*)&Job;
 	}
+	
+	return true;
+}
 
-	/** Links job into linked list with its given Priority */
-	FORCEINLINE void LinkJobWithPriority(FShaderCommonCompileJob& Job)
+// Returns a rank for the preference of distributed shader controllers; Higher is better.
+static int32 GetShaderControllerPreferenceRank(IDistributedBuildController& Controller)
+{
+	const FString Name = Controller.GetName();
+	if (Name.StartsWith(TEXT("UBA")))
 	{
-		int32 PriorityIndex = (int32)Job.Priority;
-		check((uint32)PriorityIndex < (uint32)NumShaderCompileJobPriorities);
-		check(Job.PendingPriority == EShaderCompileJobPriority::None);
-		NumPendingJobs[PriorityIndex]++;
-#if UE_SHADERCOMPILER_FIFO_JOB_EXECUTION
-		LinkTailAtomic(Job, PendingJobsTail[PriorityIndex]);
-#else
-		LinkHeadAtomic(Job, PendingJobsHead[PriorityIndex]);
-#endif
-		Job.PendingPriority = Job.Priority;
+		return 2;
 	}
-
-	/** Unlinks job from linked list with its current PendingPriority */
-	FORCEINLINE void UnlinkJobWithPriority(FShaderCommonCompileJob& Job)
+	else if (Name.StartsWith(TEXT("XGE")))
 	{
-		int32 PriorityIndex = (int32)Job.PendingPriority;
-		check((uint32)PriorityIndex < (uint32)NumShaderCompileJobPriorities);
-		check(NumPendingJobs[PriorityIndex] > 0);
-		NumPendingJobs[PriorityIndex]--;
-#if UE_SHADERCOMPILER_FIFO_JOB_EXECUTION
-		UnlinkWithTail(Job, PendingJobsTail[PriorityIndex]);
-#else
-		Unlink(Job);
-#endif
-		Job.PendingPriority = EShaderCompileJobPriority::None;
+		return 1;
 	}
-
-	/** From FShaderCompilingManager, guards access to FShaderMapCompileResults written in ProcessFinishedJob */
-	FCriticalSection& CompileQueueSection;
-
-	/** Guards access to the structure */
-	mutable FRWLock JobLock;
-
-	/** Needed to detect if DDC query callback is completing in the same thread as SubmitJob */
-	static thread_local bool bInSubmitJobThread;
-
-	/** List of jobs waiting on SubmitJob task or DDC query (not yet added to a pending queue). */
-	FShaderCommonCompileJob* PendingSubmitJobTaskJobs = nullptr;
-
-	/** Queue of tasks that haven't been assigned to a worker yet. */
-	TStaticArray<FShaderCommonCompileJob*, NumShaderCompileJobPriorities> PendingJobsHead;
-	TStaticArray<std::atomic_int32_t, NumShaderCompileJobPriorities> NumPendingJobs;
-#if UE_SHADERCOMPILER_FIFO_JOB_EXECUTION
-	TStaticArray<FShaderCommonCompileJob**, NumShaderCompileJobPriorities> PendingJobsTail;
-#endif
-
-	/** Number of jobs currently being compiled.  This includes PendingJobs and any jobs that have been assigned to workers but aren't complete yet. */
-	FThreadSafeCounter NumOutstandingJobs;
-
-	/** Active jobs */
-	FShaderCompilerJobTable JobTable;
-
-	/* A lot of outputs can be duplicated, so they are deduplicated before storing */
-	TMap<FJobOutputHash, FStoredOutput*> Outputs;
-
-	TMap<FJobOutputHash, FString> CachedJobNames;
-
-	/** Map of input hashes to job data (in flight jobs and output) */
-	FShaderJobDataMap InputHashToJobData;
-
-	/** List of duplicate jobs */
-	TArray<FShaderCommonCompileJob*> DuplicateJobs;
-
-	/** Statistics - total number of times we tried to Find() some input hash */
-	uint64 TotalSearchAttempts = 0;
-
-	/** Statistics - total number of times we succeded in Find()ing output for some input hash */
-	uint64 TotalCacheHits = 0;
-
-	/** Statistics - total number of times a duplicate job was added (duplicate jobs are processed when the original finishes compiling) */
-	uint64 TotalCacheDuplicates = 0;
-
-	/** Statistics - total number of times a per-shader DDC query was issued */
-	uint64 TotalCacheDDCQueries = 0;
-
-	/** Statistics - total number of times a per-shader DDC query succeeded for some input hash */
-	uint64 TotalCacheDDCHits = 0;
-
-	/** Statistics - allocated memory. If the number is non-zero, we can trust it as accurate. Otherwise, recalculate. */
-	uint64 CurrentlyAllocatedMemory = 0;
-};
-
-thread_local bool FShaderJobCache::bInSubmitJobThread;
-
-static FShaderJobData& GetShaderJobData(const FShaderJobCacheRef& CacheRef)
-{
-	check(CacheRef.Block);
-	return CacheRef.Block->Data[CacheRef.IndexInBlock];
-}
-
-FShaderJobData* FShaderJobDataMap::Find(const FShaderCompilerInputHash& Key)
-{
-	// Search for key with linear probing
-	for (uint32 TableIndex = GetTypeHash(Key) & HashTableMask; HashTable[TableIndex] != INDEX_NONE; TableIndex = (TableIndex + 1) & HashTableMask)
+	else
 	{
-		if ((*this)[TableIndex].InputHash == Key)
-		{
-			return &(*this)[TableIndex];
-		}
+		return 0;
 	}
-	return nullptr;
 }
 
-FShaderJobCacheRef FShaderJobDataMap::FindOrAdd(const FShaderCompilerInputHash& Key)
+IDistributedBuildController* FShaderCompilingManager::FindRemoteCompilerController() const
 {
-	// Search for key with linear probing
-	uint32 TableIndex;
-	for (TableIndex = GetTypeHash(Key) & HashTableMask; HashTable[TableIndex] != INDEX_NONE; TableIndex = (TableIndex + 1) & HashTableMask)
+	// no controllers needed if not compiling
+	if (!AllowShaderCompiling())
 	{
-		int32 ItemIndex = HashTable[TableIndex];
-		if ((*this)[ItemIndex].InputHash == Key)
-		{
-			return FShaderJobCacheRef({ &DataBlocks[ItemIndex / FShaderJobDataBlock::BlockSize], ItemIndex & (FShaderJobDataBlock::BlockSize - 1), INDEX_NONE });
-		}
+		return nullptr;
 	}
 
-	// Ensure there is space for item
-	Reserve(NumItems + 1);
-
-	// Initialize allocated item
-	int32 AllocatedIndex = NumItems++;
-	FShaderJobCacheRef AllocatedItem({ &DataBlocks[AllocatedIndex / FShaderJobDataBlock::BlockSize], AllocatedIndex & (FShaderJobDataBlock::BlockSize - 1), INDEX_NONE });
-	GetShaderJobData(AllocatedItem).InputHash = Key;
-
-	// Add to empty spot in hash table
-	HashTable[TableIndex] = AllocatedIndex;
-
-	return AllocatedItem;
-}
+	TArray<IDistributedBuildController*> AvailableControllers = IModularFeatures::Get().GetModularFeatureImplementations<IDistributedBuildController>(IDistributedBuildController::GetModularFeatureType());
 
-void FShaderJobDataMap::ReHash(int32 HashTableSize)
-{
-	// Resize table and rehash
-	HashTable.SetNumUninitialized(HashTableSize);
-	memset(HashTable.GetData(), 0xff, HashTable.Num() * HashTable.GetTypeSize());
-	HashTableMask = HashTableSize - 1;
+	// Prefer UBA, then XGE, and fallback to any other controller otherwise
+	int32 SupportedControllerPreferenceRank = 0;
+	IDistributedBuildController* SupportedController = nullptr;
 
-	for (int32 OuterIndex = 0; OuterIndex < DataBlocks.Num(); OuterIndex++)
+	for (IDistributedBuildController* Controller : AvailableControllers)
 	{
-		FShaderJobData* Data = DataBlocks[OuterIndex].Data;
-
-		for (int32 InnerIndex = 0; InnerIndex < FShaderJobDataBlock::BlockSize; InnerIndex++)
+		if (Controller != nullptr && Controller->IsSupported())
 		{
-			int32 Index = OuterIndex * FShaderJobDataBlock::BlockSize + InnerIndex;
-			if (Index >= NumItems)
-			{
-				OuterIndex = DataBlocks.Num();
-				break;
-			}
-
-			// Find table entry for key -- keys will be unique when rehashing, so we don't need to check for existing keys
-			for (uint32 TableIndex = GetTypeHash(Data[InnerIndex].InputHash) & HashTableMask;; TableIndex = (TableIndex + 1) & HashTableMask)
+			const int32 PreferenceRank = GetShaderControllerPreferenceRank(*Controller);
+			if (SupportedController == nullptr || SupportedControllerPreferenceRank < PreferenceRank)
 			{
-				if (HashTable[TableIndex] == INDEX_NONE)
-				{
-					HashTable[TableIndex] = Index;
-					break;
-				}
+				SupportedController = Controller;
+				SupportedControllerPreferenceRank = PreferenceRank;
 			}
 		}
 	}
+
+	if (SupportedController != nullptr)
+	{
+		SupportedController->InitializeController();
+		return SupportedController;
+	}
+
+	return nullptr;
 }
 
-void FShaderJobDataMap::Reserve(int32 NumReserve)
+void FShaderCompilingManager::ReportMemoryUsage()
 {
-	if (NumReserve > DataBlocks.Num() * FShaderJobDataBlock::BlockSize)
+	// This function runs from within an OOM callback. It should not take locks, as much as possible.
+	constexpr bool bAllowToWaitForLock = false;
+	for (const TUniquePtr<FShaderCompileThreadRunnableBase>& ThreadPtr : Threads)
 	{
-		while (NumReserve > DataBlocks.Num() * FShaderJobDataBlock::BlockSize)
-		{
-			DataBlocks.Add(new FShaderJobDataBlock);
-		}
-
-		int32 HashTableSize = GetDesiredHashTableSize();
-		if (HashTableSize != HashTable.Num())
-		{
-			ReHash(HashTableSize);
-		}
+		ThreadPtr->PrintWorkerMemoryUsage(bAllowToWaitForLock);
 	}
 }
 
-void FShaderJobCache::CullOutputsToMemoryBudget(uint64 TargetBudgetBytes)
+FShaderCompilingManager::FShaderCompilingManager() :
+	bCompilingDuringGame(false),
+	NumExternalJobs(0),
+	AllJobs(CompileQueueSection),
+	NumSingleThreadedRunsBeforeRetry(GSingleThreadedRunsIdle),
+	SuppressedShaderPlatforms(0),
+	BuildDistributionController(nullptr),
+	bNoShaderCompilation(false),
+	bAllowForIncompleteShaderMaps(false),
+	Notification(MakeUnique<FAsyncCompilationNotification>(GetAssetNameFormat()))
 {
-	// Track consecutive empty items.  We can delete empty blocks from the front of the map at the end.
-	int32 ConsecutiveEmptyItems = 0;
-	uint64 EmptyBlockSavings = 0;
-
-	// We don't cull items from the last block
-
-	for (int32 ItemIndex = 0; ItemIndex < InputHashToJobData.Num(); ItemIndex++)
+	// don't perform any initialization if compiling is not allowed
+	if (!AllowShaderCompiling())
 	{
-		FShaderJobData& JobData = InputHashToJobData[ItemIndex];
+		// use existing flag to disable compiling
+		bNoShaderCompilation = true;
+		return;
+	}
 
-		// Check if we are in budget yet
-		if (CurrentlyAllocatedMemory - EmptyBlockSavings <= TargetBudgetBytes)
-		{
-			break;
+	bIsEngineLoopInitialized = false;
+	FCoreDelegates::OnFEngineLoopInitComplete.AddLambda([&]() 
+		{ 
+			bIsEngineLoopInitialized = true; 
 		}
+	);
 
-		// We can only free this output if there is no in-flight job
-		if (JobData.JobInFlight == nullptr)
-		{
-			// Empty this item out (if not already empty), by removing the reference to the output and zeroing it out
-			if (!JobData.OutputHash.IsZero())
-			{
-				FStoredOutput** FoundStoredOutput = Outputs.Find(JobData.OutputHash);
-
-				if (FoundStoredOutput)
-				{
-					FStoredOutput* StoredOutput = *FoundStoredOutput;
-					checkf(StoredOutput, TEXT("Invalid entry found in FShaderJobCache Output hash table. All values are expected to be valid pointers."));
-
-					const uint64 OutputSize = StoredOutput->GetAllocatedSize();
-
-					// Decrement reference count and remove cached object if it's no longer referenced by any input hashes
-					if (StoredOutput->Release() == 0)
-					{
-						Outputs.Remove(JobData.OutputHash);
-						CachedJobNames.Remove(JobData.OutputHash);
-						CurrentlyAllocatedMemory -= OutputSize;
-					}
-				}
-
-				JobData.OutputHash.Reset();
-			}
-
-			// Track if this is another consecutive empty item
-			if (ItemIndex == ConsecutiveEmptyItems)
-			{
-				ConsecutiveEmptyItems++;
+	WorkersBusyTime = 0;
 
-				// Take into account that we will be removing empty job data blocks at the end, by adding the savings when we reach a full block
-				if ((ConsecutiveEmptyItems & (FShaderJobDataBlock::BlockSize - 1)) == 0)
-				{
-					EmptyBlockSavings += sizeof(FShaderJobDataBlock);
-				}
-			}
-		}
-	}
+#if PLATFORM_WINDOWS
+	FString ExecutableName("ShaderCompileWorker.exe");
+#else
+	FString ExecutableName("ShaderCompileWorker");
+#endif
 
-	int32 ConsecutiveEmptyBlocks = ConsecutiveEmptyItems / FShaderJobDataBlock::BlockSize;
-	if (ConsecutiveEmptyBlocks > 0)
+	// first look for project-specific version
+	ShaderCompileWorkerName = FPaths::Combine(FPaths::ProjectDir(), TEXT("Binaries"), FPlatformProcess::GetBinariesSubdirectory(), ExecutableName);
+	if (!IFileManager::Get().FileExists(*ShaderCompileWorkerName))
 	{
-		uint64 InputHashToJobDataOriginalSize = InputHashToJobData.GetAllocatedSize();
-
-		InputHashToJobData.RemoveLeadingBlocks(ConsecutiveEmptyBlocks);
-
-		CurrentlyAllocatedMemory += InputHashToJobData.GetAllocatedSize() - InputHashToJobDataOriginalSize;
+		// fallback to standard Engine location
+		ShaderCompileWorkerName = FPaths::Combine(FPaths::EngineDir(), TEXT("Binaries"), FPlatformProcess::GetBinariesSubdirectory(), ExecutableName);
 	}
-}
 
+	// Threads must use absolute paths on Windows in case the current directory is changed on another thread!
+	ShaderCompileWorkerName = FPaths::ConvertRelativePathToFull(ShaderCompileWorkerName);
 
-#if WITH_EDITOR
+	// Read values from the engine ini
+	verify(GConfig->GetBool( TEXT("DevOptions.Shaders"), TEXT("bAllowCompilingThroughWorkers"), bAllowCompilingThroughWorkers, GEngineIni ));
+	verify(GConfig->GetBool( TEXT("DevOptions.Shaders"), TEXT("bAllowAsynchronousShaderCompiling"), bAllowAsynchronousShaderCompiling, GEngineIni ));
 
-static FDelayedAutoRegisterHelper GKickOffShaderAutoGenForPlatforms(EDelayedRegisterRunPhase::DeviceProfileManagerReady, []
-{
-	// also do this for all active target platforms (e.g. when cooking)
-	ITargetPlatformManagerModule* TPM = GetTargetPlatformManager();
-	if (TPM)
+	// Explicitly load ShaderPreprocessor module so it will run its initialization step
+	FModuleManager::LoadModuleChecked<IModuleInterface>(TEXT("ShaderPreprocessor"));
+	
+	// override the use of workers, can be helpful for debugging shader compiler code
+	static const IConsoleVariable* CVarAllowCompilingThroughWorkers = IConsoleManager::Get().FindConsoleVariable(TEXT("r.Shaders.AllowCompilingThroughWorkers"), false);
+	if (!FPlatformProcess::SupportsMultithreading() || FParse::Param(FCommandLine::Get(), TEXT("noshaderworker")) || (CVarAllowCompilingThroughWorkers && CVarAllowCompilingThroughWorkers->GetInt() == 0))
 	{
-		const TArray<ITargetPlatform*>& Platforms = TPM->GetActiveTargetPlatforms();
-
-		for (int32 Index = 0; Index < Platforms.Num(); ++Index)
-		{
-			TArray<FName> DesiredShaderFormats;
-			checkf(Platforms[Index], TEXT("Null platform on the list of active platforms!"));
-			Platforms[Index]->GetAllTargetedShaderFormats(DesiredShaderFormats);
-
-			for (int32 FormatIndex = 0; FormatIndex < DesiredShaderFormats.Num(); ++FormatIndex)
-			{
-				FShaderCompileUtilities::GenerateBrdfHeaders(DesiredShaderFormats[FormatIndex]);
-			}
-		}
+		bAllowCompilingThroughWorkers = false;
 	}
 
-	// also do this for the editor mobile preview
-	EShaderPlatform MobilePreviewShaderPlatform = GShaderPlatformForFeatureLevel[ERHIFeatureLevel::ES3_1];
-	if (MobilePreviewShaderPlatform != SP_NumPlatforms)
+	if (!FPlatformProcess::SupportsMultithreading())
 	{
-		FShaderCompileUtilities::GenerateBrdfHeaders(MobilePreviewShaderPlatform);
+		bAllowAsynchronousShaderCompiling = false;
 	}
-});
-#endif
 
-/** Helper functions for logging more debug info */
-namespace ShaderCompiler
-{
-	FString GetTargetPlatformName(const ITargetPlatform* TargetPlatform)
-	{
-		if (TargetPlatform)
-		{
-			return TargetPlatform->PlatformName();
-		}
-
-		return TEXT("(current)");
-	}
-}
-
-/** Storage for the global shadar map(s) that have been replaced by new one(s), which aren't yet compiled.
- * 
- *	Sometimes a mesh drawing command references a pointer to global SM's memory. To nix these MDCs when we're replacing a global SM, we would just recreate the render state for all the components, but
- *	we may need to access a global shader during such an update, creating a catch 22. So deleting the global SM and updating components is deferred until the new one is compiled. 
- */
-FGlobalShaderMap* GGlobalShaderMap_DeferredDeleteCopy[SP_NumPlatforms] = {nullptr};
+	verify(GConfig->GetInt( TEXT("DevOptions.Shaders"), TEXT("MaxShaderJobBatchSize"), MaxShaderJobBatchSize, GEngineIni ));
+	verify(GConfig->GetBool( TEXT("DevOptions.Shaders"), TEXT("bPromptToRetryFailedShaderCompiles"), bPromptToRetryFailedShaderCompiles, GEngineIni ));
+	verify(GConfig->GetBool(TEXT("DevOptions.Shaders"), TEXT("bDebugBreakOnPromptToRetryShaderCompile"), bDebugBreakOnPromptToRetryShaderCompile, GEngineIni));
+	verify(GConfig->GetBool( TEXT("DevOptions.Shaders"), TEXT("bLogJobCompletionTimes"), bLogJobCompletionTimes, GEngineIni ));
+	GConfig->GetFloat(TEXT("DevOptions.Shaders"), TEXT("WorkerTimeToLive"), GRegularWorkerTimeToLive, GEngineIni);
+	GConfig->GetFloat(TEXT("DevOptions.Shaders"), TEXT("BuildWorkerTimeToLive"), GBuildWorkerTimeToLive, GEngineIni);
 
-#if ENABLE_COOK_STATS
-namespace GlobalShaderCookStats
-{
-	FCookStats::FDDCResourceUsageStats UsageStats;
-	static int32 ShadersCompiled = 0;
+	verify(GConfig->GetFloat( TEXT("DevOptions.Shaders"), TEXT("ProcessGameThreadTargetTime"), ProcessGameThreadTargetTime, GEngineIni ));
 
-	static FCookStatsManager::FAutoRegisterCallback RegisterCookStats([](FCookStatsManager::AddStatFuncRef AddStat)
-		{
-			UsageStats.LogStats(AddStat, TEXT("GlobalShader.Usage"), TEXT(""));
-			AddStat(TEXT("GlobalShader.Misc"), FCookStatsManager::CreateKeyValueArray(
-				TEXT("ShadersCompiled"), ShadersCompiled
-			));
-		});
-}
+#if UE_BUILD_DEBUG
+	// Increase budget for processing results in debug or else it takes forever to finish due to poor framerate
+	ProcessGameThreadTargetTime *= 3;
 #endif
 
-const FString& GetGlobalShaderMapDDCKey()
-{
-	static FString GlobalShaderMapDDCKey = GetGlobalShaderMapDDCGuid().ToString();
-	return GlobalShaderMapDDCKey;
-}
-
-const FGuid& GetGlobalShaderMapDDCGuid()
-{
-	static FGuid GlobalShaderMapDDCGuid = FDevSystemGuids::GetSystemGuid(FDevSystemGuids::Get().GLOBALSHADERMAP_DERIVEDDATA_VER);
-	return GlobalShaderMapDDCGuid;
-}
-
-const FString& GetMaterialShaderMapDDCKey()
-{
-	static FString MaterialShaderMapDDCKey = GetMaterialShaderMapDDCGuid().ToString();
-	return MaterialShaderMapDDCKey;
-}
-
-const FGuid& GetMaterialShaderMapDDCGuid()
-{
-	static FGuid MaterialShaderMapDDCGuid = FDevSystemGuids::GetSystemGuid(FDevSystemGuids::Get().MATERIALSHADERMAP_DERIVEDDATA_VER);
-	return MaterialShaderMapDDCGuid;
-}
-
-bool ShouldDumpShaderDDCKeys()
-{
-	return CVarShaderCompilerDumpDDCKeys.GetValueOnAnyThread();
-}
-
-void DumpShaderDDCKeyToFile(const EShaderPlatform InPlatform, bool bWithEditor, const FString& FileName, const FString& DDCKey)
-{
-	// deprecated version
-	const FString SubDirectory = bWithEditor ? TEXT("Editor") : TEXT("Game");
-	const FString TempPath = FPaths::ProjectSavedDir() / TEXT("ShaderDDCKeys") / SubDirectory / LexToString(InPlatform);
-	IFileManager::Get().MakeDirectory(*TempPath, true);
-
-	const FString TempFile = TempPath / FileName;
-
-	TUniquePtr<FArchive> DumpAr(IFileManager::Get().CreateFileWriter(*TempFile));
-	// serializing the string via << produces a non-textual file because it saves string's length, too
-	DumpAr->Serialize(const_cast<TCHAR*>(*DDCKey), DDCKey.Len() * sizeof(TCHAR));
-}
-
-void DumpShaderDDCKeyToFile(const EShaderPlatform InPlatform, bool bEditorOnly, const TCHAR* DebugGroupName, const FString& DDCKey)
-{
-	const FString FileName = FString::Printf(TEXT("DDCKey-%s.txt"), bEditorOnly ? TEXT("Editor") : TEXT("Game"));
-
-	const FString TempPath = GShaderCompilingManager->GetAbsoluteShaderDebugInfoDirectory() / FGenericDataDrivenShaderPlatformInfo::GetName(InPlatform).ToString() / DebugGroupName;
-	IFileManager::Get().MakeDirectory(*TempPath, true);
-
-	const FString TempFile = TempPath / FileName;
+	// Get the current process Id, this will be used by the worker app to shut down when it's parent is no longer running.
+	ProcessId = FPlatformProcess::GetCurrentProcessId();
 
-	FFileHelper::SaveStringToFile(DDCKey, *TempFile);
-}
+	// Use a working directory unique to this game, process and thread so that it will not conflict 
+	// With processes from other games, processes from the same game or threads in this same process.
+	// Use IFileManager to do path conversion to properly handle sandbox paths (outside of standard paths in particular).
+	//ShaderBaseWorkingDirectory = FPlatformProcess::ShaderWorkingDir() / FString::FromInt(ProcessId) + TEXT("/");
 
-namespace ShaderCompiler
-{
-	bool IsJobCacheEnabled()
 	{
-		return GShaderCompilerJobCache != 0;
+		FGuid Guid;
+		Guid = FGuid::NewGuid();
+		FString LegacyShaderWorkingDirectory = FPaths::ProjectIntermediateDir() / TEXT("Shaders/WorkingDirectory/")  / FString::FromInt(ProcessId) + TEXT("/");
+		ShaderBaseWorkingDirectory = FPaths::ShaderWorkingDir() / *Guid.ToString(EGuidFormats::Digits) + TEXT("/");
+		UE_LOG(LogShaderCompilers, Log, TEXT("Guid format shader working directory is %d characters bigger than the processId version (%s)."), ShaderBaseWorkingDirectory.Len() - LegacyShaderWorkingDirectory.Len(), *LegacyShaderWorkingDirectory );
 	}
 
-	bool IsJobCacheDebugValidateEnabled()
+	if (!IFileManager::Get().DeleteDirectory(*ShaderBaseWorkingDirectory, false, true))
 	{
-		return IsJobCacheEnabled() && CVarShaderCompilerDebugValidateJobCache.GetValueOnAnyThread();
+		UE_LOG(LogShaderCompilers, Fatal, TEXT("Could not delete the shader compiler working directory '%s'."), *ShaderBaseWorkingDirectory);
 	}
-
-	bool IsRemoteCompilingAllowed()
+	else
 	{
-		// commandline switches override the CVars
-		static bool bDisabledFromCommandline = FParse::Param(FCommandLine::Get(), TEXT("NoRemoteShaderCompile"));
-		if (bDisabledFromCommandline)
-		{
-			return false;
-		}
-
-		return GShaderCompilerAllowDistributedCompilation != 0;
+		UE_LOG(LogShaderCompilers, Log, TEXT("Cleaned the shader compiler working directory '%s'."), *ShaderBaseWorkingDirectory);
 	}
-}
-
-
-FShaderCompileJobCollection::FShaderCompileJobCollection(FCriticalSection& InCompileQueueSection)
-{
-	PrintStatsCmd = IConsoleManager::Get().RegisterConsoleCommand(
-		TEXT("r.ShaderCompiler.PrintStats"),
-		TEXT("Prints out to the log the stats for the shader compiler."),
-		FConsoleCommandDelegate::CreateRaw(this, &FShaderCompileJobCollection::HandlePrintStats),
-		ECVF_Default
-	);
-
-	JobsCache = MakePimpl<FShaderJobCache>(InCompileQueueSection);
-}
-
-
-// Pass through functions to inner FShaderJobCache implementation class
-FShaderCompileJob* FShaderCompileJobCollection::PrepareJob(uint32 InId, const FShaderCompileJobKey& InKey, EShaderCompileJobPriority InPriority)
-{
-	return JobsCache->PrepareJob<FShaderCompileJob>(InId, InKey, InPriority);
-}
-FShaderPipelineCompileJob* FShaderCompileJobCollection::PrepareJob(uint32 InId, const FShaderPipelineCompileJobKey& InKey, EShaderCompileJobPriority InPriority)
-{
-	return JobsCache->PrepareJob<FShaderPipelineCompileJob>(InId, InKey, InPriority);
-}
-void FShaderCompileJobCollection::RemoveJob(FShaderCommonCompileJob* InJob)
-{
-	JobsCache->RemoveJob(InJob);
-}
-int32 FShaderCompileJobCollection::RemoveAllPendingJobsWithId(uint32 InId)
-{
-	return JobsCache->RemoveAllPendingJobsWithId(InId);
-}
-void FShaderCompileJobCollection::SubmitJobs(const TArray<FShaderCommonCompileJobPtr>& InJobs)
-{
-	JobsCache->SubmitJobs(InJobs);
-}
-void FShaderCompileJobCollection::ProcessFinishedJob(FShaderCommonCompileJob* FinishedJob, bool bCompilationSkipped)
-{
-	JobsCache->ProcessFinishedJob(FinishedJob, bCompilationSkipped);
-}
-void FShaderCompileJobCollection::AddToCacheAndProcessPending(FShaderCommonCompileJob* FinishedJob)
-{
-	JobsCache->AddToCacheAndProcessPending(FinishedJob);
-}
-void FShaderCompileJobCollection::GetCachingStats(FShaderCompilerStats& OutStats) const
-{
-	JobsCache->GetStats(OutStats);
-}
-int32 FShaderCompileJobCollection::GetNumPendingJobs(EShaderCompileJobPriority InPriority) const
-{
-	return JobsCache->GetNumPendingJobs(InPriority);
-}
-int32 FShaderCompileJobCollection::GetNumOutstandingJobs() const
-{
-	return JobsCache->GetNumOutstandingJobs();
-}
-int32 FShaderCompileJobCollection::GetNumPendingJobs() const
-{
-	return JobsCache->GetNumPendingJobs();
-}
-int32 FShaderCompileJobCollection::GetPendingJobs(EShaderCompilerWorkerType InWorkerType, EShaderCompileJobPriority InPriority, int32 MinNumJobs, int32 MaxNumJobs, TArray<FShaderCommonCompileJobPtr>& OutJobs)
-{
-	return JobsCache->GetPendingJobs(InWorkerType, InPriority, MinNumJobs, MaxNumJobs, OutJobs);
-}
-
-
-static FShaderCommonCompileJob* CloneJob_Single(const FShaderCompileJob* SrcJob)
-{
-	FShaderCompileJob* Job = new FShaderCompileJob(SrcJob->Hash, SrcJob->Id, SrcJob->Priority, SrcJob->Key);
-	Job->ShaderParameters = SrcJob->ShaderParameters;
-	Job->PendingShaderMap = SrcJob->PendingShaderMap;
-	Job->Input = SrcJob->Input;
-	Job->PreprocessOutput = SrcJob->PreprocessOutput;
-	if (SrcJob->bInputHashSet)
-	{
-		Job->InputHash = SrcJob->InputHash;
-		Job->bInputHashSet = true;
-	}
-	ensure(Job->bInputHashSet == SrcJob->bInputHashSet);
-	return Job;
-}
-
-static FShaderCommonCompileJob* CloneJob_Pipeline(const FShaderPipelineCompileJob* SrcJob)
-{
-	FShaderPipelineCompileJob* Job = new FShaderPipelineCompileJob(SrcJob->Hash, SrcJob->Id, SrcJob->Priority, SrcJob->Key);
-	check(Job->StageJobs.Num() == SrcJob->StageJobs.Num());
-	Job->PendingShaderMap = SrcJob->PendingShaderMap;
+	FString AbsoluteBaseDirectory = IFileManager::Get().ConvertToAbsolutePathForExternalAppForWrite(*ShaderBaseWorkingDirectory);
+	FPaths::NormalizeDirectoryName(AbsoluteBaseDirectory);
+	AbsoluteShaderBaseWorkingDirectory = AbsoluteBaseDirectory + TEXT("/");
 
-	for(int32 i = 0; i < SrcJob->StageJobs.Num(); ++i)
+	// Build machines should dump to the AutomationTool/Saved/Logs directory and they will upload as build artifacts via the AutomationTool.
+	FString BaseDebugInfoPath = FPaths::ProjectSavedDir();
+	if (GIsBuildMachine)
 	{
-		Job->StageJobs[i]->Input = SrcJob->StageJobs[i]->Input;
-		Job->StageJobs[i]->PreprocessOutput = SrcJob->StageJobs[i]->PreprocessOutput;
+		BaseDebugInfoPath = FPaths::Combine(*FPaths::EngineDir(), TEXT("Programs"), TEXT("AutomationTool"), TEXT("Saved"), TEXT("Logs"));
 	}
 
-	if (SrcJob->bInputHashSet)
+	FString AbsoluteDebugInfoDirectory = IFileManager::Get().ConvertToAbsolutePathForExternalAppForWrite(*(BaseDebugInfoPath / TEXT("ShaderDebugInfo")));
+	const FString OverrideShaderDebugDir = CVarShaderOverrideDebugDir.GetValueOnAnyThread();
+	if (!OverrideShaderDebugDir.IsEmpty())
 	{
-		Job->InputHash = SrcJob->InputHash;
-		Job->bInputHashSet = true;
+		AbsoluteDebugInfoDirectory = OverrideShaderDebugDir;
 	}
-	ensure(Job->bInputHashSet == SrcJob->bInputHashSet);
-	return Job;
-}
+	FPaths::NormalizeDirectoryName(AbsoluteDebugInfoDirectory);
+	AbsoluteShaderDebugInfoDirectory = AbsoluteDebugInfoDirectory;
 
-static FShaderCommonCompileJob* CloneJob(const FShaderCommonCompileJob* SrcJob)
-{
-	switch (SrcJob->Type)
+	CalculateNumberOfCompilingThreads(FPlatformMisc::NumberOfCores(), FPlatformMisc::NumberOfCoresIncludingHyperthreads());
+
+	TUniquePtr<FShaderCompileThreadRunnableBase> RemoteCompileThread;
+	const bool bCanUseRemoteCompiling = bAllowCompilingThroughWorkers && ShaderCompiler::IsRemoteCompilingAllowed() && AllTargetPlatformSupportsRemoteShaderCompiling();
+	BuildDistributionController = bCanUseRemoteCompiling ? FindRemoteCompilerController() : nullptr;
+	
+	if (BuildDistributionController)
 	{
-	case EShaderCompileJobType::Single: return CloneJob_Single(static_cast<const FShaderCompileJob*>(SrcJob));
-	case EShaderCompileJobType::Pipeline:  return CloneJob_Pipeline(static_cast<const FShaderPipelineCompileJob*>(SrcJob));
-	default: checkNoEntry(); return nullptr;
+		UE_LOG(LogShaderCompilers, Display, TEXT("Using %s for Shader Compilation."), *BuildDistributionController->GetName());
+		RemoteCompileThread = MakeUnique<FShaderCompileDistributedThreadRunnable_Interface>(this, *BuildDistributionController);
 	}
-}
 
-void FShaderJobCache::InternalSetPriority(FShaderCommonCompileJob* Job, EShaderCompileJobPriority InPriority)
-{
-	const int32 PriorityIndex = (int32)InPriority;
+	GConfig->SetBool(TEXT("/Script/UnrealEd.UnrealEdOptions"), TEXT("UsingXGE"), RemoteCompileThread.IsValid(), GEditorIni);
 
-	if (Job->PendingPriority != EShaderCompileJobPriority::None)
+	TUniquePtr<FShaderCompileThreadRunnableBase> LocalThread = MakeUnique<FShaderCompileThreadRunnable>(this);
+	if (RemoteCompileThread)
 	{
-		// Need write lock to call UnlinkJobWithPriority
-		FWriteScopeLock Locker(JobLock);
-
-		// Check priority again, as the job may have been kicked off by another thread while waiting on the lock
-		if (Job->PendingPriority != EShaderCompileJobPriority::None)
-		{
-			// Job hasn't started yet, move it to the pending list for the new priority
-			check(Job->PendingPriority == Job->Priority);
-			UnlinkJobWithPriority(*Job);
+		checkf(ShaderCompiler::IsRemoteCompilingAllowed(), TEXT("We have a remote compiling thread without the remote compilation being allowed"));
 
-			ensure(!ShaderCompiler::IsJobCacheEnabled() || Job->bInputHashSet);
-			Job->Priority = InPriority;
-			LinkJobWithPriority(*Job);
-			
-			return;
-		}
+		// Only force-local jobs are guaranteed to stay on the local machine. Going wide with High priority jobs is important for the startup times,
+		// since special materials use High priority. Possibly the partition by priority is too rigid in general.
+		RemoteCompileThread->SetPriorityRange(EShaderCompileJobPriority::Low, EShaderCompileJobPriority::High);
+		LocalThread->SetPriorityRange(EShaderCompileJobPriority::Normal, EShaderCompileJobPriority::ForceLocal);
+		Threads.Add(MoveTemp(RemoteCompileThread));
 	}
-
-	if (!Job->bFinalized &&
-		Job->CurrentWorker == EShaderCompilerWorkerType::Distributed &&
-		InPriority == EShaderCompileJobPriority::ForceLocal)
+	else
 	{
-		FShaderCommonCompileJob* NewJob = CloneJob(Job);
-		NewJob->Priority = InPriority;
-		const int32 NewNumPendingJobs = NewJob->PendingShaderMap->NumPendingJobs.Increment();
-		checkf(NewNumPendingJobs > 1, TEXT("Invalid number of pending jobs %d, should have had at least 1 job previously"), NewNumPendingJobs);
-		JobTable.AddExistingJob(NewJob);
-
-		GShaderCompilerStats->RegisterNewPendingJob(*NewJob);
-		ensureMsgf(NewJob->bInputHashSet == Job->bInputHashSet, TEXT("Cloned and original jobs should either both have input hash, or both not have it. Job->bInputHashSet=%d, NewJob->bInputHashSet=%d"),
-			Job->bInputHashSet,
-			NewJob->bInputHashSet
-			);
-		ensureMsgf(!ShaderCompiler::IsJobCacheEnabled() || NewJob->GetInputHash() == Job->GetInputHash(),
-			TEXT("If shader jobs cache is enabled, cloned job should have the same input hash as the original, and it doesn't.")
-			);
-
-		FWriteScopeLock Locker(JobLock);
-		NumOutstandingJobs.Increment();
-		LinkJobWithPriority(*NewJob);
-
-		//UE_LOG(LogShaderCompilers, Display, TEXT("Submitted duplicate 'ForceLocal' shader compile job to replace existing XGE job"));
-	}
-}
-
-int32 FShaderJobCache::RemoveAllPendingJobsWithId(uint32 InId)
-{
-	int32 NumRemoved = 0;
+		UE_LOG(LogShaderCompilers, Display, TEXT("Using Local Shader Compiler with %d workers."), NumShaderCompilingThreads);
 
-#if WITH_EDITOR
-	TArray<FShaderCommonCompileJobPtr> JobsWithRequestsToCancel;
-#endif
-	{
-		// Look for jobs that are waiting on a SubmitJob task or async DDC query.  These can just be unlinked which will cause them to be
-		// discarded in SubmitJob or the DDC completion callback.  We also need to get a list of jobs with DDC requests to cancel.  We
-		// can't cancel the requests inside the loop, as the response callback uses JobLock, and it will deadlock.  We also need a
-		// reference pointer to the jobs, so the jobs (and the TPimplPtr<UE::DerivedData::FRequestOwner> contained therein) can't be
-		// deleted while a DDC completion callback is in flight, which also leads to a deadlock.
-		FWriteScopeLock Locker(JobLock);
-		for (FShaderCommonCompileJobIterator It(PendingSubmitJobTaskJobs); It;)
+		if (GIsBuildMachine)
 		{
-			FShaderCommonCompileJob& Job = *It;
-			It.Next();
-
-			if (Job.Id == InId)
+			int32 MinSCWsToSpawnBeforeWarning = 8; // optional, default to 8
+			GConfig->GetInt(TEXT("DevOptions.Shaders"), TEXT("MinSCWsToSpawnBeforeWarning"), MinSCWsToSpawnBeforeWarning, GEngineIni);
+			if (NumShaderCompilingThreads < static_cast<uint32>(MinSCWsToSpawnBeforeWarning))
 			{
-				Unlink(Job);		// from PendingSubmitJobTaskJobs
-				RemoveJob(&Job);
-				++NumRemoved;
-
-#if WITH_EDITOR
-				if (Job.RequestOwner)
-				{
-					JobsWithRequestsToCancel.Add(&Job);
-				}
-#endif
+				UE_LOG(LogShaderCompilers, Warning, TEXT("Only %d SCWs will be spawned, which will result in longer shader compile times."), NumShaderCompilingThreads);
 			}
 		}
 	}
+	Threads.Add(MoveTemp(LocalThread));
 
-#if WITH_EDITOR
-	for (FShaderCommonCompileJobPtr JobWithRequestToCancel : JobsWithRequestsToCancel)
+	for (const auto& Thread : Threads)
 	{
-		// Cancelling should short circuit the request, and make "Wait" finish immediately
-		JobWithRequestToCancel->RequestOwner->Cancel();
-		JobWithRequestToCancel->RequestOwner->Wait();
+		Thread->StartThread();
 	}
-#endif
 
+	OutOfMemoryDelegateHandle = FCoreDelegates::GetOutOfMemoryDelegate().AddRaw(this, &FShaderCompilingManager::ReportMemoryUsage);
+
+	FAssetCompilingManager::Get().RegisterManager(this);
+
+#if WITH_EDITOR
+	static const bool bAllowShaderRecompileOnSave = CVarRecompileShadersOnSave.GetValueOnAnyThread();
+	if (bAllowShaderRecompileOnSave)
 	{
-		FWriteScopeLock Locker(JobLock);
-		for (int32 PriorityIndex = 0; PriorityIndex < NumShaderCompileJobPriorities; ++PriorityIndex)
+		if (IDirectoryWatcher* DirectoryWatcher = FModuleManager::LoadModuleChecked<FDirectoryWatcherModule>(TEXT("DirectoryWatcher")).Get())
 		{
-			for (FShaderCommonCompileJobIterator It(PendingJobsHead[PriorityIndex]); It;)
+			// Handle if we are watching a directory for changes.
 			{
-				FShaderCommonCompileJob& Job = *It;
-				It.Next();
+				UE_LOG(LogShaderCompilers, Display, TEXT("Register directory watchers for shader files."));
+
+				const TMap<FString, FString>& ShaderSourceDirectoryMappings = AllShaderSourceDirectoryMappings();
 
-				if (Job.Id == InId)
+				DirectoryWatcherHandles.Reserve(ShaderSourceDirectoryMappings.Num());
+
+				for (const auto& It : ShaderSourceDirectoryMappings)
 				{
-					if (ShaderCompiler::IsJobCacheEnabled())
+					FString DirectoryToWatch = It.Value;
+					if (FPaths::IsRelative(DirectoryToWatch))
 					{
-						if (Job.JobCacheRef.Block)
-						{
-							FShaderJobData& JobData = GetShaderJobData(Job.JobCacheRef);
+						DirectoryToWatch = FPaths::ConvertRelativePathToFull(DirectoryToWatch);
+					}
 
-							check(JobData.JobInFlight == &Job);
+					FDelegateHandle& DirectoryWatcherHandle = DirectoryWatcherHandles.Add(DirectoryToWatch, FDelegateHandle());
 
-							// If we are removing an in-flight job, we need to promote a duplicate to be the new in-flight job, if present.
-							// Make sure the duplicate we choose doesn't have the same ID as what we're removing.
-							FShaderCommonCompileJob* DuplicateJob;
-							for (DuplicateJob = JobData.DuplicateJobsWaitList; DuplicateJob; DuplicateJob = DuplicateJob->NextLink)
-							{
-								if (DuplicateJob->Id != InId)
-								{
-									break;
-								}
-							}
+					DirectoryWatcher->RegisterDirectoryChangedCallback_Handle(
+						DirectoryToWatch,
+						IDirectoryWatcher::FDirectoryChanged::CreateLambda([](const TArray<FFileChangeData>& InFileChangeDatas) {
 
-							if (DuplicateJob)
+							TRACE_CPUPROFILER_EVENT_SCOPE(HandleDirectoryChanged);
+
+							if (!bAllowShaderRecompileOnSave)
 							{
-								// Advance head if we are unlinking the head, then remove
-								if (JobData.DuplicateJobsWaitList == DuplicateJob)
-								{
-									JobData.DuplicateJobsWaitList = DuplicateJob->NextLink;
-								}
-								Unlink(*DuplicateJob);
-								RemoveDuplicateJob(DuplicateJob);
+								return;
+							}
 
-								// Add it as pending at the appropriate priority
-								GShaderCompilerStats->RegisterNewPendingJob(*DuplicateJob);
+							TArray<FString> ChangedShaderFiles;
+							for (const FFileChangeData& It : InFileChangeDatas)
+							{
+								if (It.Filename.EndsWith(TEXT(".usf")) || It.Filename.EndsWith(TEXT(".ush")) || It.Filename.EndsWith(TEXT(".h")))
+								{
+									UE_LOG(LogShaderCompilers, Display, TEXT("Detected change on %s"), *It.Filename);
 
-								LinkJobWithPriority(*DuplicateJob);
+									ChangedShaderFiles.AddUnique(It.Filename);
+								}
 							}
 
-							// DuplicateJob will be nullptr if there was no duplicate to promote
-							JobData.JobInFlight = DuplicateJob;
-						}
-					}
+							if (ChangedShaderFiles.Num())
+							{
+								// Mappings from:
+								// Key:   /Engine to
+								// Value: ../../../Engine/Shaders
+								const TMap<FString, FString>& ShaderSourceDirectoryMappings = AllShaderSourceDirectoryMappings();
 
-					check((int32)Job.PendingPriority == PriorityIndex);
-					UnlinkJobWithPriority(Job);
-					RemoveJob(&Job);
-					++NumRemoved;
-				}
-			}
-		}
+								FString RemappedShaderFileName;
+								for (const auto& It : ShaderSourceDirectoryMappings)
+								{
+									// ChangedShaderFiles will be of format: ../../../Engine/Shaders/Private/PostProcessGBufferHints.usf
+									if (ChangedShaderFiles[0].StartsWith(It.Value))
+									{
+										// Change from relative path to Engine absolute path.
+										// i.e. change `../../../Engine/Shaders/Private/PostProcessGBufferHints.usf` to `/Engine/Shaders/Private/PostProcessGBufferHints.usf`
+										RemappedShaderFileName = ChangedShaderFiles[0].Replace(*It.Value, *It.Key);
+									}
+								}
 
-		if (ShaderCompiler::IsJobCacheEnabled())
-		{
-			// Also look into duplicate jobs that are cached -- we don't increment in the "for" loop because the current item may be deleted
-			for (int32 DuplicateIndex = 0; DuplicateIndex < DuplicateJobs.Num();)
-			{
-				FShaderCommonCompileJob* DuplicateJob = DuplicateJobs[DuplicateIndex];
-				check(DuplicateJob->JobCacheRef.DuplicateIndex == DuplicateIndex);
+								// Issue a `recompileshaders /Engine/Shaders/Private/PostProcessGBufferHints.usf` command, which will just compile that shader source file.
+								RecompileShaders(*RemappedShaderFileName, *GLog);
 
-				if (DuplicateJob->Id == InId)
-				{
-					FShaderJobData& JobData = GetShaderJobData(DuplicateJob->JobCacheRef);
+								UE_LOG(LogShaderCompilers, Display, TEXT("Ready for new shader file changes"));
+							}
+						}),
+						DirectoryWatcherHandle);
 
-					// if we're removing the list head, we need to update it to the next
-					if (JobData.DuplicateJobsWaitList == DuplicateJob)
+					if (DirectoryWatcherHandle.IsValid())
+					{
+						UE_LOG(LogShaderCompilers, Display, TEXT("Watching %s -> %s"), *It.Key, *DirectoryToWatch);
+					}
+					else
 					{
-						JobData.DuplicateJobsWaitList = JobData.DuplicateJobsWaitList->NextLink;
+						UE_LOG(LogShaderCompilers, Error, TEXT("Failed to set up directory watcher %s -> %s"), *It.Key, *DirectoryToWatch);
 					}
-
-					// This removes the current job (at DuplicateIndex), so we don't increment in this case
-					RemoveDuplicateJob(DuplicateJob);
-
-					// Duplicate jobs are in their own list, not one of the priority lists, so don't use UnlinkJobWithPriority
-					check(DuplicateJob->PendingPriority == EShaderCompileJobPriority::None);
-					Unlink(*DuplicateJob);
-					RemoveJob(DuplicateJob);
-					++NumRemoved;
-				}
-				else
-				{
-					// Didn't remove a job, increment!
-					DuplicateIndex++;
 				}
 			}
 		}
 	}
-
-	InternalSubtractNumOutstandingJobs(NumRemoved);
-
-	return NumRemoved;
+#endif // WITH_EDITOR
 }
 
-void FShaderJobCache::SubmitJob(FShaderCommonCompileJob* Job)
+FShaderCompilingManager::~FShaderCompilingManager()
 {
-	// Set thread local so DDC query callback can detect if it's in the same thread, and we need to run through the non-async code path.
-	struct InSubmitJobScope
+	// we never initialized, so nothing to do
+	if (!AllowShaderCompiling())
 	{
-		InSubmitJobScope() { bInSubmitJobThread = true; }
-		~InSubmitJobScope() { bInSubmitJobThread = false; }
-	};
-	InSubmitJobScope InSubmitJob;
+		return;
+	}
 
-	check(Job->Priority != EShaderCompileJobPriority::None);
-	check(Job->PendingPriority == EShaderCompileJobPriority::None);
+	for (const auto& Thread : Threads)
+	{
+		Thread->Stop();
+		Thread->WaitForCompletion();
+	}
 
-	const int32 PriorityIndex = (int32)Job->Priority;
-	bool bNewJob = true;
-	bool bJobCacheLocked = false;
+	FCoreDelegates::GetOutOfMemoryDelegate().Remove(OutOfMemoryDelegateHandle);
 
-	// check caches unless we're running in validation mode (which runs _all_ jobs and compares hashes of outputs)
-	if (ShaderCompiler::IsJobCacheEnabled() && !ShaderCompiler::IsJobCacheDebugValidateEnabled())
+#if WITH_EDITOR
+	const bool bAllowShaderRecompileOnSave = CVarRecompileShadersOnSave.GetValueOnAnyThread();
+	if (bAllowShaderRecompileOnSave)
 	{
-		bJobCacheLocked = true;
+		if (IDirectoryWatcher* DirectoryWatcher = FModuleManager::LoadModuleChecked<FDirectoryWatcherModule>(TEXT("DirectoryWatcher")).Get())
+		{
+			for (const auto& It : DirectoryWatcherHandles)
+			{
+				DirectoryWatcher->UnregisterDirectoryChangedCallback_Handle(It.Key, It.Value);
+			}
+		}
+	}
+#endif // WITH_EDITOR
 
-		const FShaderCompilerInputHash& InputHash = Job->GetInputHash();
+	FAssetCompilingManager::Get().UnregisterManager(this);
+}
 
-		const bool bCheckDDC = GShaderCompilerPerShaderDDCGlobal || !(Job->bIsDefaultMaterial || Job->bIsGlobalShader);
+void FShaderCompilingManager::CalculateNumberOfCompilingThreads(int32 NumberOfCores, int32 NumberOfCoresIncludingHyperthreads)
+{
+	const int32 NumVirtualCores = NumberOfCoresIncludingHyperthreads;
 
-		// We don't use a scope here, because we need to release this lock before calling ProcessFinishedJob, which needs to acquire
-		// CompileQueueSection.  It's not safe to acquire CompileQueueSection where JobLock is locked first, as it will cause
-		// deadlocks due to FShaderCompileThreadRunnable::CompilingLoop calling GetPendingJobs, which acquires those two locks in
-		// the opposite order.
-		double StallStart = FPlatformTime::Seconds();
-		JobLock.WriteLock();
-		Job->TimeTaskSubmitJobsStall += FPlatformTime::Seconds() - StallStart;
+	int32 NumUnusedShaderCompilingThreads;
+	verify(GConfig->GetInt(TEXT("DevOptions.Shaders"), TEXT("NumUnusedShaderCompilingThreads"), NumUnusedShaderCompilingThreads, GEngineIni));
 
-		// Job was linked in PendingSubmitJobTaskJobs before calling SubmitJob -- if it's not linked now, it means it was cancelled via
-		// call to RemoveAllPendingJobsWithId, so we can ignore it and just return.
-		if (!Job->PrevLink)
-		{
-			UE_LOG(LogShaderCompilers, Display, TEXT("Cancelled job 0x%p with pending SubmitJob call."), Job);
+	int32 NumUnusedShaderCompilingThreadsDuringGame;
+	verify(GConfig->GetInt(TEXT("DevOptions.Shaders"), TEXT("NumUnusedShaderCompilingThreadsDuringGame"), NumUnusedShaderCompilingThreadsDuringGame, GEngineIni));
 
-			JobLock.WriteUnlock();
-			return;
-		}
-		check(Job->JobIndex != INDEX_NONE);
+	int32 ShaderCompilerCoreCountThreshold;
+	verify(GConfig->GetInt(TEXT("DevOptions.Shaders"), TEXT("ShaderCompilerCoreCountThreshold"), ShaderCompilerCoreCountThreshold, GEngineIni));
 
-		FSharedBuffer* ExistingOutput;
-		FShaderJobCacheRef JobCacheRef = FindOrAdd(InputHash, Job->Priority, bCheckDDC, Job->RequestOwner, ExistingOutput);
+	bool bForceUseSCWMemoryPressureLimits = false;
+	GConfig->GetBool(TEXT("DevOptions.Shaders"), TEXT("bForceUseSCWMemoryPressureLimits"), bForceUseSCWMemoryPressureLimits, GEngineIni);
 
-		// see if there are already cached results for this job
-		if (ExistingOutput)
-		{
-			Unlink(*Job);		// from PendingSubmitJobTaskJobs
+	// Don't reserve threads based on a percentage if we are in a commandlet or on a low core machine.
+	// In these scenarios we should try to use as many threads as possible.
+	if (!IsRunningCommandlet() && !GIsBuildMachine && NumVirtualCores > ShaderCompilerCoreCountThreshold)
+	{
+		// Reserve a percentage of the threads for general background work.
+		float PercentageUnusedShaderCompilingThreads;
+		verify(GConfig->GetFloat(TEXT("DevOptions.Shaders"), TEXT("PercentageUnusedShaderCompilingThreads"), PercentageUnusedShaderCompilingThreads, GEngineIni));
 
-			// Need to release the lock before calling ProcessFinishedJob, as mentioned above (and it's also good for performance to 
-			// release the lock before the relatively costly "SerializeOutput" call).
-			JobLock.WriteUnlock();
-			bNewJob = false;
-			bJobCacheLocked = false;
+		// ensure we get a valid multiplier.
+		PercentageUnusedShaderCompilingThreads = FMath::Clamp(PercentageUnusedShaderCompilingThreads, 0.0f, 100.0f) / 100.0f;
 
-			UE_LOG(LogShaderCompilers, UE_SHADERCACHE_LOG_LEVEL, TEXT("There is already a cached job with the ihash %s, processing the new one immediately."), *LexToString(InputHash));
-			FMemoryReaderView MemReader(*ExistingOutput);
-			Job->SerializeOutput(MemReader);
+		NumUnusedShaderCompilingThreads = FMath::CeilToInt(NumVirtualCores * PercentageUnusedShaderCompilingThreads);
+		NumUnusedShaderCompilingThreadsDuringGame = NumUnusedShaderCompilingThreads;
+	}
 
-			// finish the job instantly
-			Job->TimeTaskSubmitJobsStall += ProcessFinishedJob(Job, true);
-		}
-		else
-		{
-			FShaderJobData& JobData = GetShaderJobData(JobCacheRef);
-			Job->JobCacheRef = JobCacheRef;
+	// Use all the cores on the build machines.
+	if (GForceAllCoresForShaderCompiling != 0)
+	{
+		NumUnusedShaderCompilingThreads = 0;
+	}
 
-			// see if another job with the same input hash is being worked on
-			if (JobData.JobInFlight)
-			{
-				UE_LOG(LogShaderCompilers, UE_SHADERCACHE_LOG_LEVEL, TEXT("There is an outstanding job with the ihash %s, not submitting another one (adding to wait list)."), *LexToString(InputHash));
+	NumShaderCompilingThreads = (bAllowCompilingThroughWorkers && NumVirtualCores > NumUnusedShaderCompilingThreads) ? (NumVirtualCores - NumUnusedShaderCompilingThreads) : 1;
 
-				Unlink(*Job);		// from PendingSubmitJobTaskJobs
+	// Make sure there's at least one worker allowed to be active when compiling during the game
+	NumShaderCompilingThreadsDuringGame = (bAllowCompilingThroughWorkers && NumVirtualCores > NumUnusedShaderCompilingThreadsDuringGame) ? (NumVirtualCores - NumUnusedShaderCompilingThreadsDuringGame) : 1;
 
-				// because of the cloned jobs, we need to maintain a separate mapping
-				FShaderCommonCompileJob** WaitListHead = &JobData.DuplicateJobsWaitList;
-				if (*WaitListHead)
+	// On machines with few cores, each core will have a massive impact on compile time, so we prioritize compile latency over editor performance during the build
+	if (NumVirtualCores <= 4)
+	{
+		NumShaderCompilingThreads = NumVirtualCores - 1;
+		NumShaderCompilingThreadsDuringGame = NumVirtualCores - 1;
+	}
+#if PLATFORM_DESKTOP
+	else if (GIsBuildMachine || bForceUseSCWMemoryPressureLimits)
+	{
+		// Cooker ends up running OOM so use a simple heuristic based on some INI values
+		float CookerMemoryUsedInGB = 0.0f;
+		float MemoryToLeaveForTheOSInGB = 0.0f;
+		float MemoryUsedPerSCWProcessInGB = 0.0f;
+		bool bFoundEntries = true;
+		bFoundEntries = bFoundEntries && GConfig->GetFloat(TEXT("DevOptions.Shaders"), TEXT("CookerMemoryUsedInGB"), CookerMemoryUsedInGB, GEngineIni);
+		bFoundEntries = bFoundEntries && GConfig->GetFloat(TEXT("DevOptions.Shaders"), TEXT("MemoryToLeaveForTheOSInGB"), MemoryToLeaveForTheOSInGB, GEngineIni);
+		bFoundEntries = bFoundEntries && GConfig->GetFloat(TEXT("DevOptions.Shaders"), TEXT("MemoryUsedPerSCWProcessInGB"), MemoryUsedPerSCWProcessInGB, GEngineIni);
+		if (bFoundEntries)
+		{
+			uint32 PhysicalGBRam = FPlatformMemory::GetPhysicalGBRam();
+			float AvailableMemInGB = (float)PhysicalGBRam - CookerMemoryUsedInGB;
+			if (AvailableMemInGB > 0.0f)
+			{
+				if (AvailableMemInGB > MemoryToLeaveForTheOSInGB)
 				{
-					LinkAfter(*Job, *WaitListHead);
+					AvailableMemInGB -= MemoryToLeaveForTheOSInGB;
 				}
 				else
 				{
-					*WaitListHead = Job;
+					UE_LOG(LogShaderCompilers, Warning, TEXT("Machine has %d GBs of RAM, cooker might take %f GBs, but not enough memory left for the OS! (Requested %f GBs for the OS)"), PhysicalGBRam, CookerMemoryUsedInGB, MemoryToLeaveForTheOSInGB);
 				}
-				++TotalCacheDuplicates;
-
-				AddDuplicateJob(Job);
-				JobLock.WriteUnlock();
-				bNewJob = false;
-				bJobCacheLocked = false;
 			}
 			else
 			{
-				// track new jobs so we can dedupe them
-				JobData.JobInFlight = Job;
+				UE_LOG(LogShaderCompilers, Warning, TEXT("Machine has %d GBs of RAM, but cooker might take %f GBs!"), PhysicalGBRam, CookerMemoryUsedInGB);
 			}
-		}
-	}
-	else if (ShaderCompiler::IsJobCacheDebugValidateEnabled())
-	{
-		FSharedBuffer* ExistingOutput;
-		const FShaderCompilerInputHash& InputHash = Job->GetInputHash();
-		const bool bCheckDDC = !(Job->bIsDefaultMaterial || Job->bIsGlobalShader);
-		JobLock.WriteLock();
-		Job->JobCacheRef = FindOrAdd(InputHash, Job->Priority, bCheckDDC, Job->RequestOwner, ExistingOutput);
-		bJobCacheLocked = true;
-	}
-
-	// new job
-	if (bNewJob)
-	{
-		GShaderCompilerStats->RegisterNewPendingJob(*Job);
-		ensure(!ShaderCompiler::IsJobCacheEnabled() || Job->bInputHashSet);
-
-		// If cache is disabled, we skipped the code that grabs the write lock above, so we need to do it here, before modifying the pending queue
-		if (bJobCacheLocked == false)
-		{
-			bJobCacheLocked = true;
-			JobLock.WriteLock();
-
-			// Job was linked in PendingSubmitJobTaskJobs before calling SubmitJob -- if it's not linked now, it means it was cancelled via
-			// call to RemoveAllPendingJobsWithId, so we can ignore it and just return.
-			if (!Job->PrevLink)
+			if (MemoryUsedPerSCWProcessInGB > 0.0f)
 			{
-				UE_LOG(LogShaderCompilers, Log, TEXT("Cancelled job 0x%p with pending SubmitJob call."), Job);
+				float NumSCWs = AvailableMemInGB / MemoryUsedPerSCWProcessInGB;
+				NumShaderCompilingThreads = FMath::RoundToInt(NumSCWs);
 
-				JobLock.WriteUnlock();
-				return;
+				bool bUseVirtualCores = true;
+				GConfig->GetBool(TEXT("DevOptions.Shaders"), TEXT("bUseVirtualCores"), bUseVirtualCores, GEngineIni);
+				uint32 MaxNumCoresToUse = bUseVirtualCores ? NumVirtualCores : NumberOfCores;
+				NumShaderCompilingThreads = FMath::Clamp<uint32>(NumShaderCompilingThreads, 1, MaxNumCoresToUse - 1);
+				NumShaderCompilingThreadsDuringGame = FMath::Min<int32>(NumShaderCompilingThreads, NumShaderCompilingThreadsDuringGame);
 			}
-			check(Job->JobIndex != INDEX_NONE);
 		}
-
-		// If an async DDC request is in flight, that will add the job to the pending queue for processing when the request completes,
-		// if the request didn't find a result.  Otherwise we add it to the pending queue immediately.
-		if (Job->RequestOwner.IsValid() == false)
+		else if (bForceUseSCWMemoryPressureLimits)
 		{
-			check(Job->PrevLink);
-			Unlink(*Job);		// from PendingSubmitJobTaskJobs
+			UE_LOG(LogShaderCompilers, Warning, TEXT("bForceUseSCWMemoryPressureLimits was set but missing one or more prerequisite setting(s): CookerMemoryUsedInGB, MemoryToLeaveForTheOSInGB, MemoryUsedPerSCWProcessInGB.  Ignoring bForceUseSCWMemoryPressureLimits"));
+		}
 
-			LinkJobWithPriority(*Job);
+		if (GIsBuildMachine)
+		{
+			// force crashes on hung shader maps on build machines, to prevent builds running for days
+			GCrashOnHungShaderMaps = 1;
 		}
 	}
+#endif
 
-	if (bJobCacheLocked)
-	{
-		JobLock.WriteUnlock();
-	}
+	NumShaderCompilingThreads = FMath::Max<int32>(1, NumShaderCompilingThreads);
+	NumShaderCompilingThreadsDuringGame = FMath::Max<int32>(1, NumShaderCompilingThreadsDuringGame);
+
+	NumShaderCompilingThreadsDuringGame = FMath::Min<int32>(NumShaderCompilingThreadsDuringGame, NumShaderCompilingThreads);
 }
 
-void FShaderJobCache::SubmitJobs(const TArray<FShaderCommonCompileJobPtr>& InJobs)
+void FShaderCompilingManager::OnMachineResourcesChanged(int32 NumberOfCores, int32 NumberOfCoresIncludingHyperthreads)
 {
-	if (InJobs.Num() > 0)
+	CalculateNumberOfCompilingThreads(NumberOfCores, NumberOfCoresIncludingHyperthreads);
+	for (const auto& Thread : Threads)
 	{
-		// all jobs (not just actually submitted ones) count as outstanding. This needs to be done early because
-		// we may fulfill some of the jobs from the cache (and we will be subtracting them)
-		NumOutstandingJobs.Add(InJobs.Num());
-
-		{
-			// Add pending jobs to a list to support cancelling while SubmitJob tasks or async DDC queries are in flight
-			FWriteScopeLock JobLocker(JobLock);
-			for (FShaderCommonCompileJob* Job : InJobs)
-			{
-				LinkHead(*Job, PendingSubmitJobTaskJobs);
-			}
-		}
-
-		for (const FShaderCommonCompileJobPtr& Job : InJobs)
-		{
-			UE::Tasks::ETaskPriority Prio = IsRunningCookCommandlet() ? UE::Tasks::ETaskPriority::Normal : UE::Tasks::ETaskPriority::BackgroundNormal;
-			UE::Tasks::Launch(UE_SOURCE_LOCATION, [Job, this]()
-				{
-					TRACE_CPUPROFILER_EVENT_SCOPE(ShaderJobTask);
-					double TimeStart = FPlatformTime::Seconds();
-
-					if (GShaderCompilerDebugStallSubmitJob > 0)
-					{
-						FPlatformProcess::Sleep(GShaderCompilerDebugStallSubmitJob * 0.001f);
-					}
-
-					bool bSubmitJob = true;
-					if (ShaderCompiler::IsJobCacheEnabled())
-					{
-						bSubmitJob = PreprocessShader(Job);
-						Job->GetInputHash();
-					}
-
-					if (bSubmitJob)
-					{
-						SubmitJob(Job);
-					}
-					else // if preprocessing ran and failed, finish the job immediately
-					{
-						ProcessFinishedJob(Job, /* bCompilationSkipped = */true);
-					}
-
-					Job->TimeTaskSubmitJobs = FPlatformTime::Seconds() - TimeStart;
-				}, Prio);
-		}
+		Thread->OnMachineResourcesChanged();
 	}
 }
 
-void FShaderCompileJobCollection::HandlePrintStats()
+FName FShaderCompilingManager::GetStaticAssetTypeName()
 {
-	GShaderCompilingManager->PrintStats();
+	return TEXT("UE-Shader");
 }
 
-double FShaderJobCache::ProcessFinishedJob(FShaderCommonCompileJob* FinishedJob, bool bCompilationSkipped)
+FName FShaderCompilingManager::GetAssetTypeName() const
 {
-	double StallTime;
-
-	FinishedJob->OnComplete();
+	return GetStaticAssetTypeName();
+}
 
-	GShaderCompilerStats->RegisterFinishedJob(*FinishedJob, bCompilationSkipped);
+FTextFormat FShaderCompilingManager::GetAssetNameFormat() const
+{
+	return LOCTEXT("ShaderNameFormat", "{0}|plural(one=Shader,other=Shaders)");
+}
 
+TArrayView<FName> FShaderCompilingManager::GetDependentTypeNames() const
+{
+#if WITH_EDITOR
+	static FName DependentTypeNames[] = 
 	{
-		// Need to protect writes to FShaderMapCompileResults
-		double StallStart = FPlatformTime::Seconds();
-		FScopeLock Lock(&CompileQueueSection);
-		StallTime = FPlatformTime::Seconds() - StallStart;
-
-		FShaderMapCompileResults& ShaderMapResults = *(FinishedJob->PendingShaderMap);
-		ShaderMapResults.FinishedJobs.Add(FinishedJob);
-		ShaderMapResults.bAllJobsSucceeded = ShaderMapResults.bAllJobsSucceeded && FinishedJob->bSucceeded;
-
-		const int32 NumPendingJobsForSM = ShaderMapResults.NumPendingJobs.Decrement();
-		checkf(NumPendingJobsForSM >= 0, TEXT("Problem tracking pending jobs for a SM (%d), number of pending jobs (%d) is negative!"), FinishedJob->Id, NumPendingJobsForSM);
-	}
+		// Texture can require materials to be updated,
+		// they should be processed first to avoid unecessary material updates.
+		FTextureCompilingManager::GetStaticAssetTypeName() 
+	};
+	return TArrayView<FName>(DependentTypeNames);
+#else
+	return TArrayView<FName>();
+#endif	
+}
 
-	InternalSubtractNumOutstandingJobs(1);
-	if (!bCompilationSkipped && ShaderCompiler::IsJobCacheEnabled())
-	{
-		AddToCacheAndProcessPending(FinishedJob);
-	}
+int32 FShaderCompilingManager::GetNumRemainingAssets() const
+{
+	// Currently, jobs are difficult to track but the purpose of the GetNumRemainingAssets function is to never return 0
+	// if there are still shaders that have not had their primitives updated on the render thread.
+	// So we track jobs first and when everything is finished compiling but are still lying around in other structures
+	// waiting to be further processed, we show those numbers and ultimately we always return 1 unless IsCompiling() is false.
+	return FMath::Max3(GetNumRemainingJobs(), ShaderMapJobs.Num() + PendingFinalizeShaderMaps.Num(), IsCompiling() ? 1 : 0);
+}
 
-	return StallTime;
+void FShaderCompilingManager::ProcessAsyncTasks(bool bLimitExecutionTime)
+{
+	ProcessAsyncResults(bLimitExecutionTime, false);
 }
 
-void FShaderJobCache::AddToCacheAndProcessPending(FShaderCommonCompileJob* FinishedJob)
+void FShaderCompilingManager::ProcessAsyncTasks(const AssetCompilation::FProcessAsyncTaskParams& Params)
 {
-	// Cloned jobs won't include an entry in the job cache, so skip the caching logic.  The non-cloned version of the same
-	// job will handle adding data to the cache when it completes.
-	if (!ShaderCompiler::IsJobCacheEnabled() || !FinishedJob->JobCacheRef.Block)
+	// Shader compilations are not required for PIE to begin.
+	if (Params.bPlayInEditorAssetsOnly)
 	{
 		return;
 	}
 
-	ensureMsgf(FinishedJob->bInputHashSet, TEXT("Finished job didn't have input hash set, was shader compiler jobs cache toggled runtime?"));
-
-	const FShaderCompilerInputHash& InputHash = FinishedJob->GetInputHash();
-	TArray<uint8> Output;
-	FMemoryWriter Writer(Output);
-	FinishedJob->SerializeOutput(Writer);
+	ProcessAsyncResults(Params.bLimitExecutionTime, false);
+}
 
-	FSharedBuffer Buffer = MakeSharedBufferFromArray(MoveTemp(Output));
+int32 FShaderCompilingManager::GetNumPendingJobs() const
+{
+	return AllJobs.GetNumPendingJobs();
+}
+
+int32 FShaderCompilingManager::GetNumOutstandingJobs() const
+{
+	return AllJobs.GetNumOutstandingJobs();
+}
 
-	FShaderJobData& JobData = GetShaderJobData(FinishedJob->JobCacheRef);
+FShaderCompilingManager::EDumpShaderDebugInfo FShaderCompilingManager::GetDumpShaderDebugInfo() const
+{
+	if (GDumpShaderDebugInfo < EDumpShaderDebugInfo::Never || GDumpShaderDebugInfo > EDumpShaderDebugInfo::OnErrorOrWarning)
+	{
+		return EDumpShaderDebugInfo::Never;
+	}
 
-	// see if there are outstanding jobs that also need to be resolved
-	TArray<FShaderCommonCompileJob*> FinishedDuplicateJobs;
+	return static_cast<FShaderCompilingManager::EDumpShaderDebugInfo>(GDumpShaderDebugInfo);
+}
 
+EShaderDebugInfoFlags FShaderCompilingManager::GetDumpShaderDebugInfoFlags() const
+{
+	EShaderDebugInfoFlags Flags = EShaderDebugInfoFlags::Default;
+	if (GDumpShaderDebugInfoSCWCommandLine)
+	{
+		Flags |= EShaderDebugInfoFlags::DirectCompileCommandLine;
+	}
+	
+	if (CVarDebugDumpJobInputHashes.GetValueOnAnyThread())
 	{
-		FWriteScopeLock JobLocker(JobLock);
+		Flags |= EShaderDebugInfoFlags::InputHash;
+	}
 
-		FShaderCommonCompileJob* CurHead = JobData.DuplicateJobsWaitList;
-		while (CurHead)
-		{
-			checkf(CurHead != FinishedJob, TEXT("Job that is being added to cache was also on a waiting list! Error in bookkeeping."));
+	if (CVarDebugDumpJobDiagnostics.GetValueOnAnyThread())
+	{
+		Flags |= EShaderDebugInfoFlags::Diagnostics;
+	}
 
-			// Need to add these to a list, and process them outside the JobLock scope.  ProcessFinishedJob locks CompileQueueSection,
-			// and we don't want to lock that inside a block that also locks JobLock, as it can cause a deadlock given that other
-			// code paths obtain the locks in the opposite order.  This is also good for perf, as it avoids holding the lock during
-			// the relatively costly SerializeOutput.
-			FinishedDuplicateJobs.Add(CurHead);
+	if (CVarDebugDumpShaderCode.GetValueOnAnyThread())
+	{
+		Flags |= EShaderDebugInfoFlags::ShaderCodeBinary;
+	}
 
-			// This needs to happen inside the JobLocker scope
-			RemoveDuplicateJob(CurHead);
+	if (CVarDebugDumpDetailedShaderSource.GetValueOnAnyThread())
+	{
+		Flags |= EShaderDebugInfoFlags::DetailedSource;
+	}
 
-			CurHead = CurHead->NextLink;
-		}
+	return Flags;
+}
 
-		JobData.DuplicateJobsWaitList = nullptr;
+FString FShaderCompilingManager::CreateShaderDebugInfoPath(const FShaderCompilerInput& ShaderCompilerInput) const
+{
+	TRACE_CPUPROFILER_EVENT_SCOPE(FShaderCompilingManager::CreateShaderDebugInfoPath);
 
-		if (FinishedJob->bSucceeded)
-		{
-			const bool bAddToDDC = GShaderCompilerPerShaderDDCGlobal || !(FinishedJob->bIsDefaultMaterial || FinishedJob->bIsGlobalShader);
-			// we only cache jobs that succeded
-			AddJobOutput(JobData, FinishedJob, InputHash, Buffer, FinishedDuplicateJobs.Num(), bAddToDDC);
-		}
+	FString DumpDebugInfoPath = ShaderCompilerInput.DumpDebugInfoRootPath / ShaderCompilerInput.DebugGroupName + ShaderCompilerInput.DebugExtension;
+
+	// Sanitize the name to be used as a path
+	// List mostly comes from set of characters not allowed by windows in a path.  Just try to rename a file and type one of these for the list.
+	DumpDebugInfoPath.ReplaceInline(TEXT("<"), TEXT("("));
+	DumpDebugInfoPath.ReplaceInline(TEXT(">"), TEXT(")"));
+	DumpDebugInfoPath.ReplaceInline(TEXT("::"), TEXT("=="));
+	DumpDebugInfoPath.ReplaceInline(TEXT("|"), TEXT("_"));
+	DumpDebugInfoPath.ReplaceInline(TEXT("*"), TEXT("-"));
+	DumpDebugInfoPath.ReplaceInline(TEXT("?"), TEXT("!"));
+	DumpDebugInfoPath.ReplaceInline(TEXT("\""), TEXT("\'"));
 
-		// remove ourselves from the jobs in flight
-		if (JobData.JobInFlight)
+	if (!IFileManager::Get().DirectoryExists(*DumpDebugInfoPath))
+	{
+		if (!IFileManager::Get().MakeDirectory(*DumpDebugInfoPath, true))
 		{
-#if WITH_EDITOR
-			if (JobData.JobInFlight->RequestOwner.IsValid())
-			{
-				JobData.JobInFlight->RequestOwner->KeepAlive();
-			}
-#endif
-			JobData.JobInFlight = nullptr;
+			const uint32 ErrorCode = FPlatformMisc::GetLastError();
+			UE_LOG(LogShaderCompilers, Warning, TEXT("Last Error %u: Failed to create directory for shader debug info '%s'. Try enabling large file paths or r.DumpShaderDebugShortNames."), ErrorCode, *DumpDebugInfoPath);
+			return FString();
 		}
-		FinishedJob->JobCacheRef.Clear();
 	}
 
-	if (FinishedDuplicateJobs.Num())
+	return DumpDebugInfoPath;
+}
+
+bool FShaderCompilingManager::ShouldRecompileToDumpShaderDebugInfo(const FShaderCompileJob& Job) const
+{
+	return ShouldRecompileToDumpShaderDebugInfo(Job.Input, Job.Output, Job.bSucceeded);
+}
+
+bool FShaderCompilingManager::ShouldRecompileToDumpShaderDebugInfo(const FShaderCompilerInput& Input, const FShaderCompilerOutput& Output, bool bSucceeded) const
+{
+	if (Input.DumpDebugInfoPath.IsEmpty())
 	{
-		UE_LOG(LogShaderCompilers, UE_SHADERCACHE_LOG_LEVEL, TEXT("Processed %d outstanding jobs with the same ihash %s."), FinishedDuplicateJobs.Num(), *LexToString(InputHash));
+		const EDumpShaderDebugInfo DumpShaderDebugInfo = GetDumpShaderDebugInfo();
+		const bool bErrors = !bSucceeded;
+		const bool bWarnings = Output.Errors.Num() > 0;
 
-		for (FShaderCommonCompileJob* DuplicateJob : FinishedDuplicateJobs)
+		bool bShouldDump = true;
+		if (GIsBuildMachine)
 		{
-			FMemoryReaderView MemReader(Buffer);
-			DuplicateJob->SerializeOutput(MemReader);
-			checkf(DuplicateJob->bSucceeded == FinishedJob->bSucceeded, TEXT("Different success status for the job with the same ihash"));
+			// Build machines dump these as build artifacts and they should only upload so many due to size constraints.
+			bShouldDump = (NumDumpedShaderSources < GMaxNumDumpedShaderSources);
+		}
 
-			// finish the job instantly
-			ProcessFinishedJob(DuplicateJob, true);
+		if (DumpShaderDebugInfo == EDumpShaderDebugInfo::OnError)
+		{
+			return bShouldDump && (bErrors);
+		}
+		else if (DumpShaderDebugInfo == EDumpShaderDebugInfo::OnErrorOrWarning)
+		{
+			return bShouldDump && (bErrors || bWarnings);
 		}
 	}
-}
 
-int32 FShaderJobCache::GetNumPendingJobs(EShaderCompileJobPriority InPriority) const
-{
-	return NumPendingJobs[(int32)InPriority];
+	return false;
 }
 
-int32 FShaderJobCache::GetNumOutstandingJobs() const
+void FShaderCompilingManager::ReleaseJob(FShaderCommonCompileJobPtr& Job)
 {
-	return NumOutstandingJobs.GetValue();
+	ReleaseJob(Job.GetReference());
+	Job.SafeRelease();
 }
 
-int32 FShaderJobCache::GetNumPendingJobs() const
+void FShaderCompilingManager::ReleaseJob(FShaderCommonCompileJob* Job)
 {
-	FReadScopeLock Locker(JobLock);
-	int32 NumJobs = 0;
-	for (int32 i = 0; i < NumShaderCompileJobPriorities; ++i)
-	{
-		NumJobs += NumPendingJobs[i];
-	}
-	return NumJobs;
+	Job->PendingShaderMap.SafeRelease();
+	Job->bReleased = true;
+	AllJobs.RemoveJob(Job);
 }
 
-int32 FShaderJobCache::GetPendingJobs(EShaderCompilerWorkerType InWorkerType, EShaderCompileJobPriority InPriority, int32 MinNumJobs, int32 MaxNumJobs, TArray<FShaderCommonCompileJobPtr>& OutJobs)
+void FShaderCompilingManager::SubmitJobs(TArray<FShaderCommonCompileJobPtr>& NewJobs, const FString MaterialBasePath, const FString PermutationString)
 {
-	check(InWorkerType != EShaderCompilerWorkerType::None);
-	check(InPriority != EShaderCompileJobPriority::None);
+	LLM_SCOPE_BYTAG(ShaderCompiler);
 
-	const int32 PriorityIndex = (int32)InPriority;
-	int32 NumPendingJobsOfPriority = 0;
+	// make sure no compiling can start if not allowed
+	if (!AllowShaderCompiling())
 	{
-		FReadScopeLock Locker(JobLock);
-		NumPendingJobsOfPriority = NumPendingJobs[PriorityIndex].load();
+		return;
 	}
+	
+	TRACE_CPUPROFILER_EVENT_SCOPE(FShaderCompilingManager::SubmitJobs);
+	check(!FPlatformProperties::RequiresCookedData());
 
-	if (NumPendingJobsOfPriority < MinNumJobs)
+	if (NewJobs.Num() == 0)
 	{
-		// Not enough jobs
-		return 0;
+		return;
 	}
 
-	FWriteScopeLock Locker(JobLock);
-
-	// there was a time window before we checked and then acquired the write lock - make sure the number is still sufficient
-	NumPendingJobsOfPriority = NumPendingJobs[PriorityIndex].load();
-	if (NumPendingJobsOfPriority < MinNumJobs)
+	check(GShaderCompilerStats);
+	if (FShaderCompileJob* SingleJob = NewJobs[0]->GetSingleShaderJob()) //assume that all jobs are for the same platform
 	{
-		// Not enough jobs
-		return 0;
+		GShaderCompilerStats->RegisterCompiledShaders(NewJobs.Num(), SingleJob->Input.Target.GetPlatform(), MaterialBasePath, PermutationString);
 	}
-	
-	OutJobs.Reserve(OutJobs.Num() + FMath::Min(MaxNumJobs, NumPendingJobsOfPriority));
-	int32 NumJobs = FMath::Min(MaxNumJobs, NumPendingJobsOfPriority);
-	FShaderCommonCompileJobIterator It(PendingJobsHead[PriorityIndex]);
-	// Randomize job selection by randomly skipping over jobs while traversing the list.
-	// Say, we need to pick 3 jobs out of 5 total. We can skip over 2 jobs in total, e.g. like this:
-	// pick one (4 more to go and we need to get 2 of 4), skip one (3 more to go, picking 2 out of 3), pick one (2 more to go, picking 1 of 2), skip one, pick one.
-	// It is possible that we won't skip at all and instead pick consequential jobs
-	int32 MaxJobsWeCanSkipOver = NumPendingJobsOfPriority - NumJobs;
-	for (int32 i = 0; i < NumJobs; ++i)
+	else
 	{
-		FShaderCommonCompileJob& Job = *It;
-
-		GShaderCompilerStats->RegisterAssignedJob(Job);
-		// Temporary commented out until r.ShaderDevelopmentMode=1 shader error retry crash gets fixed
-		//check(Job.CurrentWorker == EShaderCompilerWorkerType::None);
-		//check(Job.PendingPriority == InPriority);
-		ensure(!ShaderCompiler::IsJobCacheEnabled() || Job.bInputHashSet);
-
-		It.Next();
-
-		check((int32)Job.PendingPriority == PriorityIndex);
-		UnlinkJobWithPriority(Job);
-
-		Job.CurrentWorker = InWorkerType;
-		OutJobs.Add(&Job);
+		GShaderCompilerStats->RegisterCompiledShaders(NewJobs.Num(), SP_NumPlatforms, MaterialBasePath, PermutationString);
+	}
 
-		// get a random number of jobs to skip (if we can). We're skipping after taking the first job so we can ensure that we always take the latest job into the batch
-		if (MaxJobsWeCanSkipOver > 0)
+	{
+		FScopeLock Lock(&CompileQueueSection);
+		for (auto& Job : NewJobs)
 		{
-			int32 NumJobsToSkipOver = FMath::RandHelper(MaxJobsWeCanSkipOver + 1);
-			while (NumJobsToSkipOver > 0 && It)
+			FPendingShaderMapCompileResultsPtr& PendingShaderMap = ShaderMapJobs.FindOrAdd(Job->Id);
+			if (!PendingShaderMap)
 			{
-				It.Next();
-				--NumJobsToSkipOver;
-				--MaxJobsWeCanSkipOver;
+				PendingShaderMap = new FPendingShaderMapCompileResults();
 			}
-			checkf(MaxJobsWeCanSkipOver >= 0, TEXT("We skipped over too many jobs"));
-			checkf(MaxJobsWeCanSkipOver <= NumPendingJobsOfPriority - i, TEXT("Number of jobs to skip should stay less or equal than the number of nodes to go"));
+			PendingShaderMap->NumPendingJobs.Increment();
+			Job->PendingShaderMap = PendingShaderMap;
 		}
-	}
 
-	return NumJobs;
-}
+		// in the case of submitting jobs from worker threads we need to be sure that the lock extends to
+		// include AllJobs.SubmitJobs().  This will increase contention for the lock, but this will let us
+		// prototype getting shader translation and preprocessing being done on worker threads.
+		if (IsInGameThread())
+		{
+			Lock.Unlock();
+		}
 
-static float GRegularWorkerTimeToLive = 20.0f;
-static float GBuildWorkerTimeToLive = 600.0f;
+		AllJobs.SubmitJobs(NewJobs);
+	}
 
-// Configuration to retry shader compile through workers after a worker has been abandoned
-static constexpr int32 GSingleThreadedRunsIdle = -1;
-static constexpr int32 GSingleThreadedRunsDisabled = -2;
-static constexpr int32 GSingleThreadedRunsIncreaseFactor = 8;
-static constexpr int32 GSingleThreadedRunsMaxCount = (1 << 24);
+	UpdateNumRemainingAssets();
+}
 
-static void ModalErrorOrLog(const FString& Title, const FString& Text, int64 CurrentFilePos = 0, int64 ExpectedFileSize = 0, bool bIsErrorFatal = true)
+bool FShaderCompilingManager::IsCompilingShaderMap(uint32 Id)
 {
-	static FThreadSafeBool bModalReported;
-
-	FString BadFile;
-	if (CurrentFilePos > ExpectedFileSize)
-	{
-		// Corrupt file
-		BadFile = FString::Printf(TEXT(" (Truncated or corrupt output file! Current file pos %lld, file size %lld)"), CurrentFilePos, ExpectedFileSize);
-	}
-
-	if (FPlatformProperties::SupportsWindowedMode() && !FApp::IsUnattended())
+	if (Id != 0u)
 	{
-		UE_LOG(LogShaderCompilers, Error, TEXT("%s\n%s"), *Text, *BadFile);
-		if (!bModalReported.AtomicSet(true))
+		FScopeLock Lock(&CompileQueueSection);
+		FPendingShaderMapCompileResultsPtr* PendingShaderMapPtr = ShaderMapJobs.Find(Id);
+		if (PendingShaderMapPtr)
 		{
-			// Show dialog box with error message and request exit
-			FMessageDialog::Open(EAppMsgType::Ok, FText::FromString(Text), FText::FromString(Title));
-			FPlatformMisc::RequestExit(false, TEXT("ShaderCompiler.ModalErrorOrLog"));
+			return true;
 		}
-		else
+
+		FShaderMapFinalizeResults* FinalizedShaderMapPtr = PendingFinalizeShaderMaps.Find(Id);
+		if (FinalizedShaderMapPtr)
 		{
-			// Another thread already opened a dialog box and requests exit
-			FPlatformProcess::SleepInfinite();
+			return true;
 		}
 	}
-	else if (bIsErrorFatal)
+	return false;
+}
+
+FShaderCompileJob* FShaderCompilingManager::PrepareShaderCompileJob(uint32 Id, const FShaderCompileJobKey& Key, EShaderCompileJobPriority Priority)
+{
+	// don't allow any jobs if not allowed
+	if (!AllowShaderCompiling())
 	{
-		UE_LOG(LogShaderCompilers, Fatal, TEXT("%s\n%s\n%s"), *Title, *Text, *BadFile);
+		return nullptr;
 	}
-	else
+
+	return AllJobs.PrepareJob(Id, Key, Priority);
+}
+
+FShaderPipelineCompileJob* FShaderCompilingManager::PreparePipelineCompileJob(uint32 Id, const FShaderPipelineCompileJobKey& Key, EShaderCompileJobPriority Priority)
+{
+	// don't allow any jobs if not allowed
+	if (!AllowShaderCompiling())
 	{
-		UE_LOG(LogShaderCompilers, Error, TEXT("%s\n%s\n%s"), *Title, *Text, *BadFile);
+		return nullptr;
 	}
+
+	return AllJobs.PrepareJob(Id, Key, Priority);
 }
 
-template<class EnumType>
-constexpr auto& CastEnumToUnderlyingTypeReference(EnumType& Type)
+void FShaderCompilingManager::ProcessFinishedJob(FShaderCommonCompileJob* FinishedJob)
 {
-	static_assert(TIsEnum<EnumType>::Value, "");
-	using UnderType = __underlying_type(EnumType);
-	return reinterpret_cast<UnderType&>(Type);
+	AllJobs.ProcessFinishedJob(FinishedJob);
 }
 
-// Set to 1 to debug ShaderCompileWorker.exe. Set a breakpoint in LaunchWorker() to get the cmd-line.
-#define DEBUG_SHADERCOMPILEWORKER 0
-
-// Default value comes from bPromptToRetryFailedShaderCompiles in BaseEngine.ini
-// This is set as a global variable to allow changing in the debugger even in release
-// For example if there are a lot of content shader compile errors you want to skip over without relaunching
-bool GRetryShaderCompilation = true;
+/** Launches the worker, returns the launched process handle. */
+FProcHandle FShaderCompilingManager::LaunchWorker(const FString& WorkingDirectory, uint32 InProcessId, uint32 ThreadId, const FString& WorkerInputFile, const FString& WorkerOutputFile)
+{
+	// don't allow any jobs if not allowed
+	if (!AllowShaderCompiling())
+	{
+		return FProcHandle();
+	}
 
-static FShaderCompilingManager::EDumpShaderDebugInfo GDumpShaderDebugInfo = FShaderCompilingManager::EDumpShaderDebugInfo::Never;
-static FAutoConsoleVariableRef CVarDumpShaderDebugInfo(
-	TEXT("r.DumpShaderDebugInfo"),
-	CastEnumToUnderlyingTypeReference(GDumpShaderDebugInfo),
-	TEXT("Dumps debug info for compiled shaders to GameName/Saved/ShaderDebugInfo\n")
-	TEXT("When set to 1, debug info is dumped for all compiled shader\n")
-	TEXT("When set to 2, it is restricted to shaders with compilation errors\n")
-	TEXT("When set to 3, it is restricted to shaders with compilation errors or warnings\n")
-	TEXT("The debug info is platform dependent, but usually includes a preprocessed version of the shader source.\n")
-	TEXT("Global shaders automatically dump debug info if r.ShaderDevelopmentMode is enabled, this cvar is not necessary.\n")
-	TEXT("On iOS, if the PowerVR graphics SDK is installed to the default path, the PowerVR shader compiler will be called and errors will be reported during the cook.")
-	);
+	// Setup the parameters that the worker application needs
+	// Surround the working directory with double quotes because it may contain a space 
+	// WorkingDirectory ends with a '\', so we have to insert another to meet the Windows commandline parsing rules 
+	// http://msdn.microsoft.com/en-us/library/17w5ykft.aspx 
+	// Use IFileManager to do path conversion to properly handle sandbox paths (outside of standard paths in particular).
+	FString WorkerAbsoluteDirectory = IFileManager::Get().ConvertToAbsolutePathForExternalAppForWrite(*WorkingDirectory);
+	FPaths::NormalizeDirectoryName(WorkerAbsoluteDirectory);
+	FString WorkerParameters = FString(TEXT("\"")) + WorkerAbsoluteDirectory + TEXT("/\" ") + FString::FromInt(InProcessId) + TEXT(" ") + FString::FromInt(ThreadId) + TEXT(" ") + WorkerInputFile + TEXT(" ") + WorkerOutputFile;
+	WorkerParameters += FString(TEXT(" -communicatethroughfile "));
+	if ( GIsBuildMachine )
+	{
+		WorkerParameters += FString::Printf(TEXT(" -TimeToLive=%f -buildmachine"), GBuildWorkerTimeToLive);
+	}
+	else
+	{
+		WorkerParameters += FString::Printf(TEXT(" -TimeToLive=%f"), GRegularWorkerTimeToLive);
+	}
+	if (PLATFORM_LINUX) //-V560
+	{
+		// suppress log generation as much as possible
+		WorkerParameters += FString(TEXT(" -logcmds=\"Global None\" "));
 
-static TAutoConsoleVariable<bool> CVarDumpShaderOutputCacheHits(
-	TEXT("r.DumpShaderOutputCacheHits"),
-	false,
-	TEXT("Dumps shader output bytecode and cache hits with reference to original output.\n")
-	TEXT("Dumping shader output bytecode for all compile shaders also requires CVar r.DumpShaderDebugInfo=1."),
-	ECVF_ReadOnly);
+		if (UE_BUILD_DEBUG)
+		{
+			// when running a debug build under Linux, make SCW crash with core for easier debugging
+			WorkerParameters += FString(TEXT(" -core "));
+		}
+	}
+	WorkerParameters += FCommandLine::GetSubprocessCommandline();
 
-static int32 GDumpShaderDebugInfoShort = 0;
-static FAutoConsoleVariableRef CVarDumpShaderDebugShortNames(
-	TEXT("r.DumpShaderDebugShortNames"),
-	GDumpShaderDebugInfoShort,
-	TEXT("Only valid when r.DumpShaderDebugInfo > 0.\n")
-	TEXT("When set to 1, will shorten names factory and shader type folder names to avoid issues with long paths.")
-	);
-
-static int32 GDumpShaderDebugInfoBindless = 0;
-static FAutoConsoleVariableRef CVarDumpShaderDebugBindlessNames(
-	TEXT("r.DumpShaderDebugBindlessNames"),
-	GDumpShaderDebugInfoBindless,
-	TEXT("Only valid when r.DumpShaderDebugInfo > 0.\n")
-	TEXT("When set to 1, will add bindless folder names.")
-	);
-
-static int32 GDumpShaderDebugInfoSCWCommandLine = 0;
-static FAutoConsoleVariableRef CVarDumpShaderDebugSCWCommandLine(
-	TEXT("r.DumpShaderDebugWorkerCommandLine"),
-	GDumpShaderDebugInfoSCWCommandLine,
-	TEXT("Only valid when r.DumpShaderDebugInfo > 0.\n")
-	TEXT("When set to 1, it will generate a file that can be used with ShaderCompileWorker's -directcompile.")
-	);
-
-static int32 GShaderMapCompilationTimeout = 2 * 60 * 60;	// anything below an hour can hit a false positive
-static FAutoConsoleVariableRef CVarShaderMapCompilationTimeout(
-	TEXT("r.ShaderCompiler.ShadermapCompilationTimeout"),
-	GShaderMapCompilationTimeout,
-	TEXT("Maximum number of seconds a single shadermap (which can be comprised of multiple jobs) can be compiled after being considered hung.")
-);
-
-static int32 GCrashOnHungShaderMaps = 0;
-static FAutoConsoleVariableRef CVarCrashOnHungShaderMaps(
-	TEXT("r.ShaderCompiler.CrashOnHungShaderMaps"),
-	GCrashOnHungShaderMaps,
-	TEXT("If set to 1, the shader compiler will crash on hung shadermaps.")
-);
-
-static int32 GLogShaderCompilerStats = 0;
-static FAutoConsoleVariableRef CVarLogShaderCompilerStats(
-	TEXT("r.LogShaderCompilerStats"),
-	GLogShaderCompilerStats,
-	TEXT("When set to 1, Log detailed shader compiler stats.")
-);
-
-static int32 GForceAllCoresForShaderCompiling = 0;
-static FAutoConsoleVariableRef CVarForceAllCoresForShaderCompiling(
-	TEXT("r.ForceAllCoresForShaderCompiling"),
-	GForceAllCoresForShaderCompiling,
-	TEXT("When set to 1, it will ignore INI settings and launch as many ShaderCompileWorker instances as cores are available.\n")
-	TEXT("Improves shader throughput but for big projects it can make the machine run OOM")
-);
-
-static TAutoConsoleVariable<int32> CVarShadersSymbols(
-	TEXT("r.Shaders.Symbols"),
-	0,
-	TEXT("Enables debugging of shaders in platform specific graphics debuggers. This will generate and write shader symbols.\n")
-	TEXT("This enables the behavior of both r.Shaders.GenerateSymbols and r.Shaders.WriteSymbols.\n")
-	TEXT("Enables shader debugging features that require shaders to be recompiled. This compiles shaders with symbols and also includes extra runtime information like shader names. When using graphical debuggers it can be useful to enable this on startup.\n")
-	TEXT("This setting can be overriden in any Engine.ini under the [ShaderCompiler] section."),
-	ECVF_ReadOnly);
-
-static TAutoConsoleVariable<int32> CVarShadersSymbolsInfo(
-	TEXT("r.Shaders.SymbolsInfo"),
-	0,
-	TEXT("In lieu of a full set of platform shader PDBs, save out a slimmer ShaderSymbols.Info which contains shader platform hashes and shader debug info.\n")
-	TEXT("An option for when it is not practical to save PDBs for shaders all the time.\n")
-	TEXT("This setting can be overriden in any Engine.ini under the [ShaderCompiler] section."),
-	ECVF_ReadOnly);
-
-static TAutoConsoleVariable<int32> CVarShadersGenerateSymbols(
-	TEXT("r.Shaders.GenerateSymbols"),
-	0,
-	TEXT("Enables generation of data for shader debugging when compiling shaders. This explicitly does not write any shader symbols to disk.\n")
-	TEXT("This setting can be overriden in any Engine.ini under the [ShaderCompiler] section."),
-	ECVF_ReadOnly);
-
-static TAutoConsoleVariable<int32> CVarShadersWriteSymbols(
-	TEXT("r.Shaders.WriteSymbols"),
-	0,
-	TEXT("Enables writing shader symbols to disk for platforms that support that. This explicitly does not enable generation of shader symbols.\n")
-	TEXT("This setting can be overriden in any Engine.ini under the [ShaderCompiler] section."),
-	ECVF_ReadOnly);
-
-static TAutoConsoleVariable<FString> CVarShadersSymbolPathOverride(
-	TEXT("r.Shaders.SymbolPathOverride"),
-	"",
-	TEXT("Override output location of shader symbols. If the path contains the text '{Platform}', that will be replaced with the shader platform string.\n")
-	TEXT("Empty: use default location Saved/ShaderSymbols/{Platform}\n")
-	TEXT("This setting can be overriden in any Engine.ini under the [ShaderCompiler] section."),
-	ECVF_ReadOnly);
-
-static TAutoConsoleVariable<int32> CVarAllowUniqueDebugInfo(
-	TEXT("r.Shaders.AllowUniqueSymbols"),
-	0,
-	TEXT("When enabled, this tells supported shader compilers to generate symbols based on source files.\n")
-	TEXT("Enabling this can cause a drastic increase in the number of symbol files, enable only if absolutely necessary.\n")
-	TEXT("This setting can be overriden in any Engine.ini under the [ShaderCompiler] section."),
-	ECVF_ReadOnly);
-
-static TAutoConsoleVariable<int32> CVarShadersWriteSymbolsZip(
-	TEXT("r.Shaders.WriteSymbols.Zip"),
-	0,
-	TEXT(" 0: Export as loose files.\n")
-	TEXT(" 1: Export as an uncompressed archive.\n"),
-	ECVF_ReadOnly);
-
-static TAutoConsoleVariable<int32> CVarShadersEnableExtraData(
-	TEXT("r.Shaders.ExtraData"),
-	0,
-	TEXT("Enables generation of extra shader data that can be used at runtime. This includes shader names and other platform specific data.\n")
-	TEXT("This can add bloat to compiled shaders and can prevent shaders from being deduplicated.\n")
-	TEXT("This setting can be overriden in any Engine.ini under the [ShaderCompiler] section."),
-	ECVF_ReadOnly);
-
-static TAutoConsoleVariable<int32> CVarOptimizeShaders(
-	TEXT("r.Shaders.Optimize"),
-	1,
-	TEXT("Whether to optimize shaders.  When using graphical debuggers like Nsight it can be useful to disable this on startup.\n")
-	TEXT("This setting can be overriden in any Engine.ini under the [ShaderCompiler] section."),
-	ECVF_ReadOnly);
-
-static TAutoConsoleVariable<int32> CVarShaderFastMath(
-	TEXT("r.Shaders.FastMath"),
-	1,
-	TEXT("Whether to use fast-math optimisations in shaders."),
-	ECVF_ReadOnly);
-
-static TAutoConsoleVariable<int32> CVarShaderZeroInitialise(
-	TEXT("r.Shaders.ZeroInitialise"),
-	1,
-	TEXT("Whether to enforce zero initialise local variables of primitive type in shaders. Defaults to 1 (enabled). Not all shader languages can omit zero initialisation."),
-	ECVF_ReadOnly);
-
-static TAutoConsoleVariable<int32> CVarShaderBoundsChecking(
-	TEXT("r.Shaders.BoundsChecking"),
-	1,
-	TEXT("Whether to enforce bounds-checking & flush-to-zero/ignore for buffer reads & writes in shaders. Defaults to 1 (enabled). Not all shader languages can omit bounds checking."),
-	ECVF_ReadOnly);
-
-static TAutoConsoleVariable<int32> CVarShaderWarningsAsErrors(
-	TEXT("r.Shaders.WarningsAsErrors"),
-	0,
-	TEXT("Whether to treat warnings as errors when compiling shaders. (0: disabled (default), 1: global shaders only, 2: all shaders)). This setting may be ignored on older platforms."),
-	ECVF_ReadOnly);
-
-static TAutoConsoleVariable<int32> CVarShaderFlowControl(
-	TEXT("r.Shaders.FlowControlMode"),
-	0,
-	TEXT("Specifies whether the shader compiler should preserve or unroll flow-control in shader code.\n")
-	TEXT("This is primarily a debugging aid and will override any per-shader or per-material settings if not left at the default value (0).\n")
-	TEXT("\t0: Off (Default) - Entirely at the discretion of the platform compiler or the specific shader/material.\n")
-	TEXT("\t1: Prefer - Attempt to preserve flow-control.\n")
-	TEXT("\t2: Avoid - Attempt to unroll and flatten flow-control.\n"),
-	ECVF_ReadOnly);
-
-static TAutoConsoleVariable<int32> CVarD3DCheckedForTypedUAVs(
-	TEXT("r.D3D.CheckedForTypedUAVs"),
-	1,
-	TEXT("Whether to disallow usage of typed UAV loads, as they are unavailable in Windows 7 D3D 11.0.\n")
-	TEXT(" 0: Allow usage of typed UAV loads.\n")
-	TEXT(" 1: Disallow usage of typed UAV loads. (default)"),
-	ECVF_ReadOnly);
-
-static TAutoConsoleVariable<int32> CVarD3DForceDXC(
-	TEXT("r.D3D.ForceDXC"),
-	0,
-	TEXT("Forces DirectX Shader Compiler (DXC) to be used for all D3D shaders. Shaders compiled with this option are only compatible with D3D12.\n")
-	TEXT(" 0: Disable (default)\n")
-	TEXT(" 1: Force new compiler for all shaders"),
-	ECVF_ReadOnly);
-
-static TAutoConsoleVariable<int32> CVarWarpCulling(
-	TEXT("r.WarpCulling"),
-	0,
-	TEXT("Enable Warp Culling optimization for platforms that support it.\n")
-	TEXT(" 0: Disable (default)\n")
-	TEXT(" 1: Enable"),
-	ECVF_ReadOnly);
-
-static TAutoConsoleVariable<int32> CVarCullBeforeFetch(
-	TEXT("r.CullBeforeFetch"),
-	0,
-	TEXT("Enable Cull-Before-Fetch optimization for platforms that support it.\n")
-	TEXT(" 0: Disable (default)\n")
-	TEXT(" 1: Enable"),
-	ECVF_ReadOnly);
-
-ENGINE_API int32 GCreateShadersOnLoad = 0;
-static FAutoConsoleVariableRef CVarCreateShadersOnLoad(
-	TEXT("r.CreateShadersOnLoad"),
-	GCreateShadersOnLoad,
-	TEXT("Whether to create shaders on load, which can reduce hitching, but use more memory.  Otherwise they will be created as needed.")
-);
-
-static TAutoConsoleVariable<FString> CVarShaderOverrideDebugDir(
-	TEXT("r.OverrideShaderDebugDir"),
-	"",
-	TEXT("Override output location of shader debug files\n")
-	TEXT("Empty: use default location Saved\\ShaderDebugInfo.\n"),
-	ECVF_ReadOnly);
-
-static TAutoConsoleVariable<int32> CVarShadersValidation(
-	TEXT("r.Shaders.Validation"),
-	1,
-	TEXT("Enabled shader compiler validation warnings and errors."),
-	ECVF_ReadOnly);
-
-static TAutoConsoleVariable<int32> CVarShadersRemoveDeadCode(
-	TEXT("r.Shaders.RemoveDeadCode"),
-	1,
-	TEXT("Run a preprocessing step that removes unreferenced code before compiling shaders.\n")
-	TEXT("This can improve the compilation speed for shaders which include many large utility headers.\n")
-	TEXT("\t0: Keep all input source code.\n")
-	TEXT("\t1: Remove unreferenced code before compilation (Default)\n"),
-	ECVF_ReadOnly);
-
-static TAutoConsoleVariable<bool> CVarShadersPropagateLocalWorkerOOMs(
-	TEXT("r.Shaders.PropagateLocalWorkerOOMs"),
-	false,
-	TEXT("When set, out-of-memory conditions in a local shader compile worker will be treated as regular out-of-memory conditions and propagated to the main process.\n")
-	TEXT("This is useful when running in environment with hard memory limits, where it does not matter which process in particular caused us to violate the memory limit."),
-	ECVF_Default);
-
-#if ENABLE_COOK_STATS
-namespace ShaderCompilerCookStats
-{
-	static double BlockingTimeSec = 0.0;
-	static double GlobalBeginCompileShaderTimeSec = 0.0;
-	static int32 GlobalBeginCompileShaderCalls = 0;
-	static double ProcessAsyncResultsTimeSec = 0.0;
-	static double AsyncCompileTimeSec = 0.0;
-
-	static FCookStatsManager::FAutoRegisterCallback RegisterCookStats([](FCookStatsManager::AddStatFuncRef AddStat)
-		{
-			AddStat(TEXT("ShaderCompiler"), FCookStatsManager::CreateKeyValueArray(
-				TEXT("BlockingTimeSec"), BlockingTimeSec,
-				TEXT("AsyncCompileTimeSec"), AsyncCompileTimeSec,
-				TEXT("GlobalBeginCompileShaderTimeSec"), GlobalBeginCompileShaderTimeSec,
-				TEXT("GlobalBeginCompileShaderCalls"), GlobalBeginCompileShaderCalls,
-				TEXT("ProcessAsyncResultsTimeSec"), ProcessAsyncResultsTimeSec
-			));
-		});
-}
-#endif
-
-static void ReissueShaderCompileJobs(const TArray<FShaderCommonCompileJob*>& SourceJobs)
-{
-	if (SourceJobs.Num())
-	{
-		TRACE_CPUPROFILER_EVENT_SCOPE(ReissueShaderCompileJobs);
-
-		TArray<FShaderCommonCompileJobPtr> ReissueJobs;
-		ReissueJobs.Reserve(SourceJobs.Num());
-		const uint32 JobId = FShaderCommonCompileJob::GetNextJobId();
-		for (const FShaderCommonCompileJob* SourceJob : SourceJobs)
-		{
-			if (const FShaderCompileJob* SingleSourceJob = SourceJob->GetSingleShaderJob())
-			{
-				if (FShaderCompileJob* ReissueJob = GShaderCompilingManager->PrepareShaderCompileJob(JobId, SingleSourceJob->Key, SingleSourceJob->Priority))
-				{
-					ReissueJob->Input = SingleSourceJob->Input;
-					ReissueJobs.Add(FShaderCommonCompileJobPtr(ReissueJob));
-				}
-			}
-			else if (const FShaderPipelineCompileJob* PipelineSourceJob = SourceJob->GetShaderPipelineJob())
-			{
-				if (FShaderPipelineCompileJob* ReissueJob = GShaderCompilingManager->PreparePipelineCompileJob(JobId, PipelineSourceJob->Key, PipelineSourceJob->Priority))
-				{
-					ReissueJob->StageJobs = PipelineSourceJob->StageJobs;
-					ReissueJobs.Add(FShaderCommonCompileJobPtr(ReissueJob));
-				}
-			}
-			else
-			{
-				checkf(0, TEXT("Reissued shader compile job is neither a single nor a pipeline job"));
-			}
-		}
-
-		GShaderCompilingManager->SubmitJobs(ReissueJobs, FString(""), FString(""));
-	}
-}
-
-// Make functions so the crash reporter can disambiguate the actual error because of the different callstacks
-namespace ShaderCompileWorkerError
-{
-	void HandleGeneralCrash(const TCHAR* ExceptionInfo, const TCHAR* Callstack)
-	{
-		GLog->Panic();
-		UE_LOG(LogShaderCompilers, Error, TEXT("ShaderCompileWorker crashed!\n%s\n%s"), ExceptionInfo, Callstack);
-	}
-
-	void HandleBadShaderFormatVersion(const TCHAR* Data)
-	{
-		ModalErrorOrLog(TEXT("ShaderCompileWorker failed"), Data);
-	}
-
-	void HandleBadInputVersion(const TCHAR* Data)
-	{
-		ModalErrorOrLog(TEXT("ShaderCompileWorker failed"), Data);
-	}
-
-	void HandleBadSingleJobHeader(const TCHAR* Data)
-	{
-		ModalErrorOrLog(TEXT("ShaderCompileWorker failed"), Data);
-	}
-
-	void HandleBadPipelineJobHeader(const TCHAR* Data)
-	{
-		ModalErrorOrLog(TEXT("ShaderCompileWorker failed"), Data);
-	}
-
-	void HandleCantDeleteInputFile(const TCHAR* Data)
-	{
-		ModalErrorOrLog(TEXT("ShaderCompileWorker failed"), Data);
-	}
-
-	void HandleCantSaveOutputFile(const TCHAR* Data)
-	{
-		ModalErrorOrLog(TEXT("ShaderCompileWorker failed"), Data);
-	}
-
-	void HandleNoTargetShaderFormatsFound(const TCHAR* Data)
-	{
-		ModalErrorOrLog(TEXT("ShaderCompileWorker failed"), Data);
-	}
-
-	void HandleCantCompileForSpecificFormat(const TCHAR* Data)
-	{
-		ModalErrorOrLog(TEXT("ShaderCompileWorker failed"), Data);
-	}
-
-	void HandleOutputFileEmpty(const TCHAR* Filename)
-	{
-		ModalErrorOrLog(TEXT("ShaderCompileWorker failed"), FString::Printf(TEXT("Output file %s size is 0. Are you out of disk space?"), Filename));
-	}
-
-	void HandleOutputFileCorrupted(const TCHAR* Filename, int64 ExpectedSize, int64 ActualSize)
-	{
-		ModalErrorOrLog(TEXT("ShaderCompileWorker failed"), FString::Printf(TEXT("Output file corrupted (expected %I64d bytes, but only got %I64d): %s"), ExpectedSize, ActualSize, Filename));
-	}
-
-	void HandleCrashInsidePlatformCompiler(const TCHAR* Data)
-	{
-		// If the crash originates from a platform compiler, the error code must have been reported and we don't have to assume a corrupted output file.
-		// In that case, don't crash the cooker with a fatal error, just report the error so the cooker can dump debug info.
-		constexpr bool bIsErrorFatal = false;
-		ModalErrorOrLog(TEXT("ShaderCompileWorker failed"), FString::Printf(TEXT("Crash inside the platform compiler:\n%s"), Data), 0, 0, bIsErrorFatal);
-	}
-
-	void HandleBadInputFile(const TCHAR* Data)
-	{
-		ModalErrorOrLog(TEXT("ShaderCompileWorker failed"), FString::Printf(TEXT("Bad-input-file exception:\n%s"), Data));
-	}
-
-	bool HandleOutOfMemory(const TCHAR* ExceptionInfo, const TCHAR* Hostname, const FPlatformMemoryStats& MemoryStats, const TArray<FShaderCommonCompileJobPtr>& QueuedJobs)
-	{
-		constexpr int64 Gibibyte = 1024 * 1024 * 1024;
-		const FString ErrorReport = FString::Printf(
-			TEXT("ShaderCompileWorker failed with out-of-memory (OOM) exception on machine \"%s\" (%s); MemoryStats:")
-			TEXT("\n\tAvailablePhysical %llu (%.2f GiB)")
-			TEXT("\n\t AvailableVirtual %llu (%.2f GiB)")
-			TEXT("\n\t     UsedPhysical %llu (%.2f GiB)")
-			TEXT("\n\t PeakUsedPhysical %llu (%.2f GiB)")
-			TEXT("\n\t      UsedVirtual %llu (%.2f GiB)")
-			TEXT("\n\t  PeakUsedVirtual %llu (%.2f GiB)"),
-			Hostname,
-			(ExceptionInfo[0] == TEXT('\0') ? TEXT("No exception information") : ExceptionInfo),
-			MemoryStats.AvailablePhysical, double(MemoryStats.AvailablePhysical) / Gibibyte,
-			MemoryStats.AvailableVirtual, double(MemoryStats.AvailableVirtual) / Gibibyte,
-			MemoryStats.UsedPhysical, double(MemoryStats.UsedPhysical) / Gibibyte,
-			MemoryStats.PeakUsedPhysical, double(MemoryStats.PeakUsedPhysical) / Gibibyte,
-			MemoryStats.UsedVirtual, double(MemoryStats.UsedVirtual) / Gibibyte,
-			MemoryStats.PeakUsedVirtual, double(MemoryStats.PeakUsedVirtual) / Gibibyte
-		);
-
-		if (GShaderCompilingManager->IsRemoteCompilingEnabled())
-		{
-			// Remote shader compiler supports re-compiling jobs on local machine
-			UE_LOG(LogShaderCompilers, Warning, TEXT("%s\nRecompile %d shader compile %s locally"), *ErrorReport, QueuedJobs.Num(), (QueuedJobs.Num() == 1 ? TEXT("job") : TEXT("jobs")));
-			return true;
-		}
-		else
-		{
-			if (CVarShadersPropagateLocalWorkerOOMs.GetValueOnAnyThread())
-			{
-				FPlatformMemory::OnOutOfMemory(0, 64);
-			}
-			ModalErrorOrLog(TEXT("ShaderCompileWorker failed"), ErrorReport);
-			return false;
-		}
-	}
-}
-
-static TMap<FString, uint32> GetFormatVersionMap()
-{
-	TMap<FString, uint32> FormatVersionMap;
-
-	const TArray<const class IShaderFormat*>& ShaderFormats = GetTargetPlatformManagerRef().GetShaderFormats();
-	check(ShaderFormats.Num());
-	for (int32 Index = 0; Index < ShaderFormats.Num(); Index++)
-	{
-		TArray<FName> OutFormats;
-		ShaderFormats[Index]->GetSupportedFormats(OutFormats);
-		check(OutFormats.Num());
-		for (int32 InnerIndex = 0; InnerIndex < OutFormats.Num(); InnerIndex++)
-		{
-			uint32 Version = ShaderFormats[Index]->GetVersion(OutFormats[InnerIndex]);
-			FormatVersionMap.Add(OutFormats[InnerIndex].ToString(), Version);
-		}
-	}
-
-	return FormatVersionMap;
-}
-
-static int32 GetNumTotalJobs(const TArray<FShaderCommonCompileJobPtr>& Jobs)
-{
-	int32 NumJobs = 0;
-	for (int32 Index = 0; Index < Jobs.Num(); ++Index)
-	{
-		auto* PipelineJob = Jobs[Index]->GetShaderPipelineJob();
-		NumJobs += PipelineJob ? PipelineJob->StageJobs.Num() : 1;
-	}
-
-	return NumJobs;
-}
-
-static void SplitJobsByType(const TArray<FShaderCommonCompileJobPtr>& QueuedJobs, TArray<FShaderCompileJob*>& OutQueuedSingleJobs, TArray<FShaderPipelineCompileJob*>& OutQueuedPipelineJobs)
-{
-	for (int32 Index = 0; Index < QueuedJobs.Num(); ++Index)
-	{
-		FShaderCommonCompileJobPtr CommonJob = QueuedJobs[Index];
-		if (FShaderCompileJob* SingleJob = CommonJob->GetSingleShaderJob())
-		{
-			OutQueuedSingleJobs.Add(SingleJob);
-
-		}
-		else if (FShaderPipelineCompileJob* PipelineJob = CommonJob->GetShaderPipelineJob())
-		{
-			OutQueuedPipelineJobs.Add(PipelineJob);
-		}
-		else
-		{
-			checkf(0, TEXT("FShaderCommonCompileJob::Type=%d is not a valid type for a shader compile job"), (int32)CommonJob->Type);
-		}
-	}
-}
-
-bool DoWriteTasksInner(const TArray<FShaderCommonCompileJobPtr>& QueuedJobs, FArchive& InTransferFile, IDistributedBuildController* BuildDistributionController, bool bUseRelativePaths, bool bCompressTaskFile)
-{
-	int32 InputVersion = ShaderCompileWorkerInputVersion;
-	InTransferFile << InputVersion;
-
-	TArray<uint8> UncompressedArray;
-	FMemoryWriter TransferMemory(UncompressedArray);
-	FArchive& TransferFile = bCompressTaskFile ? TransferMemory : InTransferFile;
-	if (!bCompressTaskFile)
-	{
-		// still write NAME_None as string
-		FString FormatNone = FName(NAME_None).ToString();
-		TransferFile << FormatNone;
-	}
-
-	static TMap<FString, uint32> FormatVersionMap = GetFormatVersionMap();
-
-	TransferFile << FormatVersionMap;
-
-	TArray<FShaderCompileJob*> QueuedSingleJobs;
-	TArray<FShaderPipelineCompileJob*> QueuedPipelineJobs;
-	SplitJobsByType(QueuedJobs, QueuedSingleJobs, QueuedPipelineJobs);
-
-	TArray<TRefCountPtr<FSharedShaderCompilerEnvironment>> SharedEnvironments;
-	TArray<const FShaderParametersMetadata*> RequestShaderParameterStructures;
-
-	// gather shared environments and parameter structures, these tend to be shared between jobs
-	{
-
-		for (int32 JobIndex = 0; JobIndex < QueuedSingleJobs.Num(); JobIndex++)
-		{
-			QueuedSingleJobs[JobIndex]->Input.GatherSharedInputs(SharedEnvironments, RequestShaderParameterStructures);
-		}
-
-		for (int32 JobIndex = 0; JobIndex < QueuedPipelineJobs.Num(); JobIndex++)
-		{
-			auto* PipelineJob = QueuedPipelineJobs[JobIndex];
-			int32 NumStageJobs = PipelineJob->StageJobs.Num();
-
-			for (int32 Index = 0; Index < NumStageJobs; Index++)
-			{
-				PipelineJob->StageJobs[Index]->Input.GatherSharedInputs(SharedEnvironments, RequestShaderParameterStructures);
-			}
-		}
-
-		int32 NumSharedEnvironments = SharedEnvironments.Num();
-		TransferFile << NumSharedEnvironments;
-
-		for (int32 EnvironmentIndex = 0; EnvironmentIndex < SharedEnvironments.Num(); EnvironmentIndex++)
-		{
-			TransferFile << *SharedEnvironments[EnvironmentIndex];
-		}
-	}
-
-	// Write shader parameter structures
-	TArray<const FShaderParametersMetadata*> AllShaderParameterStructures;
-	{
-		// List all dependencies.
-		for (int32 StructId = 0; StructId < RequestShaderParameterStructures.Num(); StructId++)
-		{
-			RequestShaderParameterStructures[StructId]->IterateStructureMetadataDependencies(
-				[&](const FShaderParametersMetadata* Struct)
-			{
-				AllShaderParameterStructures.AddUnique(Struct);
-			});
-		}
-
-		// Write all shader parameter structure.
-		int32 NumParameterStructures = AllShaderParameterStructures.Num();
-		TransferFile << NumParameterStructures;
-		for (const FShaderParametersMetadata* Struct : AllShaderParameterStructures)
-		{
-			FString LayoutName = Struct->GetLayout().GetDebugName();
-			FString StructTypeName = Struct->GetStructTypeName();
-			FString ShaderVariableName = Struct->GetShaderVariableName();
-			uint8 UseCase = uint8(Struct->GetUseCase());
-			FString StructFileName = FString(ANSI_TO_TCHAR(Struct->GetFileName()));
-			int32 StructFileLine = Struct->GetFileLine();
-			uint32 Size = Struct->GetSize();
-			int32 MemberCount = Struct->GetMembers().Num();
-
-			static_assert(sizeof(UseCase) == sizeof(FShaderParametersMetadata::EUseCase), "Cast failure.");
-
-			TransferFile << LayoutName;
-			TransferFile << StructTypeName;
-			TransferFile << ShaderVariableName;
-			TransferFile << UseCase;
-			TransferFile << StructFileName;
-			TransferFile << StructFileLine;
-			TransferFile << Size;
-			TransferFile << MemberCount;
-
-			for (const FShaderParametersMetadata::FMember& Member : Struct->GetMembers())
-			{
-				FString Name = Member.GetName();
-				FString ShaderType = Member.GetShaderType();
-				int32 FileLine = Member.GetFileLine();
-				uint32 Offset = Member.GetOffset();
-				uint8 BaseType = uint8(Member.GetBaseType());
-				uint8 PrecisionModifier = uint8(Member.GetPrecision());
-				uint32 NumRows = Member.GetNumRows();
-				uint32 NumColumns = Member.GetNumColumns();
-				uint32 NumElements = Member.GetNumElements();
-				int32 StructMetadataIndex = INDEX_NONE;
-				if (Member.GetStructMetadata())
-				{
-					StructMetadataIndex = AllShaderParameterStructures.Find(Member.GetStructMetadata());
-					check(StructMetadataIndex != INDEX_NONE);
-				}
-
-				static_assert(sizeof(BaseType) == sizeof(EUniformBufferBaseType), "Cast failure.");
-				static_assert(sizeof(PrecisionModifier) == sizeof(EShaderPrecisionModifier::Type), "Cast failure.");
-
-				TransferFile << Name;
-				TransferFile << ShaderType;
-				TransferFile << FileLine;
-				TransferFile << Offset;
-				TransferFile << BaseType;
-				TransferFile << PrecisionModifier;
-				TransferFile << NumRows;
-				TransferFile << NumColumns;
-				TransferFile << NumElements;
-				TransferFile << StructMetadataIndex;
-			}
-		}
-	}
-
-	// Write individual shader jobs
-	{
-		int32 SingleJobHeader = ShaderCompileWorkerSingleJobHeader;
-		TransferFile << SingleJobHeader;
-
-		int32 NumBatches = QueuedSingleJobs.Num();
-		TransferFile << NumBatches;
-
-		// Serialize all the batched jobs
-		for (int32 JobIndex = 0; JobIndex < QueuedSingleJobs.Num(); JobIndex++)
-		{
-			QueuedSingleJobs[JobIndex]->SerializeWorkerInput(TransferFile);
-			QueuedSingleJobs[JobIndex]->Input.SerializeSharedInputs(TransferFile, SharedEnvironments, AllShaderParameterStructures);
-		}
-	}
-
-	// Write shader pipeline jobs
-	{
-		int32 PipelineJobHeader = ShaderCompileWorkerPipelineJobHeader;
-		TransferFile << PipelineJobHeader;
-
-		int32 NumBatches = QueuedPipelineJobs.Num();
-		TransferFile << NumBatches;
-		for (int32 JobIndex = 0; JobIndex < QueuedPipelineJobs.Num(); JobIndex++)
-		{
-			auto* PipelineJob = QueuedPipelineJobs[JobIndex];
-			FString PipelineName = PipelineJob->Key.ShaderPipeline->GetName();
-			TransferFile << PipelineName;
-			int32 NumStageJobs = PipelineJob->StageJobs.Num();
-			TransferFile << NumStageJobs;
-			for (int32 Index = 0; Index < NumStageJobs; Index++)
-			{
-				PipelineJob->StageJobs[Index]->SerializeWorkerInput(TransferFile);
-				PipelineJob->StageJobs[Index]->Input.SerializeSharedInputs(TransferFile, SharedEnvironments, AllShaderParameterStructures);
-			}
-		}
-	}
-
-	if (bCompressTaskFile)
-	{
-		TransferFile.Close();
-
-		FName CompressionFormatToUse = NAME_LZ4;
-
-		FString FormatName = CompressionFormatToUse.ToString();
-		InTransferFile << FormatName;
-
-		// serialize uncompressed data size
-		int32 UncompressedDataSize = UncompressedArray.Num();
-		checkf(UncompressedDataSize != 0, TEXT("Did not write any data to the task file for the compression."));
-		InTransferFile << UncompressedDataSize;
-
-		// not using SerializeCompressed because it splits into smaller chunks
-		int32 CompressedSizeBound = FCompression::CompressMemoryBound(CompressionFormatToUse, static_cast<int32>(UncompressedDataSize));
-		TArray<uint8> CompressedBuffer;
-		CompressedBuffer.SetNumUninitialized(CompressedSizeBound);
-
-		int32 ActualCompressedSize = CompressedSizeBound;
-		bool bSucceeded = FCompression::CompressMemory(CompressionFormatToUse, CompressedBuffer.GetData(), ActualCompressedSize, UncompressedArray.GetData(), UncompressedDataSize, COMPRESS_BiasSpeed);
-		checkf(ActualCompressedSize <= CompressedSizeBound, TEXT("Compressed size was larger than the bound - we stomped the memory."));
-		CompressedBuffer.SetNum(ActualCompressedSize, EAllowShrinking::No);
-
-		InTransferFile << CompressedBuffer;
-		UE_LOG(LogShaderCompilers, Verbose, TEXT("Compressed the task file from %d bytes to %d bytes (%.2f%% savings)"), UncompressedDataSize, ActualCompressedSize,
-			100.0 * (UncompressedDataSize - ActualCompressedSize) / static_cast<double>(UncompressedDataSize));
-	}
-
-	return InTransferFile.Close();
-}
-
-const TCHAR* DebugWorkerInputFileName = TEXT("DebugSCW.in");
-const TCHAR* DebugWorkerOutputFileName = TEXT("DebugSCW.out");
-
-FString CreateShaderCompilerWorkerDebugCommandLine(FString DebugWorkerInputFilePath)
-{
-	// 0 is parent PID, pass zero TTL and KeepInput to make SCW process the single job then exit without deleting the input file
-	return FString::Printf(TEXT("\"%s\" 0 \"DebugSCW\" %s %s -TimeToLive=0.0f -KeepInput"),
-		*DebugWorkerInputFilePath, // working directory for SCW
-		DebugWorkerInputFileName,
-		DebugWorkerOutputFileName);
-}
-
-static void DumpWorkerInputs(TConstArrayView<FShaderCommonCompileJobPtr> QueuedJobs)
-{
-	if (CVarDebugDumpWorkerInputs.GetValueOnAnyThread())
-	{
-		for (const FShaderCommonCompileJobPtr& CommonJob : QueuedJobs)
-		{
-			FShaderPipelineCompileJob* PipelineJob = CommonJob->GetShaderPipelineJob();
-			FString DebugWorkerInputFilePath;
-			if (PipelineJob)
-			{
-				// for pipeline jobs, write out the worker input for the whole pipeline, but only for the first stage
-				// would be better to put in a parent folder probably...
-				DebugWorkerInputFilePath = PipelineJob->StageJobs[0]->Input.DumpDebugInfoPath;
-			}
-			else
-			{
-				DebugWorkerInputFilePath = CommonJob->GetSingleShaderJob()->Input.DumpDebugInfoPath;
-			}
-			if (!DebugWorkerInputFilePath.IsEmpty())
-			{
-				TArray<FShaderCommonCompileJobPtr> SingleJobArray;
-				SingleJobArray.Add(CommonJob);
-
-				FArchive* DebugWorkerInputFileWriter = IFileManager::Get().CreateFileWriter(*(DebugWorkerInputFilePath / DebugWorkerInputFileName), FILEWRITE_NoFail);
-				DoWriteTasksInner(
-					SingleJobArray,
-					*DebugWorkerInputFileWriter,
-					nullptr,	// Don't pass a IDistributedBuildController, this is only used for conversion to relative paths which we do not want for debug files
-					false,		// As above, use absolute paths not relative
-					true);		// Always compress the debug files; they are rather large so this saves some disk space
-				DebugWorkerInputFileWriter->Close();
-				delete DebugWorkerInputFileWriter;
-
-				FFileHelper::SaveStringToFile(
-					CreateShaderCompilerWorkerDebugCommandLine(DebugWorkerInputFilePath),
-					*(DebugWorkerInputFilePath / TEXT("DebugCompileArgs.txt")));
-			}
-		}
-	}
-}
-
-// Serialize Queued Job information
-bool FShaderCompileUtilities::DoWriteTasks(const TArray<FShaderCommonCompileJobPtr>& QueuedJobs, FArchive& InTransferFile, IDistributedBuildController* BuildDistributionController, bool bUseRelativePaths, bool bCompressTaskFile)
-{
-	DumpWorkerInputs(QueuedJobs);
-
-	return DoWriteTasksInner(QueuedJobs, InTransferFile, BuildDistributionController, bUseRelativePaths, bCompressTaskFile);
-}
-
-static void BuildErrorStringAndReport(const FShaderDiagnosticInfo& DiagInfo, FString& ErrorString)
-{
-	bool bReportedDebugInfo = false;
-
-	for (int32 ErrorIndex = 0; ErrorIndex < DiagInfo.UniqueErrors.Num(); ErrorIndex++)
-	{
-		FString UniqueErrorString = DiagInfo.UniqueErrors[ErrorIndex] + TEXT("\n");
-
-		if (FPlatformMisc::IsDebuggerPresent())
-		{
-			// Using OutputDebugString to avoid any text getting added before the filename,
-			// Which will throw off VS.NET's ability to take you directly to the file and line of the error when double clicking it in the output window.
-			FPlatformMisc::LowLevelOutputDebugStringf(*UniqueErrorString);
-		}
-		else
-		{
-			UE_LOG(LogShaderCompilers, Warning, TEXT("%s"), *UniqueErrorString);
-		}
-
-		ErrorString += UniqueErrorString;
-	}
-}
-
-static bool ReadSingleJob(FShaderCompileJob* CurrentJob, FArchive& WorkerOutputFileReader)
-{
-	TRACE_CPUPROFILER_EVENT_SCOPE(ReadSingleJob);
-
-	check(!CurrentJob->bFinalized);
-	CurrentJob->bFinalized = true;
-
-	// Deserialize the shader compilation output.
-	CurrentJob->SerializeWorkerOutput(WorkerOutputFileReader);
-
-	// The job should already have a non-zero output hash
-	checkf(CurrentJob->Output.OutputHash != FSHAHash() || !CurrentJob->bSucceeded, TEXT("OutputHash for a successful job was not set in the shader compile worker!"));
-
-
-	// Support dumping debug info for only failed compilations or those with warnings
-	if (GShaderCompilingManager->ShouldRecompileToDumpShaderDebugInfo(*CurrentJob))
-	{
-		// Build debug info path and create the directory if it doesn't already exist
-		CurrentJob->Input.DumpDebugInfoPath = GShaderCompilingManager->CreateShaderDebugInfoPath(CurrentJob->Input);
-
-		// We failed to compile a shader, we will retry and dump the shader source. so increment the number of shader sources we've dumped so far.
-		GShaderCompilingManager->IncrementNumDumpedShaderSources();
-
-		return true;
-	}
-
-	return false;
-}
-
-static const TCHAR* GetCompileJobSuccessText(FShaderCompileJob* SingleJob)
-{
-	if (SingleJob)
-	{
-		return SingleJob->Output.bSucceeded ? TEXT("Succeeded") : TEXT("Failed");
-	}
-	return TEXT("");
-}
-
-static void LogQueuedCompileJobs(const TArray<FShaderCommonCompileJobPtr>& QueuedJobs, int32 NumProcessedJobs)
-{
-	if (NumProcessedJobs == -1)
-	{
-		UE_LOG(LogShaderCompilers, Error, TEXT("SCW %d Queued Jobs, Unknown number of processed jobs!"), QueuedJobs.Num());
-	}
-	else
-	{
-		UE_LOG(LogShaderCompilers, Error, TEXT("SCW %d Queued Jobs, Finished %d single jobs"), QueuedJobs.Num(), NumProcessedJobs);
-	}
-
-	for (int32 Index = 0; Index < QueuedJobs.Num(); ++Index)
-	{
-		if (FShaderCompileJob* SingleJob = QueuedJobs[Index]->GetSingleShaderJob())
-		{
-			UE_LOG(LogShaderCompilers, Error, TEXT("Job %d [Single] %s: %s"), Index, GetCompileJobSuccessText(SingleJob), *GetSingleJobCompilationDump(SingleJob));
-		}
-		else
-		{
-			FShaderPipelineCompileJob* PipelineJob = QueuedJobs[Index]->GetShaderPipelineJob();
-			UE_LOG(LogShaderCompilers, Error, TEXT("Job %d: Pipeline %s "), Index, PipelineJob->Key.ShaderPipeline->GetName());
-			for (int32 JobIndex = 0; JobIndex < PipelineJob->StageJobs.Num(); ++JobIndex)
-			{
-				FShaderCompileJob* StageJob = PipelineJob->StageJobs[JobIndex]->GetSingleShaderJob();
-				UE_LOG(LogShaderCompilers, Error, TEXT("PipelineJob %d %s: %s"), JobIndex, GetCompileJobSuccessText(StageJob), *GetSingleJobCompilationDump(StageJob));
-			}
-		}
-	}
-
-	// Force a log flush so we can track the crash before the cooker potentially crashes before the output shows up
-	GLog->Flush();
-}
-
-// Disable optimization for this crash handler to get full access to the entire stack frame when debugging a crash dump
-UE_DISABLE_OPTIMIZATION_SHIP
-static bool HandleWorkerCrash(const TArray<FShaderCommonCompileJobPtr>& QueuedJobs, FArchive& OutputFile, int32 OutputVersion, int64 FileSize, FSCWErrorCode::ECode ErrorCode, int32 NumProcessedJobs, int32 CallstackLength, int32 ExceptionInfoLength, int32 HostnameLength)
-{
-	TArray<TCHAR> Callstack;
-	Callstack.AddUninitialized(CallstackLength + 1);
-	OutputFile.Serialize(Callstack.GetData(), CallstackLength * sizeof(TCHAR));
-	Callstack[CallstackLength] = TEXT('\0');
-
-	TArray<TCHAR> ExceptionInfo;
-	ExceptionInfo.AddUninitialized(ExceptionInfoLength + 1);
-	OutputFile.Serialize(ExceptionInfo.GetData(), ExceptionInfoLength * sizeof(TCHAR));
-	ExceptionInfo[ExceptionInfoLength] = TEXT('\0');
-
-	TArray<TCHAR> Hostname;
-	Hostname.AddUninitialized(HostnameLength + 1);
-	OutputFile.Serialize(Hostname.GetData(), HostnameLength * sizeof(TCHAR));
-	Hostname[HostnameLength] = TEXT('\0');
-
-	// Read available and used physical memory from worker machine on OOM error
-	FPlatformMemoryStats MemoryStats;
-	if (ErrorCode == FSCWErrorCode::OutOfMemory)
-	{
-		OutputFile
-			<< MemoryStats.AvailablePhysical
-			<< MemoryStats.AvailableVirtual
-			<< MemoryStats.UsedPhysical
-			<< MemoryStats.PeakUsedPhysical
-			<< MemoryStats.UsedVirtual
-			<< MemoryStats.PeakUsedVirtual
-			;
-	}
-
-	// Store primary job information onto stack to make it part of a crash dump
-	static const int32 MaxNumCharsForSourcePaths = 8192;
-	int32 JobInputSourcePathsLength = 0;
-	ANSICHAR JobInputSourcePaths[MaxNumCharsForSourcePaths];
-	JobInputSourcePaths[0] = 0;
-
-	auto WriteInputSourcePathOntoStack = [&JobInputSourcePathsLength, &JobInputSourcePaths](const ANSICHAR* InputSourcePath)
-	{
-		if (InputSourcePath != nullptr && JobInputSourcePathsLength + 3 < MaxNumCharsForSourcePaths)
-		{
-			// Copy input source path into stack buffer
-			int32 InputSourcePathLength = FMath::Min(FCStringAnsi::Strlen(InputSourcePath), (MaxNumCharsForSourcePaths - JobInputSourcePathsLength - 2));
-			FMemory::Memcpy(JobInputSourcePaths + JobInputSourcePathsLength, InputSourcePath, InputSourcePathLength);
-
-			// Write newline character and put NUL character at the end
-			JobInputSourcePathsLength += InputSourcePathLength;
-			JobInputSourcePaths[JobInputSourcePathsLength] = TEXT('\n');
-			++JobInputSourcePathsLength;
-			JobInputSourcePaths[JobInputSourcePathsLength] = 0;
-		}
-	};
-
-	auto StoreInputDebugInfo = [&WriteInputSourcePathOntoStack, &JobInputSourcePathsLength, &JobInputSourcePaths](const FShaderCompilerInput& Input)
-	{
-		FString DebugInfo = FString::Printf(TEXT("%s:%s"), *Input.VirtualSourceFilePath, *Input.EntryPointName);
-		WriteInputSourcePathOntoStack(TCHAR_TO_UTF8(*DebugInfo));
-	};
-
-	for (auto CommonJob : QueuedJobs)
-	{
-		if (FShaderCompileJob* SingleJob = CommonJob->GetSingleShaderJob())
-		{
-			StoreInputDebugInfo(SingleJob->Input);
-		}
-		else if (FShaderPipelineCompileJob* PipelineJob = CommonJob->GetShaderPipelineJob())
-		{
-			for (int32 Job = 0; Job < PipelineJob->StageJobs.Num(); ++Job)
-			{
-				if (FShaderCompileJob* SingleStageJob = PipelineJob->StageJobs[Job])
-				{
-					StoreInputDebugInfo(SingleStageJob->Input);
-				}
-			}
-		}
-	}
-
-	// One entry per error code as we want to have different callstacks for crash reporter...
-	switch (ErrorCode)
-	{
-	default:
-	case FSCWErrorCode::GeneralCrash:
-		LogQueuedCompileJobs(QueuedJobs, NumProcessedJobs);
-		ShaderCompileWorkerError::HandleGeneralCrash(ExceptionInfo.GetData(), Callstack.GetData());
-		break;
-	case FSCWErrorCode::BadShaderFormatVersion:
-		ShaderCompileWorkerError::HandleBadShaderFormatVersion(ExceptionInfo.GetData());
-		break;
-	case FSCWErrorCode::BadInputVersion:
-		ShaderCompileWorkerError::HandleBadInputVersion(ExceptionInfo.GetData());
-		break;
-	case FSCWErrorCode::BadSingleJobHeader:
-		ShaderCompileWorkerError::HandleBadSingleJobHeader(ExceptionInfo.GetData());
-		break;
-	case FSCWErrorCode::BadPipelineJobHeader:
-		ShaderCompileWorkerError::HandleBadPipelineJobHeader(ExceptionInfo.GetData());
-		break;
-	case FSCWErrorCode::CantDeleteInputFile:
-		ShaderCompileWorkerError::HandleCantDeleteInputFile(ExceptionInfo.GetData());
-		break;
-	case FSCWErrorCode::CantSaveOutputFile:
-		ShaderCompileWorkerError::HandleCantSaveOutputFile(ExceptionInfo.GetData());
-		break;
-	case FSCWErrorCode::NoTargetShaderFormatsFound:
-		ShaderCompileWorkerError::HandleNoTargetShaderFormatsFound(ExceptionInfo.GetData());
-		break;
-	case FSCWErrorCode::CantCompileForSpecificFormat:
-		ShaderCompileWorkerError::HandleCantCompileForSpecificFormat(ExceptionInfo.GetData());
-		break;
-	case FSCWErrorCode::CrashInsidePlatformCompiler:
-		LogQueuedCompileJobs(QueuedJobs, NumProcessedJobs);
-		ShaderCompileWorkerError::HandleCrashInsidePlatformCompiler(ExceptionInfo.GetData());
-		break;
-	case FSCWErrorCode::BadInputFile:
-		ShaderCompileWorkerError::HandleBadInputFile(ExceptionInfo.GetData());
-		break;
-	case FSCWErrorCode::OutOfMemory:
-		return ShaderCompileWorkerError::HandleOutOfMemory(ExceptionInfo.GetData(), Hostname.GetData(), MemoryStats, QueuedJobs);
-	case FSCWErrorCode::Success:
-		// Can't get here...
-		return true;
-	}
-	return false;
-}
-UE_ENABLE_OPTIMIZATION_SHIP
-
-// Helper struct to provide consistent error report with detailed information about corrupted ShaderCompileWorker output file.
-struct FSCWOutputFileContext
-{
-	FArchive& OutputFile;
-	int64 FileSize = 0;
-
-	FSCWOutputFileContext(FArchive& OutputFile) :
-		OutputFile(OutputFile)
-	{
-	}
-
-	template <typename FmtType, typename... Types>
-	void ModalErrorOrLog(const FmtType& Format, Types&&... Args)
-	{
-		FString Text = FString::Printf(Format, Args...);
-		Text = FString::Printf(TEXT("File path: \"%s\"\n%s\nForgot to build ShaderCompileWorker or delete invalidated DerivedDataCache?"), *OutputFile.GetArchiveName(), *Text);
-		const TCHAR* Title = TEXT("Corrupted ShaderCompileWorker output file");
-		if (FileSize > 0)
-		{
-			::ModalErrorOrLog(Title, Text, OutputFile.Tell(), FileSize);
-		}
-		else
-		{
-			::ModalErrorOrLog(Title, Text, 0, 0);
-		}
-	}
-};
-
-// Process results from Worker Process.
-// Returns false if reading the tasks failed but we were able to recover from handing a crash report. In this case, all jobs must be submitted/processed again.
-FSCWErrorCode::ECode FShaderCompileUtilities::DoReadTaskResults(const TArray<FShaderCommonCompileJobPtr>& QueuedJobs, FArchive& OutputFile)
-{
-	FSCWOutputFileContext OutputFileContext(OutputFile);
-
-	if (OutputFile.TotalSize() == 0)
-	{
-		ShaderCompileWorkerError::HandleOutputFileEmpty(*OutputFile.GetArchiveName());
-	}
-
-	int32 OutputVersion = ShaderCompileWorkerOutputVersion;
-	OutputFile << OutputVersion;
-
-	if (ShaderCompileWorkerOutputVersion != OutputVersion)
-	{
-		OutputFileContext.ModalErrorOrLog(TEXT("Expecting output version %d, got %d instead!"), ShaderCompileWorkerOutputVersion, OutputVersion);
-	}
-
-	OutputFile << OutputFileContext.FileSize;
-
-	// Check for corrupted output file
-	if (OutputFileContext.FileSize > OutputFile.TotalSize())
-	{
-		ShaderCompileWorkerError::HandleOutputFileCorrupted(*OutputFile.GetArchiveName(), OutputFileContext.FileSize, OutputFile.TotalSize());
-	}
-
-	int32 ErrorCode = 0;
-	OutputFile << ErrorCode;
-
-	int32 NumProcessedJobs = 0;
-	OutputFile << NumProcessedJobs;
-
-	int32 CallstackLength = 0;
-	OutputFile << CallstackLength;
-
-	int32 ExceptionInfoLength = 0;
-	OutputFile << ExceptionInfoLength;
-
-	int32 HostnameLength = 0;
-	OutputFile << HostnameLength;
-
-	if (ErrorCode != FSCWErrorCode::Success)
-	{
-		// If worker crashed in a way we were able to recover from, return and expect the compile jobs to be reissued already
-		if (HandleWorkerCrash(QueuedJobs, OutputFile, OutputVersion, OutputFileContext.FileSize, (FSCWErrorCode::ECode)ErrorCode, NumProcessedJobs, CallstackLength, ExceptionInfoLength, HostnameLength))
-		{
-			FSCWErrorCode::Reset();
-			return (FSCWErrorCode::ECode)ErrorCode;
-		}
-	}
-
-	TArray<FShaderCompileJob*> QueuedSingleJobs;
-	TArray<FShaderPipelineCompileJob*> QueuedPipelineJobs;
-	SplitJobsByType(QueuedJobs, QueuedSingleJobs, QueuedPipelineJobs);
-	TArray<FShaderCommonCompileJob*> ReissueSourceJobs;
-
-	// Read single jobs
-	{
-		int32 SingleJobHeader = -1;
-		OutputFile << SingleJobHeader;
-		if (SingleJobHeader != ShaderCompileWorkerSingleJobHeader)
-		{
-			OutputFileContext.ModalErrorOrLog(TEXT("Expecting single job header ID 0x%08X, got 0x%08X instead!"), ShaderCompileWorkerSingleJobHeader, SingleJobHeader);
-		}
-
-		int32 NumJobs;
-		OutputFile << NumJobs;
-		if (NumJobs != QueuedSingleJobs.Num())
-		{
-			OutputFileContext.ModalErrorOrLog(TEXT("Expecting %d single %s, got %d instead!"), QueuedSingleJobs.Num(), (QueuedSingleJobs.Num() == 1 ? TEXT("job") : TEXT("jobs")), NumJobs);
-		}
-		else
-		{
-			for (int32 JobIndex = 0; JobIndex < NumJobs; JobIndex++)
-			{
-				auto* CurrentJob = QueuedSingleJobs[JobIndex];
-				if (ReadSingleJob(CurrentJob, OutputFile))
-				{
-					ReissueSourceJobs.Add(CurrentJob);
-				}
-			}
-		}
-	}
-
-	// Pipeline jobs
-	{
-		int32 PipelineJobHeader = -1;
-		OutputFile << PipelineJobHeader;
-		if (PipelineJobHeader != ShaderCompileWorkerPipelineJobHeader)
-		{
-			OutputFileContext.ModalErrorOrLog(TEXT("Expecting pipeline jobs header ID 0x%08X, got 0x%08X instead!"), ShaderCompileWorkerPipelineJobHeader, PipelineJobHeader);
-		}
-
-		int32 NumJobs;
-		OutputFile << NumJobs;
-		if (NumJobs != QueuedPipelineJobs.Num())
-		{
-			OutputFileContext.ModalErrorOrLog(TEXT("Expecting %d pipeline %s, got %d instead!"), QueuedPipelineJobs.Num(), (QueuedPipelineJobs.Num() == 1 ? TEXT("job") : TEXT("jobs")), NumJobs);
-		}
-		else
-		{
-			for (int32 JobIndex = 0; JobIndex < NumJobs; JobIndex++)
-			{
-				FShaderPipelineCompileJob* CurrentJob = QueuedPipelineJobs[JobIndex];
-
-				FString PipelineName;
-				OutputFile << PipelineName;
-				bool bSucceeded = false;
-				OutputFile << bSucceeded;
-				CurrentJob->bSucceeded = bSucceeded;
-				if (PipelineName != CurrentJob->Key.ShaderPipeline->GetName())
-				{
-					OutputFileContext.ModalErrorOrLog(TEXT("Expecting pipeline job \"%s\", got \"%s\" instead!"), CurrentJob->Key.ShaderPipeline->GetName(), *PipelineName);
-				}
-
-				check(!CurrentJob->bFinalized);
-				CurrentJob->bFinalized = true;
-
-				int32 NumStageJobs = -1;
-				OutputFile << NumStageJobs;
-
-				if (NumStageJobs != CurrentJob->StageJobs.Num())
-				{
-					OutputFileContext.ModalErrorOrLog(TEXT("Expecting %d stage pipeline %s, got %d instead!"), CurrentJob->StageJobs.Num(), (CurrentJob->StageJobs.Num() == 1 ? TEXT("job") : TEXT("jobs")), NumStageJobs);
-				}
-				else
-				{
-					for (int32 Index = 0; Index < NumStageJobs; Index++)
-					{
-						FShaderCompileJob* SingleJob = CurrentJob->StageJobs[Index];
-						// cannot reissue a single stage of a pipeline job
-						ReadSingleJob(SingleJob, OutputFile);
-					}
-				}
-			}
-		}
-	}
-	
-	// Requeue any jobs we wish to run again
-	ReissueShaderCompileJobs(ReissueSourceJobs);
-
-	return FSCWErrorCode::Success;
-}
-
-#if WITH_EDITOR
-static bool CheckSingleJob(const FShaderCompileJob& SingleJob, TArray<FString>& OutErrors)
-{
-	if (SingleJob.bSucceeded)
-	{
-		checkf(SingleJob.Output.ShaderCode.GetShaderCodeSize() > 0, TEXT("Abnormal shader code size for a succeded job: %d bytes"), SingleJob.Output.ShaderCode.GetShaderCodeSize());
-	}
-
-	bool bSucceeded = SingleJob.bSucceeded;
-
-	if (SingleJob.Key.ShaderType)
-	{
-		// Allow the shader validation to fail the compile if it sees any parameters bound that aren't supported.
-		const bool bValidationResult = SingleJob.Key.ShaderType->ValidateCompiledResult(
-			(EShaderPlatform)SingleJob.Input.Target.Platform,
-			SingleJob.Output.ParameterMap,
-			OutErrors);
-		bSucceeded = bValidationResult && bSucceeded;
-	}
-
-	if (SingleJob.Key.VFType)
-	{
-		const int32 OriginalNumErrors = OutErrors.Num();
-
-		// Allow the vertex factory to fail the compile if it sees any parameters bound that aren't supported
-		SingleJob.Key.VFType->ValidateCompiledResult((EShaderPlatform)SingleJob.Input.Target.Platform, SingleJob.Output.ParameterMap, OutErrors);
-
-		if (OutErrors.Num() > OriginalNumErrors)
-		{
-			bSucceeded = false;
-		}
-	}
-
-	return bSucceeded;
-};
-#endif // WITH_EDITOR
-
-/** Information tracked for each shader compile worker process instance. */
-struct FShaderCompileWorkerInfo
-{
-	/** Process handle of the worker app once launched.  Invalid handle means no process. */
-	FProcHandle WorkerProcess;
-
-	/** Tracks whether tasks have been issued to the worker. */
-	bool bIssuedTasksToWorker;	
-
-	/** Whether the worker has been launched for this set of tasks. */
-	bool bLaunchedWorker;
-
-	/** Tracks whether all tasks issued to the worker have been received. */
-	bool bComplete;
-
-	/** Whether this worker is available for new jobs. It will be false when shutting down the worker. */
-	bool bAvailable; 
-
-	/** Time at which the worker started the most recent batch of tasks. */
-	double StartTime;
-
-	/** Time at which the worker ended the most recent batch of tasks. */
-	double FinishTime = 0.0;
-
-	/** Jobs that this worker is responsible for compiling. */
-	TArray<FShaderCommonCompileJobPtr> QueuedJobs;
-
-	FShaderCompileWorkerInfo() :
-		bIssuedTasksToWorker(false),		
-		bLaunchedWorker(false),
-		bComplete(false),
-		bAvailable(true),
-		StartTime(0)
-	{
-	}
-
-	// warning: not virtual
-	~FShaderCompileWorkerInfo()
-	{
-		if(WorkerProcess.IsValid())
-		{
-			FPlatformProcess::TerminateProc(WorkerProcess);
-			FPlatformProcess::CloseProc(WorkerProcess);
-		}
-	}
-};
-
-FShaderCompileThreadRunnableBase::FShaderCompileThreadRunnableBase(FShaderCompilingManager* InManager)
-	: Manager(InManager)
-	, Thread(nullptr)
-	, MinPriorityIndex(0)
-	, MaxPriorityIndex(NumShaderCompileJobPriorities - 1)
-	, bForceFinish(false)
-{
-}
-void FShaderCompileThreadRunnableBase::StartThread()
-{
-	if (Manager->bAllowAsynchronousShaderCompiling && !FPlatformProperties::RequiresCookedData())
-	{
-		Thread = FRunnableThread::Create(this, GetThreadName(), 0, TPri_Normal, FPlatformAffinity::GetPoolThreadMask());
-	}
-}
-
-FShaderCompileThreadRunnable::FShaderCompileThreadRunnable(FShaderCompilingManager* InManager)
-	: FShaderCompileThreadRunnableBase(InManager)
-	, LastCheckForWorkersTime(0)
-{
-	for (uint32 WorkerIndex = 0; WorkerIndex < Manager->NumShaderCompilingThreads; WorkerIndex++)
-	{
-		WorkerInfos.Add(MakeUnique<FShaderCompileWorkerInfo>());
-	}
-}
-
-FShaderCompileThreadRunnable::~FShaderCompileThreadRunnable()
-{
-	FScopeLock WorkerScopeLock(&WorkerInfosLock);
-	WorkerInfos.Empty();
-}
-
-void FShaderCompileThreadRunnable::OnMachineResourcesChanged()
-{
-	bool bWaitForWorkersToShutdown = false;
-	{
-		FScopeLock WorkerScopeLock(&WorkerInfosLock);
-		// Set all bAvailable flags back to true
-		for (TUniquePtr< FShaderCompileWorkerInfo>& WorkerInfo : WorkerInfos)
-		{
-			WorkerInfo->bAvailable = true;
-		}
-
-		if (Manager->NumShaderCompilingThreads >= static_cast<uint32>(WorkerInfos.Num()))
-		{
-			while (static_cast<uint32>(WorkerInfos.Num()) < Manager->NumShaderCompilingThreads)
-			{
-				WorkerInfos.Add(MakeUnique<FShaderCompileWorkerInfo>());
-			}
-		}
-		else
-		{
-			for (int32 Index = 0; Index < WorkerInfos.Num(); ++Index)
-			{
-				FShaderCompileWorkerInfo& WorkerInfo = *WorkerInfos[Index];
-				bool bReadyForShutdown = WorkerInfo.QueuedJobs.Num() == 0;
-				if (bReadyForShutdown)
-				{
-					WorkerInfos.RemoveAtSwap(Index--);
-					if (WorkerInfos.Num() == Manager->NumShaderCompilingThreads)
-					{
-						break;
-					}
-				}
-			}
-			bWaitForWorkersToShutdown = Manager->NumShaderCompilingThreads < static_cast<uint32>(WorkerInfos.Num());
-			for (int32 Index = WorkerInfos.Num() - 1;
-				static_cast<uint32>(Index) >= Manager->NumShaderCompilingThreads; --Index)
-			{
-				WorkerInfos[Index]->bAvailable = false;
-			}
-		}
-	}
-	const double StartTime = FPlatformTime::Seconds();
-	constexpr float MaxDurationToWait = 60.f;
-	const double MaxTimeToWait = StartTime + MaxDurationToWait;
-	while (bWaitForWorkersToShutdown)
-	{
-		FPlatformProcess::Sleep(0.01f);
-		const double CurrentTime = FPlatformTime::Seconds();
-		if (CurrentTime > MaxTimeToWait)
-		{
-			UE_LOG(LogShaderCompilers, Warning, TEXT("OnMachineResourcesChanged timedout waiting %.0f seconds for WorkerInfos to complete. Workers will remain allocated."),
-				(float)(CurrentTime - StartTime));
-			break;
-		}
-
-		FScopeLock WorkerScopeLock(&WorkerInfosLock);
-		for (int32 Index = WorkerInfos.Num() - 1;
-			static_cast<uint32>(Index) >= Manager->NumShaderCompilingThreads; --Index)
-		{
-			FShaderCompileWorkerInfo& WorkerInfo = *WorkerInfos[Index];
-			check(!WorkerInfos[Index]->bAvailable); // It should still be set to false from when we changed it above
-			bool bReadyForShutdown = WorkerInfo.QueuedJobs.Num() == 0;
-			if (bReadyForShutdown)
-			{
-				WorkerInfos.RemoveAtSwap(Index);
-			}
-		}
-		bWaitForWorkersToShutdown = Manager->NumShaderCompilingThreads < static_cast<uint32>(WorkerInfos.Num());
-	}
-}
-
-/** Entry point for the shader compiling thread. */
-uint32 FShaderCompileThreadRunnableBase::Run()
-{
-	LLM_SCOPE_BYTAG(ShaderCompiler);
-	check(Manager->bAllowAsynchronousShaderCompiling);
-	while (!bForceFinish)
-	{
-		CompilingLoop();
-	}
-	UE_LOG(LogShaderCompilers, Display, TEXT("Shaders left to compile 0"));
-
-	return 0;
-}
-
-int32 FShaderCompileThreadRunnable::PullTasksFromQueue()
-{
-	TRACE_CPUPROFILER_EVENT_SCOPE(FShaderCompileThreadRunnable::PullTasksFromQueue);
-
-	FScopeLock WorkerScopeLock(&WorkerInfosLock); // Must be entered before CompileQueueSection
-	int32 NumActiveThreads = 0;
-	int32 NumJobsStarted[NumShaderCompileJobPriorities] = { 0 };
-	{
-		// Enter the critical section so we can access the input and output queues
-		FScopeLock Lock(&Manager->CompileQueueSection);
-
-		const int32 NumWorkersToFeed = Manager->bCompilingDuringGame ? Manager->NumShaderCompilingThreadsDuringGame : WorkerInfos.Num();
-
-		for (int32 PriorityIndex = MaxPriorityIndex; PriorityIndex >= MinPriorityIndex; --PriorityIndex)
-		{
-			int32 NumPendingJobs = Manager->AllJobs.GetNumPendingJobs((EShaderCompileJobPriority)PriorityIndex);
-			// Try to distribute the work evenly between the workers
-			const auto NumJobsPerWorker = (NumPendingJobs / NumWorkersToFeed) + 1;
-
-			for (int32 WorkerIndex = 0; WorkerIndex < WorkerInfos.Num(); WorkerIndex++)
-			{
-				FShaderCompileWorkerInfo& CurrentWorkerInfo = *WorkerInfos[WorkerIndex];
-
-				// If this worker doesn't have any queued jobs, look for more in the input queue
-				if (CurrentWorkerInfo.QueuedJobs.Num() == 0 && CurrentWorkerInfo.bAvailable && WorkerIndex < NumWorkersToFeed)
-				{
-					check(!CurrentWorkerInfo.bComplete);
-
-					NumPendingJobs = Manager->AllJobs.GetNumPendingJobs((EShaderCompileJobPriority)PriorityIndex);
-					if (NumPendingJobs > 0)
-					{
-						UE_LOG(LogShaderCompilers, Verbose, TEXT("Worker (%d/%d): shaders left to compile %i"), WorkerIndex + 1, WorkerInfos.Num(), NumPendingJobs);
-
-						int32 MaxNumJobs = 1;
-						// high priority jobs go in 1 per "batch", unless the engine is still starting up
-						if (PriorityIndex < (int32)EShaderCompileJobPriority::High || Manager->IgnoreAllThrottling())
-						{
-							MaxNumJobs = FMath::Min3(NumJobsPerWorker, NumPendingJobs, Manager->MaxShaderJobBatchSize);
-						}
-
-						NumJobsStarted[PriorityIndex] += Manager->AllJobs.GetPendingJobs(EShaderCompilerWorkerType::LocalThread, (EShaderCompileJobPriority)PriorityIndex, 1, MaxNumJobs, CurrentWorkerInfo.QueuedJobs);
-
-						// Update the worker state as having new tasks that need to be issued					
-						// don't reset worker app ID, because the shadercompileworkers don't shutdown immediately after finishing a single job queue.
-						CurrentWorkerInfo.bIssuedTasksToWorker = false;
-						CurrentWorkerInfo.bLaunchedWorker = false;
-						CurrentWorkerInfo.StartTime = FPlatformTime::Seconds();
-						NumActiveThreads++;
-
-						if (CurrentWorkerInfo.FinishTime > 0.0)
-						{
-							const double WorkerIdleTime = CurrentWorkerInfo.StartTime - CurrentWorkerInfo.FinishTime;
-							GShaderCompilerStats->RegisterLocalWorkerIdleTime(WorkerIdleTime);
-							if (Manager->bLogJobCompletionTimes)
-							{
-								UE_LOG(LogShaderCompilers, Display, TEXT("  Worker (%d/%d) started working after being idle for %fs"), WorkerIndex + 1, WorkerInfos.Num(), WorkerIdleTime);
-							}
-						}
-					}
-				}
-			}
-		}
-	}
-
-	for (int32 WorkerIndex = 0; WorkerIndex < WorkerInfos.Num(); WorkerIndex++)
-	{
-		if (WorkerInfos[WorkerIndex]->QueuedJobs.Num() > 0)
-		{
-			NumActiveThreads++;
-		}
-	}
-
-	for (int32 PriorityIndex = 0; PriorityIndex < NumShaderCompileJobPriorities; ++PriorityIndex)
-	{
-		if (NumJobsStarted[PriorityIndex] > 0)
-		{
-			UE_LOG(LogShaderCompilers, Verbose, TEXT("Started %d 'Local' shader compile jobs with '%s' priority"),
-				NumJobsStarted[PriorityIndex],
-				ShaderCompileJobPriorityToString((EShaderCompileJobPriority)PriorityIndex));
-		}
-	}
-
-	return NumActiveThreads;
-}
-
-void FShaderCompileThreadRunnable::PushCompletedJobsToManager()
-{
-	FScopeLock WorkerScopeLock(&WorkerInfosLock); // Must be entered before CompileQueueSection
-
-	for (int32 WorkerIndex = 0; WorkerIndex < WorkerInfos.Num(); WorkerIndex++)
-	{
-		FShaderCompileWorkerInfo& CurrentWorkerInfo = *WorkerInfos[WorkerIndex];
-
-		// Add completed jobs to the output queue, which is ShaderMapJobs
-		if (CurrentWorkerInfo.bComplete)
-		{
-			// Enter the critical section so we can access the input and output queues
-			FScopeLock Lock(&Manager->CompileQueueSection);
-
-			for (int32 JobIndex = 0; JobIndex < CurrentWorkerInfo.QueuedJobs.Num(); JobIndex++)
-			{
-				auto& Job = CurrentWorkerInfo.QueuedJobs[JobIndex];
-
-				Manager->ProcessFinishedJob(Job.GetReference());
-			}
-
-			const float ElapsedTime = FPlatformTime::Seconds() - CurrentWorkerInfo.StartTime;
-
-			Manager->WorkersBusyTime += ElapsedTime;
-			COOK_STAT(ShaderCompilerCookStats::AsyncCompileTimeSec += ElapsedTime);
-
-			// Log if requested or if there was an exceptionally slow batch, to see the offender easily
-			if (Manager->bLogJobCompletionTimes || ElapsedTime > 60.0f)
-			{
-				TArray<FShaderCommonCompileJobPtr> SortedJobs = CurrentWorkerInfo.QueuedJobs;
-				SortedJobs.Sort([](const FShaderCommonCompileJobPtr& JobA, const FShaderCommonCompileJobPtr& JobB)
-					{
-						const FShaderCompileJob* SingleJobA = JobA->GetSingleShaderJob();
-						const FShaderCompileJob* SingleJobB = JobB->GetSingleShaderJob();
-
-						const float TimeA = SingleJobA ? SingleJobA->Output.CompileTime : 0.0f;
-						const float TimeB = SingleJobB ? SingleJobB->Output.CompileTime : 0.0f;
-
-						return TimeA > TimeB;
-					});
-
-				FString JobNames;
-
-				for (int32 JobIndex = 0; JobIndex < SortedJobs.Num(); JobIndex++)
-				{
-					const FShaderCommonCompileJob& Job = *SortedJobs[JobIndex];
-					if (const FShaderCompileJob* SingleJob = Job.GetSingleShaderJob())
-					{
-						const TCHAR* JobName = Manager->bLogJobCompletionTimes ? *SingleJob->Input.DebugGroupName : SingleJob->Key.ShaderType->GetName();
-						JobNames += FString::Printf(TEXT("%s:%d(%s) [WorkerTime=%.3fs]"), JobName, SingleJob->Key.PermutationId, *SingleJob->Input.ShaderFormat.ToString(), SingleJob->Output.CompileTime);
-					}
-					else
-					{
-						const FShaderPipelineCompileJob* PipelineJob = Job.GetShaderPipelineJob();
-						JobNames += FString(PipelineJob->Key.ShaderPipeline->GetName());
-					}
-					if (JobIndex < SortedJobs.Num() - 1)
-					{
-						JobNames += TEXT(", ");
-					}
-				}
-
-				UE_LOG(LogShaderCompilers, Display, TEXT("Worker (%d/%d) finished batch of %u jobs in %.3fs, %s"), WorkerIndex + 1, WorkerInfos.Num(), SortedJobs.Num(), ElapsedTime, *JobNames);
-			}
-
-			CurrentWorkerInfo.FinishTime = FPlatformTime::Seconds();
-			CurrentWorkerInfo.bComplete = false;
-			CurrentWorkerInfo.QueuedJobs.Empty();
-		}
-	}
-}
-
-void FShaderCompileThreadRunnable::WriteNewTasks()
-{
-	TRACE_CPUPROFILER_EVENT_SCOPE(ShaderCompiler.WriteNewTasks);
-	FScopeLock WorkerScopeLock(&WorkerInfosLock);
-
-	// first, a quick check if anything is needed just to avoid hammering the task graph
-	bool bHasTasksToWrite = false;
-	for (int32 WorkerIndex = 0, NumWorkers = WorkerInfos.Num(); WorkerIndex < NumWorkers; ++WorkerIndex)
-	{
-		FShaderCompileWorkerInfo& CurrentWorkerInfo = *WorkerInfos[WorkerIndex];
-		if (!CurrentWorkerInfo.bIssuedTasksToWorker && CurrentWorkerInfo.QueuedJobs.Num() > 0)
-		{
-			bHasTasksToWrite = true;
-			break;
-		}
-	}
-
-	if (!bHasTasksToWrite)
-	{
-		return;
-	}
-
-
-	auto LoopBody = [this](int32 WorkerIndex)
-	{
-		// The calling thread holds the WorkerInfosLock and will not modify WorkerInfos, 
-		// so we can access it here without entering the lock
-		FShaderCompileWorkerInfo& CurrentWorkerInfo = *WorkerInfos[WorkerIndex];
-
-		// Only write tasks once
-		if (!CurrentWorkerInfo.bIssuedTasksToWorker && CurrentWorkerInfo.QueuedJobs.Num() > 0)
-		{
-			TRACE_CPUPROFILER_EVENT_SCOPE(ShaderCompiler.WriteNewTasksForWorker);
-			CurrentWorkerInfo.bIssuedTasksToWorker = true;
-
-			const FString WorkingDirectory = Manager->AbsoluteShaderBaseWorkingDirectory + FString::FromInt(WorkerIndex);
-
-			// To make sure that the process waiting for input file won't try to read it until it's ready
-			// we use a temp file name during writing.
-			FString TransferFileName;
-			do
-			{
-				FGuid Guid;
-				FPlatformMisc::CreateGuid(Guid);
-				TransferFileName = WorkingDirectory + Guid.ToString();
-			} while (IFileManager::Get().FileSize(*TransferFileName) != INDEX_NONE);
-
-			// Write out the file that the worker app is waiting for, which has all the information needed to compile the shader.
-			// 'Only' indicates that the worker should keep checking for more tasks after this one
-			FArchive* TransferFile = nullptr;
-
-			int32 RetryCount = 0;
-			// Retry over the next two seconds if we can't write out the input file
-			// Anti-virus and indexing applications can interfere and cause this write to fail
-			//@todo - switch to shared memory or some other method without these unpredictable hazards
-			while (TransferFile == nullptr && RetryCount < 2000)
-			{
-				if (RetryCount > 0)
-				{
-					FPlatformProcess::Sleep(0.01f);
-				}
-				TransferFile = IFileManager::Get().CreateFileWriter(*TransferFileName, FILEWRITE_EvenIfReadOnly);
-				RetryCount++;
-				if (TransferFile == nullptr)
-				{
-					UE_LOG(LogShaderCompilers, Warning, TEXT("Could not create the shader compiler transfer file '%s', retrying..."), *TransferFileName);
-				}
-			}
-			if (TransferFile == nullptr)
-			{
-				UE_LOG(LogShaderCompilers, Fatal, TEXT("Could not create the shader compiler transfer file '%s'."), *TransferFileName);
-			}
-			check(TransferFile);
-
-			GShaderCompilerStats->RegisterJobBatch(CurrentWorkerInfo.QueuedJobs.Num(), FShaderCompilerStats::EExecutionType::Local);
-			if (!FShaderCompileUtilities::DoWriteTasks(CurrentWorkerInfo.QueuedJobs, *TransferFile))
-			{
-				uint64 TotalDiskSpace = 0;
-				uint64 FreeDiskSpace = 0;
-				FPlatformMisc::GetDiskTotalAndFreeSpace(TransferFileName, TotalDiskSpace, FreeDiskSpace);
-				UE_LOG(LogShaderCompilers, Error, TEXT("Could not write the shader compiler transfer filename to '%s' (Free Disk Space: %llu."), *TransferFileName, FreeDiskSpace);
-			}
-			delete TransferFile;
-
-#if 0 // debugging code to dump the worker inputs
-			static FCriticalSection ArchiveLock;
-			{
-				FScopeLock Locker(&ArchiveLock);
-				static int ArchivedTransferFileNum = 0;
-				FString JobCacheDir = ShaderCompiler::IsJobCacheEnabled() ? TEXT("JobCache") : TEXT("NoJobCache");
-				FString ArchiveDir = FPaths::ProjectSavedDir() / TEXT("ArchivedWorkerInputs") / JobCacheDir;
-				FString ArchiveName = FString::Printf(TEXT("Input-%d"), ArchivedTransferFileNum++);
-				FString ArchivePath = ArchiveDir / ArchiveName;
-				if (!IFileManager::Get().Copy(*ArchivePath, *TransferFileName))
-				{
-					UE_LOG(LogInit, Error, TEXT("Could not copy file %s to %s"), *TransferFileName, *ArchivePath);
-					ensure(false);
-				}
-			}
-#endif
-
-			// Change the transfer file name to proper one
-			FString ProperTransferFileName = WorkingDirectory / TEXT("WorkerInputOnly.in");
-			if (!IFileManager::Get().Move(*ProperTransferFileName, *TransferFileName))
-			{
-				uint64 TotalDiskSpace = 0;
-				uint64 FreeDiskSpace = 0;
-				FPlatformMisc::GetDiskTotalAndFreeSpace(TransferFileName, TotalDiskSpace, FreeDiskSpace);
-				UE_LOG(LogShaderCompilers, Error, TEXT("Could not rename the shader compiler transfer filename to '%s' from '%s' (Free Disk Space: %llu)."), *ProperTransferFileName, *TransferFileName, FreeDiskSpace);
-			}
-		}
-	};
-
-	if (bParallelizeIO)
-	{
-		ParallelFor( TEXT("ShaderCompiler.WriteNewTasks.PF"), WorkerInfos.Num(),1, LoopBody, EParallelForFlags::Unbalanced);
-	}
-	else
-	{
-		double StartIOWork = FPlatformTime::Seconds();
-		for (int32 WorkerIndex = 0; WorkerIndex < WorkerInfos.Num(); WorkerIndex++)
-		{
-			LoopBody(WorkerIndex);
-		}
-
-		double IODuration = FPlatformTime::Seconds() - StartIOWork;
-		if (IODuration > GShaderCompilerTooLongIOThresholdSeconds)
-		{
-			UE_LOG(LogShaderCompilers, Display, TEXT("FShaderCompileThreadRunnable::WriteNewTasks()() took too long (%.3f seconds, threshold is %.3f s), will parallelize next time."), IODuration, GShaderCompilerTooLongIOThresholdSeconds);
-			bParallelizeIO = true;
-		}
-	}
-}
-
-bool FShaderCompileThreadRunnable::LaunchWorkersIfNeeded()
-{
-	TRACE_CPUPROFILER_EVENT_SCOPE(FShaderCompileThreadRunnable::LaunchWorkersIfNeeded);
-
-	const double CurrentTime = FPlatformTime::Seconds();
-	// Limit how often we check for workers running since IsApplicationRunning eats up some CPU time on Windows
-	const bool bCheckForWorkerRunning = (CurrentTime - LastCheckForWorkersTime > .1f);
-	bool bAbandonWorkers = false;
-	uint32_t NumberLaunched = 0;
-
-	if (bCheckForWorkerRunning)
-	{
-		LastCheckForWorkersTime = CurrentTime;
-	}
-
-	bool bClearStaleOutputs = Manager->bClearStaleWorkerOutputs;
-	Manager->bClearStaleWorkerOutputs = false;
-
-	FScopeLock WorkerScopeLock(&WorkerInfosLock);
-	for (int32 WorkerIndex = 0; WorkerIndex < WorkerInfos.Num(); WorkerIndex++)
-	{
-		FShaderCompileWorkerInfo& CurrentWorkerInfo = *WorkerInfos[WorkerIndex];
-		if (CurrentWorkerInfo.QueuedJobs.Num() == 0)
-		{
-			// Skip if nothing to do
-			// Also, use the opportunity to free OS resources by cleaning up handles of no more running processes
-			if (CurrentWorkerInfo.WorkerProcess.IsValid() && !FShaderCompilingManager::IsShaderCompilerWorkerRunning(CurrentWorkerInfo.WorkerProcess))
-			{
-				FPlatformProcess::CloseProc(CurrentWorkerInfo.WorkerProcess);
-				CurrentWorkerInfo.WorkerProcess = FProcHandle();
-			}
-			continue;
-		}
-
-		if (!CurrentWorkerInfo.WorkerProcess.IsValid() || (bCheckForWorkerRunning && !FShaderCompilingManager::IsShaderCompilerWorkerRunning(CurrentWorkerInfo.WorkerProcess)))
-		{
-			TRACE_CPUPROFILER_EVENT_SCOPE(FShaderCompileThreadRunnable::LaunchingWorkers);
-
-			// @TODO: dubious design - worker should not be launched unless we know there's more work to do.
-			bool bLaunchAgain = true;
-
-			// Detect when the worker has exited due to fatal error
-			// bLaunchedWorker check here is necessary to distinguish between 'process isn't running because it crashed' and 'process isn't running because it exited cleanly and the outputfile was already consumed'
-			if (CurrentWorkerInfo.WorkerProcess.IsValid())
-			{
-				// shader compiler exited one way or another, so clear out the stale PID.
-				int32 ReturnCode = 0;
-				FPlatformProcess::GetProcReturnCode(CurrentWorkerInfo.WorkerProcess, &ReturnCode);
-				FPlatformProcess::CloseProc(CurrentWorkerInfo.WorkerProcess);
-				CurrentWorkerInfo.WorkerProcess = FProcHandle();
-
-				if (CurrentWorkerInfo.bLaunchedWorker)
-				{
-					const FString WorkingDirectory = Manager->AbsoluteShaderBaseWorkingDirectory + FString::FromInt(WorkerIndex) + TEXT("/");
-					const FString OutputFileNameAndPath = WorkingDirectory + TEXT("WorkerOutputOnly.out");
-
-					if (FPlatformFileManager::Get().GetPlatformFile().FileExists(*OutputFileNameAndPath))
-					{
-						// If the worker is no longer running but it successfully wrote out the output, no need to assert
-						bLaunchAgain = false;
-					}
-					else
-					{
-						UE_LOG(LogShaderCompilers, Error, TEXT("ShaderCompileWorker terminated unexpectedly, return code %d! Falling back to directly compiling which will be very slow.  Thread %u."), ReturnCode, WorkerIndex);
-						LogQueuedCompileJobs(CurrentWorkerInfo.QueuedJobs, -1);
-
-						bAbandonWorkers = true;
-						break;
-					}
-				}
-			}
-
-			if (bLaunchAgain)
-			{
-				const FString WorkingDirectory = Manager->ShaderBaseWorkingDirectory + FString::FromInt(WorkerIndex) + TEXT("/");
-				FString InputFileName(TEXT("WorkerInputOnly.in"));
-				FString OutputFileName(TEXT("WorkerOutputOnly.out"));
-				
-				if (bClearStaleOutputs)
-				{
-					// Delete any potential stale output files; these can persist if we had previously abandoned running with workers due to an unexpected SCW termination.
-					IFileManager::Get().Delete(*(WorkingDirectory / OutputFileName));
-				}
-
-				// Store the handle with this thread so that we will know not to launch it again
-				CurrentWorkerInfo.WorkerProcess = Manager->LaunchWorker(WorkingDirectory, Manager->ProcessId, WorkerIndex, InputFileName, OutputFileName);
-				CurrentWorkerInfo.bLaunchedWorker = true;
-
-				NumberLaunched++;
-			}
-		}
-	}
-
-	const double FinishTime = FPlatformTime::Seconds();
-	if (NumberLaunched > 0 && (FinishTime - CurrentTime) >= 10.0)
-	{
-		UE_LOG(LogShaderCompilers, Warning, TEXT("Performance Warning: It took %f seconds to launch %d ShaderCompileWorkers"), FinishTime - CurrentTime, NumberLaunched);
-	}
-
-	return bAbandonWorkers;
-}
-
-int32 FShaderCompileThreadRunnable::ReadAvailableResults()
-{
-	TRACE_CPUPROFILER_EVENT_SCOPE(ShaderCompiler.ReadAvailableResults);
-	int32 NumProcessed = 0;
-	FScopeLock WorkerScopeLock(&WorkerInfosLock);
-
-	// first, a quick check if anything is needed just to avoid hammering the task graph
-	bool bHasQueuedJobs = false;
-	for (int32 WorkerIndex = 0, NumWorkers = WorkerInfos.Num(); WorkerIndex < NumWorkers; ++WorkerIndex)
-	{
-		if (WorkerInfos[WorkerIndex]->QueuedJobs.Num() > 0)
-		{
-			bHasQueuedJobs = true;
-			break;
-		}
-	}
-
-	if (!bHasQueuedJobs)
-	{
-		return NumProcessed;
-	}
-
-	auto LoopBody = [this, &NumProcessed](int32 WorkerIndex)
-	{
-		// The calling thread holds the WorkerInfosLock and will not modify WorkerInfos, 
-		// so we can access it here without entering the lock
-		FShaderCompileWorkerInfo& CurrentWorkerInfo = *WorkerInfos[WorkerIndex];
-
-		// Check for available result files
-		if (CurrentWorkerInfo.QueuedJobs.Num() > 0)
-		{
-			// Distributed compiles always use the same directory
-			// 'Only' indicates to the worker that it should log and continue checking for the input file after the first one is processed
-			TStringBuilder<512> OutputFileNameAndPath;
-			OutputFileNameAndPath << Manager->AbsoluteShaderBaseWorkingDirectory << WorkerIndex << TEXT("/WorkerOutputOnly.out");
-
-			// In the common case the output file will not exist, so check for existence before opening
-			// This is only a win if FileExists is faster than CreateFileReader, which it is on Windows
-			if (FPlatformFileManager::Get().GetPlatformFile().FileExists(*OutputFileNameAndPath))
-			{
-				TRACE_CPUPROFILER_EVENT_SCOPE(FShaderCompileThreadRunnable::ProcessOutputFile);
-
-				if (TUniquePtr<FArchive> OutputFile = TUniquePtr<FArchive>(IFileManager::Get().CreateFileReader(*OutputFileNameAndPath, FILEREAD_Silent)))
-				{
-					check(!CurrentWorkerInfo.bComplete);
-					FShaderCompileUtilities::DoReadTaskResults(CurrentWorkerInfo.QueuedJobs, *OutputFile);
-
-					// Close the output file.
-					OutputFile.Reset();
-
-					// Delete the output file now that we have consumed it, to avoid reading stale data on the next compile loop.
-					bool bDeletedOutput = IFileManager::Get().Delete(*OutputFileNameAndPath, true, true);
-					int32 RetryCount = 0;
-					// Retry over the next two seconds if we couldn't delete it
-					while (!bDeletedOutput && RetryCount < 200)
-					{
-						TRACE_CPUPROFILER_EVENT_SCOPE(FShaderCompileThreadRunnable::DeleteOutputFile);
-
-						FPlatformProcess::Sleep(0.01f);
-						bDeletedOutput = IFileManager::Get().Delete(*OutputFileNameAndPath, true, true);
-						RetryCount++;
-					}
-					checkf(bDeletedOutput, TEXT("Failed to delete %s!"), *OutputFileNameAndPath);
-
-					CurrentWorkerInfo.bComplete = true;
-				}
-
-				FPlatformAtomics::InterlockedIncrement(&NumProcessed);
-			}
-		}
-	};
-
-	if (bParallelizeIO)
-	{
-		ParallelFor( TEXT("ShaderCompiler.ReadAvailableResults.PF"),WorkerInfos.Num(),1, LoopBody, EParallelForFlags::Unbalanced);
-	}
-	else 
-	{
-		double StartIOWork = FPlatformTime::Seconds();
-		for (int32 WorkerIndex = 0; WorkerIndex < WorkerInfos.Num(); WorkerIndex++)
-		{
-			LoopBody(WorkerIndex);
-		}
-
-		double IODuration = FPlatformTime::Seconds() - StartIOWork;
-		if (IODuration > 0.3)
-		{
-			UE_LOG(LogShaderCompilers, Display, TEXT("FShaderCompileThreadRunnable::WriteNewTasks() took too long (%.3f seconds, threshold is %.3f s), will parallelize next time."), IODuration, GShaderCompilerTooLongIOThresholdSeconds);
-			bParallelizeIO = true;
-		}
-	}
-
-	return NumProcessed;
-}
-
-void FShaderCompileThreadRunnable::CompileDirectlyThroughDll()
-{
-	// If we aren't compiling through workers, so we can just track the serial time here.
-	COOK_STAT(FScopedDurationTimer CompileTimer (ShaderCompilerCookStats::AsyncCompileTimeSec));
-
-	FScopeLock WorkerScopeLock(&WorkerInfosLock);
-	for (int32 WorkerIndex = 0; WorkerIndex < WorkerInfos.Num(); WorkerIndex++)
-	{
-		FShaderCompileWorkerInfo& CurrentWorkerInfo = *WorkerInfos[WorkerIndex];
-
-		if (CurrentWorkerInfo.QueuedJobs.Num() > 0)
-		{
-			DumpWorkerInputs(CurrentWorkerInfo.QueuedJobs);
-
-			for (int32 JobIndex = 0; JobIndex < CurrentWorkerInfo.QueuedJobs.Num(); JobIndex++)
-			{
-				FShaderCommonCompileJob& CurrentJob = *CurrentWorkerInfo.QueuedJobs[JobIndex];
-				FShaderCompileUtilities::ExecuteShaderCompileJob(CurrentJob);
-			}
-
-			CurrentWorkerInfo.bComplete = true;
-		}
-	}
-}
-
-void FShaderCompileThreadRunnable::PrintWorkerMemoryUsageWithLockTaken()
-{
-	FPlatformProcessMemoryStats TotalMemoryStats{};
-	int32 NumValidWorkers = 0;
-	constexpr int64 Gibibyte = 1024 * 1024 * 1024;
-	for (int32 Iter = 0, End = WorkerInfos.Num(); Iter < End; Iter++)
-	{
-		const TUniquePtr<FShaderCompileWorkerInfo>& WorkerInfo = WorkerInfos[Iter];
-		FProcHandle ProcHandle = WorkerInfo->WorkerProcess;
-		if (!ProcHandle.IsValid())
-		{
-			continue;
-		}
-		FPlatformProcessMemoryStats MemoryStats;
-		if (FPlatformProcess::TryGetMemoryUsage(ProcHandle, MemoryStats))
-		{
-			NumValidWorkers++;
-			UE_LOG(LogShaderCompilers, Display,
-				TEXT("ShaderCompileWorker [%d/%d] MemoryStats:")
-				TEXT("\n\t     UsedPhysical %llu (%.2f GiB)")
-				TEXT("\n\t PeakUsedPhysical %llu (%.2f GiB)")
-				TEXT("\n\t      UsedVirtual %llu (%.2f GiB)")
-				TEXT("\n\t  PeakUsedVirtual %llu (%.2f GiB)"),
-				Iter + 1,
-				End,
-				MemoryStats.UsedPhysical, double(MemoryStats.UsedPhysical) / Gibibyte,
-				MemoryStats.PeakUsedPhysical, double(MemoryStats.PeakUsedPhysical) / Gibibyte,
-				MemoryStats.UsedVirtual, double(MemoryStats.UsedVirtual) / Gibibyte,
-				MemoryStats.PeakUsedVirtual, double(MemoryStats.PeakUsedVirtual) / Gibibyte
-			);
-			TotalMemoryStats.UsedPhysical += MemoryStats.UsedPhysical;
-			TotalMemoryStats.PeakUsedPhysical += MemoryStats.PeakUsedPhysical;
-			TotalMemoryStats.UsedVirtual += MemoryStats.PeakUsedVirtual;
-			TotalMemoryStats.PeakUsedVirtual += MemoryStats.PeakUsedVirtual;
-		}
-		LogQueuedCompileJobs(WorkerInfo->QueuedJobs, -1);
-	}
-
-	if (NumValidWorkers > 0)
-	{
-		UE_LOG(LogShaderCompilers, Display,
-			TEXT("Sum of MemoryStats for %d ShaderCompileWorker(s):")
-			TEXT("\n\t     UsedPhysical %llu (%.2f GiB)")
-			TEXT("\n\t PeakUsedPhysical %llu (%.2f GiB)")
-			TEXT("\n\t      UsedVirtual %llu (%.2f GiB)")
-			TEXT("\n\t  PeakUsedVirtual %llu (%.2f GiB)"),
-			NumValidWorkers,
-			TotalMemoryStats.UsedPhysical, double(TotalMemoryStats.UsedPhysical) / Gibibyte,
-			TotalMemoryStats.PeakUsedPhysical, double(TotalMemoryStats.PeakUsedPhysical) / Gibibyte,
-			TotalMemoryStats.UsedVirtual, double(TotalMemoryStats.UsedVirtual) / Gibibyte,
-			TotalMemoryStats.PeakUsedVirtual, double(TotalMemoryStats.PeakUsedVirtual) / Gibibyte
-		);
-	}
-}
-
-bool FShaderCompileThreadRunnable::PrintWorkerMemoryUsage(bool bAllowToWaitForLock)
-{
-	if (bAllowToWaitForLock)
-	{
-		FScopeLock WorkerScopeLock(&WorkerInfosLock);
-		PrintWorkerMemoryUsageWithLockTaken();
-		return true;
-	}
-	else
-	{
-		FScopeTryLock WorkerScopeLock(&WorkerInfosLock);
-		if (WorkerScopeLock.IsLocked())
-		{
-			PrintWorkerMemoryUsageWithLockTaken();
-			return true;
-		}
-		return false;
-	}
-}
-
-FShaderCompileMemoryUsage FShaderCompileThreadRunnable::GetExternalWorkerMemoryUsage()
-{
-	FScopeLock WorkerScopeLock(&WorkerInfosLock);
-	FShaderCompileMemoryUsage MemoryUsage{};
-	for (const TUniquePtr<FShaderCompileWorkerInfo>& WorkerInfo : WorkerInfos)
-	{
-		FProcHandle ProcHandle = WorkerInfo->WorkerProcess;
-		if (!ProcHandle.IsValid())
-		{
-			continue;
-		}
-		FPlatformProcessMemoryStats MemoryStats;
-		if (FPlatformProcess::TryGetMemoryUsage(ProcHandle, MemoryStats))
-		{
-			// Virtual memory is committed memory on Windows.
-			MemoryUsage.VirtualMemory += MemoryStats.UsedVirtual;
-			MemoryUsage.PhysicalMemory += MemoryStats.UsedPhysical;
-		}
-	}
-	return MemoryUsage;
-}
-
-void FShaderCompileUtilities::ExecuteShaderCompileJob(FShaderCommonCompileJob& Job)
-{
-	TRACE_CPUPROFILER_EVENT_SCOPE(FShaderCompileUtilities::ExecuteShaderCompileJob);
-
-	check(!Job.bFinalized);
-
-	FString WorkingDir = FPlatformProcess::ShaderDir();
-	static ITargetPlatformManagerModule& TPM = GetTargetPlatformManagerRef();
-	auto* SingleJob = Job.GetSingleShaderJob();
-	const TArray<const IShaderFormat*> ShaderFormats = TPM.GetShaderFormats();
-	if (SingleJob)
-	{
-		const FName Format = (SingleJob->Input.ShaderFormat != NAME_None) ? SingleJob->Input.ShaderFormat : LegacyShaderPlatformToShaderFormat(EShaderPlatform(SingleJob->Input.Target.Platform));
-		CompileShader(ShaderFormats, *SingleJob, WorkingDir);
-	}
-	else
-	{
-		FShaderPipelineCompileJob* PipelineJob = Job.GetShaderPipelineJob();
-		check(PipelineJob);
-
-		EShaderPlatform Platform = (EShaderPlatform)PipelineJob->StageJobs[0]->Input.Target.Platform;
-		const FName Format = LegacyShaderPlatformToShaderFormat(Platform);
-
-		// Verify same platform on all stages
-		for (int32 Index = 1; Index < PipelineJob->StageJobs.Num(); ++Index)
-		{
-			auto SingleStage = PipelineJob->StageJobs[Index];
-			if (!SingleStage)
-			{
-				UE_LOG(LogShaderCompilers, Fatal, TEXT("Can't nest Shader Pipelines inside Shader Pipeline '%s'!"), PipelineJob->Key.ShaderPipeline->GetName());
-			}
-			else if (Platform != SingleStage->Input.Target.Platform)
-			{
-				UE_LOG(LogShaderCompilers, Fatal, TEXT("Mismatched Target Platform %s while compiling Shader Pipeline '%s'."), *Format.GetPlainNameString(), PipelineJob->Key.ShaderPipeline->GetName());
-			}
-		}
-
-		CompileShaderPipeline(ShaderFormats, PipelineJob, WorkingDir);
-	}
-
-	Job.bFinalized = true;
-}
-
-FArchive* FShaderCompileUtilities::CreateFileHelper(const FString& Filename)
-{
-	// TODO: This logic came from FShaderCompileThreadRunnable::WriteNewTasks().
-	// We can't avoid code duplication unless we refactored the local worker too.
-
-	FArchive* File = nullptr;
-	int32 RetryCount = 0;
-	// Retry over the next two seconds if we can't write out the file.
-	// Anti-virus and indexing applications can interfere and cause this to fail.
-	while (File == nullptr && RetryCount < 200)
-	{
-		if (RetryCount > 0)
-		{
-			FPlatformProcess::Sleep(0.01f);
-		}
-		File = IFileManager::Get().CreateFileWriter(*Filename, FILEWRITE_EvenIfReadOnly);
-		RetryCount++;
-	}
-	if (File == nullptr)
-	{
-		File = IFileManager::Get().CreateFileWriter(*Filename, FILEWRITE_EvenIfReadOnly | FILEWRITE_NoFail);
-	}
-	checkf(File, TEXT("Failed to create file %s!"), *Filename);
-	return File;
-}
-
-void FShaderCompileUtilities::MoveFileHelper(const FString& To, const FString& From)
-{
-	IPlatformFile& PlatformFile = FPlatformFileManager::Get().GetPlatformFile();
-
-	if (PlatformFile.FileExists(*From))
-	{
-		FString DirectoryName;
-		int32 LastSlashIndex;
-		if (To.FindLastChar('/', LastSlashIndex))
-		{
-			DirectoryName = To.Left(LastSlashIndex);
-		} else
-		{
-			DirectoryName = To;
-		}
-
-		// TODO: This logic came from FShaderCompileThreadRunnable::WriteNewTasks().
-		// We can't avoid code duplication unless we refactored the local worker too.
-
-		bool Success = false;
-		int32 RetryCount = 0;
-		// Retry over the next two seconds if we can't move the file.
-		// Anti-virus and indexing applications can interfere and cause this to fail.
-		while (!Success && RetryCount < 200)
-		{
-			if (RetryCount > 0)
-			{
-				FPlatformProcess::Sleep(0.01f);
-			}
-
-			// MoveFile does not create the directory tree, so try to do that now...
-			Success = PlatformFile.CreateDirectoryTree(*DirectoryName);
-			if (Success)
-			{
-				Success = PlatformFile.MoveFile(*To, *From);
-			}
-			RetryCount++;
-		}
-		checkf(Success, TEXT("Failed to move file %s to %s!"), *From, *To);
-	}
-}
-
-void FShaderCompileUtilities::DeleteFileHelper(const FString& Filename)
-{
-	// TODO: This logic came from FShaderCompileThreadRunnable::WriteNewTasks().
-	// We can't avoid code duplication unless we refactored the local worker too.
-
-	if (FPlatformFileManager::Get().GetPlatformFile().FileExists(*Filename))
-	{
-		bool bDeletedOutput = IFileManager::Get().Delete(*Filename, true, true);
-
-		// Retry over the next two seconds if we couldn't delete it
-		int32 RetryCount = 0;
-		while (!bDeletedOutput && RetryCount < 200)
-		{
-			FPlatformProcess::Sleep(0.01f);
-			bDeletedOutput = IFileManager::Get().Delete(*Filename, true, true);
-			RetryCount++;
-		}
-		checkf(bDeletedOutput, TEXT("Failed to delete %s!"), *Filename);
-	}
-}
-
-int32 FShaderCompileThreadRunnable::CompilingLoop()
-{
-	if (!Manager->bAllowCompilingThroughWorkers && CVarCompileParallelInProcess.GetValueOnAnyThread())
-	{
-		int32 NumJobs = Manager->GetNumPendingJobs();
-		if (NumJobs == 0)
-		{
-			return 0;
-		}
-
-		// if -noshaderworker is specified and the experimental in-process parallel compile is enabled, submit tasks for a batch 
-		// of pending jobs to run wide in-process (and a tailing task to mark those jobs as complete in the manager)
-		TUniquePtr<TArray<FShaderCommonCompileJobPtr>> Jobs = TUniquePtr<TArray<FShaderCommonCompileJobPtr>>(new TArray<FShaderCommonCompileJobPtr>());
-		TUniquePtr<TArray<float>> JobTimes = TUniquePtr<TArray<float>>(new TArray<float>());
-		{
-			FScopeLock Lock(&Manager->CompileQueueSection);
-			for (int32 PriorityIndex = MaxPriorityIndex; PriorityIndex >= MinPriorityIndex; --PriorityIndex)
-			{
-				// Throttle how many jobs we kick per tick so we get more frequent progress updates in the UI;
-				// this doesn't seem to have much effect on overall throughput.
-				const int32 MaxNumJobs = 64;
-				Manager->AllJobs.GetPendingJobs(EShaderCompilerWorkerType::LocalThread, (EShaderCompileJobPriority)PriorityIndex, 1, MaxNumJobs, *Jobs);
-			}
-		}
-
-		JobTimes->SetNum(Jobs->Num());
-
-		TArray<UE::Tasks::FTask> CompileTasks;
-		CompileTasks.Reserve(Jobs->Num());
-		for (TArray<FShaderCommonCompileJobPtr>::SizeType JobIndex = 0; JobIndex < Jobs->Num(); ++JobIndex)
-		{
-			FShaderCommonCompileJob* Job = (*Jobs)[JobIndex];
-			float* Time = &(*JobTimes)[JobIndex];
-			CompileTasks.Add(UE::Tasks::Launch(UE_SOURCE_LOCATION, [Job, Time]()
-				{
-					const float StartTime = FPlatformTime::Seconds();
-					FShaderCompileUtilities::ExecuteShaderCompileJob(*Job);
-					*Time = FPlatformTime::Seconds() - StartTime;
-				}));
-		}
-
-		if (!Jobs->IsEmpty())
-		{
-			UE::Tasks::Launch(UE_SOURCE_LOCATION, [Jobs = MoveTemp(Jobs), JobTimes = MoveTemp(JobTimes), this]()
-				{
-					FScopeLock Lock(&Manager->CompileQueueSection);
-					for (FShaderCommonCompileJobPtr Job : *Jobs)
-					{
-						Manager->ProcessFinishedJob(Job.GetReference());
-					}
-
-					float ElapsedTime = 0.0f;
-					for (const float& JobTime : *JobTimes)
-					{
-						ElapsedTime += JobTime;
-					}
-
-					Manager->WorkersBusyTime += ElapsedTime;
-					COOK_STAT(ShaderCompilerCookStats::AsyncCompileTimeSec += ElapsedTime);
-				}, CompileTasks);
-			// num active threads is up to the task system; as long as we return non-zero
-			// this will indicate to the caller that pending jobs remain
-			return 1;
-		}
-		return 0;
-	}
-	else // compile either through worker processes or single-threaded in-process (depending on Manager->bAllowCompilingThroughWorkers)
-	{
-		// push completed jobs to Manager->ShaderMapJobs before asking for new ones, so we can free the workers now and avoid them waiting a cycle
-		PushCompletedJobsToManager();
-
-		// Grab more shader compile jobs from the input queue
-		const int32 NumActiveThreads = PullTasksFromQueue();
-
-		if (NumActiveThreads == 0 && Manager->bAllowAsynchronousShaderCompiling)
-		{
-			// Yield while there's nothing to do
-			// Note: sleep-looping is bad threading practice, wait on an event instead!
-			// The shader worker thread does it because it needs to communicate with other processes through the file system
-			FPlatformProcess::Sleep(.010f);
-		}
-
-		if (Manager->bAllowCompilingThroughWorkers)
-		{
-			// Write out the files which are input to the shader compile workers
-			WriteNewTasks();
-
-			// Launch shader compile workers if they are not already running
-			// Workers can time out when idle so they may need to be relaunched
-			bool bAbandonWorkers = LaunchWorkersIfNeeded();
-
-			if (bAbandonWorkers)
-			{
-				// Fall back to local compiles if the SCW crashed.
-				// This is nasty but needed to work around issues where message passing through files to SCW is unreliable on random PCs
-				Manager->bAllowCompilingThroughWorkers = false;
-
-				// Try to recover from abandoned workers after a certain amount of single-threaded compilations
-				if (Manager->NumSingleThreadedRunsBeforeRetry == GSingleThreadedRunsIdle)
-				{
-					// First try to recover, only run single-threaded approach once
-					Manager->NumSingleThreadedRunsBeforeRetry = 1;
-				}
-				else if (Manager->NumSingleThreadedRunsBeforeRetry > GSingleThreadedRunsMaxCount)
-				{
-					// Stop retry approach after too many retries have failed
-					Manager->NumSingleThreadedRunsBeforeRetry = GSingleThreadedRunsDisabled;
-				}
-				else
-				{
-					// Next time increase runs by factor X
-					Manager->NumSingleThreadedRunsBeforeRetry *= GSingleThreadedRunsIncreaseFactor;
-				}
-			}
-			else
-			{
-				// Read files which are outputs from the shader compile workers
-				int32 NumProcessedResults = ReadAvailableResults();
-				if (NumProcessedResults == 0)
-				{
-					// Reduce filesystem query rate while actively waiting for results.
-					FPlatformProcess::Sleep(0.1f);
-				}
-			}
-		}
-		else
-		{
-			// Execute all pending worker tasks single-threaded
-			CompileDirectlyThroughDll();
-
-			// If single-threaded mode was enabled by an abandoned worker, try to recover after the given amount of runs
-			if (Manager->NumSingleThreadedRunsBeforeRetry > 0)
-			{
-				Manager->NumSingleThreadedRunsBeforeRetry--;
-				if (Manager->NumSingleThreadedRunsBeforeRetry == 0)
-				{
-					UE_LOG(LogShaderCompilers, Display, TEXT("Retry shader compiling through workers."));
-					Manager->bClearStaleWorkerOutputs = true;
-					Manager->bAllowCompilingThroughWorkers = true;
-				}
-			}
-		}
-
-		return NumActiveThreads;
-	}
-}
-
-FShaderCompilerStats* GShaderCompilerStats = nullptr;
-
-void FShaderCompilerStats::WriteStats(FOutputDevice* Ar)
-{
-#if ALLOW_DEBUG_FILES
-	constexpr static const TCHAR DebugText[] = TEXT("Wrote shader compile stats to file '%s'.");
-	{
-		FlushRenderingCommands();
-
-		FString FileName = FPaths::Combine(*FPaths::ProjectSavedDir(),
-			FString::Printf(TEXT("MaterialStats/Stats-%s.csv"), *FDateTime::Now().ToString()));
-		auto DebugWriter = IFileManager::Get().CreateFileWriter(*FileName);
-		FDiagnosticTableWriterCSV StatWriter(DebugWriter);
-		const TSparseArray<ShaderCompilerStats>& PlatformStats = GetShaderCompilerStats();
-
-		StatWriter.AddColumn(TEXT("Path"));
-		StatWriter.AddColumn(TEXT("Platform"));
-		StatWriter.AddColumn(TEXT("Compiled"));
-		StatWriter.AddColumn(TEXT("Cooked"));
-		StatWriter.AddColumn(TEXT("Permutations"));
-		StatWriter.AddColumn(TEXT("Compiletime"));
-		StatWriter.AddColumn(TEXT("CompiledDouble"));
-		StatWriter.AddColumn(TEXT("CookedDouble"));
-		StatWriter.CycleRow();
-
-		
-		for(int32 Platform = 0; Platform < PlatformStats.GetMaxIndex(); ++Platform)
-		{
-			if(PlatformStats.IsValidIndex(Platform))
-			{
-				const ShaderCompilerStats& Stats = PlatformStats[Platform];
-				for (const auto& Pair : Stats)
-				{
-					const FString& Path = Pair.Key;
-					const FShaderCompilerStats::FShaderStats& SingleStats = Pair.Value;
-
-					StatWriter.AddColumn(*Path);
-					StatWriter.AddColumn(TEXT("%u"), Platform);
-					StatWriter.AddColumn(TEXT("%u"), SingleStats.Compiled);
-					StatWriter.AddColumn(TEXT("%u"), SingleStats.Cooked);
-					StatWriter.AddColumn(TEXT("%u"), SingleStats.PermutationCompilations.Num());
-					StatWriter.AddColumn(TEXT("%f"), SingleStats.CompileTime);
-					StatWriter.AddColumn(TEXT("%u"), SingleStats.CompiledDouble);
-					StatWriter.AddColumn(TEXT("%u"), SingleStats.CookedDouble);
-					StatWriter.CycleRow();
-					if(GLogShaderCompilerStats)
-					{
-						UE_LOG(LogShaderCompilers, Log, TEXT("SHADERSTATS %s, %u, %u, %u, %u, %u, %u\n"), *Path, Platform, SingleStats.Compiled, SingleStats.Cooked, SingleStats.PermutationCompilations.Num(), SingleStats.CompiledDouble, SingleStats.CookedDouble);
-					}
-				}
-			}
-		}
-		DebugWriter->Close();
-
-		FString FullFileName = FPaths::ConvertRelativePathToFull(FileName);
-		if (Ar)
-		{
-			Ar->Logf(DebugText, *FullFileName);
-		}
-		else
-		{
-			UE_LOG(LogShaderCompilers, Log, DebugText, *FullFileName);
-		}
-
-		if (FParse::Param(FCommandLine::Get(), TEXT("mirrorshaderstats")))
-		{
-			FString MirrorLocation;
-			GConfig->GetString(TEXT("/Script/Engine.ShaderCompilerStats"), TEXT("MaterialStatsLocation"), MirrorLocation, GGameIni);
-			FParse::Value(FCommandLine::Get(), TEXT("MaterialStatsMirror="), MirrorLocation);
-
-			if (!MirrorLocation.IsEmpty())
-			{
-				FString TargetType = TEXT("Default");
-				FParse::Value(FCommandLine::Get(), TEXT("target="), TargetType);
-				if (TargetType == TEXT("Default"))
-				{
-					FParse::Value(FCommandLine::Get(), TEXT("targetplatform="), TargetType);
-				}
-				FString CopyLocation = FPaths::Combine(*MirrorLocation, FApp::GetProjectName(), *FApp::GetBranchName(), FString::Printf(TEXT("Stats-Latest-%d(%s).csv"), FEngineVersion::Current().GetChangelist() , *TargetType));
-				TArray <FString> ExistingFiles;
-				IFileManager::Get().FindFiles(ExistingFiles, *FPaths::Combine(*MirrorLocation, FApp::GetProjectName(), *FApp::GetBranchName()));
-				for (FString CurFile : ExistingFiles)
-				{
-					if (CurFile.Contains(FString::Printf(TEXT("(%s)"), *TargetType)))
-					{
-						IFileManager::Get().Delete(*FPaths::Combine(*MirrorLocation, FApp::GetProjectName(), *FApp::GetBranchName(), *CurFile), false, true);
-					}
-				}
-				IFileManager::Get().Copy(*CopyLocation, *FileName, true, true);
-			}
-		}
-	}
-	{
-
-		FString FileName = FString::Printf(TEXT("%s/MaterialStatsDebug/StatsDebug-%s.csv"), *FPaths::ProjectSavedDir(), *FDateTime::Now().ToString());
-		auto DebugWriter = IFileManager::Get().CreateFileWriter(*FileName);
-		FDiagnosticTableWriterCSV StatWriter(DebugWriter);
-		const TSparseArray<ShaderCompilerStats>& PlatformStats = GetShaderCompilerStats();
-		StatWriter.AddColumn(TEXT("Name"));
-		StatWriter.AddColumn(TEXT("Platform"));
-		StatWriter.AddColumn(TEXT("Compiles"));
-		StatWriter.AddColumn(TEXT("CompilesDouble"));
-		StatWriter.AddColumn(TEXT("Uses"));
-		StatWriter.AddColumn(TEXT("UsesDouble"));
-		StatWriter.AddColumn(TEXT("PermutationString"));
-		StatWriter.CycleRow();
-
-
-		for (int32 Platform = 0; Platform < PlatformStats.GetMaxIndex(); ++Platform)
-		{
-			if (PlatformStats.IsValidIndex(Platform))
-			{
-				const ShaderCompilerStats& Stats = PlatformStats[Platform];
-				for (const auto& Pair : Stats)
-				{
-					const FString& Path = Pair.Key;
-					const FShaderCompilerStats::FShaderStats& SingleStats = Pair.Value;
-					for (const FShaderCompilerStats::FShaderCompilerSinglePermutationStat& Stat : SingleStats.PermutationCompilations)
-					{
-						StatWriter.AddColumn(*Path);
-						StatWriter.AddColumn(TEXT("%u"), Platform);
-						StatWriter.AddColumn(TEXT("%u"), Stat.Compiled);
-						StatWriter.AddColumn(TEXT("%u"), Stat.CompiledDouble);
-						StatWriter.AddColumn(TEXT("%u"), Stat.Cooked);
-						StatWriter.AddColumn(TEXT("%u"), Stat.CookedDouble);
-						StatWriter.AddColumn(TEXT("%s"), *Stat.PermutationString);
-						StatWriter.CycleRow();
-					}
-				}
-
-			}
-		}
-
-		FString FullFileName = FPaths::ConvertRelativePathToFull(FileName);
-		if (Ar)
-		{
-			Ar->Logf(DebugText, *FullFileName);
-		}
-		else
-		{
-			UE_LOG(LogShaderCompilers, Log, DebugText, *FullFileName);
-		}
-	}
-#endif // ALLOW_DEBUG_FILES
-}
-
-template <typename T>
-static FString FormatNumber(T Number)
-{
-	static const FNumberFormattingOptions FormattingOptions = FNumberFormattingOptions().SetUseGrouping(true);
-	return FText::AsNumber(Number, &FormattingOptions).ToString();
-}
-
-static FString PrintJobsCompletedPercentageToString(int64 JobsAssigned, int64 JobsCompleted)
-{
-	if (JobsAssigned == 0)
-	{
-		return TEXT("0%");
-	}
-	if (JobsAssigned == JobsCompleted)
-	{
-		return TEXT("100%");
-	}
-
-	// With more than a million compile jobs but only a small number that didn't complete,
-	// the output might be rounded up to 100%. To avoid a misleading output, we clamp this value to 99.99%
-	double JobsCompletedPercentage = 100.0 * (double)JobsCompleted / (double)JobsAssigned;
-	return FString::Printf(TEXT("%.2f%%"), FMath::Min(JobsCompletedPercentage, 99.99));
-}
-
-void FShaderCompilerStats::WriteStatSummary()
-{
-	const uint32 TotalCompiled = GetTotalShadersCompiled();
-	if (TotalCompiled == 0)
-	{
-		// early out if we haven't done anything yet
-		return;
-	}
-
-	UE_LOG(LogShaderCompilers, Display, TEXT("================================================"));
-
-	const TCHAR* AggregatedSuffix = bMultiProcessAggregated ? TEXT(" (aggregated across all cook processes)") : TEXT("");
-
-	// Only log cache stats if the cache has been queried at least once (this will always be 0 if the job cache is disabled)
-	if (Counters.TotalCacheSearchAttempts > 0)
-	{
-		UE_LOG(LogShaderCompilers, Display, TEXT("=== FShaderJobCache stats%s ==="), AggregatedSuffix);
-		UE_LOG(LogShaderCompilers, Display, TEXT("Total job queries %s, among them cache hits %s (%.2f%%), DDC hits %s (%.2f%%), Duplicates %s (%.2f%%)"),
-			*FormatNumber(Counters.TotalCacheSearchAttempts),
-			*FormatNumber(Counters.TotalCacheHits),
-			100.0 * static_cast<double>(Counters.TotalCacheHits) / static_cast<double>(Counters.TotalCacheSearchAttempts),
-			*FormatNumber(Counters.TotalCacheDDCHits),
-			100.0 * static_cast<double>(Counters.TotalCacheDDCHits) / static_cast<double>(Counters.TotalCacheSearchAttempts),
-			*FormatNumber(Counters.TotalCacheDuplicates),
-			100.0 * static_cast<double>(Counters.TotalCacheDuplicates) / static_cast<double>(Counters.TotalCacheSearchAttempts));
-
-		UE_LOG(LogShaderCompilers, Display, TEXT("Tracking %s distinct input hashes that result in %s distinct outputs (%.2f%%)"),
-			*FormatNumber(Counters.UniqueCacheInputHashes),
-			*FormatNumber(Counters.UniqueCacheOutputs),
-			(Counters.UniqueCacheInputHashes > 0) ? 100.0 * static_cast<double>(Counters.UniqueCacheOutputs) / static_cast<double>(Counters.UniqueCacheInputHashes) : 0.0);
-
-		static const FNumberFormattingOptions SizeFormattingOptions = FNumberFormattingOptions().SetMinimumFractionalDigits(2).SetMaximumFractionalDigits(2);
-
-		if (Counters.CacheMemBudget > 0)
-		{
-			UE_LOG(LogShaderCompilers, Display, TEXT("RAM used: %s of %s budget. Usage: %.2f%%"),
-				*FText::AsMemory(Counters.CacheMemUsed, &SizeFormattingOptions, nullptr, EMemoryUnitStandard::IEC).ToString(),
-				*FText::AsMemory(Counters.CacheMemBudget, &SizeFormattingOptions, nullptr, EMemoryUnitStandard::IEC).ToString(),
-				100.0 * Counters.CacheMemUsed / Counters.CacheMemBudget);
-		}
-		else
-		{
-			UE_LOG(LogShaderCompilers, Display, TEXT("RAM used: %s, no memory limit set"), *FText::AsMemory(Counters.CacheMemUsed, &SizeFormattingOptions, nullptr, EMemoryUnitStandard::IEC).ToString());
-		}
-	}
-
-	const double TotalTimeAtLeastOneJobWasInFlight = GetTimeShaderCompilationWasActive();
-
-	UE_LOG(LogShaderCompilers, Display, TEXT("=== Shader Compilation stats%s ==="), AggregatedSuffix);
-	UE_LOG(LogShaderCompilers, Display, TEXT("Shaders Compiled: %s"), *FormatNumber(TotalCompiled));
-
-	FScopeLock Lock(&CompileStatsLock);	// make a local copy for all the stats?
-	UE_LOG(LogShaderCompilers, Display, TEXT("Jobs assigned %s, completed %s (%s)"), 
-		*FormatNumber(Counters.JobsAssigned),
-		*FormatNumber(Counters.JobsCompleted),
-		*PrintJobsCompletedPercentageToString(Counters.JobsAssigned, Counters.JobsCompleted));
-
-	if (Counters.TimesLocalWorkersWereIdle > 0.0)
-	{
-		UE_LOG(LogShaderCompilers, Display, TEXT("Average time worker was idle: %.2f s"), Counters.AccumulatedLocalWorkerIdleTime / Counters.TimesLocalWorkersWereIdle);
-	}
-
-	if (Counters.JobsAssigned > 0)
-	{
-		UE_LOG(LogShaderCompilers, Display, TEXT("Time job spent in pending queue: average %.2f s, longest %.2f s"), Counters.AccumulatedPendingTime / (double)Counters.JobsAssigned, Counters.MaxPendingTime);
-	}
-
-	if (Counters.JobsCompleted > 0)
-	{
-		UE_LOG(LogShaderCompilers, Display, TEXT("Job execution time: average %.2f s, max %.2f s"), Counters.AccumulatedJobExecutionTime / (double)Counters.JobsCompleted, Counters.MaxJobExecutionTime);
-		UE_LOG(LogShaderCompilers, Display, TEXT("Job life time (pending + execution): average %.2f s, max %.2f"), Counters.AccumulatedJobLifeTime / (double)Counters.JobsCompleted, Counters.MaxJobLifeTime);
-	}
-
-	if (Counters.NumAccumulatedShaderCodes > 0)
-	{
-		const FString AvgCodeSizeStr = FText::AsMemory((uint64)((double)Counters.AccumulatedShaderCodeSize / (double)Counters.NumAccumulatedShaderCodes)).ToString();
-		const FString MinCodeSizeStr = FText::AsMemory((uint64)Counters.MinShaderCodeSize).ToString();
-		const FString MaxCodeSizeStr = FText::AsMemory((uint64)Counters.MaxShaderCodeSize).ToString();
-		UE_LOG(LogShaderCompilers, Display, TEXT("Shader code size: average %s, min %s, max %s"), *AvgCodeSizeStr, *MinCodeSizeStr, *MaxCodeSizeStr);
-	}
-
-	UE_LOG(LogShaderCompilers, Display, TEXT("Time at least one job was in flight (either pending or executed): %.2f s"), TotalTimeAtLeastOneJobWasInFlight);
-
-	if (Counters.AccumulatedTaskSubmitJobs > 0.0)
-	{
-		UE_LOG(LogShaderCompilers, Display, TEXT("Mutex wait stall in FShaderJobCache::SubmitJobs:  %.2f%%"), 100.0 * Counters.AccumulatedTaskSubmitJobsStall / Counters.AccumulatedTaskSubmitJobs );
-	}
-
-	// print stats about the batches
-	if (Counters.LocalJobBatchesSeen > 0 && Counters.DistributedJobBatchesSeen > 0)
-	{
-		int64 JobBatchesSeen = Counters.LocalJobBatchesSeen + Counters.DistributedJobBatchesSeen;
-		double TotalJobsReportedInJobBatches = Counters.TotalJobsReportedInLocalJobBatches + Counters.TotalJobsReportedInDistributedJobBatches;
-
-		UE_LOG(LogShaderCompilers, Display, TEXT("Jobs were issued in %s batches (%s local, %s distributed), average %.2f jobs/batch (%.2f jobs/local batch. %.2f jobs/distributed batch)"),
-			*FormatNumber(JobBatchesSeen), *FormatNumber(Counters.LocalJobBatchesSeen), *FormatNumber(Counters.DistributedJobBatchesSeen),
-			static_cast<double>(TotalJobsReportedInJobBatches) / static_cast<double>(JobBatchesSeen),
-			static_cast<double>(Counters.TotalJobsReportedInLocalJobBatches) / static_cast<double>(Counters.LocalJobBatchesSeen),
-			static_cast<double>(Counters.TotalJobsReportedInDistributedJobBatches) / static_cast<double>(Counters.DistributedJobBatchesSeen)
-		);
-	}
-	else if (Counters.LocalJobBatchesSeen > 0)
-	{
-		UE_LOG(LogShaderCompilers, Display, TEXT("Jobs were issued in %s batches (only local compilation was used), average %.2f jobs/batch"), 
-			*FormatNumber(Counters.LocalJobBatchesSeen), static_cast<double>(Counters.TotalJobsReportedInLocalJobBatches) / static_cast<double>(Counters.LocalJobBatchesSeen));
-	}
-	else if (Counters.DistributedJobBatchesSeen > 0)
-	{
-		UE_LOG(LogShaderCompilers, Display, TEXT("Jobs were issued in %s batches (only distributed compilation was used), average %.2f jobs/batch"),
-			*FormatNumber(Counters.DistributedJobBatchesSeen), static_cast<double>(Counters.TotalJobsReportedInDistributedJobBatches) / static_cast<double>(Counters.DistributedJobBatchesSeen));
-	}
-
-	if (TotalTimeAtLeastOneJobWasInFlight > 0.0)
-	{
-		UE_LOG(LogShaderCompilers, Display, TEXT("Average processing rate: %.2f jobs/sec"), (double)Counters.JobsCompleted / TotalTimeAtLeastOneJobWasInFlight);
-	}
-
-	if (ShaderTimings.Num())
-	{
-		// calculate effective parallelization (total time needed to compile all shaders divided by actual wall clock time spent processing at least 1 shader)
-		double TotalThreadTimeForAllShaders = 0.0;
-		double TotalThreadPreprocessTimeForAllShaders = 0.0;
-		for (TMap<FString, FShaderTimings>::TConstIterator Iter(ShaderTimings); Iter; ++Iter)
-		{
-			TotalThreadTimeForAllShaders += Iter.Value().TotalCompileTime;
-			TotalThreadPreprocessTimeForAllShaders += Iter.Value().TotalPreprocessTime;
-		}
-
-		UE_LOG(LogShaderCompilers, Display, TEXT("Total thread time: %s s"), *FormatNumber(TotalThreadTimeForAllShaders));
-		UE_LOG(LogShaderCompilers, Display, TEXT("Total thread preprocess time: %s s"), *FormatNumber(TotalThreadPreprocessTimeForAllShaders));
-		UE_LOG(LogShaderCompilers, Display, TEXT("Percentage time preprocessing: %.2f%%"), TotalThreadTimeForAllShaders > 0.0 ? (TotalThreadPreprocessTimeForAllShaders / TotalThreadTimeForAllShaders) * 100.0 : 0.0);
-
-		if (Counters.MaxRemoteAgents > 0)
-		{
-			UE_LOG(LogShaderCompilers, Display, TEXT("Highest number of remote agents active in parallel: %u (%u active cores peak)"), Counters.MaxRemoteAgents, Counters.MaxActiveAgentCores);
-		}
-
-		if (TotalTimeAtLeastOneJobWasInFlight > 0.0)
-		{
-			double EffectiveParallelization = TotalThreadTimeForAllShaders / TotalTimeAtLeastOneJobWasInFlight;
-			if (Counters.DistributedJobBatchesSeen == 0)
-			{
-				UE_LOG(LogShaderCompilers, Display, TEXT("Effective parallelization: %.2f (times faster than compiling all shaders on one thread). Compare with number of workers: %d"), EffectiveParallelization, GShaderCompilingManager->GetNumLocalWorkers());
-			}
-			else
-			{
-				UE_LOG(LogShaderCompilers, Display, TEXT("Effective parallelization: %.2f (times faster than compiling all shaders on one thread). Distributed compilation was used."), EffectiveParallelization);
-			}
-		}
-
-
-		// sort by avg time
-		ShaderTimings.ValueSort([](const FShaderTimings& A, const FShaderTimings& B) { return A.AverageCompileTime > B.AverageCompileTime; });
-
-		const int32 MaxShadersToPrint = FMath::Min(ShaderTimings.Num(), GMaxShaderStatsToLog);
-		UE_LOG(LogShaderCompilers, Display, TEXT("Top %d most expensive shader types by average time:"), MaxShadersToPrint);
-
-		int32 Idx = 0;
-		for (TMap<FString, FShaderTimings>::TConstIterator Iter(ShaderTimings); Iter; ++Iter)
-		{
-			const FShaderTimings& Timings = Iter.Value();
-
-			UE_LOG(LogShaderCompilers, Display, TEXT("%60s (compiled %4d times, average %4.2f sec, max %4.2f sec, min %4.2f sec)"), *Iter.Key(), Timings.NumCompiled, Timings.AverageCompileTime, Timings.MaxCompileTime, Timings.MinCompileTime);
-			if (++Idx >= MaxShadersToPrint)
-			{
-				break;
-			}
-		}
-
-		// sort by total time
-		ShaderTimings.ValueSort([](const FShaderTimings& A, const FShaderTimings& B) { return A.TotalCompileTime > B.TotalCompileTime; });
-
-		UE_LOG(LogShaderCompilers, Display, TEXT("Top %d shader types by total compile time:"), MaxShadersToPrint);
-
-		Idx = 0;
-		for (TMap<FString, FShaderTimings>::TConstIterator Iter(ShaderTimings); Iter; ++Iter)
-		{
-			const FShaderTimings& Timings = Iter.Value();
-
-			UE_LOG(LogShaderCompilers, Display, TEXT("%60s - %.2f%% of total time (compiled %4d times, average %4.2f sec, max %4.2f sec, min %4.2f sec)"), 
-				*Iter.Key(), 100.0 * Timings.TotalCompileTime / TotalThreadTimeForAllShaders, Timings.NumCompiled, Timings.AverageCompileTime, Timings.MaxCompileTime, Timings.MinCompileTime);
-			if (++Idx >= MaxShadersToPrint)
-			{
-				break;
-			}
-		}
-	}
-
-	MaterialCounters.WriteStatSummary(AggregatedSuffix);
-
-	UE_LOG(LogShaderCompilers, Display, TEXT("================================================"));
-}
-
-void FShaderCompilerStats::GatherAnalytics(const FString& BaseName, TArray<FAnalyticsEventAttribute>& Attributes)
-{
-	const double TotalTimeAtLeastOneJobWasInFlight = GetTimeShaderCompilationWasActive();
-
-	FScopeLock Lock(&CompileStatsLock);
-
-	{
-		FString AttrName = BaseName + TEXT("ShadersCompiled");
-		Attributes.Emplace(MoveTemp(AttrName), Counters.JobsCompleted);
-	}
-
-	if (ShaderTimings.Num())
-	{
-		double TotalThreadTimeForAllShaders = 0.0;
-		double TotalThreadPreprocessTimeForAllShaders = 0.0;
-		for (TMap<FString, FShaderTimings>::TConstIterator Iter(ShaderTimings); Iter; ++Iter)
-		{
-			TotalThreadTimeForAllShaders += Iter.Value().TotalCompileTime;
-			TotalThreadPreprocessTimeForAllShaders += Iter.Value().TotalPreprocessTime;
-		}
-
-		{
-			FString AttrName = BaseName + TEXT("TotalThreadTime");
-			Attributes.Emplace(MoveTemp(AttrName), TotalThreadTimeForAllShaders);
-		}
-
-		{
-			FString AttrName = BaseName + TEXT("TotalThreadPreprocessTime");
-			Attributes.Emplace(MoveTemp(AttrName), TotalThreadPreprocessTimeForAllShaders);
-		}
-
-		{
-            const double EffectiveParallelization = TotalTimeAtLeastOneJobWasInFlight > 0.0 ? TotalThreadTimeForAllShaders / TotalTimeAtLeastOneJobWasInFlight : 0.0;
-			FString AttrName = BaseName + TEXT("EffectiveParallelization");
-			Attributes.Emplace(MoveTemp(AttrName), EffectiveParallelization);
-		}
-	}
-
-	if (Counters.TotalCacheSearchAttempts)
-	{
-		const FString ChildName = TEXT("JobCache_");
-
-		{
-			FString AttrName = BaseName + ChildName + TEXT("Queries");
-			Attributes.Emplace(MoveTemp(AttrName), Counters.TotalCacheSearchAttempts);
-		}
-
-		{
-			FString AttrName = BaseName + ChildName + TEXT("Hits");
-			Attributes.Emplace(MoveTemp(AttrName), Counters.TotalCacheHits);
-		}
-
-		{
-			FString AttrName = BaseName + ChildName + TEXT("DDCHits");
-			Attributes.Emplace(MoveTemp(AttrName), Counters.TotalCacheDDCHits);
-		}
-
-		{
-			FString AttrName = BaseName + ChildName + TEXT("NumInputs");
-			Attributes.Emplace(MoveTemp(AttrName), Counters.UniqueCacheInputHashes);
-		}
-
-		{
-			FString AttrName = BaseName + ChildName + TEXT("NumOutputs");
-			Attributes.Emplace(MoveTemp(AttrName), Counters.UniqueCacheOutputs);
-		}
-
-		{
-			FString AttrName = BaseName + ChildName + TEXT("MemUsed");
-			Attributes.Emplace(MoveTemp(AttrName), Counters.CacheMemUsed);
-		}
-
-		{
-			FString AttrName = BaseName + ChildName + TEXT("MemBudget");
-			Attributes.Emplace(MoveTemp(AttrName), Counters.CacheMemBudget);
-		}
-	}
-
-	MaterialCounters.GatherAnalytics(Attributes);
-}
-
-uint32 FShaderCompilerStats::GetTotalShadersCompiled()
-{
-	FScopeLock Lock(&CompileStatsLock);
-	return (uint32)FMath::Max(0ll, Counters.JobsCompleted);
-}
-
-void AddToInterval(TArray<TInterval<double>>& Accumulator, const TInterval<double>& NewInterval)
-{
-	bool bFoundOverlap = false;
-	TInterval<double> New = NewInterval;
-	int32 Idx = 0;
-	do
-	{
-		bFoundOverlap = false;
-		for (; Idx < Accumulator.Num(); ++Idx)
-		{
-			const TInterval<double>& Existing = Accumulator[Idx];
-			if (Existing.Max < New.Min)
-			{
-				continue;	// no overlap but the new interval starts after this one ends, keep searching
-			}
-
-			if (New.Max < Existing.Min)
-			{
-				break;		// no overlap, but the new interval ends before this one starts, insert here
-			}
-
-			// if fully contained within existing interval, just ignore
-			if (Existing.Min <= New.Min && New.Max <= Existing.Max)
-			{
-				return;
-			}
-
-			bFoundOverlap = true;
-			// if there's an overlap, remove the existing interval, merge with the new one and attempt to add again
-			TInterval<double> Merged(FMath::Min(Existing.Min, New.Min), FMath::Max(Existing.Max, New.Max));
-			check(Merged.Size() >= Existing.Size());
-			check(Merged.Size() >= New.Size());
-			Accumulator.RemoveAt(Idx);
-			New = Merged;
-			break;
-		}
-	} while (bFoundOverlap);
-
-	// if we arrived here without an overlap, we have a new one; insert in the appropriate place
-	if (!bFoundOverlap)
-	{
-		Accumulator.Insert(New, Idx);
-	}
-}
-
-void FShaderCompilerStats::Aggregate(FShaderCompilerStats& Other)
-{
-	// note: intentionally not taking local lock as this should only ever be called on a local copy of the stats object
-	FScopeLock Lock(&Other.CompileStatsLock);
-	Counters += Other.Counters;
-
-	for (TSparseArray<ShaderCompilerStats>::TConstIterator It(Other.CompileStats); It; ++It)
-	{
-		if (!CompileStats.IsValidIndex(It.GetIndex()))
-		{
-			CompileStats.EmplaceAt(It.GetIndex());
-		}
-
-		ShaderCompilerStats& Stats = CompileStats[It.GetIndex()];
-		for (const TPair<FString, FShaderStats>& StatsKeyValue : *It)
-		{
-			FShaderStats* Current = Stats.Find(StatsKeyValue.Key);
-			if (Current)
-			{
-				*Current += StatsKeyValue.Value;
-			}
-			else
-			{
-				Stats.Add(StatsKeyValue);
-			}
-		}
-	}
-
-	// note: this is suboptimal (O(n^2)) but there aren't a lot of these in practice
-	for (const TInterval<double>& Interval : Other.JobLifeTimeIntervals)
-	{
-		AddToInterval(JobLifeTimeIntervals, Interval);
-	}
-
-	for (const TPair<FString, FShaderTimings>& TimingsKeyValue : Other.ShaderTimings)
-	{
-		FShaderTimings* Current = ShaderTimings.Find(TimingsKeyValue.Key);
-		if (Current)
-		{
-			*Current += TimingsKeyValue.Value;
-		}
-		else
-		{
-			ShaderTimings.Add(TimingsKeyValue);
-		}
-	}
-
-	MaterialCounters += Other.MaterialCounters;
-}
-
-void FShaderCompilerStats::WriteToCompactBinary(FCbWriter& Writer)
-{
-	FScopeLock Lock(&CompileStatsLock);
-	Writer.AddBinary("Counters", &Counters, sizeof(Counters));
-
-	Writer.AddBinary("MaterialCounters", &MaterialCounters, sizeof(MaterialCounters));
-
-	Writer.BeginArray("CompileStatIndices");	
-	// Write the array of valid indices this worker has in the compile stats sparse array
-	for (TSparseArray<ShaderCompilerStats>::TConstIterator It(CompileStats); It; ++It)
-	{
-		if (CompileStats.IsValidIndex(It.GetIndex()))
-		{
-			Writer << It.GetIndex();
-		}
-	}
-	Writer.EndArray();
-
-	Writer.BeginArray("CompileStats");
-		// Then write the actual compile stats maps in the same order as the above indices
-	for (TSparseArray<ShaderCompilerStats>::TConstIterator It(CompileStats); It; ++It)
-	{
-		if (!CompileStats.IsValidIndex(It.GetIndex()))
-		{
-			continue;
-		}
-
-		Writer.BeginObject();
-		Writer.BeginArray("CompileStatsKeys");
-		for (TPair<FString, FShaderStats> Pair : *It)
-		{
-			Writer << Pair.Key;
-		}
-		Writer.EndArray();
-
-		Writer.BeginArray("CompileStatsValues");
-		for (const TPair<FString, FShaderStats>& Pair : *It)
-		{
-			Writer.BeginObject();
-			Writer << "Compiled" << Pair.Value.Compiled;
-			Writer << "CompiledDouble" << Pair.Value.CompiledDouble;
-			Writer << "CompileTime" << Pair.Value.CompileTime;
-			Writer << "Cooked" << Pair.Value.Cooked;
-			Writer << "CookedDouble" << Pair.Value.CookedDouble;
-			Writer.BeginArray("PermutationCompilations");
-			for (const FShaderCompilerSinglePermutationStat& Stat : Pair.Value.PermutationCompilations)
-			{
-				Writer.BeginObject();
-				Writer << "Compiled" << Stat.Compiled;
-				Writer << "CompiledDouble" << Stat.CompiledDouble;
-				Writer << "Cooked" << Stat.Cooked;
-				Writer << "CookedDouble" << Stat.CookedDouble;
-				Writer << "PermutationString" << Stat.PermutationString;
-				Writer.EndObject();
-			}
-			Writer.EndArray();
-			Writer.EndObject();
-		}
-		Writer.EndArray();
-		Writer.EndObject();
-	}
-	Writer.EndArray();
-
-	Writer.BeginArray("JobLifeTimeIntervals");
-	for (const TInterval<double>& Interval : JobLifeTimeIntervals)
-	{
-		Writer.AddBinary(&Interval, sizeof(TInterval<double>));
-	}
-	Writer.EndArray();
-
-	Writer.BeginArray("ShaderTimingsKeys");
-	for (const TPair<FString, FShaderTimings>& TimingPair : ShaderTimings)
-	{
-		Writer << TimingPair.Key;
-	}
-	Writer.EndArray();
-
-	Writer.BeginArray("ShaderTimingsValues");
-	for (const TPair<FString, FShaderTimings>& TimingPair : ShaderTimings)
-	{
-		Writer.AddBinary(&TimingPair.Value, sizeof(FShaderTimings));
-	}
-	Writer.EndArray();
-}
-
-void FShaderCompilerStats::ReadFromCompactBinary(FCbObjectView& Reader)
-{
-	FScopeLock Lock(&CompileStatsLock);
-	FMemoryView CountersMem = Reader["Counters"].AsBinaryView();
-	check(CountersMem.GetSize() == sizeof(FCounters));
-	Counters = *reinterpret_cast<const FCounters*>(CountersMem.GetData());
-
-	FMemoryView MaterialCountersMem = Reader["MaterialCounters"].AsBinaryView();
-	check(MaterialCountersMem.GetSize() == sizeof(FMaterialCounters));
-	MaterialCounters = *reinterpret_cast<const FMaterialCounters*>(MaterialCountersMem.GetData());
-
-	FCbArrayView CompileStatIndicesView = Reader["CompileStatIndices"].AsArrayView();
-	FCbArrayView CompileStatsView = Reader["CompileStats"].AsArrayView();
-	check(CompileStatIndicesView.Num() == CompileStatsView.Num());
-
-	FCbFieldViewIterator IndexIt = CompileStatIndicesView.CreateViewIterator();
-	FCbFieldViewIterator StatsIt = CompileStatsView.CreateViewIterator();
-
-	while (IndexIt && StatsIt)
-	{
-		if (!CompileStats.IsValidIndex(IndexIt.AsUInt32()))
-		{
-			FSparseArrayAllocationInfo AllocInfo = CompileStats.InsertUninitialized(IndexIt.AsUInt32());
-			new(AllocInfo) ShaderCompilerStats();
-		}
-		ShaderCompilerStats& Stats = CompileStats[IndexIt.AsUInt32()];
-
-		FCbObjectView PlatformStatsObject = StatsIt->AsObjectView();
-		FCbArrayView StatsKeysView = PlatformStatsObject["CompileStatsKeys"].AsArrayView();
-		FCbArrayView StatsValuesView = PlatformStatsObject["CompileStatsValues"].AsArrayView();
-		check(StatsKeysView.Num() == StatsValuesView.Num());
-
-		Stats.Reserve(StatsKeysView.Num());
-
-		FCbFieldViewIterator KeysIt = StatsKeysView.CreateViewIterator();
-		FCbFieldViewIterator ValuesIt = StatsValuesView.CreateViewIterator();
-
-		while (KeysIt && ValuesIt)
-		{
-			FCbObjectView ShaderStatsObject = ValuesIt->AsObjectView();
-			FShaderStats& ShaderStats = Stats.Add(FString(KeysIt->AsString()));
-			ShaderStats.Compiled = ShaderStatsObject["Compiled"].AsUInt32();
-			ShaderStats.CompiledDouble = ShaderStatsObject["CompiledDouble"].AsUInt32();
-			ShaderStats.CompileTime = ShaderStatsObject["CompileTime"].AsFloat();
-			ShaderStats.Cooked = ShaderStatsObject["Cooked"].AsUInt32();
-			ShaderStats.CookedDouble = ShaderStatsObject["CookedDouble"].AsUInt32();
-
-			FCbArrayView PermutationsArrayView = ShaderStatsObject["PermutationCompilations"].AsArrayView();
-			ShaderStats.PermutationCompilations.Reset(PermutationsArrayView.Num());
-			for (FCbFieldView CompilationField : PermutationsArrayView)
-			{
-				FCbObjectView PermutationObject = CompilationField.AsObjectView();
-				uint32 Index = ShaderStats.PermutationCompilations.Emplace
-				(
-					FString(PermutationObject["PermutationString"].AsString()),
-					PermutationObject["Compiled"].AsUInt32(),
-					PermutationObject["Cooked"].AsUInt32()
-				);
-				ShaderStats.PermutationCompilations[Index].CompiledDouble = PermutationObject["CompiledDouble"].AsUInt32();
-				ShaderStats.PermutationCompilations[Index].CookedDouble = PermutationObject["CookedDouble"].AsUInt32();
-			}
-
-			++ValuesIt;
-			++KeysIt;
-		}
-
-		++IndexIt;
-		++StatsIt;
-	}
-
-	FCbArrayView JobLifeTimeIntervalsView = Reader["JobLifeTimeIntervals"].AsArrayView();
-	JobLifeTimeIntervals.Reset(JobLifeTimeIntervalsView.Num());
-	for (FCbFieldView JobLifeTimeField : JobLifeTimeIntervalsView)
-	{
-		FMemoryView IntervalObj = JobLifeTimeField.AsBinaryView();
-		check(IntervalObj.GetSize() == sizeof(TInterval<double>));
-		JobLifeTimeIntervals.Add(*reinterpret_cast<const TInterval<double>*>(IntervalObj.GetData()));
-	}
-
-	FCbArrayView TimingsKeysView = Reader["ShaderTimingsKeys"].AsArrayView();
-	FCbArrayView TimingsValuesView = Reader["ShaderTimingsValues"].AsArrayView();
-	check(TimingsKeysView.Num() == TimingsValuesView.Num());
-
-	ShaderTimings.Reserve(TimingsKeysView.Num());
-
-	FCbFieldViewIterator TimingsKeysIt = TimingsKeysView.CreateViewIterator();
-	FCbFieldViewIterator TimingsValuesIt = TimingsValuesView.CreateViewIterator();
-
-	while (TimingsKeysIt && TimingsValuesIt)
-	{
-		FMemoryView TimingsValuesBinary = TimingsValuesIt->AsBinaryView();
-		check(TimingsValuesBinary.GetSize() == sizeof(FShaderTimings));
-		ShaderTimings.Add(FString(TimingsKeysIt->AsString()), *reinterpret_cast<const FShaderTimings*>(TimingsValuesBinary.GetData()));
-		++TimingsKeysIt;
-		++TimingsValuesIt;
-	}
-}
-
-void FShaderCompilerStats::RegisterLocalWorkerIdleTime(double IdleTime)
-{
-	FScopeLock Lock(&CompileStatsLock);
-	Counters.AccumulatedLocalWorkerIdleTime += IdleTime;
-	Counters.TimesLocalWorkersWereIdle++;
-}
-
-void FShaderCompilerStats::RegisterNewPendingJob(FShaderCommonCompileJob& Job)
-{
-	// accessing job timestamps isn't arbitrated by any lock. It is assumed that the registration of a job at one of the stages
-	// of its lifetime happens before the code can move it to another stage (i.e. new pending job is registered before it is added to the pending queue,
-	// so it cannot be given away to a worker while it's still being registered, and an assigned job is registered before it is actually given to the worker,
-	// so it cannot end up being registered as finished at the same time on some other thread).
-	Job.TimeAddedToPendingQueue = FPlatformTime::Seconds();
-}
-
-void FShaderCompilerStats::RegisterAssignedJob(FShaderCommonCompileJob& Job)
-{
-	ensure(Job.TimeAddedToPendingQueue != 0.0);
-	Job.TimeAssignedToExecution = FPlatformTime::Seconds();
-
-	FScopeLock Lock(&CompileStatsLock);
-	Counters.JobsAssigned++;
-	double TimeSpendPending = (Job.TimeAssignedToExecution - Job.TimeAddedToPendingQueue);
-	Counters.AccumulatedPendingTime += TimeSpendPending;
-	Counters.MaxPendingTime = FMath::Max(TimeSpendPending, Counters.MaxPendingTime);
-}
-
-TRACE_DECLARE_INT_COUNTER(Shaders_Compiled, TEXT("Shaders/Compiled"));
-void FShaderCompilerStats::RegisterFinishedJob(FShaderCommonCompileJob& Job, bool bCompilationSkipped)
-{
-	FScopeLock Lock(&CompileStatsLock);
-	
-	if (!bCompilationSkipped)
-	{
-		ensure(Job.TimeAssignedToExecution != 0.0);
-		Job.TimeExecutionCompleted = FPlatformTime::Seconds();
-		TRACE_COUNTER_ADD(Shaders_Compiled, 1);
-		Counters.JobsCompleted++;
-
-		double ExecutionTime = (Job.TimeExecutionCompleted - Job.TimeAssignedToExecution);
-		Counters.AccumulatedJobExecutionTime += ExecutionTime;
-		Counters.MaxJobExecutionTime = FMath::Max(ExecutionTime, Counters.MaxJobExecutionTime);
-
-		double LifeTime = (Job.TimeExecutionCompleted - Job.TimeAddedToPendingQueue);
-		Counters.AccumulatedJobLifeTime += LifeTime;
-		Counters.MaxJobLifeTime = FMath::Max(LifeTime, Counters.MaxJobLifeTime);
-		
-		// estimate lifetime without an overlap
-		ensure(Job.TimeAddedToPendingQueue != 0.0 && Job.TimeAddedToPendingQueue <= Job.TimeExecutionCompleted);
-		AddToInterval(JobLifeTimeIntervals, TInterval<double>(Job.TimeAddedToPendingQueue, Job.TimeExecutionCompleted));
-	}
-	
-	if (Job.TimeTaskSubmitJobs)
-	{
-		Counters.AccumulatedTaskSubmitJobs += Job.TimeTaskSubmitJobs;
-		Counters.AccumulatedTaskSubmitJobsStall += Job.TimeTaskSubmitJobsStall;
-	}
-
-	auto RegisterStatsFromSingleJob = [this, bCompilationSkipped](const FShaderCompileJob& SingleJob)
-	{
-		// Register min/max/average shader code sizes for single job output
-		const int32 ShaderCodeSize = SingleJob.Output.ShaderCode.GetShaderCodeSize();
-		if (!bCompilationSkipped && ShaderCodeSize > 0)
-		{
-			Counters.MinShaderCodeSize = (Counters.MinShaderCodeSize > 0 ? FMath::Min(Counters.MinShaderCodeSize, ShaderCodeSize) : ShaderCodeSize);
-			Counters.MaxShaderCodeSize = (Counters.MaxShaderCodeSize > 0 ? FMath::Max(Counters.MaxShaderCodeSize, ShaderCodeSize) : ShaderCodeSize);
-			Counters.AccumulatedShaderCodeSize += (uint64)ShaderCodeSize;
-			++Counters.NumAccumulatedShaderCodes;
-		}
-
-		// Sanity check; compile time should be 0 for cache hits
-		check(!bCompilationSkipped || SingleJob.Output.CompileTime == 0.0f);
-		// Preprocess time should always be non-zero if preprocessing succeeded; note that preprocessing for pipeline stage jobs may be skipped 
-		// in the case preprocessing a preceding stage of the pipeline failed
-		check(!SingleJob.PreprocessOutput.GetSucceeded() || SingleJob.Output.PreprocessTime > 0.0f);
-
-		const FString ShaderName(SingleJob.Key.ShaderType->GetName());
-		if (FShaderTimings* Existing = ShaderTimings.Find(ShaderName))
-		{
-			// Always want to log preprocess time, in case preprocessed cache is enabled and preprocessing ran in the cooker prior to compilation
-			// (PreprocessTime will be 0 if preprocessed cache is disabled)
-			Existing->TotalPreprocessTime += SingleJob.Output.PreprocessTime;
-			if (!bCompilationSkipped)
-			{
-				// If no actual compiles have been logged yet, min compile time is just the compile time of this job (first to actually run)
-				Existing->MinCompileTime = Existing->NumCompiled ? FMath::Min(Existing->MinCompileTime, static_cast<float>(SingleJob.Output.CompileTime)) : SingleJob.Output.CompileTime;
-				Existing->MaxCompileTime = FMath::Max(Existing->MaxCompileTime, static_cast<float>(SingleJob.Output.CompileTime));
-				Existing->TotalCompileTime += SingleJob.Output.CompileTime;
-				Existing->NumCompiled++;
-				// calculate as an optimization to make sorting later faster
-				Existing->AverageCompileTime = Existing->TotalCompileTime / static_cast<float>(Existing->NumCompiled);
-			}
-		}
-		else
-		{
-			FShaderTimings New;
-			New.MinCompileTime = SingleJob.Output.CompileTime;
-			New.MaxCompileTime = New.MinCompileTime;
-			New.TotalCompileTime = New.MinCompileTime;
-			New.AverageCompileTime = New.MinCompileTime;
-			// It's possible the first entry for a given shader didn't actually compile (i.e. hit in DDC)
-			// so we need to account for that in the stats
-			New.NumCompiled = bCompilationSkipped ? 0 : 1;
-			New.TotalPreprocessTime += SingleJob.Output.PreprocessTime;
-
-			ShaderTimings.Add(ShaderName, New);
-		}
-	};
-
-	Job.ForEachSingleShaderJob(RegisterStatsFromSingleJob);
-}
-
-void FShaderCompilerStats::RegisterJobBatch(int32 NumJobs, EExecutionType ExecType)
-{
-	if (ExecType == EExecutionType::Local)
-	{
-		FScopeLock Lock(&CompileStatsLock);
-		++Counters.LocalJobBatchesSeen;
-		Counters.TotalJobsReportedInLocalJobBatches += NumJobs;
-	}
-	else if (ExecType == EExecutionType::Distributed)
-	{
-		FScopeLock Lock(&CompileStatsLock);
-		++Counters.DistributedJobBatchesSeen;
-		Counters.TotalJobsReportedInDistributedJobBatches += NumJobs;
-	}
-	else
-	{
-		checkNoEntry();
-	}
-}
-
-void FShaderCompilerStats::RegisterDistributedBuildStats(const FDistributedBuildStats& InStats)
-{
-	FScopeLock Lock(&CompileStatsLock);
-	Counters.MaxRemoteAgents = FMath::Max(Counters.MaxRemoteAgents, InStats.MaxRemoteAgents);
-	Counters.MaxActiveAgentCores = FMath::Max(Counters.MaxActiveAgentCores, InStats.MaxActiveAgentCores);
-}
-
-void FShaderCompilerStats::FMaterialCounters::WriteStatSummary(const TCHAR* AggregatedSuffix)
-{
-	auto CalcTimePercentage = [&](double Val) {
-		return  (int)round(Val / FMath::Max(1e-6, MaterialTranslateTotalTimeSec) * 100);
-	};
-
-	UE_LOG(LogShaderCompilers, Display, TEXT("=== Material stats%s ==="), AggregatedSuffix);
-	UE_LOG(LogShaderCompilers, Display, TEXT("Materials Cooked:        %d"), NumMaterialsCooked);
-	UE_LOG(LogShaderCompilers, Display, TEXT("Materials Translated:    %d"), MaterialTranslateCalls);
-	UE_LOG(LogShaderCompilers, Display, TEXT("Material Total Translate Time: %.2f s"), MaterialTranslateTotalTimeSec);
-	UE_LOG(LogShaderCompilers, Display, TEXT("Material Translation Only: %.2f s (%d%%)"), MaterialTranslateTranslationOnlyTimeSec, CalcTimePercentage(MaterialTranslateTranslationOnlyTimeSec));
-	UE_LOG(LogShaderCompilers, Display, TEXT("Material DDC Serialization Only: %.2f s (%d%%)"), MaterialTranslateSerializationOnlyTimeSec, CalcTimePercentage(MaterialTranslateSerializationOnlyTimeSec));
-
-	int HitsPercentage = MaterialTranslateCalls ? (int)roundf(float(MaterialCacheHits) / MaterialTranslateCalls * 100) : 0;
-	UE_LOG(LogShaderCompilers, Display, TEXT("Material Cache Hits: %d (%d%%)"), MaterialCacheHits, HitsPercentage);
-}
-
-void FShaderCompilerStats::FMaterialCounters::GatherAnalytics(TArray<FAnalyticsEventAttribute>& Attributes)
-{
-	Attributes.Emplace(TEXT("Material_NumMaterialsCooked"), NumMaterialsCooked);
-	Attributes.Emplace(TEXT("Material_MaterialTranslateCalls"), MaterialTranslateCalls);
-	Attributes.Emplace(TEXT("Material_MaterialTranslateTimeSec"), MaterialTranslateTotalTimeSec);
-	Attributes.Emplace(TEXT("Material_MaterialTranslateTranslationOnlyTimeSec"), MaterialTranslateTranslationOnlyTimeSec);
-	Attributes.Emplace(TEXT("Material_MaterialTranslateSerializationOnlyTimeSec"), MaterialTranslateSerializationOnlyTimeSec);
-	Attributes.Emplace(TEXT("Material_MaterialCacheHits"), MaterialCacheHits);
-}
-
-void FShaderCompilerStats::IncrementMaterialCook()
-{
-	FScopeLock Lock(&CompileStatsLock);
-	MaterialCounters.NumMaterialsCooked++;
-}
-
-void FShaderCompilerStats::IncrementMaterialTranslated(double InTotalTime, double InTranslationOnlyTime, double InSerializeTime)
-{
-	FScopeLock Lock(&CompileStatsLock);
-	MaterialCounters.MaterialTranslateCalls++;
-	MaterialCounters.MaterialTranslateTotalTimeSec += InTotalTime;
-	MaterialCounters.MaterialTranslateTranslationOnlyTimeSec += InTranslationOnlyTime;
-	MaterialCounters.MaterialTranslateSerializationOnlyTimeSec += InSerializeTime;
-}
-
-void FShaderCompilerStats::IncrementMaterialCacheHit()
-{
-	FScopeLock Lock(&CompileStatsLock);
-	MaterialCounters.MaterialCacheHits++;
-}
-
-void FShaderCompilerStats::RegisterCookedShaders(uint32 NumCooked, float CompileTime, EShaderPlatform Platform, const FString MaterialPath, FString PermutationString)
-{
-	FScopeLock Lock(&CompileStatsLock);
-	if(!CompileStats.IsValidIndex(Platform))
-	{
-		ShaderCompilerStats Stats;
-		CompileStats.Insert(Platform, Stats);
-	}
-
-	FShaderCompilerStats::FShaderStats& Stats = CompileStats[Platform].FindOrAdd(MaterialPath);
-	Stats.CompileTime += CompileTime;
-	bool bFound = false;
-	for (FShaderCompilerSinglePermutationStat& Stat : Stats.PermutationCompilations)
-	{
-		if (PermutationString == Stat.PermutationString)
-		{
-			bFound = true;
-			if (Stat.Cooked != 0)
-			{
-				Stat.CookedDouble += NumCooked;
-				Stats.CookedDouble += NumCooked;
-			}
-			else
-			{
-				Stat.Cooked = NumCooked;
-				Stats.Cooked += NumCooked;
-			}
-		}
-	}
-	if(!bFound)
-	{
-		Stats.Cooked += NumCooked;
-	}
-	if (!bFound)
-	{
-		Stats.PermutationCompilations.Emplace(PermutationString, 0, NumCooked);
-	}
-}
-
-void FShaderCompilerStats::RegisterCompiledShaders(uint32 NumCompiled, EShaderPlatform Platform, const FString MaterialPath, FString PermutationString)
-{
-	FScopeLock Lock(&CompileStatsLock);
-	if (!CompileStats.IsValidIndex(Platform))
-	{
-		ShaderCompilerStats Stats;
-		CompileStats.Insert(Platform, Stats);
-	}
-	FShaderCompilerStats::FShaderStats& Stats = CompileStats[Platform].FindOrAdd(MaterialPath);
-
-	bool bFound = false;
-	for (FShaderCompilerSinglePermutationStat& Stat : Stats.PermutationCompilations)
-	{
-		if (PermutationString == Stat.PermutationString)
-		{
-			bFound = true;
-			if (Stat.Compiled != 0)
-			{
-				Stat.CompiledDouble += NumCompiled;
-				Stats.CompiledDouble += NumCompiled;
-			}
-			else
-			{
-				Stat.Compiled = NumCompiled;
-				Stats.Compiled += NumCompiled;
-			}
-		}
-	}
-	if(!bFound)
-	{
-		Stats.Compiled += NumCompiled;
-	}
-
-
-	if (!bFound)
-	{
-		Stats.PermutationCompilations.Emplace(PermutationString, NumCompiled, 0);
-	}
-}
-
-void FShaderCompilerStats::AddDDCMiss(uint32 NumMisses)
-{
-	Counters.ShaderMapDDCMisses += NumMisses;
-}
-
-uint32 FShaderCompilerStats::GetDDCMisses() const
-{
-	return Counters.ShaderMapDDCMisses;
-}
-
-void FShaderCompilerStats::AddDDCHit(uint32 NumHits)
-{
-	Counters.ShaderMapDDCHits += NumHits;
-}
-
-uint32 FShaderCompilerStats::GetDDCHits() const
-{
-	return Counters.ShaderMapDDCHits;
-}
-
-double FShaderCompilerStats::GetTimeShaderCompilationWasActive()
-{
-	FScopeLock Lock(&CompileStatsLock);
-	double Sum = 0;
-	for (int32 Idx = 0; Idx < JobLifeTimeIntervals.Num(); ++Idx)
-	{
-		const TInterval<double>& Existing = JobLifeTimeIntervals[Idx];
-		Sum += Existing.Size();
-	}
-	return Sum;
-}
-
-FShaderCompilingManager* GShaderCompilingManager = nullptr;
-
-bool FShaderCompilingManager::AllTargetPlatformSupportsRemoteShaderCompiling()
-{
-	// no compiling support
-	if (!AllowShaderCompiling())
-	{
-		return false;
-	}
-
-	ITargetPlatformManagerModule* TPM = GetTargetPlatformManager();	
-	if (!TPM)
-	{
-		return false;
-	}
-	
-	const TArray<ITargetPlatform*>& Platforms = TPM->GetActiveTargetPlatforms();
-	for (int32 Index = 0; Index < Platforms.Num(); Index++)
-	{
-		if (!Platforms[Index]->CanSupportRemoteShaderCompile())
-		{
-			return false;
-		}
-	}
-	
-	return true;
-}
-
-// Returns a rank for the preference of distributed shader controllers; Higher is better.
-static int32 GetShaderControllerPreferenceRank(IDistributedBuildController& Controller)
-{
-	const FString Name = Controller.GetName();
-	if (Name.StartsWith(TEXT("UBA")))
-	{
-		return 2;
-	}
-	else if (Name.StartsWith(TEXT("XGE")))
-	{
-		return 1;
-	}
-	else
-	{
-		return 0;
-	}
-}
-
-IDistributedBuildController* FShaderCompilingManager::FindRemoteCompilerController() const
-{
-	// no controllers needed if not compiling
-	if (!AllowShaderCompiling())
-	{
-		return nullptr;
-	}
-
-	TArray<IDistributedBuildController*> AvailableControllers = IModularFeatures::Get().GetModularFeatureImplementations<IDistributedBuildController>(IDistributedBuildController::GetModularFeatureType());
-
-	// Prefer UBA, then XGE, and fallback to any other controller otherwise
-	int32 SupportedControllerPreferenceRank = 0;
-	IDistributedBuildController* SupportedController = nullptr;
-
-	for (IDistributedBuildController* Controller : AvailableControllers)
-	{
-		if (Controller != nullptr && Controller->IsSupported())
-		{
-			const int32 PreferenceRank = GetShaderControllerPreferenceRank(*Controller);
-			if (SupportedController == nullptr || SupportedControllerPreferenceRank < PreferenceRank)
-			{
-				SupportedController = Controller;
-				SupportedControllerPreferenceRank = PreferenceRank;
-			}
-		}
-	}
-
-	if (SupportedController != nullptr)
-	{
-		SupportedController->InitializeController();
-		return SupportedController;
-	}
-
-	return nullptr;
-}
-
-void FShaderCompilingManager::ReportMemoryUsage()
-{
-	// This function runs from within an OOM callback. It should not take locks, as much as possible.
-	constexpr bool bAllowToWaitForLock = false;
-	for (const TUniquePtr<FShaderCompileThreadRunnableBase>& ThreadPtr : Threads)
-	{
-		ThreadPtr->PrintWorkerMemoryUsage(bAllowToWaitForLock);
-	}
-}
-
-FShaderCompilingManager::FShaderCompilingManager() :
-	bCompilingDuringGame(false),
-	NumExternalJobs(0),
-	AllJobs(CompileQueueSection),
-	NumSingleThreadedRunsBeforeRetry(GSingleThreadedRunsIdle),
-	SuppressedShaderPlatforms(0),
-	BuildDistributionController(nullptr),
-	bNoShaderCompilation(false),
-	bAllowForIncompleteShaderMaps(false),
-	Notification(MakeUnique<FAsyncCompilationNotification>(GetAssetNameFormat()))
-{
-	// don't perform any initialization if compiling is not allowed
-	if (!AllowShaderCompiling())
-	{
-		// use existing flag to disable compiling
-		bNoShaderCompilation = true;
-		return;
-	}
-
-	bIsEngineLoopInitialized = false;
-	FCoreDelegates::OnFEngineLoopInitComplete.AddLambda([&]() 
-		{ 
-			bIsEngineLoopInitialized = true; 
-		}
-	);
-
-	WorkersBusyTime = 0;
-
-#if PLATFORM_WINDOWS
-	FString ExecutableName("ShaderCompileWorker.exe");
-#else
-	FString ExecutableName("ShaderCompileWorker");
-#endif
-
-	// first look for project-specific version
-	ShaderCompileWorkerName = FPaths::Combine(FPaths::ProjectDir(), TEXT("Binaries"), FPlatformProcess::GetBinariesSubdirectory(), ExecutableName);
-	if (!IFileManager::Get().FileExists(*ShaderCompileWorkerName))
-	{
-		// fallback to standard Engine location
-		ShaderCompileWorkerName = FPaths::Combine(FPaths::EngineDir(), TEXT("Binaries"), FPlatformProcess::GetBinariesSubdirectory(), ExecutableName);
-	}
-
-	// Threads must use absolute paths on Windows in case the current directory is changed on another thread!
-	ShaderCompileWorkerName = FPaths::ConvertRelativePathToFull(ShaderCompileWorkerName);
-
-	// Read values from the engine ini
-	verify(GConfig->GetBool( TEXT("DevOptions.Shaders"), TEXT("bAllowCompilingThroughWorkers"), bAllowCompilingThroughWorkers, GEngineIni ));
-	verify(GConfig->GetBool( TEXT("DevOptions.Shaders"), TEXT("bAllowAsynchronousShaderCompiling"), bAllowAsynchronousShaderCompiling, GEngineIni ));
-
-	// Explicitly load ShaderPreprocessor module so it will run its initialization step
-	FModuleManager::LoadModuleChecked<IModuleInterface>(TEXT("ShaderPreprocessor"));
-	
-	// override the use of workers, can be helpful for debugging shader compiler code
-	static const IConsoleVariable* CVarAllowCompilingThroughWorkers = IConsoleManager::Get().FindConsoleVariable(TEXT("r.Shaders.AllowCompilingThroughWorkers"), false);
-	if (!FPlatformProcess::SupportsMultithreading() || FParse::Param(FCommandLine::Get(), TEXT("noshaderworker")) || (CVarAllowCompilingThroughWorkers && CVarAllowCompilingThroughWorkers->GetInt() == 0))
-	{
-		bAllowCompilingThroughWorkers = false;
-	}
-
-	if (!FPlatformProcess::SupportsMultithreading())
-	{
-		bAllowAsynchronousShaderCompiling = false;
-	}
-
-	verify(GConfig->GetInt( TEXT("DevOptions.Shaders"), TEXT("MaxShaderJobBatchSize"), MaxShaderJobBatchSize, GEngineIni ));
-	verify(GConfig->GetBool( TEXT("DevOptions.Shaders"), TEXT("bPromptToRetryFailedShaderCompiles"), bPromptToRetryFailedShaderCompiles, GEngineIni ));
-	verify(GConfig->GetBool(TEXT("DevOptions.Shaders"), TEXT("bDebugBreakOnPromptToRetryShaderCompile"), bDebugBreakOnPromptToRetryShaderCompile, GEngineIni));
-	verify(GConfig->GetBool( TEXT("DevOptions.Shaders"), TEXT("bLogJobCompletionTimes"), bLogJobCompletionTimes, GEngineIni ));
-	GConfig->GetFloat(TEXT("DevOptions.Shaders"), TEXT("WorkerTimeToLive"), GRegularWorkerTimeToLive, GEngineIni);
-	GConfig->GetFloat(TEXT("DevOptions.Shaders"), TEXT("BuildWorkerTimeToLive"), GBuildWorkerTimeToLive, GEngineIni);
-
-	verify(GConfig->GetFloat( TEXT("DevOptions.Shaders"), TEXT("ProcessGameThreadTargetTime"), ProcessGameThreadTargetTime, GEngineIni ));
-
-#if UE_BUILD_DEBUG
-	// Increase budget for processing results in debug or else it takes forever to finish due to poor framerate
-	ProcessGameThreadTargetTime *= 3;
-#endif
-
-	// Get the current process Id, this will be used by the worker app to shut down when it's parent is no longer running.
-	ProcessId = FPlatformProcess::GetCurrentProcessId();
-
-	// Use a working directory unique to this game, process and thread so that it will not conflict 
-	// With processes from other games, processes from the same game or threads in this same process.
-	// Use IFileManager to do path conversion to properly handle sandbox paths (outside of standard paths in particular).
-	//ShaderBaseWorkingDirectory = FPlatformProcess::ShaderWorkingDir() / FString::FromInt(ProcessId) + TEXT("/");
-
-	{
-		FGuid Guid;
-		Guid = FGuid::NewGuid();
-		FString LegacyShaderWorkingDirectory = FPaths::ProjectIntermediateDir() / TEXT("Shaders/WorkingDirectory/")  / FString::FromInt(ProcessId) + TEXT("/");
-		ShaderBaseWorkingDirectory = FPaths::ShaderWorkingDir() / *Guid.ToString(EGuidFormats::Digits) + TEXT("/");
-		UE_LOG(LogShaderCompilers, Log, TEXT("Guid format shader working directory is %d characters bigger than the processId version (%s)."), ShaderBaseWorkingDirectory.Len() - LegacyShaderWorkingDirectory.Len(), *LegacyShaderWorkingDirectory );
-	}
-
-	if (!IFileManager::Get().DeleteDirectory(*ShaderBaseWorkingDirectory, false, true))
-	{
-		UE_LOG(LogShaderCompilers, Fatal, TEXT("Could not delete the shader compiler working directory '%s'."), *ShaderBaseWorkingDirectory);
-	}
-	else
-	{
-		UE_LOG(LogShaderCompilers, Log, TEXT("Cleaned the shader compiler working directory '%s'."), *ShaderBaseWorkingDirectory);
-	}
-	FString AbsoluteBaseDirectory = IFileManager::Get().ConvertToAbsolutePathForExternalAppForWrite(*ShaderBaseWorkingDirectory);
-	FPaths::NormalizeDirectoryName(AbsoluteBaseDirectory);
-	AbsoluteShaderBaseWorkingDirectory = AbsoluteBaseDirectory + TEXT("/");
-
-	// Build machines should dump to the AutomationTool/Saved/Logs directory and they will upload as build artifacts via the AutomationTool.
-	FString BaseDebugInfoPath = FPaths::ProjectSavedDir();
-	if (GIsBuildMachine)
-	{
-		BaseDebugInfoPath = FPaths::Combine(*FPaths::EngineDir(), TEXT("Programs"), TEXT("AutomationTool"), TEXT("Saved"), TEXT("Logs"));
-	}
-
-	FString AbsoluteDebugInfoDirectory = IFileManager::Get().ConvertToAbsolutePathForExternalAppForWrite(*(BaseDebugInfoPath / TEXT("ShaderDebugInfo")));
-	const FString OverrideShaderDebugDir = CVarShaderOverrideDebugDir.GetValueOnAnyThread();
-	if (!OverrideShaderDebugDir.IsEmpty())
-	{
-		AbsoluteDebugInfoDirectory = OverrideShaderDebugDir;
-	}
-	FPaths::NormalizeDirectoryName(AbsoluteDebugInfoDirectory);
-	AbsoluteShaderDebugInfoDirectory = AbsoluteDebugInfoDirectory;
-
-	CalculateNumberOfCompilingThreads(FPlatformMisc::NumberOfCores(), FPlatformMisc::NumberOfCoresIncludingHyperthreads());
-
-	TUniquePtr<FShaderCompileThreadRunnableBase> RemoteCompileThread;
-	const bool bCanUseRemoteCompiling = bAllowCompilingThroughWorkers && ShaderCompiler::IsRemoteCompilingAllowed() && AllTargetPlatformSupportsRemoteShaderCompiling();
-	BuildDistributionController = bCanUseRemoteCompiling ? FindRemoteCompilerController() : nullptr;
-	
-	if (BuildDistributionController)
-	{
-		UE_LOG(LogShaderCompilers, Display, TEXT("Using %s for Shader Compilation."), *BuildDistributionController->GetName());
-		RemoteCompileThread = MakeUnique<FShaderCompileDistributedThreadRunnable_Interface>(this, *BuildDistributionController);
-	}
-
-	GConfig->SetBool(TEXT("/Script/UnrealEd.UnrealEdOptions"), TEXT("UsingXGE"), RemoteCompileThread.IsValid(), GEditorIni);
-
-	TUniquePtr<FShaderCompileThreadRunnableBase> LocalThread = MakeUnique<FShaderCompileThreadRunnable>(this);
-	if (RemoteCompileThread)
-	{
-		checkf(ShaderCompiler::IsRemoteCompilingAllowed(), TEXT("We have a remote compiling thread without the remote compilation being allowed"));
-
-		// Only force-local jobs are guaranteed to stay on the local machine. Going wide with High priority jobs is important for the startup times,
-		// since special materials use High priority. Possibly the partition by priority is too rigid in general.
-		RemoteCompileThread->SetPriorityRange(EShaderCompileJobPriority::Low, EShaderCompileJobPriority::High);
-		LocalThread->SetPriorityRange(EShaderCompileJobPriority::Normal, EShaderCompileJobPriority::ForceLocal);
-		Threads.Add(MoveTemp(RemoteCompileThread));
-	}
-	else
-	{
-		UE_LOG(LogShaderCompilers, Display, TEXT("Using Local Shader Compiler with %d workers."), NumShaderCompilingThreads);
-
-		if (GIsBuildMachine)
-		{
-			int32 MinSCWsToSpawnBeforeWarning = 8; // optional, default to 8
-			GConfig->GetInt(TEXT("DevOptions.Shaders"), TEXT("MinSCWsToSpawnBeforeWarning"), MinSCWsToSpawnBeforeWarning, GEngineIni);
-			if (NumShaderCompilingThreads < static_cast<uint32>(MinSCWsToSpawnBeforeWarning))
-			{
-				UE_LOG(LogShaderCompilers, Warning, TEXT("Only %d SCWs will be spawned, which will result in longer shader compile times."), NumShaderCompilingThreads);
-			}
-		}
-	}
-	Threads.Add(MoveTemp(LocalThread));
-
-	for (const auto& Thread : Threads)
-	{
-		Thread->StartThread();
-	}
-
-	OutOfMemoryDelegateHandle = FCoreDelegates::GetOutOfMemoryDelegate().AddRaw(this, &FShaderCompilingManager::ReportMemoryUsage);
-
-	FAssetCompilingManager::Get().RegisterManager(this);
-
-#if WITH_EDITOR
-	static const bool bAllowShaderRecompileOnSave = CVarRecompileShadersOnSave.GetValueOnAnyThread();
-	if (bAllowShaderRecompileOnSave)
-	{
-		if (IDirectoryWatcher* DirectoryWatcher = FModuleManager::LoadModuleChecked<FDirectoryWatcherModule>(TEXT("DirectoryWatcher")).Get())
-		{
-			// Handle if we are watching a directory for changes.
-			{
-				UE_LOG(LogShaderCompilers, Display, TEXT("Register directory watchers for shader files."));
-
-				const TMap<FString, FString>& ShaderSourceDirectoryMappings = AllShaderSourceDirectoryMappings();
-
-				DirectoryWatcherHandles.Reserve(ShaderSourceDirectoryMappings.Num());
-
-				for (const auto& It : ShaderSourceDirectoryMappings)
-				{
-					FString DirectoryToWatch = It.Value;
-					if (FPaths::IsRelative(DirectoryToWatch))
-					{
-						DirectoryToWatch = FPaths::ConvertRelativePathToFull(DirectoryToWatch);
-					}
-
-					FDelegateHandle& DirectoryWatcherHandle = DirectoryWatcherHandles.Add(DirectoryToWatch, FDelegateHandle());
-
-					DirectoryWatcher->RegisterDirectoryChangedCallback_Handle(
-						DirectoryToWatch,
-						IDirectoryWatcher::FDirectoryChanged::CreateLambda([](const TArray<FFileChangeData>& InFileChangeDatas) {
-
-							TRACE_CPUPROFILER_EVENT_SCOPE(HandleDirectoryChanged);
-
-							if (!bAllowShaderRecompileOnSave)
-							{
-								return;
-							}
-
-							TArray<FString> ChangedShaderFiles;
-							for (const FFileChangeData& It : InFileChangeDatas)
-							{
-								if (It.Filename.EndsWith(TEXT(".usf")) || It.Filename.EndsWith(TEXT(".ush")) || It.Filename.EndsWith(TEXT(".h")))
-								{
-									UE_LOG(LogShaderCompilers, Display, TEXT("Detected change on %s"), *It.Filename);
-
-									ChangedShaderFiles.AddUnique(It.Filename);
-								}
-							}
-
-							if (ChangedShaderFiles.Num())
-							{
-								// Mappings from:
-								// Key:   /Engine to
-								// Value: ../../../Engine/Shaders
-								const TMap<FString, FString>& ShaderSourceDirectoryMappings = AllShaderSourceDirectoryMappings();
-
-								FString RemappedShaderFileName;
-								for (const auto& It : ShaderSourceDirectoryMappings)
-								{
-									// ChangedShaderFiles will be of format: ../../../Engine/Shaders/Private/PostProcessGBufferHints.usf
-									if (ChangedShaderFiles[0].StartsWith(It.Value))
-									{
-										// Change from relative path to Engine absolute path.
-										// i.e. change `../../../Engine/Shaders/Private/PostProcessGBufferHints.usf` to `/Engine/Shaders/Private/PostProcessGBufferHints.usf`
-										RemappedShaderFileName = ChangedShaderFiles[0].Replace(*It.Value, *It.Key);
-									}
-								}
-
-								// Issue a `recompileshaders /Engine/Shaders/Private/PostProcessGBufferHints.usf` command, which will just compile that shader source file.
-								RecompileShaders(*RemappedShaderFileName, *GLog);
-
-								UE_LOG(LogShaderCompilers, Display, TEXT("Ready for new shader file changes"));
-							}
-						}),
-						DirectoryWatcherHandle);
-
-					if (DirectoryWatcherHandle.IsValid())
-					{
-						UE_LOG(LogShaderCompilers, Display, TEXT("Watching %s -> %s"), *It.Key, *DirectoryToWatch);
-					}
-					else
-					{
-						UE_LOG(LogShaderCompilers, Error, TEXT("Failed to set up directory watcher %s -> %s"), *It.Key, *DirectoryToWatch);
-					}
-				}
-			}
-		}
-	}
-#endif // WITH_EDITOR
-}
-
-FShaderCompilingManager::~FShaderCompilingManager()
-{
-	// we never initialized, so nothing to do
-	if (!AllowShaderCompiling())
-	{
-		return;
-	}
-
-	for (const auto& Thread : Threads)
-	{
-		Thread->Stop();
-		Thread->WaitForCompletion();
-	}
-
-	FCoreDelegates::GetOutOfMemoryDelegate().Remove(OutOfMemoryDelegateHandle);
-
-#if WITH_EDITOR
-	const bool bAllowShaderRecompileOnSave = CVarRecompileShadersOnSave.GetValueOnAnyThread();
-	if (bAllowShaderRecompileOnSave)
-	{
-		if (IDirectoryWatcher* DirectoryWatcher = FModuleManager::LoadModuleChecked<FDirectoryWatcherModule>(TEXT("DirectoryWatcher")).Get())
-		{
-			for (const auto& It : DirectoryWatcherHandles)
-			{
-				DirectoryWatcher->UnregisterDirectoryChangedCallback_Handle(It.Key, It.Value);
-			}
-		}
-	}
-#endif // WITH_EDITOR
-
-	FAssetCompilingManager::Get().UnregisterManager(this);
-}
-
-void FShaderCompilingManager::CalculateNumberOfCompilingThreads(int32 NumberOfCores, int32 NumberOfCoresIncludingHyperthreads)
-{
-	const int32 NumVirtualCores = NumberOfCoresIncludingHyperthreads;
-
-	int32 NumUnusedShaderCompilingThreads;
-	verify(GConfig->GetInt(TEXT("DevOptions.Shaders"), TEXT("NumUnusedShaderCompilingThreads"), NumUnusedShaderCompilingThreads, GEngineIni));
-
-	int32 NumUnusedShaderCompilingThreadsDuringGame;
-	verify(GConfig->GetInt(TEXT("DevOptions.Shaders"), TEXT("NumUnusedShaderCompilingThreadsDuringGame"), NumUnusedShaderCompilingThreadsDuringGame, GEngineIni));
-
-	int32 ShaderCompilerCoreCountThreshold;
-	verify(GConfig->GetInt(TEXT("DevOptions.Shaders"), TEXT("ShaderCompilerCoreCountThreshold"), ShaderCompilerCoreCountThreshold, GEngineIni));
-
-	bool bForceUseSCWMemoryPressureLimits = false;
-	GConfig->GetBool(TEXT("DevOptions.Shaders"), TEXT("bForceUseSCWMemoryPressureLimits"), bForceUseSCWMemoryPressureLimits, GEngineIni);
-
-	// Don't reserve threads based on a percentage if we are in a commandlet or on a low core machine.
-	// In these scenarios we should try to use as many threads as possible.
-	if (!IsRunningCommandlet() && !GIsBuildMachine && NumVirtualCores > ShaderCompilerCoreCountThreshold)
-	{
-		// Reserve a percentage of the threads for general background work.
-		float PercentageUnusedShaderCompilingThreads;
-		verify(GConfig->GetFloat(TEXT("DevOptions.Shaders"), TEXT("PercentageUnusedShaderCompilingThreads"), PercentageUnusedShaderCompilingThreads, GEngineIni));
-
-		// ensure we get a valid multiplier.
-		PercentageUnusedShaderCompilingThreads = FMath::Clamp(PercentageUnusedShaderCompilingThreads, 0.0f, 100.0f) / 100.0f;
-
-		NumUnusedShaderCompilingThreads = FMath::CeilToInt(NumVirtualCores * PercentageUnusedShaderCompilingThreads);
-		NumUnusedShaderCompilingThreadsDuringGame = NumUnusedShaderCompilingThreads;
-	}
-
-	// Use all the cores on the build machines.
-	if (GForceAllCoresForShaderCompiling != 0)
-	{
-		NumUnusedShaderCompilingThreads = 0;
-	}
-
-	NumShaderCompilingThreads = (bAllowCompilingThroughWorkers && NumVirtualCores > NumUnusedShaderCompilingThreads) ? (NumVirtualCores - NumUnusedShaderCompilingThreads) : 1;
-
-	// Make sure there's at least one worker allowed to be active when compiling during the game
-	NumShaderCompilingThreadsDuringGame = (bAllowCompilingThroughWorkers && NumVirtualCores > NumUnusedShaderCompilingThreadsDuringGame) ? (NumVirtualCores - NumUnusedShaderCompilingThreadsDuringGame) : 1;
-
-	// On machines with few cores, each core will have a massive impact on compile time, so we prioritize compile latency over editor performance during the build
-	if (NumVirtualCores <= 4)
-	{
-		NumShaderCompilingThreads = NumVirtualCores - 1;
-		NumShaderCompilingThreadsDuringGame = NumVirtualCores - 1;
-	}
-#if PLATFORM_DESKTOP
-	else if (GIsBuildMachine || bForceUseSCWMemoryPressureLimits)
-	{
-		// Cooker ends up running OOM so use a simple heuristic based on some INI values
-		float CookerMemoryUsedInGB = 0.0f;
-		float MemoryToLeaveForTheOSInGB = 0.0f;
-		float MemoryUsedPerSCWProcessInGB = 0.0f;
-		bool bFoundEntries = true;
-		bFoundEntries = bFoundEntries && GConfig->GetFloat(TEXT("DevOptions.Shaders"), TEXT("CookerMemoryUsedInGB"), CookerMemoryUsedInGB, GEngineIni);
-		bFoundEntries = bFoundEntries && GConfig->GetFloat(TEXT("DevOptions.Shaders"), TEXT("MemoryToLeaveForTheOSInGB"), MemoryToLeaveForTheOSInGB, GEngineIni);
-		bFoundEntries = bFoundEntries && GConfig->GetFloat(TEXT("DevOptions.Shaders"), TEXT("MemoryUsedPerSCWProcessInGB"), MemoryUsedPerSCWProcessInGB, GEngineIni);
-		if (bFoundEntries)
-		{
-			uint32 PhysicalGBRam = FPlatformMemory::GetPhysicalGBRam();
-			float AvailableMemInGB = (float)PhysicalGBRam - CookerMemoryUsedInGB;
-			if (AvailableMemInGB > 0.0f)
-			{
-				if (AvailableMemInGB > MemoryToLeaveForTheOSInGB)
-				{
-					AvailableMemInGB -= MemoryToLeaveForTheOSInGB;
-				}
-				else
-				{
-					UE_LOG(LogShaderCompilers, Warning, TEXT("Machine has %d GBs of RAM, cooker might take %f GBs, but not enough memory left for the OS! (Requested %f GBs for the OS)"), PhysicalGBRam, CookerMemoryUsedInGB, MemoryToLeaveForTheOSInGB);
-				}
-			}
-			else
-			{
-				UE_LOG(LogShaderCompilers, Warning, TEXT("Machine has %d GBs of RAM, but cooker might take %f GBs!"), PhysicalGBRam, CookerMemoryUsedInGB);
-			}
-			if (MemoryUsedPerSCWProcessInGB > 0.0f)
-			{
-				float NumSCWs = AvailableMemInGB / MemoryUsedPerSCWProcessInGB;
-				NumShaderCompilingThreads = FMath::RoundToInt(NumSCWs);
-
-				bool bUseVirtualCores = true;
-				GConfig->GetBool(TEXT("DevOptions.Shaders"), TEXT("bUseVirtualCores"), bUseVirtualCores, GEngineIni);
-				uint32 MaxNumCoresToUse = bUseVirtualCores ? NumVirtualCores : NumberOfCores;
-				NumShaderCompilingThreads = FMath::Clamp<uint32>(NumShaderCompilingThreads, 1, MaxNumCoresToUse - 1);
-				NumShaderCompilingThreadsDuringGame = FMath::Min<int32>(NumShaderCompilingThreads, NumShaderCompilingThreadsDuringGame);
-			}
-		}
-		else if (bForceUseSCWMemoryPressureLimits)
-		{
-			UE_LOG(LogShaderCompilers, Warning, TEXT("bForceUseSCWMemoryPressureLimits was set but missing one or more prerequisite setting(s): CookerMemoryUsedInGB, MemoryToLeaveForTheOSInGB, MemoryUsedPerSCWProcessInGB.  Ignoring bForceUseSCWMemoryPressureLimits"));
-		}
-
-		if (GIsBuildMachine)
-		{
-			// force crashes on hung shader maps on build machines, to prevent builds running for days
-			GCrashOnHungShaderMaps = 1;
-		}
-	}
-#endif
-
-	NumShaderCompilingThreads = FMath::Max<int32>(1, NumShaderCompilingThreads);
-	NumShaderCompilingThreadsDuringGame = FMath::Max<int32>(1, NumShaderCompilingThreadsDuringGame);
-
-	NumShaderCompilingThreadsDuringGame = FMath::Min<int32>(NumShaderCompilingThreadsDuringGame, NumShaderCompilingThreads);
-}
-
-void FShaderCompilingManager::OnMachineResourcesChanged(int32 NumberOfCores, int32 NumberOfCoresIncludingHyperthreads)
-{
-	CalculateNumberOfCompilingThreads(NumberOfCores, NumberOfCoresIncludingHyperthreads);
-	for (const auto& Thread : Threads)
-	{
-		Thread->OnMachineResourcesChanged();
-	}
-}
-
-FName FShaderCompilingManager::GetStaticAssetTypeName()
-{
-	return TEXT("UE-Shader");
-}
-
-FName FShaderCompilingManager::GetAssetTypeName() const
-{
-	return GetStaticAssetTypeName();
-}
-
-FTextFormat FShaderCompilingManager::GetAssetNameFormat() const
-{
-	return LOCTEXT("ShaderNameFormat", "{0}|plural(one=Shader,other=Shaders)");
-}
-
-TArrayView<FName> FShaderCompilingManager::GetDependentTypeNames() const
-{
-#if WITH_EDITOR
-	static FName DependentTypeNames[] = 
-	{
-		// Texture can require materials to be updated,
-		// they should be processed first to avoid unecessary material updates.
-		FTextureCompilingManager::GetStaticAssetTypeName() 
-	};
-	return TArrayView<FName>(DependentTypeNames);
-#else
-	return TArrayView<FName>();
-#endif	
-}
-
-int32 FShaderCompilingManager::GetNumRemainingAssets() const
-{
-	// Currently, jobs are difficult to track but the purpose of the GetNumRemainingAssets function is to never return 0
-	// if there are still shaders that have not had their primitives updated on the render thread.
-	// So we track jobs first and when everything is finished compiling but are still lying around in other structures
-	// waiting to be further processed, we show those numbers and ultimately we always return 1 unless IsCompiling() is false.
-	return FMath::Max3(GetNumRemainingJobs(), ShaderMapJobs.Num() + PendingFinalizeShaderMaps.Num(), IsCompiling() ? 1 : 0);
-}
-
-void FShaderCompilingManager::ProcessAsyncTasks(bool bLimitExecutionTime)
-{
-	ProcessAsyncResults(bLimitExecutionTime, false);
-}
-
-void FShaderCompilingManager::ProcessAsyncTasks(const AssetCompilation::FProcessAsyncTaskParams& Params)
-{
-	// Shader compilations are not required for PIE to begin.
-	if (Params.bPlayInEditorAssetsOnly)
-	{
-		return;
-	}
-
-	ProcessAsyncResults(Params.bLimitExecutionTime, false);
-}
-
-int32 FShaderCompilingManager::GetNumPendingJobs() const
-{
-	return AllJobs.GetNumPendingJobs();
-}
-
-int32 FShaderCompilingManager::GetNumOutstandingJobs() const
-{
-	return AllJobs.GetNumOutstandingJobs();
-}
-
-FShaderCompilingManager::EDumpShaderDebugInfo FShaderCompilingManager::GetDumpShaderDebugInfo() const
-{
-	if (GDumpShaderDebugInfo < EDumpShaderDebugInfo::Never || GDumpShaderDebugInfo > EDumpShaderDebugInfo::OnErrorOrWarning)
-	{
-		return EDumpShaderDebugInfo::Never;
-	}
-
-	return static_cast<FShaderCompilingManager::EDumpShaderDebugInfo>(GDumpShaderDebugInfo);
-}
-
-EShaderDebugInfoFlags FShaderCompilingManager::GetDumpShaderDebugInfoFlags() const
-{
-	EShaderDebugInfoFlags Flags = EShaderDebugInfoFlags::Default;
-	if (GDumpShaderDebugInfoSCWCommandLine)
-	{
-		Flags |= EShaderDebugInfoFlags::DirectCompileCommandLine;
-	}
-	
-	if (CVarDebugDumpJobInputHashes.GetValueOnAnyThread())
-	{
-		Flags |= EShaderDebugInfoFlags::InputHash;
-	}
-
-	if (CVarDebugDumpJobDiagnostics.GetValueOnAnyThread())
-	{
-		Flags |= EShaderDebugInfoFlags::Diagnostics;
-	}
-
-	if (CVarDebugDumpShaderCode.GetValueOnAnyThread())
-	{
-		Flags |= EShaderDebugInfoFlags::ShaderCodeBinary;
-	}
-
-	if (CVarDebugDumpDetailedShaderSource.GetValueOnAnyThread())
-	{
-		Flags |= EShaderDebugInfoFlags::DetailedSource;
-	}
-
-	return Flags;
-}
-
-FString FShaderCompilingManager::CreateShaderDebugInfoPath(const FShaderCompilerInput& ShaderCompilerInput) const
-{
-	TRACE_CPUPROFILER_EVENT_SCOPE(FShaderCompilingManager::CreateShaderDebugInfoPath);
-
-	FString DumpDebugInfoPath = ShaderCompilerInput.DumpDebugInfoRootPath / ShaderCompilerInput.DebugGroupName + ShaderCompilerInput.DebugExtension;
-
-	// Sanitize the name to be used as a path
-	// List mostly comes from set of characters not allowed by windows in a path.  Just try to rename a file and type one of these for the list.
-	DumpDebugInfoPath.ReplaceInline(TEXT("<"), TEXT("("));
-	DumpDebugInfoPath.ReplaceInline(TEXT(">"), TEXT(")"));
-	DumpDebugInfoPath.ReplaceInline(TEXT("::"), TEXT("=="));
-	DumpDebugInfoPath.ReplaceInline(TEXT("|"), TEXT("_"));
-	DumpDebugInfoPath.ReplaceInline(TEXT("*"), TEXT("-"));
-	DumpDebugInfoPath.ReplaceInline(TEXT("?"), TEXT("!"));
-	DumpDebugInfoPath.ReplaceInline(TEXT("\""), TEXT("\'"));
-
-	if (!IFileManager::Get().DirectoryExists(*DumpDebugInfoPath))
-	{
-		if (!IFileManager::Get().MakeDirectory(*DumpDebugInfoPath, true))
-		{
-			const uint32 ErrorCode = FPlatformMisc::GetLastError();
-			UE_LOG(LogShaderCompilers, Warning, TEXT("Last Error %u: Failed to create directory for shader debug info '%s'. Try enabling large file paths or r.DumpShaderDebugShortNames."), ErrorCode, *DumpDebugInfoPath);
-			return FString();
-		}
-	}
-
-	return DumpDebugInfoPath;
-}
-
-bool FShaderCompilingManager::ShouldRecompileToDumpShaderDebugInfo(const FShaderCompileJob& Job) const
-{
-	return ShouldRecompileToDumpShaderDebugInfo(Job.Input, Job.Output, Job.bSucceeded);
-}
-
-bool FShaderCompilingManager::ShouldRecompileToDumpShaderDebugInfo(const FShaderCompilerInput& Input, const FShaderCompilerOutput& Output, bool bSucceeded) const
-{
-	if (Input.DumpDebugInfoPath.IsEmpty())
-	{
-		const EDumpShaderDebugInfo DumpShaderDebugInfo = GetDumpShaderDebugInfo();
-		const bool bErrors = !bSucceeded;
-		const bool bWarnings = Output.Errors.Num() > 0;
-
-		bool bShouldDump = true;
-		if (GIsBuildMachine)
-		{
-			// Build machines dump these as build artifacts and they should only upload so many due to size constraints.
-			bShouldDump = (NumDumpedShaderSources < GMaxNumDumpedShaderSources);
-		}
-
-		if (DumpShaderDebugInfo == EDumpShaderDebugInfo::OnError)
-		{
-			return bShouldDump && (bErrors);
-		}
-		else if (DumpShaderDebugInfo == EDumpShaderDebugInfo::OnErrorOrWarning)
-		{
-			return bShouldDump && (bErrors || bWarnings);
-		}
-	}
-
-	return false;
-}
-
-void FShaderCompilingManager::ReleaseJob(FShaderCommonCompileJobPtr& Job)
-{
-	ReleaseJob(Job.GetReference());
-	Job.SafeRelease();
-}
-
-void FShaderCompilingManager::ReleaseJob(FShaderCommonCompileJob* Job)
-{
-	Job->PendingShaderMap.SafeRelease();
-	Job->bReleased = true;
-	AllJobs.RemoveJob(Job);
-}
-
-void FShaderCompilingManager::SubmitJobs(TArray<FShaderCommonCompileJobPtr>& NewJobs, const FString MaterialBasePath, const FString PermutationString)
-{
-	LLM_SCOPE_BYTAG(ShaderCompiler);
-
-	// make sure no compiling can start if not allowed
-	if (!AllowShaderCompiling())
-	{
-		return;
-	}
-	
-	TRACE_CPUPROFILER_EVENT_SCOPE(FShaderCompilingManager::SubmitJobs);
-	check(!FPlatformProperties::RequiresCookedData());
-
-	if (NewJobs.Num() == 0)
-	{
-		return;
-	}
-
-	check(GShaderCompilerStats);
-	if (FShaderCompileJob* SingleJob = NewJobs[0]->GetSingleShaderJob()) //assume that all jobs are for the same platform
-	{
-		GShaderCompilerStats->RegisterCompiledShaders(NewJobs.Num(), SingleJob->Input.Target.GetPlatform(), MaterialBasePath, PermutationString);
-	}
-	else
-	{
-		GShaderCompilerStats->RegisterCompiledShaders(NewJobs.Num(), SP_NumPlatforms, MaterialBasePath, PermutationString);
-	}
-
-	{
-		FScopeLock Lock(&CompileQueueSection);
-		for (auto& Job : NewJobs)
-		{
-			FPendingShaderMapCompileResultsPtr& PendingShaderMap = ShaderMapJobs.FindOrAdd(Job->Id);
-			if (!PendingShaderMap)
-			{
-				PendingShaderMap = new FPendingShaderMapCompileResults();
-			}
-			PendingShaderMap->NumPendingJobs.Increment();
-			Job->PendingShaderMap = PendingShaderMap;
-		}
-
-		// in the case of submitting jobs from worker threads we need to be sure that the lock extends to
-		// include AllJobs.SubmitJobs().  This will increase contention for the lock, but this will let us
-		// prototype getting shader translation and preprocessing being done on worker threads.
-		if (IsInGameThread())
-		{
-			Lock.Unlock();
-		}
-
-		AllJobs.SubmitJobs(NewJobs);
-	}
-
-	UpdateNumRemainingAssets();
-}
-
-bool FShaderCompilingManager::IsCompilingShaderMap(uint32 Id)
-{
-	if (Id != 0u)
-	{
-		FScopeLock Lock(&CompileQueueSection);
-		FPendingShaderMapCompileResultsPtr* PendingShaderMapPtr = ShaderMapJobs.Find(Id);
-		if (PendingShaderMapPtr)
-		{
-			return true;
-		}
-
-		FShaderMapFinalizeResults* FinalizedShaderMapPtr = PendingFinalizeShaderMaps.Find(Id);
-		if (FinalizedShaderMapPtr)
-		{
-			return true;
-		}
-	}
-	return false;
-}
-
-FShaderCompileJob* FShaderCompilingManager::PrepareShaderCompileJob(uint32 Id, const FShaderCompileJobKey& Key, EShaderCompileJobPriority Priority)
-{
-	// don't allow any jobs if not allowed
-	if (!AllowShaderCompiling())
-	{
-		return nullptr;
-	}
-
-	return AllJobs.PrepareJob(Id, Key, Priority);
-}
-
-FShaderPipelineCompileJob* FShaderCompilingManager::PreparePipelineCompileJob(uint32 Id, const FShaderPipelineCompileJobKey& Key, EShaderCompileJobPriority Priority)
-{
-	// don't allow any jobs if not allowed
-	if (!AllowShaderCompiling())
-	{
-		return nullptr;
-	}
-
-	return AllJobs.PrepareJob(Id, Key, Priority);
-}
-
-void FShaderCompilingManager::ProcessFinishedJob(FShaderCommonCompileJob* FinishedJob)
-{
-	AllJobs.ProcessFinishedJob(FinishedJob);
-}
-
-/** Launches the worker, returns the launched process handle. */
-FProcHandle FShaderCompilingManager::LaunchWorker(const FString& WorkingDirectory, uint32 InProcessId, uint32 ThreadId, const FString& WorkerInputFile, const FString& WorkerOutputFile)
-{
-	// don't allow any jobs if not allowed
-	if (!AllowShaderCompiling())
-	{
-		return FProcHandle();
-	}
-
-	// Setup the parameters that the worker application needs
-	// Surround the working directory with double quotes because it may contain a space 
-	// WorkingDirectory ends with a '\', so we have to insert another to meet the Windows commandline parsing rules 
-	// http://msdn.microsoft.com/en-us/library/17w5ykft.aspx 
-	// Use IFileManager to do path conversion to properly handle sandbox paths (outside of standard paths in particular).
-	FString WorkerAbsoluteDirectory = IFileManager::Get().ConvertToAbsolutePathForExternalAppForWrite(*WorkingDirectory);
-	FPaths::NormalizeDirectoryName(WorkerAbsoluteDirectory);
-	FString WorkerParameters = FString(TEXT("\"")) + WorkerAbsoluteDirectory + TEXT("/\" ") + FString::FromInt(InProcessId) + TEXT(" ") + FString::FromInt(ThreadId) + TEXT(" ") + WorkerInputFile + TEXT(" ") + WorkerOutputFile;
-	WorkerParameters += FString(TEXT(" -communicatethroughfile "));
-	if ( GIsBuildMachine )
-	{
-		WorkerParameters += FString::Printf(TEXT(" -TimeToLive=%f -buildmachine"), GBuildWorkerTimeToLive);
-	}
-	else
-	{
-		WorkerParameters += FString::Printf(TEXT(" -TimeToLive=%f"), GRegularWorkerTimeToLive);
-	}
-	if (PLATFORM_LINUX) //-V560
-	{
-		// suppress log generation as much as possible
-		WorkerParameters += FString(TEXT(" -logcmds=\"Global None\" "));
-
-		if (UE_BUILD_DEBUG)
-		{
-			// when running a debug build under Linux, make SCW crash with core for easier debugging
-			WorkerParameters += FString(TEXT(" -core "));
-		}
-	}
-	WorkerParameters += FCommandLine::GetSubprocessCommandline();
-
-#if USE_SHADER_COMPILER_WORKER_TRACE
-	// When doing utrace functionality we can't run with -nothreading, since it won't create the utrace thread to send events.
-	WorkerParameters += FString(TEXT(" -trace=default "));
-#else
-	WorkerParameters += FString(TEXT(" -nothreading "));
-#endif // USE_SHADER_COMPILER_WORKER_TRACE
-
-	// Launch the worker process
-	int32 PriorityModifier = -1; // below normal
-	GConfig->GetInt(TEXT("DevOptions.Shaders"), TEXT("WorkerProcessPriority"), PriorityModifier, GEngineIni);
-
-	if (DEBUG_SHADERCOMPILEWORKER)
-	{
-		// Note: Set breakpoint here and launch the ShaderCompileWorker with WorkerParameters a cmd-line
-		const TCHAR* WorkerParametersText = *WorkerParameters;
-		FPlatformMisc::LowLevelOutputDebugStringf(TEXT("Launching shader compile worker w/ WorkerParameters\n\t%s\n"), WorkerParametersText);
-		FProcHandle DummyHandle;
-		return DummyHandle;
-	}
-	else
-	{
-#if UE_BUILD_DEBUG && PLATFORM_LINUX
-		FPlatformMisc::LowLevelOutputDebugStringf(TEXT("Launching shader compile worker:\n\t%s\n"), *WorkerParameters);
-#endif
-		// Disambiguate between SCW.exe missing vs other errors.
-		static bool bFirstLaunch = true;
-		uint32 WorkerId = 0;
-		FProcHandle WorkerHandle = FPlatformProcess::CreateProc(*ShaderCompileWorkerName, *WorkerParameters, true, false, false, &WorkerId, PriorityModifier, nullptr, nullptr);
-		if (WorkerHandle.IsValid())
-		{
-			// Process launched at least once successfully
-			bFirstLaunch = false;
-		}
-		else
-		{
-			// If this doesn't error, the app will hang waiting for jobs that can never be completed
-			if (bFirstLaunch)
-			{
-				// When using source builds users are likely to make a mistake of not building SCW (e.g. in particular on Linux, even though default makefile target builds it).
-				// Make the engine exit gracefully with a helpful message instead of a crash.
-				static bool bShowedMessageBox = false;
-				if (!bShowedMessageBox && !IsRunningCommandlet() && !FApp::IsUnattended())
-				{
-					bShowedMessageBox = true;
-					FText ErrorMessage = FText::Format(LOCTEXT("LaunchingShaderCompileWorkerFailed", "Unable to launch {0} - make sure you built ShaderCompileWorker."), FText::FromString(ShaderCompileWorkerName));
-					FPlatformMisc::MessageBoxExt(EAppMsgType::Ok, *ErrorMessage.ToString(),
-												 *LOCTEXT("LaunchingShaderCompileWorkerFailedTitle", "Unable to launch ShaderCompileWorker.").ToString());
-				}
-				UE_LOG(LogShaderCompilers, Error, TEXT("Couldn't launch %s! Make sure you build ShaderCompileWorker."), *ShaderCompileWorkerName);
-				// duplicate to printf() since threaded logs may not be always flushed
-				FPlatformMisc::LowLevelOutputDebugStringf(TEXT("Couldn't launch %s! Make sure you build ShaderCompileWorker.\n"), *ShaderCompileWorkerName);
-				FPlatformMisc::RequestExitWithStatus(true, 1);
-			}
-			else
-			{
-				UE_LOG(LogShaderCompilers, Fatal, TEXT("Couldn't launch %s!"), *ShaderCompileWorkerName);
-			}
-		}
-
-		return WorkerHandle;
-	}
-}
-
-void FShaderCompilingManager::AddCompiledResults(TMap<int32, FShaderMapFinalizeResults>& CompiledShaderMaps, int32 ShaderMapIdx, const FShaderMapFinalizeResults& Results)
-{
-	// merge with the previous unprocessed jobs, if any
-	if (FShaderMapCompileResults const* PrevResults = CompiledShaderMaps.Find(ShaderMapIdx))
-	{
-		FShaderMapFinalizeResults NewResults(Results);
-
-		NewResults.bAllJobsSucceeded = NewResults.bAllJobsSucceeded && PrevResults->bAllJobsSucceeded;
-		NewResults.bSkipResultProcessing = NewResults.bSkipResultProcessing || PrevResults->bSkipResultProcessing;
-		NewResults.TimeStarted = FMath::Min(NewResults.TimeStarted, PrevResults->TimeStarted);
-		NewResults.bIsHung = NewResults.bIsHung || PrevResults->bIsHung;
-		NewResults.FinishedJobs.Append(PrevResults->FinishedJobs);
-
-		CompiledShaderMaps.Add(ShaderMapIdx, NewResults);
-	}
-	else
-	{
-		CompiledShaderMaps.Add(ShaderMapIdx, Results);
-	}
-}
-
-/** Flushes all pending jobs for the given shader maps. */
-void FShaderCompilingManager::BlockOnShaderMapCompletion(const TArray<int32>& ShaderMapIdsToFinishCompiling, TMap<int32, FShaderMapFinalizeResults>& CompiledShaderMaps)
-{
-	// never block if no compiling, just in case
-	if (!AllowShaderCompiling())
-	{
-		return;
-	}
-
-	TRACE_CPUPROFILER_EVENT_SCOPE(FShaderCompilingManager::BlockOnShaderMapCompletion);
-
-	COOK_STAT(FScopedDurationTimer BlockingTimer(ShaderCompilerCookStats::BlockingTimeSec));
-	if (bAllowAsynchronousShaderCompiling)
-	{
-		// Calculate how many shader jobs there are total to provide the slow task with the correct amount of work.
-		int NumJobs = 0;
-		{
-			FScopeLock Lock(&CompileQueueSection);
-			for (int32 ShaderMapIndex = 0; ShaderMapIndex < ShaderMapIdsToFinishCompiling.Num(); ShaderMapIndex++)
-			{
-				FPendingShaderMapCompileResultsPtr* ResultsPtr = ShaderMapJobs.Find(ShaderMapIdsToFinishCompiling[ShaderMapIndex]);
-				if (ResultsPtr)
-				{
-					FShaderMapCompileResults* Results = *ResultsPtr;
-					NumJobs += Results->NumPendingJobs.GetValue();
-				}
-			}
-		}
-
-		FScopedSlowTask SlowTask(NumJobs, FText::Format(LOCTEXT("BlockOnShaderMapCompletion", "Compiling Shaders ({0})"), NumJobs), GIsEditor && !IsRunningCommandlet() && UE::GetPlayInEditorID() == INDEX_NONE);
-		if (NumJobs > 0)
-		{
-			SlowTask.MakeDialogDelayed(1.0f);
-		}
-
-		int32 NumPendingJobs = 0;
-		// Keep track of previous number of pending jobs so we can update the slow task with the amount of work done.
-		int32 NumPreviousPendingJobs = NumJobs;
-		int32 LogCounter = 0;
-		do 
-		{
-			NumPendingJobs = 0;
-			{
-				// Lock CompileQueueSection so we can access the input and output queues
-				FScopeLock Lock(&CompileQueueSection);
-
-				for (int32 ShaderMapIndex = 0; ShaderMapIndex < ShaderMapIdsToFinishCompiling.Num(); ShaderMapIndex++)
-				{
-					FPendingShaderMapCompileResultsPtr* ResultsPtr = ShaderMapJobs.Find(ShaderMapIdsToFinishCompiling[ShaderMapIndex]);
-					if (ResultsPtr)
-					{
-						FShaderMapCompileResults* Results = *ResultsPtr;
-
-						if (Results->NumPendingJobs.GetValue() == 0)
-						{
-							if (Results->FinishedJobs.Num() > 0)
-							{
-								AddCompiledResults(CompiledShaderMaps, ShaderMapIdsToFinishCompiling[ShaderMapIndex], *Results);
-							}
-							ShaderMapJobs.Remove(ShaderMapIdsToFinishCompiling[ShaderMapIndex]);
-						}
-						else
-						{
-							Results->CheckIfHung();
-							NumPendingJobs += Results->NumPendingJobs.GetValue();
-						}
-					}
-				}
-			}
-
-			if (NumPendingJobs > 0)
-			{
-				const float SleepTime =.01f;
-				
-				// We need to manually tick the Distributed build controller while the game thread is blocked
-				// otherwise we can get stuck in a infinite loop waiting for jobs that never will be done
-				// because for example, some controllers depend on the HTTP module which needs to be ticked in the main thread
-				if (BuildDistributionController && IsInGameThread())
-				{
-					BuildDistributionController->Tick(SleepTime);
-				}
-
-				// Progress the slow task with how many jobs we've completed since last tick.  Update the slow task message with the current number of pending jobs
-				// we are waiting on.
-				const int32 CompletedJobsSinceLastTick = NumPreviousPendingJobs - NumPendingJobs;
-				SlowTask.EnterProgressFrame(CompletedJobsSinceLastTick, FText::Format(LOCTEXT("BlockOnShaderMapCompletion", "Compiling Shaders ({0})"), NumPendingJobs));
-				NumPreviousPendingJobs = NumPendingJobs;
-				
-				// Yield CPU time while waiting
-				FPlatformProcess::Sleep(SleepTime);
-
-				// Flush threaded logs around every 500ms or so based on Sleep of 0.01f seconds above
-				if (++LogCounter > 50)
-				{
-					LogCounter = 0;
-					GLog->FlushThreadedLogs(EOutputDeviceRedirectorFlushOptions::Async);
-				}
-			}
-		} 
-		while (NumPendingJobs > 0);
-	}
-	else
-	{
-		int32 NumActiveWorkers = 0;
-		do 
-		{
-			for (const TUniquePtr<FShaderCompileThreadRunnableBase>& Thread : Threads)
-			{
-				NumActiveWorkers = Thread->CompilingLoop();
-			}
-		} 
-		while (NumActiveWorkers > 0);
-
-		check(AllJobs.GetNumPendingJobs() == 0);
-
-		for (int32 ShaderMapIndex = 0; ShaderMapIndex < ShaderMapIdsToFinishCompiling.Num(); ShaderMapIndex++)
-		{
-			const FPendingShaderMapCompileResultsPtr* ResultsPtr = ShaderMapJobs.Find(ShaderMapIdsToFinishCompiling[ShaderMapIndex]);
-
-			if (ResultsPtr)
-			{
-				const FShaderMapCompileResults* Results = *ResultsPtr;
-				check(Results->NumPendingJobs.GetValue() == 0);
-				check(Results->FinishedJobs.Num() > 0);
-
-				AddCompiledResults(CompiledShaderMaps, ShaderMapIdsToFinishCompiling[ShaderMapIndex], *Results);
-				ShaderMapJobs.Remove(ShaderMapIdsToFinishCompiling[ShaderMapIndex]);
-			}
-		}
-	}
-
-	UpdateNumRemainingAssets();
-}
-
-void FShaderCompilingManager::BlockOnAllShaderMapCompletion(TMap<int32, FShaderMapFinalizeResults>& CompiledShaderMaps)
-{
-	// never block if no compiling, just in case
-	if (!AllowShaderCompiling())
-	{
-		return;
-	}
-
-	TRACE_CPUPROFILER_EVENT_SCOPE(FShaderCompilingManager::BlockOnAllShaderMapCompletion);
-
-	COOK_STAT(FScopedDurationTimer BlockingTimer(ShaderCompilerCookStats::BlockingTimeSec));
-	if (bAllowAsynchronousShaderCompiling)
-	{
-		// Calculate how many shader jobs there are total to provide the slow task with the correct amount of work.
-		int NumJobs = 0;
-		{
-			FScopeLock Lock(&CompileQueueSection);
-			for (TMap<int32, FPendingShaderMapCompileResultsPtr>::TIterator It(ShaderMapJobs); It; ++It)
-			{
-				FShaderMapCompileResults* Results = It.Value();
-				NumJobs += Results->NumPendingJobs.GetValue();
-			}
-		}
-
-		FScopedSlowTask SlowTask(NumJobs, FText::Format(LOCTEXT("BlockOnAllShaderMapCompletion", "Compiling Shaders ({0})"), NumJobs), NumJobs && GIsEditor && !IsRunningCommandlet());
-		if (NumJobs > 0)
-		{
-			SlowTask.MakeDialog(false, true);
-		}
-
-		int32 NumPendingJobs = 0;
-		// Keep track of previous number of pending jobs so we can update the slow task with the amount of work done.
-		int32 NumPreviousPendingJobs = NumJobs;
-
-		do 
-		{
-			NumPendingJobs = 0;
-			{
-				// Lock CompileQueueSection so we can access the input and output queues
-				FScopeLock Lock(&CompileQueueSection);
-
-				int32 ShaderMapIdx = 0;
-				for (TMap<int32, FPendingShaderMapCompileResultsPtr>::TIterator It(ShaderMapJobs); It; ++It)
-				{
-					FShaderMapCompileResults* Results = It.Value();
-
-					if (Results->NumPendingJobs.GetValue() == 0)
-					{
-						AddCompiledResults(CompiledShaderMaps, It.Key(), *Results);
-						It.RemoveCurrent();
-					}
-					else
-					{
-						Results->CheckIfHung();
-						NumPendingJobs += Results->NumPendingJobs.GetValue();
-					}
-				}
-			}
-
-			if (NumPendingJobs > 0)
-			{
-				const float SleepTime =.01f;
-				
-				// We need to manually tick the Distributed build controller while the game thread is blocked
-				// otherwise we can get stuck in a infinite loop waiting for jobs that never will be done
-				// because for example, some controllers depend on the HTTP module which needs to be ticked in the main thread
-				if (BuildDistributionController && IsInGameThread())
-				{
-					BuildDistributionController->Tick(SleepTime);
-				}
-
-				// Progress the slow task with how many jobs we've completed since last tick.  Update the slow task message with the current number of pending jobs
-				// we are waiting on.
-				const int32 CompletedJobsSinceLastTick = NumPreviousPendingJobs - NumPendingJobs;
-				SlowTask.EnterProgressFrame(CompletedJobsSinceLastTick, FText::Format(LOCTEXT("BlockOnAllShaderMapCompletion", "Compiling Shaders ({0})"), NumPendingJobs));
-				NumPreviousPendingJobs = NumPendingJobs;
-				
-				// Yield CPU time while waiting
-				FPlatformProcess::Sleep(SleepTime);
-			}
-		} 
-		while (NumPendingJobs > 0);
-	}
-	else
-	{
-		int32 NumActiveWorkers = 0;
-		do 
-		{
-			for (const TUniquePtr<FShaderCompileThreadRunnableBase>& Thread : Threads)
-			{
-				NumActiveWorkers = Thread->CompilingLoop();
-			}
-
-			for (TMap<int32, FPendingShaderMapCompileResultsPtr>::TIterator It(ShaderMapJobs); It; ++It)
-			{
-				FShaderMapCompileResults* Results = It.Value();
-				Results->CheckIfHung();
-			}
-		} 
-		while (NumActiveWorkers > 0);
-
-		check(AllJobs.GetNumPendingJobs() == 0);
-
-		for (TMap<int32, FPendingShaderMapCompileResultsPtr>::TIterator It(ShaderMapJobs); It; ++It)
-		{
-			const FShaderMapCompileResults* Results = It.Value();
-			check(Results->NumPendingJobs.GetValue()== 0);
-
-			AddCompiledResults(CompiledShaderMaps, It.Key(), *Results);
-			It.RemoveCurrent();
-		}
-	}
-
-	UpdateNumRemainingAssets();
-}
-
-namespace
-{
-	void PropagateGlobalShadersToAllPrimitives()
-	{
-		// Re-register everything to work around FShader lifetime issues - it currently lives and dies with the
-		// shadermap it is stored in, while cached MDCs can reference its memory. Re-registering will
-		// re-create the cache.
-		TRACE_CPUPROFILER_EVENT_SCOPE(PropagateGlobalShadersToAllPrimitives);
-
-		FObjectCacheContextScope ObjectCacheScope;
-		TSet<FSceneInterface*> ScenesToUpdate;
-		TIndirectArray<FComponentRecreateRenderStateContext> ComponentContexts;
-		for (IPrimitiveComponent* PrimitiveComponentInterface : ObjectCacheScope.GetContext().GetPrimitiveComponents())
-		{
-			if (PrimitiveComponentInterface->IsRenderStateCreated())
-			{
-				ComponentContexts.Add(new FComponentRecreateRenderStateContext(PrimitiveComponentInterface, &ScenesToUpdate));
-#if WITH_EDITOR
-				if (UPrimitiveComponent* PrimitiveComponent = PrimitiveComponentInterface->GetUObject<UPrimitiveComponent>())
-				{
-					if (PrimitiveComponent->HasValidSettingsForStaticLighting(false))
-					{
-						FStaticLightingSystemInterface::OnPrimitiveComponentUnregistered.Broadcast(PrimitiveComponent);
-						FStaticLightingSystemInterface::OnPrimitiveComponentRegistered.Broadcast(PrimitiveComponent);
-					}
-				}
-#endif
-			}
-		}
-
-		UpdateAllPrimitiveSceneInfosForScenes(ScenesToUpdate);
-		ComponentContexts.Empty();
-		UpdateAllPrimitiveSceneInfosForScenes(ScenesToUpdate);
-	}
-}
-
-void FShaderCompilingManager::ProcessCompiledShaderMaps(
-	TMap<int32, FShaderMapFinalizeResults>& CompiledShaderMaps, 
-	float TimeBudget)
-{
-	// never process anything if not allowed, just in case
-	if (!AllowShaderCompiling())
-	{
-		return;
-	}
-
-#if WITH_EDITOR
-	TRACE_CPUPROFILER_EVENT_SCOPE(FShaderCompilingManager::ProcessCompiledShaderMaps);
-
-	TMap<TRefCountPtr<FMaterial>, TRefCountPtr<FMaterialShaderMap>> MaterialsToUpdate;
-	TArray<TRefCountPtr<FMaterial>> MaterialsToReleaseCompilingId;
-
-	// Process compiled shader maps in FIFO order, in case a shader map has been enqueued multiple times,
-	// Which can happen if a material is edited while a background compile is going on
-	for (TMap<int32, FShaderMapFinalizeResults>::TIterator ShaderMapResultIter(CompiledShaderMaps); ShaderMapResultIter; ++ShaderMapResultIter)
-	{
-		const uint32 CompilingId = ShaderMapResultIter.Key();
-
-		FShaderMapFinalizeResults& CompileResults = ShaderMapResultIter.Value();
-		TArray<FShaderCommonCompileJobPtr>& FinishedJobs = CompileResults.FinishedJobs;
-
-		if (CompileResults.bSkipResultProcessing)
-		{
-			ShaderMapResultIter.RemoveCurrent();
-			continue;
-		}
-
-		TRefCountPtr<FMaterialShaderMap> CompilingShaderMap = FMaterialShaderMap::FindCompilingShaderMap(CompilingId);
-
-		if (CompilingShaderMap)
-		{
-			TArray<TRefCountPtr<FMaterial>>& MaterialDependencies = CompilingShaderMap->CompilingMaterialDependencies;
-			FShaderDiagnosticInfo ErrorInfo(FinishedJobs);
-
-			bool bSuccess = true;
-			for (int32 JobIndex = 0; JobIndex < FinishedJobs.Num(); JobIndex++)
-			{
-				FShaderCommonCompileJob& CurrentJob = *FinishedJobs[JobIndex];
-
-				if (FShaderCompileJob* SingleJob = CurrentJob.GetSingleShaderJob())
-				{
-					const bool bCheckSucceeded = CheckSingleJob(*SingleJob, ErrorInfo.UniqueErrors);
-					bSuccess = bCheckSucceeded && bSuccess;
-				}
-				else if (FShaderPipelineCompileJob* PipelineJob = CurrentJob.GetShaderPipelineJob())
-				{
-					for (int32 Index = 0; Index < PipelineJob->StageJobs.Num(); ++Index)
-					{
-						const bool bCheckSucceeded = CheckSingleJob(*PipelineJob->StageJobs[Index], ErrorInfo.UniqueErrors);
-						bSuccess = PipelineJob->StageJobs[Index]->bSucceeded && bCheckSucceeded && bSuccess;
-					}
-				}
-				else
-				{
-					checkf(0, TEXT("FShaderCommonCompileJob::Type=%d is not a valid type for a shader compile job"), (int32)CurrentJob.Type);
-				}
-			}
-
-			if (bSuccess)
-			{
-				int32 JobIndex = 0;
-				if (FinishedJobs.Num() > 0)
-				{
-					CompilingShaderMap->ProcessCompilationResults(FinishedJobs, JobIndex, TimeBudget);
-					{
-						FScopeLock Lock(&CompileQueueSection);
-						for (int32 i = 0; i < JobIndex; ++i)
-						{
-							ReleaseJob(FinishedJobs[i]);
-						}
-					}
-					FinishedJobs.RemoveAt(0, JobIndex);
-				}
-			}
-
-			if (!bSuccess || FinishedJobs.Num() == 0)
-			{
-				ShaderMapResultIter.RemoveCurrent();
-			}
-
-			FMaterialShaderMap* ShaderMapToUseForRendering = nullptr;
-
-#if DEBUG_INFINITESHADERCOMPILE
-			UE_LOG(LogTemp, Display, TEXT("Finished compile of shader map 0x%08X%08X"), (int)((int64)(ShaderMap.GetReference()) >> 32), (int)((int64)(ShaderMap.GetReference())));
-#endif
-			int32 NumIncompleteMaterials = 0;
-			int32 MaterialIndex = 0;
-			
-			FMaterial* SingleMaterial = MaterialDependencies.Num() > 0 ? MaterialDependencies[0] : nullptr;
-			bool bRequiredComplete = false;
-
-			while (MaterialIndex < MaterialDependencies.Num())
-			{
-				FMaterial* Material = MaterialDependencies[MaterialIndex];
-				check(Material->GetGameThreadCompilingShaderMapId() == CompilingShaderMap->GetCompilingId());
-				bRequiredComplete |= Material->IsRequiredComplete();
-
-#if DEBUG_INFINITESHADERCOMPILE
-				UE_LOG(LogTemp, Display, TEXT("Shader map %s complete, GameThreadShaderMap 0x%08X%08X, marking material %s as finished"), *ShaderMap->GetFriendlyName(), (int)((int64)(ShaderMap.GetReference()) >> 32), (int)((int64)(ShaderMap.GetReference())), *Material->GetFriendlyName());
-				UE_LOG(LogTemp, Display, TEXT("Marking material as finished 0x%08X%08X"), (int)((int64)(Material) >> 32), (int)((int64)(Material)));
-#endif
-				//Material->RemoveOutstandingCompileId(ShaderMap->CompilingId);
-
-				bool bReleaseCompilingId = false;
-
-				// Only process results that still match the ID which requested a compile
-				// This avoids applying shadermaps which are out of date and a newer one is in the async compiling pipeline
-				if (Material->GetMaterialId() != CompilingShaderMap->GetShaderMapId().BaseMaterialId)
-				{
-					bReleaseCompilingId = true;
-				}
-				else if (bSuccess)
-				{
-					bool bIsComplete = CompilingShaderMap->IsComplete(Material, true) && (CompilingShaderMap->CompilingMaterialNumExternalDependencies == 0);
-
-					// If running a cook, only process complete shader maps, as there's no rendering of partially complete shader maps to worry about.
-					if (bIsComplete || IsRunningCookCommandlet() == false || bAllowForIncompleteShaderMaps)
-					{
-						if (ShaderMapToUseForRendering == nullptr)
-						{
-							// Make a clone of the compiling shader map to use for rendering
-							// This will allow rendering to proceed with the clone, while async compilation continues to potentially update the compiling shader map
-							double StartTime = FPlatformTime::Seconds();
-							ShaderMapToUseForRendering = CompilingShaderMap->AcquireFinalizedClone();
-							TimeBudget -= (FPlatformTime::Seconds() - StartTime);
-						}
-
-						MaterialsToUpdate.Add(Material, ShaderMapToUseForRendering);
-					}
-
-					if (bIsComplete)
-					{
-						bReleaseCompilingId = true;
-					}
-					else
-					{
-						++NumIncompleteMaterials;
-					}
-
-					if (ErrorInfo.UniqueWarnings.Num() > 0)
-					{
-						UE_LOG(LogShaderCompilers, Warning, TEXT("Warnings while compiling Material %s for platform %s:"),
-							*Material->GetDebugName(),
-							*LegacyShaderPlatformToShaderFormat(CompilingShaderMap->GetShaderPlatform()).ToString());
-						for (const FString& UniqueWarning : ErrorInfo.UniqueWarnings)
-						{
-							UE_LOG(LogShaders, Warning, TEXT("  %s"), *UniqueWarning);
-						}
-					}
-				}
-				else
-				{
-					bReleaseCompilingId = true;
-					// Propagate error messages
-					Material->CompileErrors = ErrorInfo.UniqueErrors;
-
-					MaterialsToUpdate.Add(Material, nullptr);
-
-					if (Material->IsDefaultMaterial())
-					{
-						FString ErrorString;
-
-						// Log the errors unsuppressed before the fatal error, so it's always obvious from the log what the compile error was
-						for (const FString& UniqueError : ErrorInfo.UniqueErrors)
-						{
-							ErrorString += FString::Printf(TEXT("  %s\n"), *UniqueError);
-						}
-
-						ErrorString += FString::Printf(TEXT("Failed to compile default material %s!"), *Material->GetBaseMaterialPathName());
-
-						if (AreShaderErrorsFatal())
-						{
-							// Assert if a default material could not be compiled, since there will be nothing for other failed materials to fall back on.
-							UE_LOG(LogShaderCompilers, Fatal, TEXT("%s"), *ErrorString);
-						}
-						else
-						{
-							UE_LOG(LogShaderCompilers, Error, TEXT("%s"), *ErrorString);
-						}
-					}
-					
-					FString ErrorString;
-
-					ErrorString += FString::Printf(TEXT("Failed to compile Material %s for platform %s, Default Material will be used in game.\n"),
-						*Material->GetDebugName(), *LegacyShaderPlatformToShaderFormat(CompilingShaderMap->GetShaderPlatform()).ToString());
-
-					for (const FString& UniqueError : ErrorInfo.UniqueErrors)
-					{
-						FString ErrorMessage = UniqueError;
-						// Work around build machine string matching heuristics that will cause a cook to fail
-						ErrorMessage.ReplaceInline(TEXT("error "), TEXT("err0r "), ESearchCase::CaseSensitive);
-						ErrorString += FString::Printf(TEXT("  %s\n"), *ErrorMessage);
-					}
-
-					UE_LOG(LogShaderCompilers, Warning, TEXT("%s"), *ErrorString);
-				}
-
-				if (bReleaseCompilingId)
-				{
-					check(Material->GameThreadCompilingShaderMapId != 0u);
-					Material->GameThreadCompilingShaderMapId = 0u;
-					Material->GameThreadPendingCompilerEnvironment.SafeRelease();
-					MaterialDependencies.RemoveAt(MaterialIndex);
-					MaterialsToReleaseCompilingId.Add(Material);
-				}
-				else
-				{
-					++MaterialIndex;
-				}
-			}
-
-			if (NumIncompleteMaterials == 0 && (IsMaterialMapDDCEnabled() || bRequiredComplete))
-			{
-				CompilingShaderMap->bCompiledSuccessfully = bSuccess;
-				CompilingShaderMap->bCompilationFinalized = true;
-				if (ShaderMapToUseForRendering)
-				{
-					// ShaderMapToUseForRendering is only initialized inside the loop over material dependencies,
-					// so it's safe to assume that SingleMaterial has been set (a material is needed to construct 
-					// the FMaterialShaderParameters struct which is in turn needed to build the DDC key).
-					check(SingleMaterial != nullptr);
-					ShaderMapToUseForRendering->bCompiledSuccessfully = true;
-					ShaderMapToUseForRendering->bCompilationFinalized = true;
-					if (ShaderMapToUseForRendering->bIsPersistent)
-					{
-						ShaderMapToUseForRendering->SaveToDerivedDataCache(FMaterialShaderParameters(SingleMaterial));
-					}
-				}
-
-				CompilingShaderMap->ReleaseCompilingId();
-			}
-
-			if (TimeBudget < 0)
-			{
-				break;
-			}
-		}
-		else
-		{
-			if (CompilingId == GlobalShaderMapId)
-			{
-				ProcessCompiledGlobalShaders(FinishedJobs);
-				PropagateGlobalShadersToAllPrimitives();
-			}
-
-			// ShaderMap was removed from compiling list or is being used by another type of shader map which is maintaining a reference
-			// to the results, either way the job can be released
-			{
-				FScopeLock Lock(&CompileQueueSection);
-				for (FShaderCommonCompileJobPtr& Job : FinishedJobs)
-				{
-					ReleaseJob(Job);
-				}
-			}
-			ShaderMapResultIter.RemoveCurrent();
-		}
-	}
-
-	if (MaterialsToReleaseCompilingId.Num() > 0)
-	{
-		ENQUEUE_RENDER_COMMAND(ReleaseCompilingShaderMapIds)([MaterialsToReleaseCompilingId = MoveTemp(MaterialsToReleaseCompilingId)](FRHICommandListImmediate& RHICmdList)
-		{
-			for (FMaterial* Material : MaterialsToReleaseCompilingId)
-			{
-				check(Material->RenderingThreadCompilingShaderMapId != 0u);
-				Material->RenderingThreadCompilingShaderMapId = 0u;
-				Material->RenderingThreadPendingCompilerEnvironment.SafeRelease();
-			}
-		});
-	}
-
-	if (MaterialsToUpdate.Num() > 0)
-	{
-		FMaterial::SetShaderMapsOnMaterialResources(MaterialsToUpdate);
-
-		for (const auto& It : MaterialsToUpdate)
-		{
-			It.Key->NotifyCompilationFinished();
-		}
-
-		if (FApp::CanEverRender())
-		{
-			// This empties MaterialsToUpdate, see the comment inside the function for the reason.
-			PropagateMaterialChangesToPrimitives(MaterialsToUpdate);
-
-			FEditorSupportDelegates::RedrawAllViewports.Broadcast();
-		}
-	}
-
-	UpdateNumRemainingAssets();
-#endif // WITH_EDITOR
-}
-
-void FShaderCompilingManager::PropagateMaterialChangesToPrimitives(TMap<TRefCountPtr<FMaterial>, TRefCountPtr<FMaterialShaderMap>>& MaterialsToUpdate)
-{
-	// don't perform any work if no compiling
-	if (!AllowShaderCompiling())
-	{
-		return;
-	}
-
-	TSet<FSceneInterface*> ScenesToUpdate;
-	FObjectCacheContextScope ObjectCacheScope;
-	{
-		TRACE_CPUPROFILER_EVENT_SCOPE(FShaderCompilingManager::PropagateMaterialChangesToPrimitives);
-
-		TArray<UMaterialInterface*> UpdatedMaterials;
-		for (TMap<TRefCountPtr<FMaterial>, TRefCountPtr<FMaterialShaderMap>>::TConstIterator MaterialIt(MaterialsToUpdate); MaterialIt; ++MaterialIt)
-		{
-			FMaterial* UpdatedMaterial = MaterialIt.Key();
-			UpdatedMaterials.Add(UpdatedMaterial->GetMaterialInterface());
-		}
-
-		for (IPrimitiveComponent* PrimitiveComponent : ObjectCacheScope.GetContext().GetPrimitivesAffectedByMaterials(UpdatedMaterials))
-		{
-			PrimitiveComponent->MarkRenderStateDirty();
-		}
-	}
-
-	// Recreating the render state for the primitives may end up recreating the material resources if some materials are missing some usage flags.
-	// For example, if some materials are not marked as used with static lighting and we build lightmaps, UMaterialInstance::CheckMaterialUsage
-	// will catch the problem and try to set the flag. However, since MaterialsToUpdate stores smart pointers, the material resources will have
-	// a refcount of 2, so the FMaterial destructor will trigger a check failure because the refcount doesn't reach 0. Empty this map before
-	// recreating the render state to allow resources to be deleted cleanly.
-	MaterialsToUpdate.Empty();
-}
-
-
-/**
- * Shutdown the shader compile manager
- * this function should be used when ending the game to shutdown shader compile threads
- * will not complete current pending shader compilation
- */
-void FShaderCompilingManager::Shutdown()
-{
-	// Shutdown has been moved to the destructor because the shader compiler lifetime is expected to
-	// be longer than other asset compilers, otherwise niagara compilations might get stuck.
-}
-
-void FShaderCompilingManager::PrintStats()
-{
-	FShaderCompilerStats LocalStats;
-	GetLocalStats(LocalStats);
-	LocalStats.WriteStatSummary();
-}
-
-void FShaderCompilingManager::GetLocalStats(FShaderCompilerStats& OutStats) const
-{
-	if (GShaderCompilerStats)
-	{
-		OutStats.Aggregate(*GShaderCompilerStats);
-		AllJobs.GetCachingStats(OutStats);
-	}
-}
-
-FShaderCompileMemoryUsage FShaderCompilingManager::GetExternalMemoryUsage()
-{
-	FShaderCompileMemoryUsage TotalMemoryUsage{};
-	for (const TUniquePtr<FShaderCompileThreadRunnableBase>& ThreadPtr : Threads)
-	{
-		FShaderCompileMemoryUsage MemoryUsage = ThreadPtr->GetExternalWorkerMemoryUsage();
-		TotalMemoryUsage.VirtualMemory += MemoryUsage.VirtualMemory;
-		TotalMemoryUsage.PhysicalMemory += MemoryUsage.PhysicalMemory;
-	}
-	return TotalMemoryUsage;
-}
-
-bool FShaderCompilingManager::HandlePotentialRetryOnError(TMap<int32, FShaderMapFinalizeResults>& CompletedShaderMaps)
-{
-	if (FApp::IsUnattended())
-	{
-		return false;
-	}
-
-	bool bRetryCompile = false;
-
-#if WITH_EDITORONLY_DATA
-	for (TMap<int32, FShaderMapFinalizeResults>::TIterator It(CompletedShaderMaps); It; ++It)
-	{
-		FShaderMapFinalizeResults& Results = It.Value();
-
-		if (!Results.bAllJobsSucceeded)
-		{
-			bool bSpecialEngineMaterial = false;
-
-			const FMaterialShaderMap* ShaderMap = FMaterialShaderMap::FindCompilingShaderMap(It.Key());
-			if (ShaderMap)
-			{
-				for (const FMaterial* Material : ShaderMap->CompilingMaterialDependencies)
-				{
-					if (Material->IsSpecialEngineMaterial())
-					{
-						bSpecialEngineMaterial = true;
-						break;
-					}
-				}
-			}
-
-			if (UE_LOG_ACTIVE(LogShaders, Log) 
-				// Always log detailed errors when a special engine material or global shader fails to compile, as those will be fatal errors
-				|| bSpecialEngineMaterial 
-				|| It.Key() == GlobalShaderMapId)
-			{
-				TArray<FShaderCommonCompileJobPtr>& CompleteJobs = Results.FinishedJobs;
-				FShaderDiagnosticInfo ShaderDiagInfo(CompleteJobs);
-
-				const TCHAR* MaterialName = ShaderMap ? ShaderMap->GetFriendlyName() : TEXT("global shaders");
-				FString ErrorString = FString::Printf(TEXT("%i Shader compiler errors compiling %s for platform %s:"), ShaderDiagInfo.UniqueErrors.Num(), MaterialName, *ShaderDiagInfo.TargetShaderPlatformString);
-				UE_LOG(LogShaderCompilers, Warning, TEXT("%s"), *ErrorString);
-				ErrorString += TEXT("\n");
-
-				bool bAnyErrorLikelyToBeCodeError = false;
-				for (const FShaderCommonCompileJob* Job : ShaderDiagInfo.ErrorJobs)
-				{
-					bAnyErrorLikelyToBeCodeError |= Job->bErrorsAreLikelyToBeCode;
-				}
-
-				BuildErrorStringAndReport(ShaderDiagInfo, ErrorString);
-
-				if (UE_LOG_ACTIVE(LogShaders, Log) && (bAnyErrorLikelyToBeCodeError || bPromptToRetryFailedShaderCompiles || bSpecialEngineMaterial))
-				{
-					// Use debug break in debug with the debugger attached, otherwise message box
-					if (bDebugBreakOnPromptToRetryShaderCompile && FPlatformMisc::IsDebuggerPresent())
-					{
-						// A shader compile error has occurred, see the debug output for information.
-						// Double click the errors in the VS.NET output window and the IDE will take you directly to the file and line of the error.
-						// Check ErrorJobs for more state on the failed shaders, for example in-memory includes like Material.usf
-						UE_DEBUG_BREAK();
-						// Set GRetryShaderCompilation to true in the debugger to enable retries in debug
-						// NOTE: MaterialTemplate.usf will not be reloaded when retrying!
-						bRetryCompile = GRetryShaderCompilation;
-					}
-					else
-					{
-						if (FPlatformMisc::MessageBoxExt(EAppMsgType::YesNo, *FText::Format(NSLOCTEXT("UnrealEd", "Error_RetryShaderCompilation", "{0}\r\n\r\nRetry compilation?"),
-							FText::FromString(ErrorString)).ToString(), TEXT("Error")) == EAppReturnType::Type::Yes)
-						{
-							bRetryCompile = true;
-						}
-					}
-				}
-
-				if (bRetryCompile)
-				{
-					break;
-				}
-			}
-		}
-	}
-
-	if (bRetryCompile)
-	{
-		// Flush the shader file cache so that any changes will be propagated.
-		FlushShaderFileCache();
-
-		TArray<int32> MapsToRemove;
-
-		for (TMap<int32, FShaderMapFinalizeResults>::TIterator It(CompletedShaderMaps); It; ++It)
-		{
-			FShaderMapFinalizeResults& Results = It.Value();
-
-			if (!Results.bAllJobsSucceeded)
-			{
-				MapsToRemove.Add(It.Key());
-
-				// Reset outputs
-				for (int32 JobIndex = 0; JobIndex < Results.FinishedJobs.Num(); JobIndex++)
-				{
-					FShaderCommonCompileJob& CurrentJob = *Results.FinishedJobs[JobIndex];
-					auto* SingleJob = CurrentJob.GetSingleShaderJob();
-
-					// NOTE: Changes to MaterialTemplate.usf before retrying won't work, because the entry for Material.usf in CurrentJob.Environment.IncludeFileNameToContentsMap isn't reset
-					if (SingleJob)
-					{
-						if (GShaderCompilingManager->ShouldRecompileToDumpShaderDebugInfo(*SingleJob))
-						{
-							SingleJob->Input.DumpDebugInfoPath = GShaderCompilingManager->CreateShaderDebugInfoPath(SingleJob->Input);
-						}
-						SingleJob->Output = FShaderCompilerOutput();
-						SingleJob->PreprocessOutput = FShaderPreprocessOutput();
-					}
-					else
-					{
-						auto* PipelineJob = CurrentJob.GetShaderPipelineJob();
-						for (auto CommonJob : PipelineJob->StageJobs)
-						{
-							CommonJob->Output = FShaderCompilerOutput();
-							CommonJob->PreprocessOutput = FShaderPreprocessOutput();
-							CommonJob->bFinalized = false;
-						}
-					}
-					
-					// Reset DDC query request owner
-					CurrentJob.RequestOwner.Reset();
-					CurrentJob.bFinalized = false;
-				}
-
-				// Send all the shaders from this shader map through the compiler again
-				SubmitJobs(Results.FinishedJobs, FString(""), FString(""));
-			}
-		}
-
-		const int32 OriginalNumShaderMaps = CompletedShaderMaps.Num();
-
-		// Remove the failed shader maps
-		for (int32 RemoveIndex = 0; RemoveIndex < MapsToRemove.Num(); RemoveIndex++)
-		{
-			CompletedShaderMaps.Remove(MapsToRemove[RemoveIndex]);
-		}
-
-		check(CompletedShaderMaps.Num() == OriginalNumShaderMaps - MapsToRemove.Num());
-
-		// Block until the failed shader maps have been compiled again
-		BlockOnShaderMapCompletion(MapsToRemove, CompletedShaderMaps);
-
-		check(CompletedShaderMaps.Num() == OriginalNumShaderMaps);
-	}
-#endif	//WITH_EDITORONLY_DATA
-
-	return bRetryCompile;
-}
-
-void FShaderMapCompileResults::CheckIfHung()
-{
-	if (!bIsHung)
-	{
-		double DurationSoFar = FPlatformTime::Seconds() - TimeStarted;
-		if (DurationSoFar >= static_cast<double>(GShaderMapCompilationTimeout))
-		{
-			bIsHung = true;
-			// always produce an error message first, even if going to crash, as the automation controller does not seem to be picking up Fatal messages
-			UE_LOG(LogShaderCompilers, Error, TEXT("Hung shadermap detected, time spent compiling: %f seconds, NumPendingJobs: %d, FinishedJobs: %d"),
-				DurationSoFar,
-				NumPendingJobs.GetValue(),
-				FinishedJobs.Num()
-			);
-
-			if (GCrashOnHungShaderMaps)
-			{
-				UE_LOG(LogShaderCompilers, Fatal, TEXT("Crashing on a hung shadermap, time spent compiling: %f seconds, NumPendingJobs: %d, FinishedJobs: %d"),
-					DurationSoFar,
-					NumPendingJobs.GetValue(),
-					FinishedJobs.Num()
-				);
-			}
-		}
-	}
-}
-
-void FShaderCompilingManager::CancelCompilation(const TCHAR* MaterialName, const TArray<int32>& ShaderMapIdsToCancel)
-{
-	// nothing to cancel here, just in case
-	if (!AllowShaderCompiling())
-	{
-		return;
-	}
-
-	check(IsInGameThread());
-	check(!FPlatformProperties::RequiresCookedData());
-
-	// Lock CompileQueueSection so we can access the input and output queues
-	FScopeLock Lock(&CompileQueueSection);
-
-	int32 TotalNumJobsRemoved = 0;
-	for (int32 IdIndex = 0; IdIndex < ShaderMapIdsToCancel.Num(); ++IdIndex)
-	{
-		int32 MapIdx = ShaderMapIdsToCancel[IdIndex];
-		if (const FPendingShaderMapCompileResultsPtr* ResultsPtr = ShaderMapJobs.Find(MapIdx))
-		{
-			const int32 NumJobsRemoved = AllJobs.RemoveAllPendingJobsWithId(MapIdx);
-	
-			TotalNumJobsRemoved += NumJobsRemoved;
-
-			FShaderMapCompileResults* ShaderMapJob = *ResultsPtr;
-			const int32 PrevNumPendingJobs = ShaderMapJob->NumPendingJobs.Subtract(NumJobsRemoved);
-			check(PrevNumPendingJobs >= NumJobsRemoved);
-
-			// The shader map job result should be skipped since it is out of date.
-			ShaderMapJob->bSkipResultProcessing = true;
-		
-			if (PrevNumPendingJobs == NumJobsRemoved && ShaderMapJob->FinishedJobs.Num() == 0)
-			{
-				//We've removed all the jobs for this shader map so remove it.
-				ShaderMapJobs.Remove(MapIdx);
-			}
-		}
-
-		// Don't continue finalizing once compilation has been canceled
-		// the CompilingId has been removed from ShaderMapsBeingCompiled, which will cause crash when attempting to do any further processing
-		const int32 NumPendingRemoved = PendingFinalizeShaderMaps.Remove(MapIdx);
-	}
-
-	if (TotalNumJobsRemoved > 0)
-	{
-		UE_LOG(LogShaders, Display, TEXT("CancelCompilation %s, Removed %d jobs"), MaterialName ? MaterialName : TEXT(""), TotalNumJobsRemoved);
-	}
-}
-
-void FShaderCompilingManager::FinishCompilation(const TCHAR* MaterialName, const TArray<int32>& ShaderMapIdsToFinishCompiling)
-{
-	TRACE_CPUPROFILER_EVENT_SCOPE(FShaderCompilingManager::FinishCompilation);
-
-	// nothing to do
-	if (!AllowShaderCompiling())
-	{
-		return;
-	}
-
-	check(IsInGameThread());
-	check(!FPlatformProperties::RequiresCookedData());
-	const double StartTime = FPlatformTime::Seconds();
-
-	FText StatusUpdate;
-	if ( MaterialName != nullptr)
-	{
-		FFormatNamedArguments Args;
-		Args.Add( TEXT("MaterialName"), FText::FromString( MaterialName ) );
-		StatusUpdate = FText::Format( NSLOCTEXT("ShaderCompilingManager", "CompilingShadersForMaterialStatus", "Compiling shaders: {MaterialName}..."), Args );
-	}
-	else
-	{
-		StatusUpdate = NSLOCTEXT("ShaderCompilingManager", "CompilingShadersStatus", "Compiling shaders...");
-	}
-
-	FScopedSlowTask SlowTask(1, StatusUpdate, GIsEditor && !IsRunningCommandlet());
-	SlowTask.EnterProgressFrame(1);
-
-	TMap<int32, FShaderMapFinalizeResults> CompiledShaderMaps;
-	CompiledShaderMaps.Append( PendingFinalizeShaderMaps );
-	PendingFinalizeShaderMaps.Empty();
-	BlockOnShaderMapCompletion(ShaderMapIdsToFinishCompiling, CompiledShaderMaps);
-
-	bool bRetry = false;
-	do 
-	{
-		bRetry = HandlePotentialRetryOnError(CompiledShaderMaps);
-	} 
-	while (bRetry);
-
-	ProcessCompiledShaderMaps(CompiledShaderMaps, FLT_MAX);
-	check(CompiledShaderMaps.Num() == 0);
-
-	const double EndTime = FPlatformTime::Seconds();
-
-	UE_LOG(LogShaders, Verbose, TEXT("FinishCompilation %s %.3fs"), MaterialName ? MaterialName : TEXT(""), (float)(EndTime - StartTime));
-}
-
-void FShaderCompilingManager::FinishAllCompilation()
-{
-#if WITH_EDITOR
-	// This is here for backward compatibility since textures are most probably expected to be ready too.
-	FTextureCompilingManager::Get().FinishAllCompilation();
-#endif
-
-	TRACE_CPUPROFILER_EVENT_SCOPE(FShaderCompilingManager::FinishAllCompilation);
-	check(IsInGameThread());
-	check(!FPlatformProperties::RequiresCookedData());
-	const double StartTime = FPlatformTime::Seconds();
-
-	TMap<int32, FShaderMapFinalizeResults> CompiledShaderMaps;
-	CompiledShaderMaps.Append( PendingFinalizeShaderMaps );
-	PendingFinalizeShaderMaps.Empty();
-	BlockOnAllShaderMapCompletion(CompiledShaderMaps);
-
-	bool bRetry = false;
-	do 
-	{
-		bRetry = HandlePotentialRetryOnError(CompiledShaderMaps);
-	} 
-	while (bRetry);
-
-	ProcessCompiledShaderMaps(CompiledShaderMaps, FLT_MAX);
-	check(CompiledShaderMaps.Num() == 0);
-
-	const double EndTime = FPlatformTime::Seconds();
-
-	UE_LOG(LogShaders, Verbose, TEXT("FinishAllCompilation %.3fs"), (float)(EndTime - StartTime));
-}
-
-void FShaderCompilingManager::ProcessAsyncResults(bool bLimitExecutionTime, bool bBlockOnGlobalShaderCompletion)
-{
-	const float TimeSlice = bLimitExecutionTime ? ProcessGameThreadTargetTime : 0.f;
-	ProcessAsyncResults(TimeSlice, bBlockOnGlobalShaderCompletion);
-}
-
-void FShaderCompilingManager::ProcessAsyncResults(float TimeSlice, bool bBlockOnGlobalShaderCompletion)
-{
-	LLM_SCOPE_BYTAG(ShaderCompiler);
-
-	TRACE_CPUPROFILER_EVENT_SCOPE(FShaderCompilingManager::ProcessAsyncResults)
-
-	COOK_STAT(FScopedDurationTimer Timer(ShaderCompilerCookStats::ProcessAsyncResultsTimeSec));
-	check(IsInGameThread());
-
-	const double StartTime = FPlatformTime::Seconds();
-
-	// Some controllers need to be manually ticked if the engine loop is not initialized or blocked
-	// to do things like tick the HTTPModule.
-	// Otherwise the results from the controller will never be processed.
-	// We check for bBlockOnGlobalShaderCompletion because the BlockOnShaderMapCompletion methods already do this.
-	if (!bBlockOnGlobalShaderCompletion && BuildDistributionController)
-	{
-		BuildDistributionController->Tick(0.0f);
-	}
-
-	// Block on global shaders before checking for shader maps to finalize
-	// So if we block on global shaders for a long time, we will get a chance to finalize all the non-global shader maps completed during that time.
-	if (bBlockOnGlobalShaderCompletion)
-	{
-		TArray<int32> ShaderMapId;
-		ShaderMapId.Add(GlobalShaderMapId);
-
-		// Block until the global shader map jobs are complete
-		GShaderCompilingManager->BlockOnShaderMapCompletion(ShaderMapId, PendingFinalizeShaderMaps);
-	}
-
-	int32 NumCompilingShaderMaps = 0;
-
-	{
-		// Lock CompileQueueSection so we can access the input and output queues
-		FScopeLock Lock(&CompileQueueSection);
-
-		if (!bBlockOnGlobalShaderCompletion)
-		{
-			bCompilingDuringGame = true;
-		}
-
-		// Get all material shader maps to finalize
-		//
-		for (TMap<int32, FPendingShaderMapCompileResultsPtr>::TIterator It(ShaderMapJobs); It; ++It)
-		{
-			FPendingShaderMapCompileResultsPtr& Results = It.Value();
-			if (Results->FinishedJobs.Num() > 0)
-			{
-				FShaderMapFinalizeResults& FinalizeResults = PendingFinalizeShaderMaps.FindOrAdd(It.Key());
-				FinalizeResults.FinishedJobs.Append(Results->FinishedJobs);
-				Results->FinishedJobs.Reset();
-				FinalizeResults.bAllJobsSucceeded = FinalizeResults.bAllJobsSucceeded && Results->bAllJobsSucceeded;
-			}
-
-			checkf(Results->FinishedJobs.Num() == 0, TEXT("Failed to remove finished jobs, %d remain"), Results->FinishedJobs.Num());
-			if (Results->NumPendingJobs.GetValue() == 0)
-			{
-				It.RemoveCurrent();
-			}
-		}
-
-		NumCompilingShaderMaps = ShaderMapJobs.Num();
-	}
-
-	int32 NumPendingShaderMaps = PendingFinalizeShaderMaps.Num();
-
-	if (PendingFinalizeShaderMaps.Num() > 0)
-	{
-		bool bRetry = false;
-		do 
-		{
-			bRetry = HandlePotentialRetryOnError(PendingFinalizeShaderMaps);
-		} 
-		while (bRetry);
-
-		const float TimeBudget = TimeSlice > 0 ? TimeSlice : FLT_MAX;
-		ProcessCompiledShaderMaps(PendingFinalizeShaderMaps, TimeBudget);
-		check(TimeSlice > 0 || PendingFinalizeShaderMaps.Num() == 0);
-	}
-
-
-	if (bBlockOnGlobalShaderCompletion && TimeSlice <= 0 && !IsRunningCookCommandlet())
-	{
-		check(PendingFinalizeShaderMaps.Num() == 0);
-
-		if (NumPendingShaderMaps - PendingFinalizeShaderMaps.Num() > 0)
-		{
-			UE_LOG(LogShaders, Warning, TEXT("Blocking ProcessAsyncResults for %.1fs, processed %u shader maps, %u being compiled"), 
-				(float)(FPlatformTime::Seconds() - StartTime),
-				NumPendingShaderMaps - PendingFinalizeShaderMaps.Num(), 
-				NumCompilingShaderMaps);
-		}
-	}
-	else if (NumPendingShaderMaps - PendingFinalizeShaderMaps.Num() > 0)
-	{
-		UE_LOG(LogShaders, Verbose, TEXT("Completed %u async shader maps, %u more pending, %u being compiled"),
-			NumPendingShaderMaps - PendingFinalizeShaderMaps.Num(), 
-			PendingFinalizeShaderMaps.Num(),
-			NumCompilingShaderMaps);
-	}
-
-	UpdateNumRemainingAssets();
-}
-
-void FShaderCompilingManager::UpdateNumRemainingAssets()
-{
-	if (IsInGameThread())
-	{
-		const int32 NumRemainingAssets = GetNumRemainingAssets();
-		if (LastNumRemainingAssets != NumRemainingAssets)
-		{
-			if (NumRemainingAssets == 0)
-			{
-				// This is important to at least broadcast once we reach 0 remaining assets
-				// even if we don't have any UObject to report because some listener are only 
-				// interested to be notified when the number of async compilation reaches 0.
-				FAssetCompilingManager::Get().OnAssetPostCompileEvent().Broadcast({});
-			}
-
-			LastNumRemainingAssets = NumRemainingAssets;
-			Notification->Update(NumRemainingAssets);
-		}
-	}
-}
-
-bool FShaderCompilingManager::IsShaderCompilerWorkerRunning(FProcHandle & WorkerHandle)
-{
-	return FPlatformProcess::IsProcRunning(WorkerHandle);
-}
-
-#if WITH_EDITOR
-
-/* Generates a uniform buffer struct member hlsl declaration using the member's metadata. */
-static void GenerateUniformBufferStructMember(FString& Result, const FShaderParametersMetadata::FMember& Member, EShaderPlatform ShaderPlatform)
-{
-	// Generate the base type name.
-	FString TypeName;
-	Member.GenerateShaderParameterType(TypeName, ShaderPlatform);
-
-	// Generate array dimension post fix
-	FString ArrayDim;
-	if (Member.GetNumElements() > 0)
-	{
-		ArrayDim = FString::Printf(TEXT("[%u]"), Member.GetNumElements());
-	}
-
-	Result = FString::Printf(TEXT("%s %s%s"), *TypeName, Member.GetName(), *ArrayDim);
-}
-
-/* Generates the instanced stereo hlsl code that's dependent on view uniform declarations. */
-void GenerateInstancedStereoCode(FString& Result, EShaderPlatform ShaderPlatform)
-{
-	// Find the InstancedView uniform buffer struct
-	const FShaderParametersMetadata* View = nullptr;
-	const FShaderParametersMetadata* InstancedView = nullptr;
-
-	for (TLinkedList<FShaderParametersMetadata*>::TIterator StructIt(FShaderParametersMetadata::GetStructList()); StructIt; StructIt.Next())
-	{
-		if (StructIt->GetShaderVariableName() == FString(TEXT("View")))
-		{
-			View = *StructIt;
-		}
-
-		if (StructIt->GetShaderVariableName() == FString(TEXT("InstancedView")))
-		{
-			InstancedView = *StructIt;
-		}
-
-		if (View && InstancedView)
-		{
-			break;
-		}
-	}
-	checkSlow(View != nullptr);
-	checkSlow(InstancedView != nullptr);
-
-	const TArray<FShaderParametersMetadata::FMember>& StructMembersView = View->GetMembers();
-	const TArray<FShaderParametersMetadata::FMember>& StructMembersInstanced = InstancedView->GetMembers();
-
-	static const auto CVarViewHasTileOffsetData = IConsoleManager::Get().FindTConsoleVariableDataInt(TEXT("r.ViewHasTileOffsetData"));
-	const bool bViewHasTileOffsetData = CVarViewHasTileOffsetData ? (CVarViewHasTileOffsetData->GetValueOnAnyThread() != 0) : false;
-
-	Result = "";
-	if (bViewHasTileOffsetData)
-	{
-		Result +=  "struct ViewStateTileOffsetData\r\n";
-		Result += "{\r\n";
-		Result += "\tFLWCVector3 WorldCameraOrigin;\r\n";
-		Result += "\tFLWCVector3 WorldViewOrigin;\r\n";
-		Result += "\tFLWCVector3 PrevWorldCameraOrigin;\r\n";
-		Result += "\tFLWCVector3 PrevWorldViewOrigin;\r\n";
-		Result += "\tFLWCVector3 PreViewTranslation;\r\n";
-		Result += "\tFLWCVector3 PrevPreViewTranslation;\r\n";
-		Result += "};\r\n";
-	}
-
-	// ViewState definition
-	Result +=  "struct ViewState\n";
-	Result += "{\n";
-	for (int32 MemberIndex = 0; MemberIndex < StructMembersInstanced.Num(); ++MemberIndex)
-	{
-		FString MemberDecl;
-		// ViewState is only supposed to contain InstancedView members however we want their original type and array-length instead of their representation in the instanced array
-		// GPUSceneViewId for example needs to return 	uint GPUSceneViewId; and not uint4 InstancedView_GPUSceneViewId[2];
-		// and that initial representation is in StructMembersView
-		GenerateUniformBufferStructMember(MemberDecl, StructMembersView[MemberIndex], ShaderPlatform);
-		Result += FString::Printf(TEXT("\t%s;\n"), *MemberDecl);
-	}
-	Result += "\tFDFInverseMatrix WorldToClip;\n";
-	Result += "\tFDFMatrix ClipToWorld;\n";
-	Result += "\tFDFMatrix ScreenToWorld;\n";
-	Result += "\tFDFMatrix PrevClipToWorld;\n";
-	Result += "\tFDFVector3 WorldCameraOrigin;\n";
-	Result += "\tFDFVector3 WorldViewOrigin;\n";
-	Result += "\tFDFVector3 PrevWorldCameraOrigin;\n";
-	Result += "\tFDFVector3 PrevWorldViewOrigin;\n";
-	Result += "\tFDFVector3 PreViewTranslation;\n";
-	Result += "\tFDFVector3 PrevPreViewTranslation;\n";
-
-	if (bViewHasTileOffsetData)
-	{
-		Result += "\tViewStateTileOffsetData TileOffset;\n";
-	}
-
-	Result += "};\n";
-
-	Result += "\tvoid FinalizeViewState(inout ViewState InOutView);\n";
-
-	// GetPrimaryView definition
-	Result += "ViewState GetPrimaryView()\n";
-	Result += "{\n";
-	Result += "\tViewState Result;\n";
-	for (int32 MemberIndex = 0; MemberIndex < StructMembersInstanced.Num(); ++MemberIndex)
-	{
-		const FShaderParametersMetadata::FMember& Member = StructMembersView[MemberIndex];
-		Result += FString::Printf(TEXT("\tResult.%s = View.%s;\n"), Member.GetName(), Member.GetName());
-	}
-	Result += "\tFinalizeViewState(Result);\n";
-	Result += "\treturn Result;\n";
-	Result += "}\n";
-
-	// GetInstancedView definition
-	Result += "#if (INSTANCED_STEREO || MOBILE_MULTI_VIEW)\n";
-	Result += "ViewState GetInstancedView(uint ViewIndex)\n";
-	Result += "{\n";
-	Result += "\tViewState Result;\n";
-	for (int32 MemberIndex = 0; MemberIndex < StructMembersInstanced.Num(); ++MemberIndex)
-	{
-		const FShaderParametersMetadata::FMember& ViewMember = StructMembersView[MemberIndex];
-		const FShaderParametersMetadata::FMember& InstancedViewMember = StructMembersInstanced[MemberIndex];
+#if USE_SHADER_COMPILER_WORKER_TRACE
+	// When doing utrace functionality we can't run with -nothreading, since it won't create the utrace thread to send events.
+	WorkerParameters += FString(TEXT(" -trace=default "));
+#else
+	WorkerParameters += FString(TEXT(" -nothreading "));
+#endif // USE_SHADER_COMPILER_WORKER_TRACE
 
-		FString ViewMemberTypeName;
-		ViewMember.GenerateShaderParameterType(ViewMemberTypeName, ShaderPlatform);
+	// Launch the worker process
+	int32 PriorityModifier = -1; // below normal
+	GConfig->GetInt(TEXT("DevOptions.Shaders"), TEXT("WorkerProcessPriority"), PriorityModifier, GEngineIni);
 
-		// this code avoids an assumption that instanced buffer only supports 2 views, to be future-proof
-		if (ViewMember.GetNumElements() >= 1 && (InstancedViewMember.GetNumElements() >= 2 * ViewMember.GetNumElements()))
-		{
-			// if View has an array (even 1-sized) for this index, and InstancedView has Nx (N>=2) the element count -> per-view array
-			// Result.TranslucencyLightingVolumeMin[0] = (float4) InstancedView_TranslucencyLightingVolumeMin[ViewIndex * 2 + 0];
-			checkf((InstancedViewMember.GetNumElements() % ViewMember.GetNumElements()) == 0, TEXT("Per-view arrays are expected to be stored in an array that is an exact multiple of the original array."));
-			for (uint32 ElementIndex = 0; ElementIndex < ViewMember.GetNumElements(); ElementIndex++)
-			{
-				Result += FString::Printf(TEXT("\tResult.%s[%u] = (%s) InstancedView.%s[ViewIndex * %u + %u];\n"),
-					ViewMember.GetName(), ElementIndex, *ViewMemberTypeName, InstancedViewMember.GetName(), ViewMember.GetNumElements(), ElementIndex);
-			}
-		}
-		else if (InstancedViewMember.GetNumElements() > 1 && ViewMember.GetNumElements() == 0)
-		{
-			// if View has a scalar field for this index, and InstancedView has an array with >1 elements -> per-view scalar
-			// 	Result.TranslatedWorldToClip = (float4x4) InstancedView_TranslatedWorldToClip[ViewIndex];
-			Result += FString::Printf(TEXT("\tResult.%s = (%s) InstancedView.%s[ViewIndex];\n"),
-				ViewMember.GetName(), *ViewMemberTypeName, InstancedViewMember.GetName());
-		}
-		else if (InstancedViewMember.GetNumElements() == ViewMember.GetNumElements())
+	if (DEBUG_SHADERCOMPILEWORKER)
+	{
+		// Note: Set breakpoint here and launch the ShaderCompileWorker with WorkerParameters a cmd-line
+		const TCHAR* WorkerParametersText = *WorkerParameters;
+		FPlatformMisc::LowLevelOutputDebugStringf(TEXT("Launching shader compile worker w/ WorkerParameters\n\t%s\n"), WorkerParametersText);
+		FProcHandle DummyHandle;
+		return DummyHandle;
+	}
+	else
+	{
+#if UE_BUILD_DEBUG && PLATFORM_LINUX
+		FPlatformMisc::LowLevelOutputDebugStringf(TEXT("Launching shader compile worker:\n\t%s\n"), *WorkerParameters);
+#endif
+		// Disambiguate between SCW.exe missing vs other errors.
+		static bool bFirstLaunch = true;
+		uint32 WorkerId = 0;
+		FProcHandle WorkerHandle = FPlatformProcess::CreateProc(*ShaderCompileWorkerName, *WorkerParameters, true, false, false, &WorkerId, PriorityModifier, nullptr, nullptr);
+		if (WorkerHandle.IsValid())
 		{
-			// if View has the same number of elements for this index as InstancedView, it's backed by a view-dependent array, assume a view-independent field
-			// 	Result.TemporalAAParams = InstancedView_TemporalAAParams;
-			Result += FString::Printf(TEXT("\tResult.%s = InstancedView.%s;\n"),
-				ViewMember.GetName(), InstancedViewMember.GetName());
+			// Process launched at least once successfully
+			bFirstLaunch = false;
 		}
 		else
 		{
-			// something unexpected, better crash now rather than generate wrong shader code and poison DDC 
-			UE_LOG(LogShaderCompilers, Fatal, TEXT("Don't know how to copy View buffers' field %s (NumElements=%d) from InstancedView field %s (NumElements=%d)"),
-				ViewMember.GetName(), ViewMember.GetNumElements(), InstancedViewMember.GetName(), InstancedViewMember.GetNumElements()
-				);
+			// If this doesn't error, the app will hang waiting for jobs that can never be completed
+			if (bFirstLaunch)
+			{
+				// When using source builds users are likely to make a mistake of not building SCW (e.g. in particular on Linux, even though default makefile target builds it).
+				// Make the engine exit gracefully with a helpful message instead of a crash.
+				static bool bShowedMessageBox = false;
+				if (!bShowedMessageBox && !IsRunningCommandlet() && !FApp::IsUnattended())
+				{
+					bShowedMessageBox = true;
+					FText ErrorMessage = FText::Format(LOCTEXT("LaunchingShaderCompileWorkerFailed", "Unable to launch {0} - make sure you built ShaderCompileWorker."), FText::FromString(ShaderCompileWorkerName));
+					FPlatformMisc::MessageBoxExt(EAppMsgType::Ok, *ErrorMessage.ToString(),
+												 *LOCTEXT("LaunchingShaderCompileWorkerFailedTitle", "Unable to launch ShaderCompileWorker.").ToString());
+				}
+				UE_LOG(LogShaderCompilers, Error, TEXT("Couldn't launch %s! Make sure you build ShaderCompileWorker."), *ShaderCompileWorkerName);
+				// duplicate to printf() since threaded logs may not be always flushed
+				FPlatformMisc::LowLevelOutputDebugStringf(TEXT("Couldn't launch %s! Make sure you build ShaderCompileWorker.\n"), *ShaderCompileWorkerName);
+				FPlatformMisc::RequestExitWithStatus(true, 1);
+			}
+			else
+			{
+				UE_LOG(LogShaderCompilers, Fatal, TEXT("Couldn't launch %s!"), *ShaderCompileWorkerName);
+			}
 		}
+
+		return WorkerHandle;
 	}
-	Result += "\tFinalizeViewState(Result);\n";
-	Result += "\treturn Result;\n";
-	Result += "}\n";
-	Result += "#endif\n";
 }
 
-void ValidateShaderFilePath(const FString& VirtualShaderFilePath, const FString& VirtualSourceFilePath)
+void FShaderCompilingManager::AddCompiledResults(TMap<int32, FShaderMapFinalizeResults>& CompiledShaderMaps, int32 ShaderMapIdx, const FShaderMapFinalizeResults& Results)
 {
-	check(CheckVirtualShaderFilePath(VirtualShaderFilePath));
-
-	checkf(VirtualShaderFilePath.Contains(TEXT("/Generated/")),
-		TEXT("Incorrect virtual shader path for generated file '%s': Generated files must be located under an "
-				"non existing 'Generated' directory, for instance: /Engine/Generated/ or /Plugin/FooBar/Generated/."),
-		*VirtualShaderFilePath);
-
-	checkf(VirtualShaderFilePath == VirtualSourceFilePath || FPaths::GetExtension(VirtualShaderFilePath) == TEXT("ush"),
-		TEXT("Incorrect virtual shader path extension for generated file '%s': Generated file must either be the "
-				"USF to compile, or a USH file to be included."),
-		*VirtualShaderFilePath);
-}
-
-/** Lock for the storage of instanced stereo code. */
-FCriticalSection GCachedGeneratedInstancedStereoCodeLock;
-
-/** Storage for instanced stereo code so it is not generated every time we compile a shader. */
-TMap<EShaderPlatform, FThreadSafeSharedAnsiStringPtr> GCachedGeneratedInstancedStereoCode;
+	// merge with the previous unprocessed jobs, if any
+	if (FShaderMapCompileResults const* PrevResults = CompiledShaderMaps.Find(ShaderMapIdx))
+	{
+		FShaderMapFinalizeResults NewResults(Results);
 
-void GlobalBeginCompileShader(
-	const FString& DebugGroupName,
-	const FVertexFactoryType* VFType,
-	const FShaderType* ShaderType,
-	const FShaderPipelineType* ShaderPipelineType,
-	int32 PermutationId,
-	const TCHAR* SourceFilename,
-	const TCHAR* FunctionName,
-	FShaderTarget Target,
-	FShaderCompilerInput& Input,
-	bool bAllowDevelopmentShaderCompile,
-	const FString& DebugDescription,
-	const FString& DebugExtension
-)
-{
-	GlobalBeginCompileShader(
-		DebugGroupName,
-		VFType,
-		ShaderType,
-		ShaderPipelineType,
-		PermutationId,
-		SourceFilename,
-		FunctionName,
-		Target,
-		Input,
-		bAllowDevelopmentShaderCompile,
-		*DebugDescription,
-		*DebugExtension
-		);
-}
+		NewResults.bAllJobsSucceeded = NewResults.bAllJobsSucceeded && PrevResults->bAllJobsSucceeded;
+		NewResults.bSkipResultProcessing = NewResults.bSkipResultProcessing || PrevResults->bSkipResultProcessing;
+		NewResults.TimeStarted = FMath::Min(NewResults.TimeStarted, PrevResults->TimeStarted);
+		NewResults.bIsHung = NewResults.bIsHung || PrevResults->bIsHung;
+		NewResults.FinishedJobs.Append(PrevResults->FinishedJobs);
 
-namespace
-{
-	bool ShaderFrequencyNeedsInstancedStereoMods(const FShaderType* ShaderType)
+		CompiledShaderMaps.Add(ShaderMapIdx, NewResults);
+	}
+	else
 	{
-		return !(IsRayTracingShaderFrequency(ShaderType->GetFrequency()));
+		CompiledShaderMaps.Add(ShaderMapIdx, Results);
 	}
 }
 
-static bool IsSubstrateSupportForShaderPipeline(const FShaderCompilerInput& Input)
-{
-	// Substrate requires HLSL2021 which must be cross-compiled for D3D11 to be consumed by FXC compiler.
-	// This cross-compilation toolchain does not support geometry shaders.
-	bool bPipelineContainsGeometryShader = false;
-	Input.Environment.GetCompileArgument(TEXT("PIPELINE_CONTAINS_GEOMETRYSHADER"), bPipelineContainsGeometryShader);
-	const bool bCanRHICompileHlsl2021GeometryShaders = !(GetMaxSupportedFeatureLevel((EShaderPlatform)Input.Target.Platform) == ERHIFeatureLevel::SM5);
-	const bool bIsSubstrateSupportedForPipeline = !bPipelineContainsGeometryShader || bCanRHICompileHlsl2021GeometryShaders;
-	return bIsSubstrateSupportedForPipeline;
-}
-
-void GlobalBeginCompileShader(
-	const FString& DebugGroupName,
-	const FVertexFactoryType* VFType,
-	const FShaderType* ShaderType,
-	const FShaderPipelineType* ShaderPipelineType,
-	int32 PermutationId,
-	const TCHAR* SourceFilename,
-	const TCHAR* FunctionName,
-	FShaderTarget Target,
-	FShaderCompilerInput& Input,
-	bool bAllowDevelopmentShaderCompile,
-	const TCHAR* DebugDescription,
-	const TCHAR* DebugExtension
-	)
+/** Flushes all pending jobs for the given shader maps. */
+void FShaderCompilingManager::BlockOnShaderMapCompletion(const TArray<int32>& ShaderMapIdsToFinishCompiling, TMap<int32, FShaderMapFinalizeResults>& CompiledShaderMaps)
 {
-	LLM_SCOPE_BYTAG(ShaderCompiler);
+	// never block if no compiling, just in case
+	if (!AllowShaderCompiling())
+	{
+		return;
+	}
 
-	TRACE_CPUPROFILER_EVENT_SCOPE(GlobalBeginCompileShader);
-	COOK_STAT(ShaderCompilerCookStats::GlobalBeginCompileShaderCalls++);
-	COOK_STAT(FScopedDurationTimer DurationTimer(ShaderCompilerCookStats::GlobalBeginCompileShaderTimeSec));
+	TRACE_CPUPROFILER_EVENT_SCOPE(FShaderCompilingManager::BlockOnShaderMapCompletion);
 
-	const EShaderPlatform ShaderPlatform = EShaderPlatform(Target.Platform);
-	const FName ShaderFormatName = LegacyShaderPlatformToShaderFormat(ShaderPlatform);
+	COOK_STAT(FScopedDurationTimer BlockingTimer(ShaderCompilerCookStats::BlockingTimeSec));
+	if (bAllowAsynchronousShaderCompiling)
+	{
+		// Calculate how many shader jobs there are total to provide the slow task with the correct amount of work.
+		int NumJobs = 0;
+		{
+			FScopeLock Lock(&CompileQueueSection);
+			for (int32 ShaderMapIndex = 0; ShaderMapIndex < ShaderMapIdsToFinishCompiling.Num(); ShaderMapIndex++)
+			{
+				FPendingShaderMapCompileResultsPtr* ResultsPtr = ShaderMapJobs.Find(ShaderMapIdsToFinishCompiling[ShaderMapIndex]);
+				if (ResultsPtr)
+				{
+					FShaderMapCompileResults* Results = *ResultsPtr;
+					NumJobs += Results->NumPendingJobs.GetValue();
+				}
+			}
+		}
 
-	ITargetPlatform* TargetPlatform = GetTargetPlatformManager()->FindTargetPlatformWithSupport(TEXT("ShaderFormat"), ShaderFormatName);
+		FScopedSlowTask SlowTask(NumJobs, FText::Format(LOCTEXT("BlockOnShaderMapCompletion", "Compiling Shaders ({0})"), NumJobs), GIsEditor && !IsRunningCommandlet() && UE::GetPlayInEditorID() == INDEX_NONE);
+		if (NumJobs > 0)
+		{
+			SlowTask.MakeDialogDelayed(1.0f);
+		}
 
-	FShaderCompileUtilities::GenerateBrdfHeaders(ShaderPlatform);
+		int32 NumPendingJobs = 0;
+		// Keep track of previous number of pending jobs so we can update the slow task with the amount of work done.
+		int32 NumPreviousPendingJobs = NumJobs;
+		int32 LogCounter = 0;
+		do 
+		{
+			NumPendingJobs = 0;
+			{
+				// Lock CompileQueueSection so we can access the input and output queues
+				FScopeLock Lock(&CompileQueueSection);
 
-	// NOTE:  Input.bCompilingForShaderPipeline is initialized by the constructor for single versus pipeline jobs, do not initialize again here!
+				for (int32 ShaderMapIndex = 0; ShaderMapIndex < ShaderMapIdsToFinishCompiling.Num(); ShaderMapIndex++)
+				{
+					FPendingShaderMapCompileResultsPtr* ResultsPtr = ShaderMapJobs.Find(ShaderMapIdsToFinishCompiling[ShaderMapIndex]);
+					if (ResultsPtr)
+					{
+						FShaderMapCompileResults* Results = *ResultsPtr;
 
-	Input.Target = Target;
-	Input.ShaderPlatformName = FDataDrivenShaderPlatformInfo::GetName(ShaderPlatform);
-	Input.ShaderFormat = ShaderFormatName;
-	Input.SupportedHardwareMask = TargetPlatform ? TargetPlatform->GetSupportedHardwareMask() : 0;
-	Input.CompressionFormat = GetShaderCompressionFormat();
-	GetShaderCompressionOodleSettings(Input.OodleCompressor, Input.OodleLevel);
-	Input.VirtualSourceFilePath = SourceFilename;
-	Input.EntryPointName = FunctionName;
-	Input.bIncludeUsedOutputs = false;
-	Input.DumpDebugInfoRootPath = GShaderCompilingManager->GetAbsoluteShaderDebugInfoDirectory() / Input.ShaderPlatformName.ToString();
-	Input.DebugInfoFlags = GShaderCompilingManager->GetDumpShaderDebugInfoFlags();
-	// asset material name or "Global"
-	Input.DebugGroupName = DebugGroupName;
-	Input.DebugDescription = DebugDescription;
-	Input.DebugExtension = DebugExtension;
-	Input.RootParametersStructure = ShaderType->GetRootParametersMetadata();
-	Input.ShaderName = ShaderType->GetName();
+						if (Results->NumPendingJobs.GetValue() == 0)
+						{
+							if (Results->FinishedJobs.Num() > 0)
+							{
+								AddCompiledResults(CompiledShaderMaps, ShaderMapIdsToFinishCompiling[ShaderMapIndex], *Results);
+							}
+							ShaderMapJobs.Remove(ShaderMapIdsToFinishCompiling[ShaderMapIndex]);
+						}
+						else
+						{
+							Results->CheckIfHung();
+							NumPendingJobs += Results->NumPendingJobs.GetValue();
+						}
+					}
+				}
+			}
 
-	if (GDumpShaderDebugInfoBindless)
-	{
-		auto GetBindlessString = [](ERHIBindlessConfiguration InConfig)
+			if (NumPendingJobs > 0)
 			{
-				switch (InConfig)
+				const float SleepTime =.01f;
+				
+				// We need to manually tick the Distributed build controller while the game thread is blocked
+				// otherwise we can get stuck in a infinite loop waiting for jobs that never will be done
+				// because for example, some controllers depend on the HTTP module which needs to be ticked in the main thread
+				if (BuildDistributionController && IsInGameThread())
 				{
-				default:
-				case ERHIBindlessConfiguration::Disabled:          return TEXT("Off");
-				case ERHIBindlessConfiguration::AllShaders:        return TEXT("On");
-				case ERHIBindlessConfiguration::RayTracingShaders: return TEXT("RT");
+					BuildDistributionController->Tick(SleepTime);
 				}
-			};
 
-		FStringBuilderBase Builder;
-		Builder.Append("BindlessR").Append(GetBindlessString(UE::ShaderCompiler::GetBindlessResourcesConfiguration(ShaderFormatName)));
-		Builder.Append("_");
-		Builder.Append("BindlessS").Append(GetBindlessString(UE::ShaderCompiler::GetBindlessSamplersConfiguration(ShaderFormatName)));
+				// Progress the slow task with how many jobs we've completed since last tick.  Update the slow task message with the current number of pending jobs
+				// we are waiting on.
+				const int32 CompletedJobsSinceLastTick = NumPreviousPendingJobs - NumPendingJobs;
+				SlowTask.EnterProgressFrame(CompletedJobsSinceLastTick, FText::Format(LOCTEXT("BlockOnShaderMapCompletion", "Compiling Shaders ({0})"), NumPendingJobs));
+				NumPreviousPendingJobs = NumPendingJobs;
+				
+				// Yield CPU time while waiting
+				FPlatformProcess::Sleep(SleepTime);
 
-		Input.DebugGroupName = Builder.ToString() / Input.DebugGroupName;
+				// Flush threaded logs around every 500ms or so based on Sleep of 0.01f seconds above
+				if (++LogCounter > 50)
+				{
+					LogCounter = 0;
+					GLog->FlushThreadedLogs(EOutputDeviceRedirectorFlushOptions::Async);
+				}
+			}
+		} 
+		while (NumPendingJobs > 0);
 	}
-
-	// Verify FShaderCompilerInput's file paths are consistent. 
-	#if DO_CHECK
-		check(CheckVirtualShaderFilePath(Input.VirtualSourceFilePath));
-
-		checkf(FPaths::GetExtension(Input.VirtualSourceFilePath) == TEXT("usf"),
-			TEXT("Incorrect virtual shader path extension for shader file to compile '%s': Only .usf files should be "
-				 "compiled. .ush file are meant to be included only."),
-			*Input.VirtualSourceFilePath);
-
-		for (const auto& Entry : Input.Environment.IncludeVirtualPathToContentsMap)
+	else
+	{
+		int32 NumActiveWorkers = 0;
+		do 
 		{
-			ValidateShaderFilePath(Entry.Key, Input.VirtualSourceFilePath);
-		}
+			for (const TUniquePtr<FShaderCompileThreadRunnableBase>& Thread : Threads)
+			{
+				NumActiveWorkers = Thread->CompilingLoop();
+			}
+		} 
+		while (NumActiveWorkers > 0);
 
-		for (const auto& Entry : Input.Environment.IncludeVirtualPathToSharedContentsMap)
-		{
-			ValidateShaderFilePath(Entry.Key, Input.VirtualSourceFilePath);
-		}
-	#endif
+		check(AllJobs.GetNumPendingJobs() == 0);
 
-	if (ShaderPipelineType)
-	{
-		Input.DebugGroupName = Input.DebugGroupName / ShaderPipelineType->GetName();
-	}
-	
-	if (VFType)
-	{
-		FString VFName = VFType->GetName();
-		if (GDumpShaderDebugInfoShort)
+		for (int32 ShaderMapIndex = 0; ShaderMapIndex < ShaderMapIdsToFinishCompiling.Num(); ShaderMapIndex++)
 		{
-			// Shorten vertex factory name
-			if (VFName[0] == TCHAR('F') || VFName[0] == TCHAR('T'))
+			const FPendingShaderMapCompileResultsPtr* ResultsPtr = ShaderMapJobs.Find(ShaderMapIdsToFinishCompiling[ShaderMapIndex]);
+
+			if (ResultsPtr)
 			{
-				VFName.RemoveAt(0);
+				const FShaderMapCompileResults* Results = *ResultsPtr;
+				check(Results->NumPendingJobs.GetValue() == 0);
+				check(Results->FinishedJobs.Num() > 0);
+
+				AddCompiledResults(CompiledShaderMaps, ShaderMapIdsToFinishCompiling[ShaderMapIndex], *Results);
+				ShaderMapJobs.Remove(ShaderMapIdsToFinishCompiling[ShaderMapIndex]);
 			}
-			VFName.ReplaceInline(TEXT("VertexFactory"), TEXT("VF"));
-			VFName.ReplaceInline(TEXT("GPUSkinAPEXCloth"), TEXT("APEX"));
-			VFName.ReplaceInline(TEXT("true"), TEXT("_1"));
-			VFName.ReplaceInline(TEXT("false"), TEXT("_0"));
 		}
-		Input.DebugGroupName = Input.DebugGroupName / VFName;
 	}
-	
+
+	UpdateNumRemainingAssets();
+}
+
+void FShaderCompilingManager::BlockOnAllShaderMapCompletion(TMap<int32, FShaderMapFinalizeResults>& CompiledShaderMaps)
+{
+	// never block if no compiling, just in case
+	if (!AllowShaderCompiling())
 	{
-		FString ShaderTypeName = ShaderType->GetName();
-		if (GDumpShaderDebugInfoShort)
+		return;
+	}
+
+	TRACE_CPUPROFILER_EVENT_SCOPE(FShaderCompilingManager::BlockOnAllShaderMapCompletion);
+
+	COOK_STAT(FScopedDurationTimer BlockingTimer(ShaderCompilerCookStats::BlockingTimeSec));
+	if (bAllowAsynchronousShaderCompiling)
+	{
+		// Calculate how many shader jobs there are total to provide the slow task with the correct amount of work.
+		int NumJobs = 0;
 		{
-			// Shorten known types
-			if (ShaderTypeName[0] == TCHAR('F') || ShaderTypeName[0] == TCHAR('T'))
+			FScopeLock Lock(&CompileQueueSection);
+			for (TMap<int32, FPendingShaderMapCompileResultsPtr>::TIterator It(ShaderMapJobs); It; ++It)
 			{
-				ShaderTypeName.RemoveAt(0);
+				FShaderMapCompileResults* Results = It.Value();
+				NumJobs += Results->NumPendingJobs.GetValue();
 			}
 		}
-		Input.DebugGroupName = Input.DebugGroupName / ShaderTypeName / FString::Printf(TEXT("%i"), PermutationId);
-		
-		if (GDumpShaderDebugInfoShort)
+
+		FScopedSlowTask SlowTask(NumJobs, FText::Format(LOCTEXT("BlockOnAllShaderMapCompletion", "Compiling Shaders ({0})"), NumJobs), NumJobs && GIsEditor && !IsRunningCommandlet());
+		if (NumJobs > 0)
 		{
-			Input.DebugGroupName.ReplaceInline(TEXT("BasePass"), TEXT("BP"));
-			Input.DebugGroupName.ReplaceInline(TEXT("ForForward"), TEXT("Fwd"));
-			Input.DebugGroupName.ReplaceInline(TEXT("Shadow"), TEXT("Shdw"));
-			Input.DebugGroupName.ReplaceInline(TEXT("LightMap"), TEXT("LM"));
-			Input.DebugGroupName.ReplaceInline(TEXT("EHeightFogFeature==E_"), TEXT(""));
-			Input.DebugGroupName.ReplaceInline(TEXT("Capsule"), TEXT("Caps"));
-			Input.DebugGroupName.ReplaceInline(TEXT("Movable"), TEXT("Mov"));
-			Input.DebugGroupName.ReplaceInline(TEXT("Culling"), TEXT("Cull"));
-			Input.DebugGroupName.ReplaceInline(TEXT("Atmospheric"), TEXT("Atm"));
-			Input.DebugGroupName.ReplaceInline(TEXT("Atmosphere"), TEXT("Atm"));
-			Input.DebugGroupName.ReplaceInline(TEXT("Exponential"), TEXT("Exp"));
-			Input.DebugGroupName.ReplaceInline(TEXT("Ambient"), TEXT("Amb"));
-			Input.DebugGroupName.ReplaceInline(TEXT("Perspective"), TEXT("Persp"));
-			Input.DebugGroupName.ReplaceInline(TEXT("Occlusion"), TEXT("Occ"));
-			Input.DebugGroupName.ReplaceInline(TEXT("Position"), TEXT("Pos"));
-			Input.DebugGroupName.ReplaceInline(TEXT("Skylight"), TEXT("Sky"));
-			Input.DebugGroupName.ReplaceInline(TEXT("LightingPolicy"), TEXT("LP"));
-			Input.DebugGroupName.ReplaceInline(TEXT("TranslucentLighting"), TEXT("TranslLight"));
-			Input.DebugGroupName.ReplaceInline(TEXT("Translucency"), TEXT("Transl"));
-			Input.DebugGroupName.ReplaceInline(TEXT("DistanceField"), TEXT("DistFiel"));
-			Input.DebugGroupName.ReplaceInline(TEXT("Indirect"), TEXT("Ind"));
-			Input.DebugGroupName.ReplaceInline(TEXT("Cached"), TEXT("Cach"));
-			Input.DebugGroupName.ReplaceInline(TEXT("Inject"), TEXT("Inj"));
-			Input.DebugGroupName.ReplaceInline(TEXT("Visualization"), TEXT("Viz"));
-			Input.DebugGroupName.ReplaceInline(TEXT("Instanced"), TEXT("Inst"));
-			Input.DebugGroupName.ReplaceInline(TEXT("Evaluate"), TEXT("Eval"));
-			Input.DebugGroupName.ReplaceInline(TEXT("Landscape"), TEXT("Land"));
-			Input.DebugGroupName.ReplaceInline(TEXT("Dynamic"), TEXT("Dyn"));
-			Input.DebugGroupName.ReplaceInline(TEXT("Vertex"), TEXT("Vtx"));
-			Input.DebugGroupName.ReplaceInline(TEXT("Output"), TEXT("Out"));
-			Input.DebugGroupName.ReplaceInline(TEXT("Directional"), TEXT("Dir"));
-			Input.DebugGroupName.ReplaceInline(TEXT("Irradiance"), TEXT("Irr"));
-			Input.DebugGroupName.ReplaceInline(TEXT("Deferred"), TEXT("Def"));
-			Input.DebugGroupName.ReplaceInline(TEXT("true"), TEXT("_1"));
-			Input.DebugGroupName.ReplaceInline(TEXT("false"), TEXT("_0"));
-			Input.DebugGroupName.ReplaceInline(TEXT("PROPAGATE_AO"), TEXT("AO"));
-			Input.DebugGroupName.ReplaceInline(TEXT("PROPAGATE_SECONDARY_OCCLUSION"), TEXT("SEC_OCC"));
-			Input.DebugGroupName.ReplaceInline(TEXT("PROPAGATE_MULTIPLE_BOUNCES"), TEXT("MULT_BOUNC"));
-			Input.DebugGroupName.ReplaceInline(TEXT("LOCAL_LIGHTS_DISABLED"), TEXT("NoLL"));
-			Input.DebugGroupName.ReplaceInline(TEXT("LOCAL_LIGHTS_ENABLED"), TEXT("LL"));
-			Input.DebugGroupName.ReplaceInline(TEXT("LOCAL_LIGHTS_PREPASS_ENABLED"), TEXT("LLPP"));
-			Input.DebugGroupName.ReplaceInline(TEXT("PostProcess"), TEXT("Post"));
-			Input.DebugGroupName.ReplaceInline(TEXT("AntiAliasing"), TEXT("AA"));
-			Input.DebugGroupName.ReplaceInline(TEXT("Mobile"), TEXT("Mob"));
-			Input.DebugGroupName.ReplaceInline(TEXT("Linear"), TEXT("Lin"));
-			Input.DebugGroupName.ReplaceInline(TEXT("INT32_MAX"), TEXT("IMAX"));
-			Input.DebugGroupName.ReplaceInline(TEXT("Policy"), TEXT("Pol"));
-			Input.DebugGroupName.ReplaceInline(TEXT("EAtmRenderFlag==E_"), TEXT(""));
+			SlowTask.MakeDialog(false, true);
 		}
-	}
 
-	// Setup the debug info path if requested, or if this is a global shader and shader development mode is enabled
-	Input.DumpDebugInfoPath.Empty();
-	if (GShaderCompilingManager->GetDumpShaderDebugInfo() == FShaderCompilingManager::EDumpShaderDebugInfo::Always)
-	{
-		Input.DumpDebugInfoPath = GShaderCompilingManager->CreateShaderDebugInfoPath(Input);
-	}
+		int32 NumPendingJobs = 0;
+		// Keep track of previous number of pending jobs so we can update the slow task with the amount of work done.
+		int32 NumPreviousPendingJobs = NumJobs;
 
-	// Add the appropriate definitions for the shader frequency.
-	{
-		SET_SHADER_DEFINE(Input.Environment, PIXELSHADER,				Target.Frequency == SF_Pixel);
-		SET_SHADER_DEFINE(Input.Environment, VERTEXSHADER,				Target.Frequency == SF_Vertex);
-		SET_SHADER_DEFINE(Input.Environment, MESHSHADER,				Target.Frequency == SF_Mesh);
-		SET_SHADER_DEFINE(Input.Environment, AMPLIFICATIONSHADER,		Target.Frequency == SF_Amplification);
-		SET_SHADER_DEFINE(Input.Environment, GEOMETRYSHADER,			Target.Frequency == SF_Geometry);
-		SET_SHADER_DEFINE(Input.Environment, COMPUTESHADER,				Target.Frequency == SF_Compute);
-		SET_SHADER_DEFINE(Input.Environment, RAYCALLABLESHADER,			Target.Frequency == SF_RayCallable);
-		SET_SHADER_DEFINE(Input.Environment, RAYHITGROUPSHADER,			Target.Frequency == SF_RayHitGroup);
-		SET_SHADER_DEFINE(Input.Environment, RAYGENSHADER,				Target.Frequency == SF_RayGen);
-		SET_SHADER_DEFINE(Input.Environment, RAYMISSSHADER,				Target.Frequency == SF_RayMiss);
-		SET_SHADER_DEFINE(Input.Environment, WORKGRAPHROOTSHADER,		Target.Frequency == SF_WorkGraphRoot);
-		SET_SHADER_DEFINE(Input.Environment, WORKGRAPHCOMPUTESHADER,	Target.Frequency == SF_WorkGraphComputeNode);
-	}
+		do 
+		{
+			NumPendingJobs = 0;
+			{
+				// Lock CompileQueueSection so we can access the input and output queues
+				FScopeLock Lock(&CompileQueueSection);
 
-	SET_SHADER_DEFINE(Input.Environment, FORWARD_SHADING_FORCES_SKYLIGHT_CUBEMAPS_BLENDING, ForwardShadingForcesSkyLightCubemapBlending(ShaderPlatform) ? 1 : 0);
+				int32 ShaderMapIdx = 0;
+				for (TMap<int32, FPendingShaderMapCompileResultsPtr>::TIterator It(ShaderMapJobs); It; ++It)
+				{
+					FShaderMapCompileResults* Results = It.Value();
 
-	// Enables HLSL 2021
-	uint32 EnablesHLSL2021ByDefault = FDataDrivenShaderPlatformInfo::GetEnablesHLSL2021ByDefault(EShaderPlatform(Target.Platform));
-	if((EnablesHLSL2021ByDefault == uint32(1) && DebugGroupName == TEXT("Global")) ||
-		EnablesHLSL2021ByDefault == uint32(2) ||
-		Target.Frequency == SF_RayGen ||			// We want to make sure that function overloads follow c++ rules for FRayDesc.
-		Target.Frequency == SF_RayHitGroup)
-	{
-		Input.Environment.CompilerFlags.Add(CFLAG_HLSL2021);
-	}
+					if (Results->NumPendingJobs.GetValue() == 0)
+					{
+						AddCompiledResults(CompiledShaderMaps, It.Key(), *Results);
+						It.RemoveCurrent();
+					}
+					else
+					{
+						Results->CheckIfHung();
+						NumPendingJobs += Results->NumPendingJobs.GetValue();
+					}
+				}
+			}
 
-	// #defines get stripped out by the preprocessor without this. We can override with this
-	SET_SHADER_DEFINE(Input.Environment, COMPILER_DEFINE, TEXT("#define"));
+			if (NumPendingJobs > 0)
+			{
+				const float SleepTime =.01f;
+				
+				// We need to manually tick the Distributed build controller while the game thread is blocked
+				// otherwise we can get stuck in a infinite loop waiting for jobs that never will be done
+				// because for example, some controllers depend on the HTTP module which needs to be ticked in the main thread
+				if (BuildDistributionController && IsInGameThread())
+				{
+					BuildDistributionController->Tick(SleepTime);
+				}
 
-	if (FSceneInterface::GetShadingPath(GetMaxSupportedFeatureLevel(ShaderPlatform)) == EShadingPath::Deferred)
-	{
-		SET_SHADER_DEFINE(Input.Environment, SHADING_PATH_DEFERRED, 1);
+				// Progress the slow task with how many jobs we've completed since last tick.  Update the slow task message with the current number of pending jobs
+				// we are waiting on.
+				const int32 CompletedJobsSinceLastTick = NumPreviousPendingJobs - NumPendingJobs;
+				SlowTask.EnterProgressFrame(CompletedJobsSinceLastTick, FText::Format(LOCTEXT("BlockOnAllShaderMapCompletion", "Compiling Shaders ({0})"), NumPendingJobs));
+				NumPreviousPendingJobs = NumPendingJobs;
+				
+				// Yield CPU time while waiting
+				FPlatformProcess::Sleep(SleepTime);
+			}
+		} 
+		while (NumPendingJobs > 0);
 	}
-
-	const bool bUsingMobileRenderer = FSceneInterface::GetShadingPath(GetMaxSupportedFeatureLevel(ShaderPlatform)) == EShadingPath::Mobile;
-	if (bUsingMobileRenderer)
+	else
 	{
-		SET_SHADER_DEFINE_AND_COMPILE_ARGUMENT(Input.Environment, SHADING_PATH_MOBILE, true);
-		
-		const bool bMobileDeferredShading = IsMobileDeferredShadingEnabled((EShaderPlatform)Target.Platform);
-		SET_SHADER_DEFINE_AND_COMPILE_ARGUMENT(Input.Environment, MOBILE_DEFERRED_SHADING, bMobileDeferredShading);
-
-		if (bMobileDeferredShading)
+		int32 NumActiveWorkers = 0;
+		do 
 		{
-			bool bGLESDeferredShading = Target.Platform == SP_OPENGL_ES3_1_ANDROID;
-			SET_SHADER_DEFINE(Input.Environment, USE_GLES_FBF_DEFERRED, bGLESDeferredShading ? 1 : 0);
-			SET_SHADER_DEFINE(Input.Environment, MOBILE_EXTENDED_GBUFFER, MobileUsesExtenedGBuffer((EShaderPlatform)Target.Platform) ? 1 : 0);
-		}
+			for (const TUniquePtr<FShaderCompileThreadRunnableBase>& Thread : Threads)
+			{
+				NumActiveWorkers = Thread->CompilingLoop();
+			}
 
-		SET_SHADER_DEFINE(Input.Environment, USE_SCENE_DEPTH_AUX, MobileRequiresSceneDepthAux(ShaderPlatform) ? 1 : 0);
+			for (TMap<int32, FPendingShaderMapCompileResultsPtr>::TIterator It(ShaderMapJobs); It; ++It)
+			{
+				FShaderMapCompileResults* Results = It.Value();
+				Results->CheckIfHung();
+			}
+		} 
+		while (NumActiveWorkers > 0);
 
-		static FShaderPlatformCachedIniValue<bool> EnableCullBeforeFetchIniValue(TEXT("r.CullBeforeFetch"));
-		if (EnableCullBeforeFetchIniValue.Get((EShaderPlatform)Target.Platform) == 1)
-		{
-			Input.Environment.CompilerFlags.Add(CFLAG_CullBeforeFetch);
-		}
+		check(AllJobs.GetNumPendingJobs() == 0);
 
-		static FShaderPlatformCachedIniValue<bool> EnableWarpCullingIniValue(TEXT("r.WarpCulling"));
-		if (EnableWarpCullingIniValue.Get((EShaderPlatform)Target.Platform) == 1)
+		for (TMap<int32, FPendingShaderMapCompileResultsPtr>::TIterator It(ShaderMapJobs); It; ++It)
 		{
-			Input.Environment.CompilerFlags.Add(CFLAG_WarpCulling);
+			const FShaderMapCompileResults* Results = It.Value();
+			check(Results->NumPendingJobs.GetValue()== 0);
+
+			AddCompiledResults(CompiledShaderMaps, It.Key(), *Results);
+			It.RemoveCurrent();
 		}
 	}
 
-	if (RenderRectLightsAsSpotLights(GetMaxSupportedFeatureLevel(ShaderPlatform)))
-	{
-		SET_SHADER_DEFINE(Input.Environment, RECT_LIGHT_AS_SPOTLIGHT, 1);
-	}
+	UpdateNumRemainingAssets();
+}
 
-	if (ShaderPlatform == SP_VULKAN_ES3_1_ANDROID || ShaderPlatform == SP_VULKAN_SM5_ANDROID)
+namespace
+{
+	void PropagateGlobalShadersToAllPrimitives()
 	{
-		bool bIsStripReflect = true;
-		GConfig->GetBool(TEXT("/Script/AndroidRuntimeSettings.AndroidRuntimeSettings"), TEXT("bStripShaderReflection"), bIsStripReflect, GEngineIni);
-		if (!bIsStripReflect)
+		// Re-register everything to work around FShader lifetime issues - it currently lives and dies with the
+		// shadermap it is stored in, while cached MDCs can reference its memory. Re-registering will
+		// re-create the cache.
+		TRACE_CPUPROFILER_EVENT_SCOPE(PropagateGlobalShadersToAllPrimitives);
+
+		FObjectCacheContextScope ObjectCacheScope;
+		TSet<FSceneInterface*> ScenesToUpdate;
+		TIndirectArray<FComponentRecreateRenderStateContext> ComponentContexts;
+		for (IPrimitiveComponent* PrimitiveComponentInterface : ObjectCacheScope.GetContext().GetPrimitiveComponents())
 		{
-			Input.Environment.SetCompileArgument(TEXT("STRIP_REFLECT_ANDROID"), false);
+			if (PrimitiveComponentInterface->IsRenderStateCreated())
+			{
+				ComponentContexts.Add(new FComponentRecreateRenderStateContext(PrimitiveComponentInterface, &ScenesToUpdate));
+#if WITH_EDITOR
+				if (UPrimitiveComponent* PrimitiveComponent = PrimitiveComponentInterface->GetUObject<UPrimitiveComponent>())
+				{
+					if (PrimitiveComponent->HasValidSettingsForStaticLighting(false))
+					{
+						FStaticLightingSystemInterface::OnPrimitiveComponentUnregistered.Broadcast(PrimitiveComponent);
+						FStaticLightingSystemInterface::OnPrimitiveComponentRegistered.Broadcast(PrimitiveComponent);
+					}
+				}
+#endif
+			}
 		}
+
+		UpdateAllPrimitiveSceneInfosForScenes(ScenesToUpdate);
+		ComponentContexts.Empty();
+		UpdateAllPrimitiveSceneInfosForScenes(ScenesToUpdate);
 	}
+}
 
-	static const auto CVarViewHasTileOffsetData = IConsoleManager::Get().FindTConsoleVariableDataInt(TEXT("r.ViewHasTileOffsetData"));
-	const bool bViewHasTileOffsetData = CVarViewHasTileOffsetData->GetValueOnAnyThread() != 0;
-	SET_SHADER_DEFINE_AND_COMPILE_ARGUMENT(Input.Environment, VIEW_HAS_TILEOFFSET_DATA, bViewHasTileOffsetData);
+void FShaderCompilingManager::ProcessCompiledShaderMaps(
+	TMap<int32, FShaderMapFinalizeResults>& CompiledShaderMaps, 
+	float TimeBudget)
+{
+	// never process anything if not allowed, just in case
+	if (!AllowShaderCompiling())
+	{
+		return;
+	}
 
-	static const auto CVarPrimitiveHasTileOffsetData = IConsoleManager::Get().FindTConsoleVariableDataInt(TEXT("r.PrimitiveHasTileOffsetData"));
-	const bool bPrimitiveHasTileOffsetData = CVarPrimitiveHasTileOffsetData->GetValueOnAnyThread() != 0;
-	SET_SHADER_DEFINE_AND_COMPILE_ARGUMENT(Input.Environment, PRIMITIVE_HAS_TILEOFFSET_DATA, bPrimitiveHasTileOffsetData);
+#if WITH_EDITOR
+	TRACE_CPUPROFILER_EVENT_SCOPE(FShaderCompilingManager::ProcessCompiledShaderMaps);
 
-	// Set VR definitions
-	if (ShaderFrequencyNeedsInstancedStereoMods(ShaderType))
+	TMap<TRefCountPtr<FMaterial>, TRefCountPtr<FMaterialShaderMap>> MaterialsToUpdate;
+	TArray<TRefCountPtr<FMaterial>> MaterialsToReleaseCompilingId;
+
+	// Process compiled shader maps in FIFO order, in case a shader map has been enqueued multiple times,
+	// Which can happen if a material is edited while a background compile is going on
+	for (TMap<int32, FShaderMapFinalizeResults>::TIterator ShaderMapResultIter(CompiledShaderMaps); ShaderMapResultIter; ++ShaderMapResultIter)
 	{
-		const UE::StereoRenderUtils::FStereoShaderAspects Aspects(ShaderPlatform);
+		const uint32 CompilingId = ShaderMapResultIter.Key();
 
-		SET_SHADER_DEFINE_AND_COMPILE_ARGUMENT(Input.Environment, INSTANCED_STEREO, Aspects.IsInstancedStereoEnabled());
-		SET_SHADER_DEFINE_AND_COMPILE_ARGUMENT(Input.Environment, MULTI_VIEW, Aspects.IsInstancedMultiViewportEnabled());
-		SET_SHADER_DEFINE_AND_COMPILE_ARGUMENT(Input.Environment, MOBILE_MULTI_VIEW, Aspects.IsMobileMultiViewEnabled());
+		FShaderMapFinalizeResults& CompileResults = ShaderMapResultIter.Value();
+		TArray<FShaderCommonCompileJobPtr>& FinishedJobs = CompileResults.FinishedJobs;
 
-		// Throw a warning if we are silently disabling ISR due to missing platform support (but don't have MMV enabled).
-		static const auto CVarInstancedStereo = IConsoleManager::Get().FindTConsoleVariableDataInt(TEXT("vr.InstancedStereo"));
-		const bool bIsInstancedStereoEnabledInSettings = CVarInstancedStereo ? (CVarInstancedStereo->GetValueOnAnyThread() != 0) : false;
-		static const auto CVarMultiview = IConsoleManager::Get().FindTConsoleVariableDataInt(TEXT("vr.MobileMultiView"));
-		const bool bIsMultiviewEnabledInSettings = CVarMultiview ? (CVarMultiview->GetValueOnAnyThread() != 0) : false;
-		bool bWarningIssued = false;
-		// warn if ISR was enabled in settings, but aspects show that it's not enabled AND we don't use Mobile MultiView as an alternative
-		if (bIsInstancedStereoEnabledInSettings && !Aspects.IsInstancedStereoEnabled() && !(bIsMultiviewEnabledInSettings && Aspects.IsMobileMultiViewEnabled()) && !GShaderCompilingManager->AreWarningsSuppressed(ShaderPlatform))
-		{
-			UE_LOG(LogShaderCompilers, Warning, TEXT("Instanced stereo rendering is not supported for %s shader platform."), *ShaderFormatName.ToString());
-			bWarningIssued = true;
-		}
-		// Warn if MMV was enabled in settings, but aspects show that it's not enabled AND we don't use Instanced Stereo as an alternative
-		if (bIsMultiviewEnabledInSettings && !Aspects.IsMobileMultiViewEnabled() && !(bIsInstancedStereoEnabledInSettings && Aspects.IsInstancedStereoEnabled()) && !GShaderCompilingManager->AreWarningsSuppressed(ShaderPlatform))
-		{
-			UE_LOG(LogShaderCompilers, Warning, TEXT("Multiview rendering is not supported for %s shader platform."), *ShaderFormatName.ToString());
-			bWarningIssued = true;
-		}
-		if (bWarningIssued)
+		if (CompileResults.bSkipResultProcessing)
 		{
-			GShaderCompilingManager->SuppressWarnings(ShaderPlatform);
+			ShaderMapResultIter.RemoveCurrent();
+			continue;
 		}
-	}
-	else
-	{
-		SET_SHADER_DEFINE_AND_COMPILE_ARGUMENT(Input.Environment, INSTANCED_STEREO, false);
-		SET_SHADER_DEFINE_AND_COMPILE_ARGUMENT(Input.Environment, MULTI_VIEW, 0);
-		SET_SHADER_DEFINE_AND_COMPILE_ARGUMENT(Input.Environment, MOBILE_MULTI_VIEW, false);
-	}
 
-	// Reserve space in maps to prevent reallocation and rehashing in AddUniformBufferIncludesToEnvironment -- plus one at the end is for GeneratedInstancedStereo.ush
-	const int32 UniformBufferReserveNum = Input.Environment.UniformBufferMap.Num() + ShaderType->GetReferencedUniformBuffers().Num() + (VFType ? VFType->GetReferencedUniformBuffers().Num() : 0) + 1;
-	Input.Environment.UniformBufferMap.Reserve(UniformBufferReserveNum);
-	Input.Environment.IncludeVirtualPathToSharedContentsMap.Reserve(UniformBufferReserveNum);
+		TRefCountPtr<FMaterialShaderMap> CompilingShaderMap = FMaterialShaderMap::FindCompilingShaderMap(CompilingId);
 
-	ShaderType->AddUniformBufferIncludesToEnvironment(Input.Environment, ShaderPlatform);
+		if (CompilingShaderMap)
+		{
+			TArray<TRefCountPtr<FMaterial>>& MaterialDependencies = CompilingShaderMap->CompilingMaterialDependencies;
+			FShaderDiagnosticInfo ErrorInfo(FinishedJobs);
 
-	if (VFType)
-	{
-		VFType->AddUniformBufferIncludesToEnvironment(Input.Environment, ShaderPlatform);
-	}
+			bool bSuccess = true;
+			for (int32 JobIndex = 0; JobIndex < FinishedJobs.Num(); JobIndex++)
+			{
+				FShaderCommonCompileJob& CurrentJob = *FinishedJobs[JobIndex];
 
-	// Add generated instanced stereo code (this code also generates ViewState, so needed not just for ISR)
-	{
-		// this function may be called on multiple threads, so protect the storage
-		FScopeLock GeneratedInstancedCodeLock(&GCachedGeneratedInstancedStereoCodeLock);
+				if (FShaderCompileJob* SingleJob = CurrentJob.GetSingleShaderJob())
+				{
+					const bool bCheckSucceeded = CheckSingleJob(*SingleJob, ErrorInfo.UniqueErrors);
+					bSuccess = bCheckSucceeded && bSuccess;
+				}
+				else if (FShaderPipelineCompileJob* PipelineJob = CurrentJob.GetShaderPipelineJob())
+				{
+					for (int32 Index = 0; Index < PipelineJob->StageJobs.Num(); ++Index)
+					{
+						const bool bCheckSucceeded = CheckSingleJob(*PipelineJob->StageJobs[Index], ErrorInfo.UniqueErrors);
+						bSuccess = PipelineJob->StageJobs[Index]->bSucceeded && bCheckSucceeded && bSuccess;
+					}
+				}
+				else
+				{
+					checkf(0, TEXT("FShaderCommonCompileJob::Type=%d is not a valid type for a shader compile job"), (int32)CurrentJob.Type);
+				}
+			}
 
-		FThreadSafeSharedAnsiStringPtr* Existing = GCachedGeneratedInstancedStereoCode.Find(ShaderPlatform);
-		FThreadSafeSharedAnsiStringPtr CachedCodePtr = Existing ? *Existing : nullptr;
-		if (!CachedCodePtr.IsValid())
-		{
-			FString CachedCode;
-			GenerateInstancedStereoCode(CachedCode, ShaderPlatform);
+			if (bSuccess)
+			{
+				int32 JobIndex = 0;
+				if (FinishedJobs.Num() > 0)
+				{
+					CompilingShaderMap->ProcessCompilationResults(FinishedJobs, JobIndex, TimeBudget);
+					{
+						FScopeLock Lock(&CompileQueueSection);
+						for (int32 i = 0; i < JobIndex; ++i)
+						{
+							ReleaseJob(FinishedJobs[i]);
+						}
+					}
+					FinishedJobs.RemoveAt(0, JobIndex);
+				}
+			}
 
-			CachedCodePtr = MakeShareable(new TArray<ANSICHAR>());
-			ShaderConvertAndStripComments(CachedCode, *CachedCodePtr);
+			if (!bSuccess || FinishedJobs.Num() == 0)
+			{
+				ShaderMapResultIter.RemoveCurrent();
+			}
 
-			GCachedGeneratedInstancedStereoCode.Add(ShaderPlatform, CachedCodePtr);
-		}
+			FMaterialShaderMap* ShaderMapToUseForRendering = nullptr;
 
-		Input.Environment.IncludeVirtualPathToSharedContentsMap.Add(TEXT("/Engine/Generated/GeneratedInstancedStereo.ush"), CachedCodePtr);
-	}
+#if DEBUG_INFINITESHADERCOMPILE
+			UE_LOG(LogTemp, Display, TEXT("Finished compile of shader map 0x%08X%08X"), (int)((int64)(ShaderMap.GetReference()) >> 32), (int)((int64)(ShaderMap.GetReference())));
+#endif
+			int32 NumIncompleteMaterials = 0;
+			int32 MaterialIndex = 0;
+			
+			FMaterial* SingleMaterial = MaterialDependencies.Num() > 0 ? MaterialDependencies[0] : nullptr;
+			bool bRequiredComplete = false;
 
-	{
-		// Check if the compile environment explicitly wants to force optimization
-		const bool bForceOptimization = Input.Environment.CompilerFlags.Contains(CFLAG_ForceOptimization);
+			while (MaterialIndex < MaterialDependencies.Num())
+			{
+				FMaterial* Material = MaterialDependencies[MaterialIndex];
+				check(Material->GetGameThreadCompilingShaderMapId() == CompilingShaderMap->GetCompilingId());
+				bRequiredComplete |= Material->IsRequiredComplete();
+
+#if DEBUG_INFINITESHADERCOMPILE
+				UE_LOG(LogTemp, Display, TEXT("Shader map %s complete, GameThreadShaderMap 0x%08X%08X, marking material %s as finished"), *ShaderMap->GetFriendlyName(), (int)((int64)(ShaderMap.GetReference()) >> 32), (int)((int64)(ShaderMap.GetReference())), *Material->GetFriendlyName());
+				UE_LOG(LogTemp, Display, TEXT("Marking material as finished 0x%08X%08X"), (int)((int64)(Material) >> 32), (int)((int64)(Material)));
+#endif
+				//Material->RemoveOutstandingCompileId(ShaderMap->CompilingId);
 
-		if (!bForceOptimization && !ShouldOptimizeShaders(ShaderFormatName))
-		{
-			Input.Environment.CompilerFlags.Add(CFLAG_Debug);
-		}
-	}
+				bool bReleaseCompilingId = false;
 
-	// Extra data (names, etc)
-	if (ShouldEnableExtraShaderData(ShaderFormatName))
-	{
-		Input.Environment.CompilerFlags.Add(CFLAG_ExtraShaderData);
-	}
+				// Only process results that still match the ID which requested a compile
+				// This avoids applying shadermaps which are out of date and a newer one is in the async compiling pipeline
+				if (Material->GetMaterialId() != CompilingShaderMap->GetShaderMapId().BaseMaterialId)
+				{
+					bReleaseCompilingId = true;
+				}
+				else if (bSuccess)
+				{
+					bool bIsComplete = CompilingShaderMap->IsComplete(Material, true) && (CompilingShaderMap->CompilingMaterialNumExternalDependencies == 0);
 
-	// Symbols
-	if (ShouldGenerateShaderSymbols(ShaderFormatName))
-	{
-		Input.Environment.CompilerFlags.Add(CFLAG_GenerateSymbols);
-	}
-	if (ShouldGenerateShaderSymbolsInfo(ShaderFormatName))
-	{
-		Input.Environment.CompilerFlags.Add(CFLAG_GenerateSymbolsInfo);
-	}
+					// If running a cook, only process complete shader maps, as there's no rendering of partially complete shader maps to worry about.
+					if (bIsComplete || IsRunningCookCommandlet() == false || bAllowForIncompleteShaderMaps)
+					{
+						if (ShaderMapToUseForRendering == nullptr)
+						{
+							// Make a clone of the compiling shader map to use for rendering
+							// This will allow rendering to proceed with the clone, while async compilation continues to potentially update the compiling shader map
+							double StartTime = FPlatformTime::Seconds();
+							ShaderMapToUseForRendering = CompilingShaderMap->AcquireFinalizedClone();
+							TimeBudget -= (FPlatformTime::Seconds() - StartTime);
+						}
 
-	// Are symbols based on source or results
-	if (ShouldAllowUniqueShaderSymbols(ShaderFormatName))
-	{
-		Input.Environment.CompilerFlags.Add(CFLAG_AllowUniqueSymbols);
-	}
+						MaterialsToUpdate.Add(Material, ShaderMapToUseForRendering);
+					}
 
-	if (CVarShaderFastMath.GetValueOnAnyThread() == 0)
-	{
-		Input.Environment.CompilerFlags.Add(CFLAG_NoFastMath);
-	}
-    
-	if (bUsingMobileRenderer)
-    {
-        static IConsoleVariable* CVar = IConsoleManager::Get().FindConsoleVariable(TEXT("r.Mobile.FloatPrecisionMode"));
-        Input.Environment.FullPrecisionInPS |= CVar ? (CVar->GetInt() == EMobileFloatPrecisionMode::Full) : false;
-    }
-	
-	{
-		int32 FlowControl = CVarShaderFlowControl.GetValueOnAnyThread();
-		switch (FlowControl)
-		{
-			case 2:
-				Input.Environment.CompilerFlags.Add(CFLAG_AvoidFlowControl);
-				break;
-			case 1:
-				Input.Environment.CompilerFlags.Add(CFLAG_PreferFlowControl);
-				break;
-			case 0:
-				// Fallback to nothing...
-			default:
-				break;
-		}
-	}
+					if (bIsComplete)
+					{
+						bReleaseCompilingId = true;
+					}
+					else
+					{
+						++NumIncompleteMaterials;
+					}
 
-	{
-		static const auto CVar = IConsoleManager::Get().FindConsoleVariable(TEXT("r.Shaders.Validation"));
-		if (CVar && CVar->GetInt() == 0)
-		{
-			Input.Environment.CompilerFlags.Add(CFLAG_SkipValidation);
-		}
-	}
+					if (ErrorInfo.UniqueWarnings.Num() > 0)
+					{
+						UE_LOG(LogShaderCompilers, Warning, TEXT("Warnings while compiling Material %s for platform %s:"),
+							*Material->GetDebugName(),
+							*LegacyShaderPlatformToShaderFormat(CompilingShaderMap->GetShaderPlatform()).ToString());
+						for (const FString& UniqueWarning : ErrorInfo.UniqueWarnings)
+						{
+							UE_LOG(LogShaders, Warning, TEXT("  %s"), *UniqueWarning);
+						}
+					}
+				}
+				else
+				{
+					bReleaseCompilingId = true;
+					// Propagate error messages
+					Material->CompileErrors = ErrorInfo.UniqueErrors;
 
-	{
-		SET_SHADER_DEFINE(Input.Environment, DO_CHECK, GSShaderCheckLevel > 0 ? 1 : 0);
-		SET_SHADER_DEFINE(Input.Environment, DO_GUARD_SLOW, GSShaderCheckLevel > 1 ? 1 : 0);
-	}
+					MaterialsToUpdate.Add(Material, nullptr);
 
-	{
-		static FShaderPlatformCachedIniValue<int32> CVarWarningsAsErrorsPerPlatform(TEXT("r.Shaders.WarningsAsErrors"));
-		const int WarnLevel = CVarWarningsAsErrorsPerPlatform.Get(ShaderPlatform);
-		if ((WarnLevel == 1 && ShaderType->GetTypeForDynamicCast() == FShaderType::EShaderTypeForDynamicCast::Global) || WarnLevel > 1)
-		{
-			Input.Environment.CompilerFlags.Add(CFLAG_WarningsAsErrors);
-		}
-	}
+					if (Material->IsDefaultMaterial())
+					{
+						FString ErrorString;
 
-	if (UseRemoveUnsedInterpolators((EShaderPlatform)Target.Platform) && !IsOpenGLPlatform((EShaderPlatform)Target.Platform))
-	{
-		Input.Environment.CompilerFlags.Add(CFLAG_ForceRemoveUnusedInterpolators);
-	}
-	
-	if (IsD3DPlatform((EShaderPlatform)Target.Platform) && IsPCPlatform((EShaderPlatform)Target.Platform))
-	{
+						// Log the errors unsuppressed before the fatal error, so it's always obvious from the log what the compile error was
+						for (const FString& UniqueError : ErrorInfo.UniqueErrors)
+						{
+							ErrorString += FString::Printf(TEXT("  %s\n"), *UniqueError);
+						}
 
+						ErrorString += FString::Printf(TEXT("Failed to compile default material %s!"), *Material->GetBaseMaterialPathName());
 
-		if (CVarD3DCheckedForTypedUAVs.GetValueOnAnyThread() == 0)
-		{
-			Input.Environment.CompilerFlags.Add(CFLAG_AllowTypedUAVLoads);
-		}
+						if (AreShaderErrorsFatal())
+						{
+							// Assert if a default material could not be compiled, since there will be nothing for other failed materials to fall back on.
+							UE_LOG(LogShaderCompilers, Fatal, TEXT("%s"), *ErrorString);
+						}
+						else
+						{
+							UE_LOG(LogShaderCompilers, Error, TEXT("%s"), *ErrorString);
+						}
+					}
+					
+					FString ErrorString;
 
-		{
-			static const auto CVar = IConsoleManager::Get().FindConsoleVariable(TEXT("r.D3D.CheckedForTypedUAVs"));
-			if (CVar && CVar->GetInt() == 0)
-			{
-				Input.Environment.CompilerFlags.Add(CFLAG_AllowTypedUAVLoads);
-			}
-		}
-	}
+					ErrorString += FString::Printf(TEXT("Failed to compile Material %s for platform %s, Default Material will be used in game.\n"),
+						*Material->GetDebugName(), *LegacyShaderPlatformToShaderFormat(CompilingShaderMap->GetShaderPlatform()).ToString());
 
-	if (IsMetalPlatform((EShaderPlatform)Target.Platform))
-	{
-		if (CVarShaderZeroInitialise.GetValueOnAnyThread() != 0)
-		{
-			Input.Environment.CompilerFlags.Add(CFLAG_ZeroInitialise);
-		}
+					for (const FString& UniqueError : ErrorInfo.UniqueErrors)
+					{
+						FString ErrorMessage = UniqueError;
+						// Work around build machine string matching heuristics that will cause a cook to fail
+						ErrorMessage.ReplaceInline(TEXT("error "), TEXT("err0r "), ESearchCase::CaseSensitive);
+						ErrorString += FString::Printf(TEXT("  %s\n"), *ErrorMessage);
+					}
 
-		if (CVarShaderBoundsChecking.GetValueOnAnyThread() != 0)
-		{
-			Input.Environment.CompilerFlags.Add(CFLAG_BoundsChecking);
-		}
-		
-		// Check whether we can compile metal shaders to bytecode - avoids poisoning the DDC
-		static ITargetPlatformManagerModule& TPM = GetTargetPlatformManagerRef();
-		const IShaderFormat* Compiler = TPM.FindShaderFormat(ShaderFormatName);
-		static const bool bCanCompileOfflineMetalShaders = Compiler && Compiler->CanCompileBinaryShaders();
-		if (!bCanCompileOfflineMetalShaders)
-		{
-			Input.Environment.CompilerFlags.Add(CFLAG_Debug);
-		}
-		
-		// Shaders built for archiving - for Metal that requires compiling the code in a different way so that we can strip it later
-		bool bArchive = false;
-		GConfig->GetBool(TEXT("/Script/UnrealEd.ProjectPackagingSettings"), TEXT("bSharedMaterialNativeLibraries"), bArchive, GGameIni);
-		if (bCanCompileOfflineMetalShaders && bArchive)
-		{
-			Input.Environment.CompilerFlags.Add(CFLAG_Archive);
-		}
-		
-		{
-			uint32 ShaderVersion = RHIGetMetalShaderLanguageVersion(EShaderPlatform(Target.Platform));
-			Input.Environment.SetCompileArgument(TEXT("SHADER_LANGUAGE_VERSION"), ShaderVersion);
-			
-			bool bAllowFastIntrinsics = false;
-			bool bForceFloats = false;
-			int32 IndirectArgumentTier = 0;
-			bool bEnableMathOptimisations = true;
-            bool bSupportAppleA8 = false;
-			bool bMetalOptimizeForSize = false;
+					UE_LOG(LogShaderCompilers, Warning, TEXT("%s"), *ErrorString);
+				}
 
-			if (IsPCPlatform(EShaderPlatform(Target.Platform)))
-			{
-				GConfig->GetBool(TEXT("/Script/MacTargetPlatform.MacTargetSettings"), TEXT("UseFastIntrinsics"), bAllowFastIntrinsics, GEngineIni);
-				GConfig->GetBool(TEXT("/Script/MacTargetPlatform.MacTargetSettings"), TEXT("EnableMathOptimisations"), bEnableMathOptimisations, GEngineIni);
-				GConfig->GetInt(TEXT("/Script/MacTargetPlatform.MacTargetSettings"), TEXT("IndirectArgumentTier"), IndirectArgumentTier, GEngineIni);
-				GConfig->GetBool(TEXT("/Script/MacTargetPlatform.MacTargetSettings"), TEXT("MetalOptimizeForSize"), bMetalOptimizeForSize, GEngineIni);
-                
-                // No half precision support on MacOS at the moment
-                bForceFloats = true;
+				if (bReleaseCompilingId)
+				{
+					check(Material->GameThreadCompilingShaderMapId != 0u);
+					Material->GameThreadCompilingShaderMapId = 0u;
+					Material->GameThreadPendingCompilerEnvironment.SafeRelease();
+					MaterialDependencies.RemoveAt(MaterialIndex);
+					MaterialsToReleaseCompilingId.Add(Material);
+				}
+				else
+				{
+					++MaterialIndex;
+				}
 			}
-			else
+
+			if (NumIncompleteMaterials == 0 && (IsMaterialMapDDCEnabled() || bRequiredComplete))
 			{
-				GConfig->GetBool(TEXT("/Script/IOSRuntimeSettings.IOSRuntimeSettings"), TEXT("UseFastIntrinsics"), bAllowFastIntrinsics, GEngineIni);
-				GConfig->GetBool(TEXT("/Script/IOSRuntimeSettings.IOSRuntimeSettings"), TEXT("EnableMathOptimisations"), bEnableMathOptimisations, GEngineIni);
-				GConfig->GetBool(TEXT("/Script/IOSRuntimeSettings.IOSRuntimeSettings"), TEXT("ForceFloats"), bForceFloats, GEngineIni);
-				GConfig->GetInt(TEXT("/Script/IOSRuntimeSettings.IOSRuntimeSettings"), TEXT("IndirectArgumentTier"), IndirectArgumentTier, GEngineIni);
-                GConfig->GetBool(TEXT("/Script/IOSRuntimeSettings.IOSRuntimeSettings"), TEXT("bSupportAppleA8"), bSupportAppleA8, GEngineIni);
-				GConfig->GetBool(TEXT("/Script/IOSRuntimeSettings.IOSRuntimeSettings"), TEXT("MetalOptimizeForSize"), bMetalOptimizeForSize, GEngineIni);
-                
-				// Force no development shaders on iOS
-				bAllowDevelopmentShaderCompile = false;
+				CompilingShaderMap->bCompiledSuccessfully = bSuccess;
+				CompilingShaderMap->bCompilationFinalized = true;
+				if (ShaderMapToUseForRendering)
+				{
+					// ShaderMapToUseForRendering is only initialized inside the loop over material dependencies,
+					// so it's safe to assume that SingleMaterial has been set (a material is needed to construct 
+					// the FMaterialShaderParameters struct which is in turn needed to build the DDC key).
+					check(SingleMaterial != nullptr);
+					ShaderMapToUseForRendering->bCompiledSuccessfully = true;
+					ShaderMapToUseForRendering->bCompilationFinalized = true;
+					if (ShaderMapToUseForRendering->bIsPersistent)
+					{
+						ShaderMapToUseForRendering->SaveToDerivedDataCache(FMaterialShaderParameters(SingleMaterial));
+					}
+				}
+
+				CompilingShaderMap->ReleaseCompilingId();
 			}
-            
-            Input.Environment.FullPrecisionInPS |= bForceFloats;
-            
-			Input.Environment.SetCompileArgument(TEXT("METAL_USE_FAST_INTRINSICS"), bAllowFastIntrinsics);
-			Input.Environment.SetCompileArgument(TEXT("METAL_INDIRECT_ARGUMENT_BUFFERS"), IndirectArgumentTier);
-            Input.Environment.SetCompileArgument(TEXT("SUPPORT_APPLE_A8"), bSupportAppleA8);
-            Input.Environment.SetCompileArgument(TEXT("METAL_OPTIMIZE_FOR_SIZE"), bMetalOptimizeForSize);
-			
-			// Same as console-variable above, but that's global and this is per-platform, per-project
-			if (!bEnableMathOptimisations)
+
+			if (TimeBudget < 0)
 			{
-				Input.Environment.CompilerFlags.Add(CFLAG_NoFastMath);
+				break;
 			}
 		}
-	}
+		else
+		{
+			if (CompilingId == GlobalShaderMapId)
+			{
+				ProcessCompiledGlobalShaders(FinishedJobs);
+				PropagateGlobalShadersToAllPrimitives();
+			}
 
-	if (IsAndroidPlatform(EShaderPlatform(Target.Platform)))
-	{
-		// Force no development shaders on Android platforms
-		bAllowDevelopmentShaderCompile = false;
+			// ShaderMap was removed from compiling list or is being used by another type of shader map which is maintaining a reference
+			// to the results, either way the job can be released
+			{
+				FScopeLock Lock(&CompileQueueSection);
+				for (FShaderCommonCompileJobPtr& Job : FinishedJobs)
+				{
+					ReleaseJob(Job);
+				}
+			}
+			ShaderMapResultIter.RemoveCurrent();
+		}
 	}
 
-	// Mobile emulation should be defined when a PC platform is using a mobile renderer (limited to feature level ES3_1)...  eg SP_PCD3D_ES3_1,SP_VULKAN_PCES3_1,SP_METAL_MACES3_1
-	if (IsSimulatedPlatform(EShaderPlatform(Target.Platform)) && bAllowDevelopmentShaderCompile)
+	if (MaterialsToReleaseCompilingId.Num() > 0)
 	{
-		SET_SHADER_DEFINE(Input.Environment, MOBILE_EMULATION, 1);
+		ENQUEUE_RENDER_COMMAND(ReleaseCompilingShaderMapIds)([MaterialsToReleaseCompilingId = MoveTemp(MaterialsToReleaseCompilingId)](FRHICommandListImmediate& RHICmdList)
+		{
+			for (FMaterial* Material : MaterialsToReleaseCompilingId)
+			{
+				check(Material->RenderingThreadCompilingShaderMapId != 0u);
+				Material->RenderingThreadCompilingShaderMapId = 0u;
+				Material->RenderingThreadPendingCompilerEnvironment.SafeRelease();
+			}
+		});
 	}
 
-	// Add compiler flag CFLAG_ForceDXC if DXC is enabled
-	const bool bHlslVersion2021 = Input.Environment.CompilerFlags.Contains(CFLAG_HLSL2021);
-	const bool bIsDxcEnabled = IsDxcEnabledForPlatform((EShaderPlatform)Target.Platform, bHlslVersion2021);
-	SET_SHADER_DEFINE(Input.Environment, COMPILER_DXC, bIsDxcEnabled);
-	if (bIsDxcEnabled)
+	if (MaterialsToUpdate.Num() > 0)
 	{
-		Input.Environment.CompilerFlags.Add(CFLAG_ForceDXC);
-	}
+		FMaterial::SetShaderMapsOnMaterialResources(MaterialsToUpdate);
 
-	bool bIsMobilePlatform = IsMobilePlatform((EShaderPlatform)Target.Platform);
+		for (const auto& It : MaterialsToUpdate)
+		{
+			It.Key->NotifyCompilationFinished();
+		}
 
-	if (bIsMobilePlatform)
-	{
-		if (IsUsingEmulatedUniformBuffers((EShaderPlatform)Target.Platform))
+		if (FApp::CanEverRender())
 		{
-			Input.Environment.CompilerFlags.Add(CFLAG_UseEmulatedUB);
+			// This empties MaterialsToUpdate, see the comment inside the function for the reason.
+			PropagateMaterialChangesToPrimitives(MaterialsToUpdate);
+
+			FEditorSupportDelegates::RedrawAllViewports.Broadcast();
 		}
 	}
 
-	SET_SHADER_DEFINE(Input.Environment, HAS_INVERTED_Z_BUFFER, (bool)ERHIZBuffer::IsInverted);
+	UpdateNumRemainingAssets();
+#endif // WITH_EDITOR
+}
 
-	if (Input.Environment.CompilerFlags.Contains(CFLAG_HLSL2021))
+void FShaderCompilingManager::PropagateMaterialChangesToPrimitives(TMap<TRefCountPtr<FMaterial>, TRefCountPtr<FMaterialShaderMap>>& MaterialsToUpdate)
+{
+	// don't perform any work if no compiling
+	if (!AllowShaderCompiling())
 	{
-		SET_SHADER_DEFINE(Input.Environment, COMPILER_SUPPORTS_HLSL2021, 1);
+		return;
 	}
 
+	TSet<FSceneInterface*> ScenesToUpdate;
+	FObjectCacheContextScope ObjectCacheScope;
 	{
-		static const auto CVar = IConsoleManager::Get().FindTConsoleVariableDataInt(TEXT("r.ClearCoatNormal"));
-		SET_SHADER_DEFINE(Input.Environment, CLEAR_COAT_BOTTOM_NORMAL, CVar ? (CVar->GetValueOnAnyThread() != 0) && !bIsMobilePlatform : 0);
-	}
+		TRACE_CPUPROFILER_EVENT_SCOPE(FShaderCompilingManager::PropagateMaterialChangesToPrimitives);
 
-	{
-		static const auto CVar = IConsoleManager::Get().FindTConsoleVariableDataInt(TEXT("r.IrisNormal"));
-		SET_SHADER_DEFINE(Input.Environment, IRIS_NORMAL, CVar ? (CVar->GetValueOnAnyThread() != 0) : 0);
-	}
+		TArray<UMaterialInterface*> UpdatedMaterials;
+		for (TMap<TRefCountPtr<FMaterial>, TRefCountPtr<FMaterialShaderMap>>::TConstIterator MaterialIt(MaterialsToUpdate); MaterialIt; ++MaterialIt)
+		{
+			FMaterial* UpdatedMaterial = MaterialIt.Key();
+			UpdatedMaterials.Add(UpdatedMaterial->GetMaterialInterface());
+		}
 
-	{
-		static const auto CVar = IConsoleManager::Get().FindTConsoleVariableDataInt(TEXT("Compat.UseDXT5NormalMaps"));
-		SET_SHADER_DEFINE(Input.Environment, DXT5_NORMALMAPS, CVar ? (CVar->GetValueOnAnyThread() != 0) : 0);
+		for (IPrimitiveComponent* PrimitiveComponent : ObjectCacheScope.GetContext().GetPrimitivesAffectedByMaterials(UpdatedMaterials))
+		{
+			PrimitiveComponent->MarkRenderStateDirty();
+		}
 	}
 
-	if (bAllowDevelopmentShaderCompile)
-	{
-		static const auto CVar = IConsoleManager::Get().FindTConsoleVariableDataInt(TEXT("r.CompileShadersForDevelopment"));
-		SET_SHADER_DEFINE(Input.Environment, COMPILE_SHADERS_FOR_DEVELOPMENT, CVar ? (CVar->GetValueOnAnyThread() != 0) : 0);
-	}
+	// Recreating the render state for the primitives may end up recreating the material resources if some materials are missing some usage flags.
+	// For example, if some materials are not marked as used with static lighting and we build lightmaps, UMaterialInstance::CheckMaterialUsage
+	// will catch the problem and try to set the flag. However, since MaterialsToUpdate stores smart pointers, the material resources will have
+	// a refcount of 2, so the FMaterial destructor will trigger a check failure because the refcount doesn't reach 0. Empty this map before
+	// recreating the render state to allow resources to be deleted cleanly.
+	MaterialsToUpdate.Empty();
+}
 
-	{
-		SET_SHADER_DEFINE(Input.Environment, ALLOW_STATIC_LIGHTING, IsStaticLightingAllowed() ? 1 : 0);
-	}
 
-	{
-		// Allow GBuffer containing a velocity target to be overridden at a higher level with GBUFFER_LAYOUT
-		bool bUsingBasePassVelocity = IsUsingBasePassVelocity((EShaderPlatform)Target.Platform);
-		SET_SHADER_DEFINE(Input.Environment, USES_BASE_PASS_VELOCITY, bUsingBasePassVelocity ? 1 : 0);
+/**
+ * Shutdown the shader compile manager
+ * this function should be used when ending the game to shutdown shader compile threads
+ * will not complete current pending shader compilation
+ */
+void FShaderCompilingManager::Shutdown()
+{
+	// Shutdown has been moved to the destructor because the shader compiler lifetime is expected to
+	// be longer than other asset compilers, otherwise niagara compilations might get stuck.
+}
 
-		bool bGBufferHasVelocity = bUsingBasePassVelocity;
-		if (!bGBufferHasVelocity)
-		{
-			const EGBufferLayout Layout = FShaderCompileUtilities::FetchGBufferLayout(Input.Environment);
-			bGBufferHasVelocity |= (Layout == GBL_ForceVelocity);
-		}
-		SET_SHADER_DEFINE(Input.Environment, GBUFFER_HAS_VELOCITY, bGBufferHasVelocity ? 1 : 0);
-	}
+void FShaderCompilingManager::PrintStats()
+{
+	FShaderCompilerStats LocalStats;
+	GetLocalStats(LocalStats);
+	LocalStats.WriteStatSummary();
+}
 
+void FShaderCompilingManager::GetLocalStats(FShaderCompilerStats& OutStats) const
+{
+	if (GShaderCompilerStats)
 	{
-		static const auto CVar = IConsoleManager::Get().FindTConsoleVariableDataInt(TEXT("r.GBufferDiffuseSampleOcclusion"));
-		SET_SHADER_DEFINE(Input.Environment, GBUFFER_HAS_DIFFUSE_SAMPLE_OCCLUSION, CVar ? (CVar->GetValueOnAnyThread() != 0) : 1);
+		OutStats.Aggregate(*GShaderCompilerStats);
+		AllJobs.GetCachingStats(OutStats);
 	}
+}
 
+FShaderCompileMemoryUsage FShaderCompilingManager::GetExternalMemoryUsage()
+{
+	FShaderCompileMemoryUsage TotalMemoryUsage{};
+	for (const TUniquePtr<FShaderCompileThreadRunnableBase>& ThreadPtr : Threads)
 	{
-		SET_SHADER_DEFINE(Input.Environment, SELECTIVE_BASEPASS_OUTPUTS, IsUsingSelectiveBasePassOutputs((EShaderPlatform)Target.Platform) ? 1 : 0);
+		FShaderCompileMemoryUsage MemoryUsage = ThreadPtr->GetExternalWorkerMemoryUsage();
+		TotalMemoryUsage.VirtualMemory += MemoryUsage.VirtualMemory;
+		TotalMemoryUsage.PhysicalMemory += MemoryUsage.PhysicalMemory;
 	}
+	return TotalMemoryUsage;
+}
 
-	{
-		SET_SHADER_DEFINE(Input.Environment, USE_DBUFFER, IsUsingDBuffers((EShaderPlatform)Target.Platform) ? 1 : 0);
-	}
+static void BuildErrorStringAndReport(const FShaderDiagnosticInfo& DiagInfo, FString& ErrorString)
+{
+	bool bReportedDebugInfo = false;
 
+	for (int32 ErrorIndex = 0; ErrorIndex < DiagInfo.UniqueErrors.Num(); ErrorIndex++)
 	{
-		static IConsoleVariable* CVar = IConsoleManager::Get().FindConsoleVariable(TEXT("r.AllowGlobalClipPlane"));
-		SET_SHADER_DEFINE(Input.Environment, PROJECT_ALLOW_GLOBAL_CLIP_PLANE, CVar ? (CVar->GetInt() != 0) : 0);
-	}
+		FString UniqueErrorString = DiagInfo.UniqueErrors[ErrorIndex] + TEXT("\n");
 
-	{
-		const bool bSupportsClipDistance = FDataDrivenShaderPlatformInfo::GetSupportsClipDistance((EShaderPlatform)Target.Platform);
-		SET_SHADER_DEFINE(Input.Environment, PLATFORM_SUPPORTS_CLIP_DISTANCE, bSupportsClipDistance ? 1u : 0u);
-	}
+		if (FPlatformMisc::IsDebuggerPresent())
+		{
+			// Using OutputDebugString to avoid any text getting added before the filename,
+			// Which will throw off VS.NET's ability to take you directly to the file and line of the error when double clicking it in the output window.
+			FPlatformMisc::LowLevelOutputDebugStringf(*UniqueErrorString);
+		}
+		else
+		{
+			UE_LOG(LogShaderCompilers, Warning, TEXT("%s"), *UniqueErrorString);
+		}
 
-	{
-		const bool bSupportsVertexShaderSRVs = FDataDrivenShaderPlatformInfo::GetSupportsVertexShaderSRVs((EShaderPlatform)Target.Platform);
-		SET_SHADER_DEFINE(Input.Environment, PLATFORM_SUPPORTS_VERTEX_SHADER_SRVS, bSupportsVertexShaderSRVs ? 1u : 0u);
+		ErrorString += UniqueErrorString;
 	}
+}
 
+bool FShaderCompilingManager::HandlePotentialRetryOnError(TMap<int32, FShaderMapFinalizeResults>& CompletedShaderMaps)
+{
+	if (FApp::IsUnattended())
 	{
-		const bool bSupportsVertexShaderUAVs = FDataDrivenShaderPlatformInfo::GetSupportsVertexShaderUAVs((EShaderPlatform)Target.Platform) != ERHIFeatureSupport::Unsupported;
-		SET_SHADER_DEFINE(Input.Environment, PLATFORM_SUPPORTS_VERTEX_SHADER_UAVS, bSupportsVertexShaderUAVs ? 1u : 0u);
+		return false;
 	}
 
-	{
-		const uint32 MaxSamplers = FDataDrivenShaderPlatformInfo::GetMaxSamplers((EShaderPlatform)Target.Platform);
-		SET_SHADER_DEFINE_AND_COMPILE_ARGUMENT(Input.Environment, PLATFORM_MAX_SAMPLERS, MaxSamplers);
-	}
+	bool bRetryCompile = false;
 
+#if WITH_EDITORONLY_DATA
+	for (TMap<int32, FShaderMapFinalizeResults>::TIterator It(CompletedShaderMaps); It; ++It)
 	{
-		const bool bSupportsIndependentSamplers = FDataDrivenShaderPlatformInfo::GetSupportsIndependentSamplers((EShaderPlatform)Target.Platform);
-		SET_SHADER_DEFINE(Input.Environment, SUPPORTS_INDEPENDENT_SAMPLERS, bSupportsIndependentSamplers ? 1 : 0);
-	}
+		FShaderMapFinalizeResults& Results = It.Value();
 
-	bool bForwardShading = false;
-	{
-		if (bIsMobilePlatform)
-		{
-			bForwardShading = !IsMobileDeferredShadingEnabled((EShaderPlatform)Target.Platform);
-		}
-		else if (TargetPlatform)
-		{
-			bForwardShading = TargetPlatform->UsesForwardShading();
-		}
-		else
+		if (!Results.bAllJobsSucceeded)
 		{
-			static IConsoleVariable* CVarForwardShading = IConsoleManager::Get().FindConsoleVariable(TEXT("r.ForwardShading"));
-			bForwardShading = CVarForwardShading ? (CVarForwardShading->GetInt() != 0) : false;
-		}
-		SET_SHADER_DEFINE(Input.Environment, FORWARD_SHADING, bForwardShading);
-	}
+			bool bSpecialEngineMaterial = false;
+
+			const FMaterialShaderMap* ShaderMap = FMaterialShaderMap::FindCompilingShaderMap(It.Key());
+			if (ShaderMap)
+			{
+				for (const FMaterial* Material : ShaderMap->CompilingMaterialDependencies)
+				{
+					if (Material->IsSpecialEngineMaterial())
+					{
+						bSpecialEngineMaterial = true;
+						break;
+					}
+				}
+			}
+
+			if (UE_LOG_ACTIVE(LogShaders, Log) 
+				// Always log detailed errors when a special engine material or global shader fails to compile, as those will be fatal errors
+				|| bSpecialEngineMaterial 
+				|| It.Key() == GlobalShaderMapId)
+			{
+				TArray<FShaderCommonCompileJobPtr>& CompleteJobs = Results.FinishedJobs;
+				FShaderDiagnosticInfo ShaderDiagInfo(CompleteJobs);
+
+				const TCHAR* MaterialName = ShaderMap ? ShaderMap->GetFriendlyName() : TEXT("global shaders");
+				FString ErrorString = FString::Printf(TEXT("%i Shader compiler errors compiling %s for platform %s:"), ShaderDiagInfo.UniqueErrors.Num(), MaterialName, *ShaderDiagInfo.TargetShaderPlatformString);
+				UE_LOG(LogShaderCompilers, Warning, TEXT("%s"), *ErrorString);
+				ErrorString += TEXT("\n");
 
-	{
-		if (VelocityEncodeDepth((EShaderPlatform)Target.Platform))
-		{
-			SET_SHADER_DEFINE(Input.Environment, VELOCITY_ENCODE_DEPTH, 1);
-		}
-		else
-		{
-			SET_SHADER_DEFINE(Input.Environment, VELOCITY_ENCODE_DEPTH, 0);
-		}
-	}
+				bool bAnyErrorLikelyToBeCodeError = false;
+				for (const FShaderCommonCompileJob* Job : ShaderDiagInfo.ErrorJobs)
+				{
+					bAnyErrorLikelyToBeCodeError |= Job->bErrorsAreLikelyToBeCode;
+				}
 
-	{
-		if (MaskedInEarlyPass((EShaderPlatform)Target.Platform))
-		{
-			SET_SHADER_DEFINE(Input.Environment, EARLY_Z_PASS_ONLY_MATERIAL_MASKING, 1);
-		}
-		else
-		{
-			SET_SHADER_DEFINE(Input.Environment, EARLY_Z_PASS_ONLY_MATERIAL_MASKING, 0);
-		}
-	}
+				BuildErrorStringAndReport(ShaderDiagInfo, ErrorString);
 
-	{
-		static IConsoleVariable* CVar = IConsoleManager::Get().FindConsoleVariable(TEXT("r.VertexFoggingForOpaque"));
-		bool bVertexFoggingForOpaque = false;
-		if (bForwardShading)
-		{
-			bVertexFoggingForOpaque = CVar ? (CVar->GetInt() != 0) : 0;
-			if (TargetPlatform)
-			{
-				const int32 PlatformHeightFogMode = TargetPlatform->GetHeightFogModeForOpaque();
-				if (PlatformHeightFogMode == 1)
+				if (UE_LOG_ACTIVE(LogShaders, Log) && (bAnyErrorLikelyToBeCodeError || bPromptToRetryFailedShaderCompiles || bSpecialEngineMaterial))
 				{
-					bVertexFoggingForOpaque = false;
+					// Use debug break in debug with the debugger attached, otherwise message box
+					if (bDebugBreakOnPromptToRetryShaderCompile && FPlatformMisc::IsDebuggerPresent())
+					{
+						// A shader compile error has occurred, see the debug output for information.
+						// Double click the errors in the VS.NET output window and the IDE will take you directly to the file and line of the error.
+						// Check ErrorJobs for more state on the failed shaders, for example in-memory includes like Material.usf
+						UE_DEBUG_BREAK();
+						// Set GRetryShaderCompilation to true in the debugger to enable retries in debug
+						// NOTE: MaterialTemplate.usf will not be reloaded when retrying!
+						bRetryCompile = GRetryShaderCompilation;
+					}
+					else
+					{
+						if (FPlatformMisc::MessageBoxExt(EAppMsgType::YesNo, *FText::Format(NSLOCTEXT("UnrealEd", "Error_RetryShaderCompilation", "{0}\r\n\r\nRetry compilation?"),
+							FText::FromString(ErrorString)).ToString(), TEXT("Error")) == EAppReturnType::Type::Yes)
+						{
+							bRetryCompile = true;
+						}
+					}
 				}
-				else if (PlatformHeightFogMode == 2)
+
+				if (bRetryCompile)
 				{
-					bVertexFoggingForOpaque = true;
+					break;
 				}
 			}
 		}
-		SET_SHADER_DEFINE(Input.Environment, PROJECT_VERTEX_FOGGING_FOR_OPAQUE, bVertexFoggingForOpaque);
 	}
 
+	if (bRetryCompile)
 	{
-		static IConsoleVariable* CVar = IConsoleManager::Get().FindConsoleVariable(TEXT("r.Mobile.DisableVertexFog"));
-		SET_SHADER_DEFINE(Input.Environment, PROJECT_MOBILE_DISABLE_VERTEX_FOG, CVar ? (CVar->GetInt() != 0) : 0);
-	}
+		// Flush the shader file cache so that any changes will be propagated.
+		FlushShaderFileCache();
 
-	bool bSupportLocalFogVolumes = false;
-	{
-		static IConsoleVariable* CVar = IConsoleManager::Get().FindConsoleVariable(TEXT("r.SupportLocalFogVolumes"));
-		bSupportLocalFogVolumes = CVar && CVar->GetInt() > 0;
-		SET_SHADER_DEFINE(Input.Environment, PROJECT_SUPPORTS_LOCALFOGVOLUME, (bSupportLocalFogVolumes ? 1 : 0));
-	}
+		TArray<int32> MapsToRemove;
 
-	{
-		static IConsoleVariable* CVar = IConsoleManager::Get().FindConsoleVariable(TEXT("r.LocalFogVolume.ApplyOnTranslucent"));
-		const bool bLocalFogVolumesApplyOnTranclucent = CVar && CVar->GetInt() > 0;
-		SET_SHADER_DEFINE(Input.Environment, PROJECT_LOCALFOGVOLUME_APPLYONTRANSLUCENT, ((bSupportLocalFogVolumes && bLocalFogVolumesApplyOnTranclucent) ? 1 : 0));
-	}
+		for (TMap<int32, FShaderMapFinalizeResults>::TIterator It(CompletedShaderMaps); It; ++It)
+		{
+			FShaderMapFinalizeResults& Results = It.Value();
 
-	bool bSupportSkyAtmosphere = false;
-	{
-		static IConsoleVariable* CVar = IConsoleManager::Get().FindConsoleVariable(TEXT("r.SupportSkyAtmosphere"));
-		bSupportSkyAtmosphere = CVar && CVar->GetInt() != 0;
-		SET_SHADER_DEFINE(Input.Environment, PROJECT_SUPPORT_SKY_ATMOSPHERE, bSupportSkyAtmosphere ? 1 : 0);
-	}
+			if (!Results.bAllJobsSucceeded)
+			{
+				MapsToRemove.Add(It.Key());
 
-	{
-		static IConsoleVariable* CVar = IConsoleManager::Get().FindConsoleVariable(TEXT("r.SupportCloudShadowOnForwardLitTranslucent"));
-		const bool bSupportCloudShadowOnForwardLitTranslucent = CVar && CVar->GetInt() > 0;
-		SET_SHADER_DEFINE(Input.Environment, SUPPORT_CLOUD_SHADOW_ON_FORWARD_LIT_TRANSLUCENT, bSupportCloudShadowOnForwardLitTranslucent ? 1 : 0);
-	}
+				// Reset outputs
+				for (int32 JobIndex = 0; JobIndex < Results.FinishedJobs.Num(); JobIndex++)
+				{
+					FShaderCommonCompileJob& CurrentJob = *Results.FinishedJobs[JobIndex];
+					auto* SingleJob = CurrentJob.GetSingleShaderJob();
 
-	{
-		static IConsoleVariable *CVar = IConsoleManager::Get().FindConsoleVariable(TEXT("r.Water.SingleLayerWater.SupportCloudShadow"));
-		const bool bSupportCloudShadowOnSingleLayerWater = CVar && CVar->GetInt() > 0;
-		SET_SHADER_DEFINE(Input.Environment, SUPPORT_CLOUD_SHADOW_ON_SINGLE_LAYER_WATER, bSupportCloudShadowOnSingleLayerWater ? 1 : 0);
-	}
+					// NOTE: Changes to MaterialTemplate.usf before retrying won't work, because the entry for Material.usf in CurrentJob.Environment.IncludeFileNameToContentsMap isn't reset
+					if (SingleJob)
+					{
+						if (GShaderCompilingManager->ShouldRecompileToDumpShaderDebugInfo(*SingleJob))
+						{
+							SingleJob->Input.DumpDebugInfoPath = GShaderCompilingManager->CreateShaderDebugInfoPath(SingleJob->Input);
+						}
+						SingleJob->Output = FShaderCompilerOutput();
+						SingleJob->PreprocessOutput = FShaderPreprocessOutput();
+					}
+					else
+					{
+						auto* PipelineJob = CurrentJob.GetShaderPipelineJob();
+						for (auto CommonJob : PipelineJob->StageJobs)
+						{
+							CommonJob->Output = FShaderCompilerOutput();
+							CommonJob->PreprocessOutput = FShaderPreprocessOutput();
+							CommonJob->bFinalized = false;
+						}
+					}
+					
+					// Reset DDC query request owner
+					CurrentJob.RequestOwner.Reset();
+					CurrentJob.bFinalized = false;
+				}
 
-	{
-		static IConsoleVariable* CVar = IConsoleManager::Get().FindConsoleVariable(TEXT("r.FirstPersonRendering.Enabled"));
-		const bool bSupportFirstPersonRendering = CVar && CVar->GetInt() > 0;
-		SET_SHADER_DEFINE(Input.Environment, SUPPORT_FIRST_PERSON_RENDERING, bSupportFirstPersonRendering ? 1 : 0);
-	}
+				// Send all the shaders from this shader map through the compiler again
+				SubmitJobs(Results.FinishedJobs, FString(""), FString(""));
+			}
+		}
 
-	{
-		const bool bTranslucentUsesLightRectLights = GetTranslucentUsesLightRectLights();
-		SET_SHADER_DEFINE(Input.Environment, SUPPORT_RECTLIGHT_ON_FORWARD_LIT_TRANSLUCENT, bTranslucentUsesLightRectLights ? 1 : 0);
-	}
+		const int32 OriginalNumShaderMaps = CompletedShaderMaps.Num();
 
-	{
-		const bool bTranslucentUsesLightIESProfiles = GetTranslucentUsesLightIESProfiles();
-		SET_SHADER_DEFINE(Input.Environment, SUPPORT_IESPROFILE_ON_FORWARD_LIT_TRANSLUCENT, bTranslucentUsesLightIESProfiles ? 1 : 0);
-	}
+		// Remove the failed shader maps
+		for (int32 RemoveIndex = 0; RemoveIndex < MapsToRemove.Num(); RemoveIndex++)
+		{
+			CompletedShaderMaps.Remove(MapsToRemove[RemoveIndex]);
+		}
 
-	{
-		static IConsoleVariable* CVar = IConsoleManager::Get().FindConsoleVariable(TEXT("r.Shadow.Virtual.TranslucentQuality"));
-		const bool bHighQualityShadow = CVar && CVar->GetInt() > 0;
-		SET_SHADER_DEFINE(Input.Environment, SUPPORT_VSM_FOWARD_QUALITY, bHighQualityShadow ? 1 : 0);
-	}
+		check(CompletedShaderMaps.Num() == OriginalNumShaderMaps - MapsToRemove.Num());
 
-	{
-		const bool bUseTriangleStrips = GetHairStrandsUsesTriangleStrips();
-		SET_SHADER_DEFINE(Input.Environment, USE_HAIR_TRIANGLE_STRIP, bUseTriangleStrips ? 1 : 0);
+		// Block until the failed shader maps have been compiled again
+		BlockOnShaderMapCompletion(MapsToRemove, CompletedShaderMaps);
+
+		check(CompletedShaderMaps.Num() == OriginalNumShaderMaps);
 	}
+#endif	//WITH_EDITORONLY_DATA
 
-	const bool bSubstrate = Substrate::IsSubstrateEnabled() && IsSubstrateSupportForShaderPipeline(Input);
-	{
-		SET_SHADER_DEFINE(Input.Environment, SUBSTRATE_ENABLED, bSubstrate ? 1 : 0);
+	return bRetryCompile;
+}
 
-		if (bSubstrate)
+void FShaderMapCompileResults::CheckIfHung()
+{
+	if (!bIsHung)
+	{
+		double DurationSoFar = FPlatformTime::Seconds() - TimeStarted;
+		if (DurationSoFar >= static_cast<double>(GShaderMapCompilationTimeout))
 		{
-			const uint32 SubstrateShadingQuality = Substrate::GetShadingQuality(Target.GetPlatform());
-			SET_SHADER_DEFINE(Input.Environment, SUBSTRATE_SHADING_QUALITY, SubstrateShadingQuality);
-
-			const bool bLowQuality = SubstrateShadingQuality > 1;
-			SET_SHADER_DEFINE(Input.Environment, USE_ACHROMATIC_BXDF_ENERGY, bLowQuality ? 1u : 0u);
-
-			const uint32 SubstrateSheenQuality = Substrate::GetSheenQuality();
-			Input.Environment.SetDefine(TEXT("SUBSTRATE_SHEEN_QUALITY"), bLowQuality ? 2 : SubstrateSheenQuality);
+			bIsHung = true;
+			// always produce an error message first, even if going to crash, as the automation controller does not seem to be picking up Fatal messages
+			UE_LOG(LogShaderCompilers, Error, TEXT("Hung shadermap detected, time spent compiling: %f seconds, NumPendingJobs: %d, FinishedJobs: %d"),
+				DurationSoFar,
+				NumPendingJobs.GetValue(),
+				FinishedJobs.Num()
+			);
 
-			const uint32 SubstrateNormalQuality = Substrate::GetNormalQuality();
-			SET_SHADER_DEFINE(Input.Environment, SUBSTRATE_NORMAL_QUALITY, SubstrateNormalQuality);
-			if (SubstrateNormalQuality == 0)
-			{
-				SET_SHADER_DEFINE(Input.Environment, SUBSTRATE_TOP_LAYER_TYPE, TEXT("uint"));
-			}
-			else
+			if (GCrashOnHungShaderMaps)
 			{
-				SET_SHADER_DEFINE(Input.Environment, SUBSTRATE_TOP_LAYER_TYPE, TEXT("uint2"));
+				UE_LOG(LogShaderCompilers, Fatal, TEXT("Crashing on a hung shadermap, time spent compiling: %f seconds, NumPendingJobs: %d, FinishedJobs: %d"),
+					DurationSoFar,
+					NumPendingJobs.GetValue(),
+					FinishedJobs.Num()
+				);
 			}
+		}
+	}
+}
 
-			const uint32 SubstrateUintPerPixel = Substrate::GetBytePerPixel(Target.GetPlatform()) / 4u;
-			SET_SHADER_DEFINE(Input.Environment, SUBSTRATE_MATERIAL_NUM_UINTS, SubstrateUintPerPixel);
-
-			const uint32 SubstrateClosurePerPixel = Substrate::GetClosurePerPixel(Target.GetPlatform());
-			SET_SHADER_DEFINE(Input.Environment, SUBSTRATE_MATERIAL_CLOSURE_COUNT, SubstrateClosurePerPixel);
+void FShaderCompilingManager::CancelCompilation(const TCHAR* MaterialName, const TArray<int32>& ShaderMapIdsToCancel)
+{
+	// nothing to cancel here, just in case
+	if (!AllowShaderCompiling())
+	{
+		return;
+	}
 
-			const bool bSubstrateDBufferPass = Substrate::IsDBufferPassEnabled(Target.GetPlatform());
-			SET_SHADER_DEFINE(Input.Environment, SUBSTRATE_USE_DBUFFER_PASS, bSubstrateDBufferPass ? 1 : 0);
+	check(IsInGameThread());
+	check(!FPlatformProperties::RequiresCookedData());
 
-			const bool bSubstrateGlints = Substrate::IsGlintEnabled(Target.GetPlatform());
-			SET_SHADER_DEFINE(Input.Environment, PLATFORM_ENABLES_SUBSTRATE_GLINTS, bSubstrateGlints ? 1 : 0);
+	// Lock CompileQueueSection so we can access the input and output queues
+	FScopeLock Lock(&CompileQueueSection);
 
-			const bool bSpecularProfileEnabled = Substrate::IsSpecularProfileEnabled(Target.GetPlatform());
-			SET_SHADER_DEFINE(Input.Environment, PLATFORM_ENABLES_SUBSTRATE_SPECULAR_PROFILE, bSpecularProfileEnabled ? 1 : 0);
-		}
-		else
+	int32 TotalNumJobsRemoved = 0;
+	for (int32 IdIndex = 0; IdIndex < ShaderMapIdsToCancel.Num(); ++IdIndex)
+	{
+		int32 MapIdx = ShaderMapIdsToCancel[IdIndex];
+		if (const FPendingShaderMapCompileResultsPtr* ResultsPtr = ShaderMapJobs.Find(MapIdx))
 		{
-			// Some global uniform buffers reference this type -- so we need to have it defined in all cases
-			SET_SHADER_DEFINE(Input.Environment, SUBSTRATE_TOP_LAYER_TYPE, TEXT("uint"));
-		}
+			const int32 NumJobsRemoved = AllJobs.RemoveAllPendingJobsWithId(MapIdx);
+	
+			TotalNumJobsRemoved += NumJobsRemoved;
 
-		const bool bSubstrateBackCompatibility = bSubstrate && Substrate::IsBackCompatibilityEnabled();
-		SET_SHADER_DEFINE(Input.Environment, PROJECT_SUBSTRATE_BACKCOMPATIBILITY, bSubstrateBackCompatibility ? 1 : 0);
+			FShaderMapCompileResults* ShaderMapJob = *ResultsPtr;
+			const int32 PrevNumPendingJobs = ShaderMapJob->NumPendingJobs.Subtract(NumJobsRemoved);
+			check(PrevNumPendingJobs >= NumJobsRemoved);
 
-		const bool bSubstrateOpaqueRoughRefrac = bSubstrate && Substrate::IsOpaqueRoughRefractionEnabled();
-		SET_SHADER_DEFINE(Input.Environment, SUBSTRATE_OPAQUE_ROUGH_REFRACTION_ENABLED, bSubstrateOpaqueRoughRefrac ? 1 : 0);
+			// The shader map job result should be skipped since it is out of date.
+			ShaderMapJob->bSkipResultProcessing = true;
+		
+			if (PrevNumPendingJobs == NumJobsRemoved && ShaderMapJob->FinishedJobs.Num() == 0)
+			{
+				//We've removed all the jobs for this shader map so remove it.
+				ShaderMapJobs.Remove(MapIdx);
+			}
+		}
 
-		const bool bSubstrateAdvDebug = bSubstrate && Substrate::IsAdvancedVisualizationEnabled();
-		SET_SHADER_DEFINE(Input.Environment, SUBSTRATE_ADVANCED_DEBUG_ENABLED, bSubstrateAdvDebug ? 1 : 0);
+		// Don't continue finalizing once compilation has been canceled
+		// the CompilingId has been removed from ShaderMapsBeingCompiled, which will cause crash when attempting to do any further processing
+		const int32 NumPendingRemoved = PendingFinalizeShaderMaps.Remove(MapIdx);
 	}
 
+	if (TotalNumJobsRemoved > 0)
 	{
-		static IConsoleVariable* CVar = IConsoleManager::Get().FindConsoleVariable(TEXT("r.Material.RoughDiffuse"));
-		const bool bMaterialRoughDiffuse = CVar && CVar->GetInt() != 0;
-		const bool bSubstrateRoughDiffuse = Substrate::IsRoughDiffuseEnabled() && !Substrate::IsBackCompatibilityEnabled();
-		SET_SHADER_DEFINE(Input.Environment, MATERIAL_ROUGHDIFFUSE, (bSubstrate ? bSubstrateRoughDiffuse : bMaterialRoughDiffuse) ? 1 : 0);
+		UE_LOG(LogShaders, Display, TEXT("CancelCompilation %s, Removed %d jobs"), MaterialName ? MaterialName : TEXT(""), TotalNumJobsRemoved);
 	}
+}
 
-	{
-		static IConsoleVariable* CVar = IConsoleManager::Get().FindConsoleVariable(TEXT("r.Lumen.Supported"));
-		const bool bLumenSupported = CVar->GetInt() != 0;
-		SET_SHADER_DEFINE(Input.Environment, PROJECT_SUPPORTS_LUMEN, bLumenSupported ? 1 : 0);
-	}
+void FShaderCompilingManager::FinishCompilation(const TCHAR* MaterialName, const TArray<int32>& ShaderMapIdsToFinishCompiling)
+{
+	TRACE_CPUPROFILER_EVENT_SCOPE(FShaderCompilingManager::FinishCompilation);
 
+	// nothing to do
+	if (!AllowShaderCompiling())
 	{
-		const bool bSupportOIT = FDataDrivenShaderPlatformInfo::GetSupportsOIT(EShaderPlatform(Target.Platform));
-		static IConsoleVariable* CVar = IConsoleManager::Get().FindConsoleVariable(TEXT("r.OIT.SortedPixels"));
-		const bool bOIT = CVar && CVar->GetInt() != 0;
-		SET_SHADER_DEFINE(Input.Environment, PROJECT_OIT, (bSupportOIT && bOIT) ? 1 : 0);
+		return;
 	}
 
-	{
-		static IConsoleVariable* CVar = IConsoleManager::Get().FindConsoleVariable(TEXT("r.Material.EnergyConservation"));
-		const bool bMaterialEnergyConservation = CVar && CVar->GetInt() != 0;
-		SET_SHADER_DEFINE(Input.Environment, LEGACY_MATERIAL_ENERGYCONSERVATION, bMaterialEnergyConservation ? 1 : 0);
-	}
+	check(IsInGameThread());
+	check(!FPlatformProperties::RequiresCookedData());
+	const double StartTime = FPlatformTime::Seconds();
 
+	FText StatusUpdate;
+	if ( MaterialName != nullptr)
 	{
-		static IConsoleVariable* CVar = IConsoleManager::Get().FindConsoleVariable(TEXT("r.SupportSkyAtmosphereAffectsHeightFog"));
-		SET_SHADER_DEFINE(Input.Environment, PROJECT_SUPPORT_SKY_ATMOSPHERE_AFFECTS_HEIGHFOG, (CVar && bSupportSkyAtmosphere) ? (CVar->GetInt() != 0) : 0);
+		FFormatNamedArguments Args;
+		Args.Add( TEXT("MaterialName"), FText::FromString( MaterialName ) );
+		StatusUpdate = FText::Format( NSLOCTEXT("ShaderCompilingManager", "CompilingShadersForMaterialStatus", "Compiling shaders: {MaterialName}..."), Args );
 	}
-
+	else
 	{
-		SET_SHADER_DEFINE(Input.Environment, PROJECT_EXPFOG_MATCHES_VFOG, DoesProjectSupportExpFogMatchesVolumetricFog() ? 1 : 0);
+		StatusUpdate = NSLOCTEXT("ShaderCompilingManager", "CompilingShadersStatus", "Compiling shaders...");
 	}
 
+	FScopedSlowTask SlowTask(1, StatusUpdate, GIsEditor && !IsRunningCommandlet());
+	SlowTask.EnterProgressFrame(1);
+
+	TMap<int32, FShaderMapFinalizeResults> CompiledShaderMaps;
+	CompiledShaderMaps.Append( PendingFinalizeShaderMaps );
+	PendingFinalizeShaderMaps.Empty();
+	BlockOnShaderMapCompletion(ShaderMapIdsToFinishCompiling, CompiledShaderMaps);
+
+	bool bRetry = false;
+	do 
 	{
-		static IConsoleVariable* CVar = IConsoleManager::Get().FindConsoleVariable(TEXT("r.PostProcessing.PropagateAlpha"));
-		int32 PropagateAlpha = CVar->GetInt();
+		bRetry = HandlePotentialRetryOnError(CompiledShaderMaps);
+	} 
+	while (bRetry);
 
-		if (bIsMobilePlatform)
-		{
-			static FShaderPlatformCachedIniValue<int32> MobilePropagateAlphaIniValue(TEXT("r.Mobile.PropagateAlpha"));
-			int MobilePropagateAlphaIniValueInt = MobilePropagateAlphaIniValue.Get((EShaderPlatform)ShaderPlatform);
-			PropagateAlpha = MobilePropagateAlphaIniValueInt > 0 ? 2 : 0;
-		}
+	ProcessCompiledShaderMaps(CompiledShaderMaps, FLT_MAX);
+	check(CompiledShaderMaps.Num() == 0);
 
-		if (PropagateAlpha < 0 || PropagateAlpha > 2)
-		{
-			PropagateAlpha = 0;
-		}
-		SET_SHADER_DEFINE(Input.Environment, POST_PROCESS_ALPHA, PropagateAlpha);
-	}
+	const double EndTime = FPlatformTime::Seconds();
 
-	if (TargetPlatform && 
-		TargetPlatform->SupportsFeature(ETargetPlatformFeatures::NormalmapLAEncodingMode))
-	{
-		SET_SHADER_DEFINE(Input.Environment, LA_NORMALMAPS, 1);
-	}
+	UE_LOG(LogShaders, Verbose, TEXT("FinishCompilation %s %.3fs"), MaterialName ? MaterialName : TEXT(""), (float)(EndTime - StartTime));
+}
 
-	SET_SHADER_DEFINE(Input.Environment, COLORED_LIGHT_FUNCTION_ATLAS, GetLightFunctionAtlasFormat() > 0 ? 1 : 0);
+void FShaderCompilingManager::FinishAllCompilation()
+{
+#if WITH_EDITOR
+	// This is here for backward compatibility since textures are most probably expected to be ready too.
+	FTextureCompilingManager::Get().FinishAllCompilation();
+#endif
 
-	SET_SHADER_DEFINE(Input.Environment, PLATFORM_SUPPORTS_SHADER_ROOT_CONSTANTS, RHISupportsShaderRootConstants(EShaderPlatform(Target.Platform)) ? 1 : 0);
-	SET_SHADER_DEFINE(Input.Environment, PLATFORM_SUPPORTS_SHADER_BUNDLE_DISPATCH, RHISupportsShaderBundleDispatch(EShaderPlatform(Target.Platform)) ? 1 : 0);
-	SET_SHADER_DEFINE(Input.Environment, PLATFORM_SUPPORTS_RENDERTARGET_WRITE_MASK, RHISupportsRenderTargetWriteMask(EShaderPlatform(Target.Platform)) ? 1 : 0);
-	SET_SHADER_DEFINE(Input.Environment, PLATFORM_SUPPORTS_PER_PIXEL_DBUFFER_MASK, FDataDrivenShaderPlatformInfo::GetSupportsPerPixelDBufferMask(EShaderPlatform(Target.Platform)) ? 1 : 0);
-	SET_SHADER_DEFINE(Input.Environment, PLATFORM_SUPPORTS_DISTANCE_FIELDS, DoesPlatformSupportDistanceFields(EShaderPlatform(Target.Platform)) ? 1 : 0);
-	SET_SHADER_DEFINE(Input.Environment, PLATFORM_SUPPORTS_MESH_SHADERS_TIER0, RHISupportsMeshShadersTier0(EShaderPlatform(Target.Platform)) ? 1 : 0);
-	SET_SHADER_DEFINE(Input.Environment, PLATFORM_SUPPORTS_MESH_SHADERS_TIER1, RHISupportsMeshShadersTier1(EShaderPlatform(Target.Platform)) ? 1 : 0);
-	SET_SHADER_DEFINE(Input.Environment, PLATFORM_ALLOW_SCENE_DATA_COMPRESSED_TRANSFORMS, FDataDrivenShaderPlatformInfo::GetSupportSceneDataCompressedTransforms(EShaderPlatform(Target.Platform)) ? 1 : 0);
-	SET_SHADER_DEFINE(Input.Environment, PLATFORM_SUPPORTS_BUFFER_LOAD_TYPE_CONVERSION, RHISupportsBufferLoadTypeConversion(ShaderPlatform) ? 1 : 0);
-	SET_SHADER_DEFINE(Input.Environment, PLATFORM_SUPPORTS_UNIFORM_BUFFER_OBJECTS, FDataDrivenShaderPlatformInfo::GetSupportsUniformBufferObjects(EShaderPlatform(Target.Platform)) ? 1 : 0);
-	SET_SHADER_DEFINE(Input.Environment, PLATFORM_SUPPORTS_RAY_TRACING_HIGH_END_EFFECTS, FDataDrivenShaderPlatformInfo::GetSupportsHighEndRayTracingEffects(EShaderPlatform(Target.Platform)) ? 1 : 0);
-	SET_SHADER_DEFINE(Input.Environment, COMPILER_SUPPORTS_BARYCENTRIC_INTRINSICS, FDataDrivenShaderPlatformInfo::GetSupportsBarycentricsIntrinsics(EShaderPlatform(Target.Platform)));
-	SET_SHADER_DEFINE(Input.Environment, PLATFORM_SUPPORTS_BARYCENTRICS_SEMANTIC, FDataDrivenShaderPlatformInfo::GetSupportsBarycentricsSemantic(EShaderPlatform(Target.Platform)) != ERHIFeatureSupport::Unsupported);
+	TRACE_CPUPROFILER_EVENT_SCOPE(FShaderCompilingManager::FinishAllCompilation);
+	check(IsInGameThread());
+	check(!FPlatformProperties::RequiresCookedData());
+	const double StartTime = FPlatformTime::Seconds();
 
-	bool bEnableBindlessMacro = false;
-	if (RHIGetBindlessSupport(ShaderPlatform) != ERHIBindlessSupport::Unsupported && !Input.Environment.CompilerFlags.Contains(CFLAG_ForceBindful))
+	TMap<int32, FShaderMapFinalizeResults> CompiledShaderMaps;
+	CompiledShaderMaps.Append( PendingFinalizeShaderMaps );
+	PendingFinalizeShaderMaps.Empty();
+	BlockOnAllShaderMapCompletion(CompiledShaderMaps);
+
+	bool bRetry = false;
+	do 
 	{
-		const bool bIsRaytracingShader = IsRayTracingShaderFrequency(Input.Target.GetFrequency());
+		bRetry = HandlePotentialRetryOnError(CompiledShaderMaps);
+	} 
+	while (bRetry);
 
-		const ERHIBindlessConfiguration ResourcesConfig = UE::ShaderCompiler::GetBindlessResourcesConfiguration(ShaderFormatName);
-		const ERHIBindlessConfiguration SamplersConfig = UE::ShaderCompiler::GetBindlessSamplersConfiguration(ShaderFormatName);
+	ProcessCompiledShaderMaps(CompiledShaderMaps, FLT_MAX);
+	check(CompiledShaderMaps.Num() == 0);
 
-		if (ResourcesConfig == ERHIBindlessConfiguration::AllShaders || (ResourcesConfig == ERHIBindlessConfiguration::RayTracingShaders && bIsRaytracingShader))
-		{
-			bEnableBindlessMacro = true;
-			Input.Environment.CompilerFlags.Add(CFLAG_BindlessResources);
-			SET_SHADER_DEFINE(Input.Environment, ENABLE_BINDLESS_RESOURCES, true);
-		}
+	const double EndTime = FPlatformTime::Seconds();
 
-		if (SamplersConfig == ERHIBindlessConfiguration::AllShaders || (SamplersConfig == ERHIBindlessConfiguration::RayTracingShaders && bIsRaytracingShader))
-		{
-			bEnableBindlessMacro = true;
-			Input.Environment.CompilerFlags.Add(CFLAG_BindlessSamplers);
-			SET_SHADER_DEFINE(Input.Environment, ENABLE_BINDLESS_SAMPLERS, true);
-		}
-	}
+	UE_LOG(LogShaders, Verbose, TEXT("FinishAllCompilation %.3fs"), (float)(EndTime - StartTime));
+}
 
-	SET_SHADER_DEFINE(Input.Environment, PLATFORM_SUPPORTS_BINDLESS, bEnableBindlessMacro);
+void FShaderCompilingManager::ProcessAsyncResults(bool bLimitExecutionTime, bool bBlockOnGlobalShaderCompletion)
+{
+	const float TimeSlice = bLimitExecutionTime ? ProcessGameThreadTargetTime : 0.f;
+	ProcessAsyncResults(TimeSlice, bBlockOnGlobalShaderCompletion);
+}
 
-	if (CVarShadersRemoveDeadCode.GetValueOnAnyThread())
-	{
-		Input.Environment.CompilerFlags.Add(CFLAG_RemoveDeadCode);
-	}
+void FShaderCompilingManager::ProcessAsyncResults(float TimeSlice, bool bBlockOnGlobalShaderCompletion)
+{
+	LLM_SCOPE_BYTAG(ShaderCompiler);
 
-	if (CVarDisableSourceStripping.GetValueOnAnyThread())
-	{
-		Input.Environment.CompilerFlags.Add(CFLAG_DisableSourceStripping);
-	}
+	TRACE_CPUPROFILER_EVENT_SCOPE(FShaderCompilingManager::ProcessAsyncResults)
+
+	COOK_STAT(FScopedDurationTimer Timer(ShaderCompilerCookStats::ProcessAsyncResultsTimeSec));
+	check(IsInGameThread());
+
+	const double StartTime = FPlatformTime::Seconds();
 
+	// Some controllers need to be manually ticked if the engine loop is not initialized or blocked
+	// to do things like tick the HTTPModule.
+	// Otherwise the results from the controller will never be processed.
+	// We check for bBlockOnGlobalShaderCompletion because the BlockOnShaderMapCompletion methods already do this.
+	if (!bBlockOnGlobalShaderCompletion && BuildDistributionController)
 	{
-		static const auto CVar = IConsoleManager::Get().FindConsoleVariable(TEXT("r.VT.AnisotropicFiltering"));
-		SET_SHADER_DEFINE(Input.Environment, VIRTUAL_TEXTURE_ANISOTROPIC_FILTERING, CVar ? (CVar->GetInt() != 0) : 0);
-		
-		if (bIsMobilePlatform)
-		{
-			static FShaderPlatformCachedIniValue<bool> CVarVTMobileManualTrilinearFiltering(TEXT("r.VT.Mobile.ManualTrilinearFiltering"));
-			SET_SHADER_DEFINE(Input.Environment, VIRTUAL_TEXTURE_MANUAL_TRILINEAR_FILTERING, (CVarVTMobileManualTrilinearFiltering.Get(Target.GetPlatform()) ? 1 : 0));
-		}
+		BuildDistributionController->Tick(0.0f);
 	}
 
-	if (bIsMobilePlatform)
+	// Block on global shaders before checking for shader maps to finalize
+	// So if we block on global shaders for a long time, we will get a chance to finalize all the non-global shader maps completed during that time.
+	if (bBlockOnGlobalShaderCompletion)
 	{
-		const bool bMobileMovableSpotlightShadowsEnabled = IsMobileMovableSpotlightShadowsEnabled(Target.GetPlatform());
-		SET_SHADER_DEFINE(Input.Environment, PROJECT_MOBILE_ENABLE_MOVABLE_SPOTLIGHT_SHADOWS, bMobileMovableSpotlightShadowsEnabled ? 1 : 0);
+		TArray<int32> ShaderMapId;
+		ShaderMapId.Add(GlobalShaderMapId);
+
+		// Block until the global shader map jobs are complete
+		GShaderCompilingManager->BlockOnShaderMapCompletion(ShaderMapId, PendingFinalizeShaderMaps);
 	}
 
+	int32 NumCompilingShaderMaps = 0;
+
 	{
-		using namespace UE::Color;
-		const bool bWorkingColorSpaceIsSRGB = FColorSpace::GetWorking().IsSRGB();
-		SET_SHADER_DEFINE(Input.Environment, WORKING_COLOR_SPACE_IS_SRGB, bWorkingColorSpaceIsSRGB ? 1 : 0);
-		
-		// We limit matrix definitions below to WORKING_COLOR_SPACE_IS_SRGB == 0.
-		if (!bWorkingColorSpaceIsSRGB)
-		{
-			const TCHAR MatrixFormat[] = TEXT("float3x3(%0.10f, %0.10f, %0.10f, %0.10f, %0.10f, %0.10f, %0.10f, %0.10f, %0.10f)");
-			const FColorSpace& WorkingColorSpace = FColorSpace::GetWorking();
+		// Lock CompileQueueSection so we can access the input and output queues
+		FScopeLock Lock(&CompileQueueSection);
 
-			// Note that we transpose the matrices during print since color matrices are usually pre-multiplied.
-			const FMatrix44d& ToXYZ = WorkingColorSpace.GetRgbToXYZ();
-			Input.Environment.SetDefine(
-				TEXT("WORKING_COLOR_SPACE_RGB_TO_XYZ_MAT"),
-				FString::Printf(MatrixFormat,
-					ToXYZ.M[0][0], ToXYZ.M[1][0], ToXYZ.M[2][0],
-					ToXYZ.M[0][1], ToXYZ.M[1][1], ToXYZ.M[2][1],
-					ToXYZ.M[0][2], ToXYZ.M[1][2], ToXYZ.M[2][2]));
+		if (!bBlockOnGlobalShaderCompletion)
+		{
+			bCompilingDuringGame = true;
+		}
 
-			const FMatrix44d& FromXYZ = WorkingColorSpace.GetXYZToRgb();
-			Input.Environment.SetDefine(
-				TEXT("XYZ_TO_RGB_WORKING_COLOR_SPACE_MAT"),
-				FString::Printf(MatrixFormat,
-					FromXYZ.M[0][0], FromXYZ.M[1][0], FromXYZ.M[2][0],
-					FromXYZ.M[0][1], FromXYZ.M[1][1], FromXYZ.M[2][1],
-					FromXYZ.M[0][2], FromXYZ.M[1][2], FromXYZ.M[2][2]));
+		// Get all material shader maps to finalize
+		//
+		for (TMap<int32, FPendingShaderMapCompileResultsPtr>::TIterator It(ShaderMapJobs); It; ++It)
+		{
+			FPendingShaderMapCompileResultsPtr& Results = It.Value();
+			if (Results->FinishedJobs.Num() > 0)
+			{
+				FShaderMapFinalizeResults& FinalizeResults = PendingFinalizeShaderMaps.FindOrAdd(It.Key());
+				FinalizeResults.FinishedJobs.Append(Results->FinishedJobs);
+				Results->FinishedJobs.Reset();
+				FinalizeResults.bAllJobsSucceeded = FinalizeResults.bAllJobsSucceeded && Results->bAllJobsSucceeded;
+			}
 
-			const FColorSpaceTransform& FromSRGB = FColorSpaceTransform::GetSRGBToWorkingColorSpace();
-			SET_SHADER_DEFINE(Input.Environment,
-				SRGB_TO_WORKING_COLOR_SPACE_MAT,
-				FString::Printf(MatrixFormat,
-					FromSRGB.M[0][0], FromSRGB.M[1][0], FromSRGB.M[2][0],
-					FromSRGB.M[0][1], FromSRGB.M[1][1], FromSRGB.M[2][1],
-					FromSRGB.M[0][2], FromSRGB.M[1][2], FromSRGB.M[2][2]));
+			checkf(Results->FinishedJobs.Num() == 0, TEXT("Failed to remove finished jobs, %d remain"), Results->FinishedJobs.Num());
+			if (Results->NumPendingJobs.GetValue() == 0)
+			{
+				It.RemoveCurrent();
+			}
 		}
 
-		static const auto CVar = IConsoleManager::Get().FindConsoleVariable(TEXT("r.LegacyLuminanceFactors"));
-		const bool bUseLegacyLuminance = CVar && CVar->GetInt() != 0;
-		SET_SHADER_DEFINE(Input.Environment, UE_LEGACY_LUMINANCE_FACTORS, bUseLegacyLuminance ? 1 : 0);
+		NumCompilingShaderMaps = ShaderMapJobs.Num();
 	}
 
-	const double TileSize = FLargeWorldRenderScalar::GetTileSize();
-	SET_SHADER_DEFINE(Input.Environment, UE_LWC_RENDER_TILE_SIZE, (float)TileSize);
-	SET_SHADER_DEFINE(Input.Environment, UE_LWC_RENDER_TILE_SIZE_SQRT, (float)FMath::Sqrt(TileSize));
-	SET_SHADER_DEFINE(Input.Environment, UE_LWC_RENDER_TILE_SIZE_RSQRT, (float)FMath::InvSqrt(TileSize));
-	SET_SHADER_DEFINE(Input.Environment, UE_LWC_RENDER_TILE_SIZE_RCP, (float)(1.0 / TileSize));
-	SET_SHADER_DEFINE(Input.Environment, UE_LWC_RENDER_TILE_SIZE_FMOD_PI, (float)FMath::Fmod(TileSize, UE_DOUBLE_PI));
-	SET_SHADER_DEFINE(Input.Environment, UE_LWC_RENDER_TILE_SIZE_FMOD_2PI, (float)FMath::Fmod(TileSize, 2.0 * UE_DOUBLE_PI));
+	int32 NumPendingShaderMaps = PendingFinalizeShaderMaps.Num();
 
-	// Add required symbols from the shader binding layout if set
-	if (Input.Environment.ShaderBindingLayout)
+	if (PendingFinalizeShaderMaps.Num() > 0)
 	{
-		Input.Environment.ShaderBindingLayout->AddRequiredSymbols(Input.RequiredSymbols);
-	}
-
-	// Allow the target shader format to modify the shader input before we add it as a job
-	const IShaderFormat* Format = GetTargetPlatformManagerRef().FindShaderFormat(ShaderFormatName);
-	checkf(Format, TEXT("Shader format %s cannot be found"), *ShaderFormatName.ToString());
-	Format->ModifyShaderCompilerInput(Input);
-
-	// Allow the GBuffer and other shader defines to cause dependend environment changes, but minimizing the #ifdef magic in the shaders, which
-	// is nearly impossible to debug when it goes wrong.
-	FShaderCompileUtilities::ApplyDerivedDefines(Input.Environment, Input.SharedEnvironment, (EShaderPlatform)Target.Platform);
-}
-
-#endif // WITH_EDITOR
+		bool bRetry = false;
+		do 
+		{
+			bRetry = HandlePotentialRetryOnError(PendingFinalizeShaderMaps);
+		} 
+		while (bRetry);
 
-/** Timer class used to report information on the 'recompileshaders' console command. */
-class FRecompileShadersTimer
-{
-public:
-	FRecompileShadersTimer(const TCHAR* InInfoStr=TEXT("Test")) :
-		InfoStr(InInfoStr),
-		bAlreadyStopped(false)
-	{
-		StartTime = FPlatformTime::Seconds();
+		const float TimeBudget = TimeSlice > 0 ? TimeSlice : FLT_MAX;
+		ProcessCompiledShaderMaps(PendingFinalizeShaderMaps, TimeBudget);
+		check(TimeSlice > 0 || PendingFinalizeShaderMaps.Num() == 0);
 	}
 
-	FRecompileShadersTimer(const FString& InInfoStr) :
-		InfoStr(InInfoStr),
-		bAlreadyStopped(false)
-	{
-		StartTime = FPlatformTime::Seconds();
-	}
 
-	void Stop(bool DisplayLog = true)
+	if (bBlockOnGlobalShaderCompletion && TimeSlice <= 0 && !IsRunningCookCommandlet())
 	{
-		if (!bAlreadyStopped)
+		check(PendingFinalizeShaderMaps.Num() == 0);
+
+		if (NumPendingShaderMaps - PendingFinalizeShaderMaps.Num() > 0)
 		{
-			bAlreadyStopped = true;
-			EndTime = FPlatformTime::Seconds();
-			TimeElapsed = EndTime-StartTime;
-			if (DisplayLog)
-			{
-				UE_LOG(LogShaderCompilers, Warning, TEXT("		[%s] took [%.4f] s"),*InfoStr,TimeElapsed);
-			}
+			UE_LOG(LogShaders, Warning, TEXT("Blocking ProcessAsyncResults for %.1fs, processed %u shader maps, %u being compiled"), 
+				(float)(FPlatformTime::Seconds() - StartTime),
+				NumPendingShaderMaps - PendingFinalizeShaderMaps.Num(), 
+				NumCompilingShaderMaps);
 		}
 	}
-
-	~FRecompileShadersTimer()
+	else if (NumPendingShaderMaps - PendingFinalizeShaderMaps.Num() > 0)
 	{
-		Stop(true);
+		UE_LOG(LogShaders, Verbose, TEXT("Completed %u async shader maps, %u more pending, %u being compiled"),
+			NumPendingShaderMaps - PendingFinalizeShaderMaps.Num(), 
+			PendingFinalizeShaderMaps.Num(),
+			NumCompilingShaderMaps);
 	}
 
-protected:
-	double StartTime,EndTime;
-	double TimeElapsed;
-	FString InfoStr;
-	bool bAlreadyStopped;
-};
+	UpdateNumRemainingAssets();
+}
 
-namespace
+void FShaderCompilingManager::UpdateNumRemainingAssets()
 {
-	void ListAllShaderTypes()
+	if (IsInGameThread())
 	{
-		UE_LOG(LogShaderCompilers, Display, TEXT("ShaderTypeName, Filename"));
-		for (TLinkedList<FShaderType*>::TIterator It(FShaderType::GetTypeList()); It; It.Next())
+		const int32 NumRemainingAssets = GetNumRemainingAssets();
+		if (LastNumRemainingAssets != NumRemainingAssets)
 		{
-			UE_LOG(LogShaderCompilers, Display, TEXT("%s, %s "), (*It)->GetName(), (*It)->GetShaderFilename());
-		}
+			if (NumRemainingAssets == 0)
+			{
+				// This is important to at least broadcast once we reach 0 remaining assets
+				// even if we don't have any UObject to report because some listener are only 
+				// interested to be notified when the number of async compilation reaches 0.
+				FAssetCompilingManager::Get().OnAssetPostCompileEvent().Broadcast({});
+			}
 
-		UE_LOG(LogShaderCompilers, Display, TEXT("VertexFactoryTypeName, Filename"));
-		for (TLinkedList<FVertexFactoryType*>::TIterator It(FVertexFactoryType::GetTypeList()); It; It.Next())
-		{
-			UE_LOG(LogShaderCompilers, Display, TEXT("%s, %s"), (*It)->GetName(), (*It)->GetShaderFilename());
+			LastNumRemainingAssets = NumRemainingAssets;
+			Notification->Update(NumRemainingAssets);
 		}
 	}
+}
 
-	ODSCRecompileCommand ParseRecompileCommandString(const TCHAR* CmdString, TArray<FString>& OutMaterialsToLoad, FString& OutShaderTypesToLoad, FString& OutRequestedMaterialName)
-	{
-		FString CmdName = FParse::Token(CmdString, 0);
+bool FShaderCompilingManager::IsShaderCompilerWorkerRunning(FProcHandle & WorkerHandle)
+{
+	return FPlatformProcess::IsProcRunning(WorkerHandle);
+}
 
-		ODSCRecompileCommand CommandType = ODSCRecompileCommand::None;
-		OutMaterialsToLoad.Empty();
+#if WITH_EDITOR
 
-		if( !CmdName.IsEmpty() && FCString::Stricmp(*CmdName,TEXT("Material"))==0 )
-		{
-			CommandType = ODSCRecompileCommand::Material;
+/* Generates a uniform buffer struct member hlsl declaration using the member's metadata. */
+static void GenerateUniformBufferStructMember(FString& Result, const FShaderParametersMetadata::FMember& Member, EShaderPlatform ShaderPlatform)
+{
+	// Generate the base type name.
+	FString TypeName;
+	Member.GenerateShaderParameterType(TypeName, ShaderPlatform);
 
-			// tell other side the material to load, by pathname
-			FString RequestedMaterialName( FParse::Token( CmdString, 0 ) );
-			OutRequestedMaterialName = RequestedMaterialName;
-			UMaterialInterface* MatchingMaterial = nullptr;
-			for (TObjectIterator<UMaterialInterface> It; It; ++It)
-			{
-				UMaterial* Material = It->GetMaterial();
+	// Generate array dimension post fix
+	FString ArrayDim;
+	if (Member.GetNumElements() > 0)
+	{
+		ArrayDim = FString::Printf(TEXT("[%u]"), Member.GetNumElements());
+	}
 
-				if (Material && Material->GetName() == RequestedMaterialName)
-				{
-					OutMaterialsToLoad.Add(It->GetPathName());
-					MatchingMaterial = Material;
-					break;
-				}
-			}
+	Result = FString::Printf(TEXT("%s %s%s"), *TypeName, Member.GetName(), *ArrayDim);
+}
 
-			// Find all material instances from the requested material and 
-			// request a compile for them.
-			if (MatchingMaterial)
-			{
-				for (TObjectIterator<UMaterialInstance> It; It; ++It)
-				{
-					if (It && It->IsDependent(MatchingMaterial))
-					{
-						OutMaterialsToLoad.Add(It->GetPathName());
-					}
-				}
-			}
-		}
-		else if (!CmdName.IsEmpty() && FCString::Stricmp(*CmdName, TEXT("Global")) == 0)
+/* Generates the instanced stereo hlsl code that's dependent on view uniform declarations. */
+void GenerateInstancedStereoCode(FString& Result, EShaderPlatform ShaderPlatform)
+{
+	// Find the InstancedView uniform buffer struct
+	const FShaderParametersMetadata* View = nullptr;
+	const FShaderParametersMetadata* InstancedView = nullptr;
+
+	for (TLinkedList<FShaderParametersMetadata*>::TIterator StructIt(FShaderParametersMetadata::GetStructList()); StructIt; StructIt.Next())
+	{
+		if (StructIt->GetShaderVariableName() == FString(TEXT("View")))
 		{
-			CommandType = ODSCRecompileCommand::Global;
+			View = *StructIt;
 		}
-		else if (!CmdName.IsEmpty() && FCString::Stricmp(*CmdName, TEXT("Changed")) == 0)
-		{
-			CommandType = ODSCRecompileCommand::Changed;
 
-			// Compile all the shaders that have changed for the materials we have loaded.
-			for (TObjectIterator<UMaterialInterface> It; It; ++It)
-			{
-				OutMaterialsToLoad.Add(It->GetPathName());
-			}
-		}
-		else if (FCString::Stricmp(*CmdName, TEXT("All")) == 0)
+		if (StructIt->GetShaderVariableName() == FString(TEXT("InstancedView")))
 		{
-			CommandType = ODSCRecompileCommand::Material;
-
-			// tell other side all the materials to load, by pathname
-			for (TObjectIterator<UMaterialInterface> It; It; ++It)
-			{
-				OutMaterialsToLoad.Add(It->GetPathName());
-			}
+			InstancedView = *StructIt;
 		}
-		else if (FCString::Stricmp(*CmdName, TEXT("listtypes")) == 0)
+
+		if (View && InstancedView)
 		{
-			ListAllShaderTypes();
+			break;
 		}
-		else
-		{
-			CommandType = ODSCRecompileCommand::SingleShader;
+	}
+	checkSlow(View != nullptr);
+	checkSlow(InstancedView != nullptr);
 
-			OutShaderTypesToLoad = CmdName;
+	const TArray<FShaderParametersMetadata::FMember>& StructMembersView = View->GetMembers();
+	const TArray<FShaderParametersMetadata::FMember>& StructMembersInstanced = InstancedView->GetMembers();
 
-			// tell other side which materials to load and compile the single
-			// shader for.
-			for (TObjectIterator<UMaterialInterface> It; It; ++It)
-			{
-				OutMaterialsToLoad.Add(It->GetPathName());
-			}
-		}
+	static const auto CVarViewHasTileOffsetData = IConsoleManager::Get().FindTConsoleVariableDataInt(TEXT("r.ViewHasTileOffsetData"));
+	const bool bViewHasTileOffsetData = CVarViewHasTileOffsetData ? (CVarViewHasTileOffsetData->GetValueOnAnyThread() != 0) : false;
 
-		return CommandType;
+	Result = "";
+	if (bViewHasTileOffsetData)
+	{
+		Result +=  "struct ViewStateTileOffsetData\r\n";
+		Result += "{\r\n";
+		Result += "\tFLWCVector3 WorldCameraOrigin;\r\n";
+		Result += "\tFLWCVector3 WorldViewOrigin;\r\n";
+		Result += "\tFLWCVector3 PrevWorldCameraOrigin;\r\n";
+		Result += "\tFLWCVector3 PrevWorldViewOrigin;\r\n";
+		Result += "\tFLWCVector3 PreViewTranslation;\r\n";
+		Result += "\tFLWCVector3 PrevPreViewTranslation;\r\n";
+		Result += "};\r\n";
 	}
-}
 
-static int32 GODSCMaterialUpdateFlags = 0;
-static FAutoConsoleVariableRef CVarODSCMaterialUpdateFlags(
-	TEXT("ODSC.MaterialUpdateFlags"),
-	GODSCMaterialUpdateFlags,
-	TEXT("Changes the material update flags when ODSC receives new shaders and needs to update the materials\n")
-	TEXT("0 (default): no additional work\n")
-	TEXT("1: Reregister all components while updating the material\n")
-	TEXT("2: Sync with the rendering thread after all the calls to RecacheUniformExpressions\n")
-	TEXT("4 (legacy): Recreates only the render state for *all* components, including the ones not changed by ODSC\n")
-);
+	// ViewState definition
+	Result +=  "struct ViewState\n";
+	Result += "{\n";
+	for (int32 MemberIndex = 0; MemberIndex < StructMembersInstanced.Num(); ++MemberIndex)
+	{
+		FString MemberDecl;
+		// ViewState is only supposed to contain InstancedView members however we want their original type and array-length instead of their representation in the instanced array
+		// GPUSceneViewId for example needs to return 	uint GPUSceneViewId; and not uint4 InstancedView_GPUSceneViewId[2];
+		// and that initial representation is in StructMembersView
+		GenerateUniformBufferStructMember(MemberDecl, StructMembersView[MemberIndex], ShaderPlatform);
+		Result += FString::Printf(TEXT("\t%s;\n"), *MemberDecl);
+	}
+	Result += "\tFDFInverseMatrix WorldToClip;\n";
+	Result += "\tFDFMatrix ClipToWorld;\n";
+	Result += "\tFDFMatrix ScreenToWorld;\n";
+	Result += "\tFDFMatrix PrevClipToWorld;\n";
+	Result += "\tFDFVector3 WorldCameraOrigin;\n";
+	Result += "\tFDFVector3 WorldViewOrigin;\n";
+	Result += "\tFDFVector3 PrevWorldCameraOrigin;\n";
+	Result += "\tFDFVector3 PrevWorldViewOrigin;\n";
+	Result += "\tFDFVector3 PreViewTranslation;\n";
+	Result += "\tFDFVector3 PrevPreViewTranslation;\n";
 
-void ProcessCookOnTheFlyShaders(bool bReloadGlobalShaders, const TArray<uint8>& MeshMaterialMaps, const TArray<FString>& MaterialsToLoad, const TArray<uint8>& GlobalShaderMap)
-{
-	TRACE_CPUPROFILER_EVENT_SCOPE(ProcessCookOnTheFlyShaders);
-	check(IsInGameThread());
+	if (bViewHasTileOffsetData)
+	{
+		Result += "\tViewStateTileOffsetData TileOffset;\n";
+	}
 
-	bool bHasFlushed = false;
+	Result += "};\n";
 
-	auto DoFlushIfNecessary = [&bHasFlushed]() {
-		if (!bHasFlushed )
-		{
-			// now we need to refresh the RHI resources
-			FlushRenderingCommands();
-			bHasFlushed = true;
-		}
-	};
+	Result += "\tvoid FinalizeViewState(inout ViewState InOutView);\n";
 
-	// reload the global shaders
-	if (bReloadGlobalShaders)
+	// GetPrimaryView definition
+	Result += "ViewState GetPrimaryView()\n";
+	Result += "{\n";
+	Result += "\tViewState Result;\n";
+	for (int32 MemberIndex = 0; MemberIndex < StructMembersInstanced.Num(); ++MemberIndex)
 	{
-		DoFlushIfNecessary();
-
-		// Some platforms rely on global shaders to be created to implement basic RHI functionality
-		TGuardValue<int32> Guard(GCreateShadersOnLoad, 1);
-		CompileGlobalShaderMap(true);
+		const FShaderParametersMetadata::FMember& Member = StructMembersView[MemberIndex];
+		Result += FString::Printf(TEXT("\tResult.%s = View.%s;\n"), Member.GetName(), Member.GetName());
 	}
+	Result += "\tFinalizeViewState(Result);\n";
+	Result += "\treturn Result;\n";
+	Result += "}\n";
 
-	// load all the mesh material shaders if any were sent back
-	if (MeshMaterialMaps.Num() > 0)
+	// GetInstancedView definition
+	Result += "#if (INSTANCED_STEREO || MOBILE_MULTI_VIEW)\n";
+	Result += "ViewState GetInstancedView(uint ViewIndex)\n";
+	Result += "{\n";
+	Result += "\tViewState Result;\n";
+	for (int32 MemberIndex = 0; MemberIndex < StructMembersInstanced.Num(); ++MemberIndex)
 	{
-		DoFlushIfNecessary();
-
-		// parse the shaders
-		FMemoryReader MemoryReader(MeshMaterialMaps, true);
-		FNameAsStringProxyArchive Ar(MemoryReader);
+		const FShaderParametersMetadata::FMember& ViewMember = StructMembersView[MemberIndex];
+		const FShaderParametersMetadata::FMember& InstancedViewMember = StructMembersInstanced[MemberIndex];
 
-		TArray<UMaterialInterface*> LoadedMaterials;
-		FMaterialShaderMap::LoadForRemoteRecompile(Ar, GMaxRHIShaderPlatform, LoadedMaterials);
+		FString ViewMemberTypeName;
+		ViewMember.GenerateShaderParameterType(ViewMemberTypeName, ShaderPlatform);
 
-		// Only update materials if we need to.
-		if (LoadedMaterials.Num())
+		// this code avoids an assumption that instanced buffer only supports 2 views, to be future-proof
+		if (ViewMember.GetNumElements() >= 1 && (InstancedViewMember.GetNumElements() >= 2 * ViewMember.GetNumElements()))
 		{
-			// this will stop the rendering thread, and reattach components, in the destructor
-			FMaterialUpdateContext UpdateContext(GODSCMaterialUpdateFlags);
-
-			// gather the shader maps to reattach
-			for (UMaterialInterface* Material : LoadedMaterials)
+			// if View has an array (even 1-sized) for this index, and InstancedView has Nx (N>=2) the element count -> per-view array
+			// Result.TranslucencyLightingVolumeMin[0] = (float4) InstancedView_TranslucencyLightingVolumeMin[ViewIndex * 2 + 0];
+			checkf((InstancedViewMember.GetNumElements() % ViewMember.GetNumElements()) == 0, TEXT("Per-view arrays are expected to be stored in an array that is an exact multiple of the original array."));
+			for (uint32 ElementIndex = 0; ElementIndex < ViewMember.GetNumElements(); ElementIndex++)
 			{
-				// ~FMaterialUpdateContext takes care of calling RecacheUniformExpressions on all MaterialInstances, no need to call it twice
-				if (Cast<UMaterialInstance>(Material) == nullptr)
-				{
-					Material->RecacheUniformExpressions(true);
-				}
-
-				UpdateContext.AddMaterialInterface(Material);
+				Result += FString::Printf(TEXT("\tResult.%s[%u] = (%s) InstancedView.%s[ViewIndex * %u + %u];\n"),
+					ViewMember.GetName(), ElementIndex, *ViewMemberTypeName, InstancedViewMember.GetName(), ViewMember.GetNumElements(), ElementIndex);
 			}
 		}
+		else if (InstancedViewMember.GetNumElements() > 1 && ViewMember.GetNumElements() == 0)
+		{
+			// if View has a scalar field for this index, and InstancedView has an array with >1 elements -> per-view scalar
+			// 	Result.TranslatedWorldToClip = (float4x4) InstancedView_TranslatedWorldToClip[ViewIndex];
+			Result += FString::Printf(TEXT("\tResult.%s = (%s) InstancedView.%s[ViewIndex];\n"),
+				ViewMember.GetName(), *ViewMemberTypeName, InstancedViewMember.GetName());
+		}
+		else if (InstancedViewMember.GetNumElements() == ViewMember.GetNumElements())
+		{
+			// if View has the same number of elements for this index as InstancedView, it's backed by a view-dependent array, assume a view-independent field
+			// 	Result.TemporalAAParams = InstancedView_TemporalAAParams;
+			Result += FString::Printf(TEXT("\tResult.%s = InstancedView.%s;\n"),
+				ViewMember.GetName(), InstancedViewMember.GetName());
+		}
+		else
+		{
+			// something unexpected, better crash now rather than generate wrong shader code and poison DDC 
+			UE_LOG(LogShaderCompilers, Fatal, TEXT("Don't know how to copy View buffers' field %s (NumElements=%d) from InstancedView field %s (NumElements=%d)"),
+				ViewMember.GetName(), ViewMember.GetNumElements(), InstancedViewMember.GetName(), InstancedViewMember.GetNumElements()
+				);
+		}
 	}
+	Result += "\tFinalizeViewState(Result);\n";
+	Result += "\treturn Result;\n";
+	Result += "}\n";
+	Result += "#endif\n";
+}
 
-	// load all the global shaders if any were sent back
-	if (GlobalShaderMap.Num() > 0)
-	{
-		DoFlushIfNecessary();
+void ValidateShaderFilePath(const FString& VirtualShaderFilePath, const FString& VirtualSourceFilePath)
+{
+	check(CheckVirtualShaderFilePath(VirtualShaderFilePath));
 
-		// parse the shaders
-		FMemoryReader MemoryReader(GlobalShaderMap, true);
-		FNameAsStringProxyArchive Ar(MemoryReader);
+	checkf(VirtualShaderFilePath.Contains(TEXT("/Generated/")),
+		TEXT("Incorrect virtual shader path for generated file '%s': Generated files must be located under an "
+				"non existing 'Generated' directory, for instance: /Engine/Generated/ or /Plugin/FooBar/Generated/."),
+		*VirtualShaderFilePath);
 
-		LoadGlobalShadersForRemoteRecompile(Ar, GMaxRHIShaderPlatform);
-	}
+	checkf(VirtualShaderFilePath == VirtualSourceFilePath || FPaths::GetExtension(VirtualShaderFilePath) == TEXT("ush"),
+		TEXT("Incorrect virtual shader path extension for generated file '%s': Generated file must either be the "
+				"USF to compile, or a USH file to be included."),
+		*VirtualShaderFilePath);
 }
 
-/**
-* Forces a recompile of the global shaders.
-*/
-void RecompileGlobalShaders()
-{
-#if WITH_EDITOR
-	if (!FPlatformProperties::RequiresCookedData())
-	{
-		// Flush pending accesses to the existing global shaders.
-		FlushRenderingCommands();
+/** Lock for the storage of instanced stereo code. */
+FCriticalSection GCachedGeneratedInstancedStereoCodeLock;
 
-		UMaterialInterface::IterateOverActiveFeatureLevels([&](ERHIFeatureLevel::Type InFeatureLevel) {
-			auto ShaderPlatform = GShaderPlatformForFeatureLevel[InFeatureLevel];
-			GetGlobalShaderMap(ShaderPlatform)->Empty();
-			VerifyGlobalShaders(ShaderPlatform, nullptr, false);
-		});
+/** Storage for instanced stereo code so it is not generated every time we compile a shader. */
+TMap<EShaderPlatform, FThreadSafeSharedAnsiStringPtr> GCachedGeneratedInstancedStereoCode;
 
-		GShaderCompilingManager->ProcessAsyncResults(false, true);
-	}
-#endif // WITH_EDITOR
+void GlobalBeginCompileShader(
+	const FString& DebugGroupName,
+	const FVertexFactoryType* VFType,
+	const FShaderType* ShaderType,
+	const FShaderPipelineType* ShaderPipelineType,
+	int32 PermutationId,
+	const TCHAR* SourceFilename,
+	const TCHAR* FunctionName,
+	FShaderTarget Target,
+	FShaderCompilerInput& Input,
+	bool bAllowDevelopmentShaderCompile,
+	const FString& DebugDescription,
+	const FString& DebugExtension
+)
+{
+	GlobalBeginCompileShader(
+		DebugGroupName,
+		VFType,
+		ShaderType,
+		ShaderPipelineType,
+		PermutationId,
+		SourceFilename,
+		FunctionName,
+		Target,
+		Input,
+		bAllowDevelopmentShaderCompile,
+		*DebugDescription,
+		*DebugExtension
+		);
 }
 
-void GetOutdatedShaderTypes(TArray<const FShaderType*>& OutdatedShaderTypes, TArray<const FShaderPipelineType*>& OutdatedShaderPipelineTypes, TArray<const FVertexFactoryType*>& OutdatedFactoryTypes)
+namespace
 {
-#if WITH_EDITOR
-	for (int PlatformIndex = 0; PlatformIndex < SP_NumPlatforms; ++PlatformIndex)
-	{
-		const FGlobalShaderMap* ShaderMap = GGlobalShaderMap[PlatformIndex];
-		if (ShaderMap)
-		{
-			ShaderMap->GetOutdatedTypes(OutdatedShaderTypes, OutdatedShaderPipelineTypes, OutdatedFactoryTypes);
-		}
-	}
-
-	FMaterialShaderMap::GetAllOutdatedTypes(OutdatedShaderTypes, OutdatedShaderPipelineTypes, OutdatedFactoryTypes);
-
-	for (int32 TypeIndex = 0; TypeIndex < OutdatedShaderTypes.Num(); TypeIndex++)
-	{
-		UE_LOG(LogShaders, Warning, TEXT("		Recompiling %s"), OutdatedShaderTypes[TypeIndex]->GetName());
-	}
-	for (int32 TypeIndex = 0; TypeIndex < OutdatedShaderPipelineTypes.Num(); TypeIndex++)
-	{
-		UE_LOG(LogShaders, Warning, TEXT("		Recompiling %s"), OutdatedShaderPipelineTypes[TypeIndex]->GetName());
-	}
-	for (int32 TypeIndex = 0; TypeIndex < OutdatedFactoryTypes.Num(); TypeIndex++)
+	bool ShaderFrequencyNeedsInstancedStereoMods(const FShaderType* ShaderType)
 	{
-		UE_LOG(LogShaders, Warning, TEXT("		Recompiling %s"), OutdatedFactoryTypes[TypeIndex]->GetName());
+		return !(IsRayTracingShaderFrequency(ShaderType->GetFrequency()));
 	}
-#endif // WITH_EDITOR
 }
 
-bool RecompileShaders(const TCHAR* Cmd, FOutputDevice& Ar)
+static bool IsSubstrateSupportForShaderPipeline(const FShaderCompilerInput& Input)
 {
-	// if this platform can't compile shaders, then we try to send a message to a file/cooker server
-	if (FPlatformProperties::RequiresCookedData())
-	{
-#if WITH_ODSC
-		TArray<FString> MaterialsToLoad;
-		FString ShaderTypesToLoad;
-		FString RequestedMaterialName;
-		ODSCRecompileCommand CommandType = ParseRecompileCommandString(Cmd, MaterialsToLoad, ShaderTypesToLoad, RequestedMaterialName);
-
-		ERHIFeatureLevel::Type TargetFeatureLevel = GetMaxSupportedFeatureLevel(GMaxRHIShaderPlatform);
-		const EMaterialQualityLevel::Type ActiveQualityLevel = GetCachedScalabilityCVars().MaterialQualityLevel;
-		GODSCManager->AddThreadedRequest(MaterialsToLoad, ShaderTypesToLoad, GMaxRHIShaderPlatform, TargetFeatureLevel, ActiveQualityLevel, CommandType, RequestedMaterialName);
-#endif
-		return true;
-	}
+	// Substrate requires HLSL2021 which must be cross-compiled for D3D11 to be consumed by FXC compiler.
+	// This cross-compilation toolchain does not support geometry shaders.
+	bool bPipelineContainsGeometryShader = false;
+	Input.Environment.GetCompileArgument(TEXT("PIPELINE_CONTAINS_GEOMETRYSHADER"), bPipelineContainsGeometryShader);
+	const bool bCanRHICompileHlsl2021GeometryShaders = !(GetMaxSupportedFeatureLevel((EShaderPlatform)Input.Target.Platform) == ERHIFeatureLevel::SM5);
+	const bool bIsSubstrateSupportedForPipeline = !bPipelineContainsGeometryShader || bCanRHICompileHlsl2021GeometryShaders;
+	return bIsSubstrateSupportedForPipeline;
+}
 
-#if WITH_EDITOR
-	FString FlagStr(FParse::Token(Cmd, 0));
-	if (FlagStr.Len() > 0)
-	{
-		TRACE_CPUPROFILER_EVENT_SCOPE(RecompileShaders);
-		GWarn->BeginSlowTask( NSLOCTEXT("ShaderCompilingManager", "BeginRecompilingShadersTask", "Recompiling shaders"), true );
+void GlobalBeginCompileShader(
+	const FString& DebugGroupName,
+	const FVertexFactoryType* VFType,
+	const FShaderType* ShaderType,
+	const FShaderPipelineType* ShaderPipelineType,
+	int32 PermutationId,
+	const TCHAR* SourceFilename,
+	const TCHAR* FunctionName,
+	FShaderTarget Target,
+	FShaderCompilerInput& Input,
+	bool bAllowDevelopmentShaderCompile,
+	const TCHAR* DebugDescription,
+	const TCHAR* DebugExtension
+	)
+{
+	LLM_SCOPE_BYTAG(ShaderCompiler);
 
-		// Flush the shader file cache so that any changes to shader source files will be detected
-		FlushShaderFileCache();
-		FlushRenderingCommands();
+	TRACE_CPUPROFILER_EVENT_SCOPE(GlobalBeginCompileShader);
+	COOK_STAT(ShaderCompilerCookStats::GlobalBeginCompileShaderCalls++);
+	COOK_STAT(FScopedDurationTimer DurationTimer(ShaderCompilerCookStats::GlobalBeginCompileShaderTimeSec));
 
-		if (FCString::Stricmp(*FlagStr,TEXT("Changed")) == 0)
-		{
-			TArray<const FShaderType*> OutdatedShaderTypes;
-			TArray<const FVertexFactoryType*> OutdatedFactoryTypes;
-			TArray<const FShaderPipelineType*> OutdatedShaderPipelineTypes;
-			{
-				FRecompileShadersTimer SearchTimer(TEXT("Searching for changed files"));
-				GetOutdatedShaderTypes(OutdatedShaderTypes, OutdatedShaderPipelineTypes, OutdatedFactoryTypes);
-			}
+	const EShaderPlatform ShaderPlatform = EShaderPlatform(Target.Platform);
+	const FName ShaderFormatName = LegacyShaderPlatformToShaderFormat(ShaderPlatform);
 
-			if (OutdatedShaderPipelineTypes.Num() > 0 || OutdatedShaderTypes.Num() > 0 || OutdatedFactoryTypes.Num() > 0)
-			{
-				FRecompileShadersTimer TestTimer(TEXT("RecompileShaders Changed"));
+	ITargetPlatform* TargetPlatform = GetTargetPlatformManager()->FindTargetPlatformWithSupport(TEXT("ShaderFormat"), ShaderFormatName);
 
-				UpdateReferencedUniformBufferNames(OutdatedShaderTypes, OutdatedFactoryTypes, OutdatedShaderPipelineTypes);
+	FShaderCompileUtilities::GenerateBrdfHeaders(ShaderPlatform);
 
-				// Kick off global shader recompiles
-				UMaterialInterface::IterateOverActiveFeatureLevels([&](ERHIFeatureLevel::Type InFeatureLevel) {
-					auto ShaderPlatform = GShaderPlatformForFeatureLevel[InFeatureLevel];
-					BeginRecompileGlobalShaders(OutdatedShaderTypes, OutdatedShaderPipelineTypes, ShaderPlatform);
-					// Block on global shader compilation. Do this for each feature level/platform compiled as otherwise global shader compile job IDs collide.
-					FinishRecompileGlobalShaders();
-				});
+	// NOTE:  Input.bCompilingForShaderPipeline is initialized by the constructor for single versus pipeline jobs, do not initialize again here!
 
-				// Kick off material shader recompiles
-				UMaterialInterface::IterateOverActiveFeatureLevels([&](ERHIFeatureLevel::Type InFeatureLevel) {
-					auto ShaderPlatform = GShaderPlatformForFeatureLevel[InFeatureLevel];
-					UMaterial::UpdateMaterialShaders(OutdatedShaderTypes, OutdatedShaderPipelineTypes, OutdatedFactoryTypes, ShaderPlatform);
-				});
+	Input.Target = Target;
+	Input.ShaderPlatformName = FDataDrivenShaderPlatformInfo::GetName(ShaderPlatform);
+	Input.ShaderFormat = ShaderFormatName;
+	Input.SupportedHardwareMask = TargetPlatform ? TargetPlatform->GetSupportedHardwareMask() : 0;
+	Input.CompressionFormat = GetShaderCompressionFormat();
+	GetShaderCompressionOodleSettings(Input.OodleCompressor, Input.OodleLevel);
+	Input.VirtualSourceFilePath = SourceFilename;
+	Input.EntryPointName = FunctionName;
+	Input.bIncludeUsedOutputs = false;
+	Input.DumpDebugInfoRootPath = GShaderCompilingManager->GetAbsoluteShaderDebugInfoDirectory() / Input.ShaderPlatformName.ToString();
+	Input.DebugInfoFlags = GShaderCompilingManager->GetDumpShaderDebugInfoFlags();
+	// asset material name or "Global"
+	Input.DebugGroupName = DebugGroupName;
+	Input.DebugDescription = DebugDescription;
+	Input.DebugExtension = DebugExtension;
+	Input.RootParametersStructure = ShaderType->GetRootParametersMetadata();
+	Input.ShaderName = ShaderType->GetName();
 
-				GWarn->StatusUpdate(0, 1, NSLOCTEXT("ShaderCompilingManager", "CompilingGlobalShaderStatus", "Compiling global shaders..."));
-			}
-			else
+	if (GDumpShaderDebugInfoBindless)
+	{
+		auto GetBindlessString = [](ERHIBindlessConfiguration InConfig)
 			{
-				UE_LOG(LogShaderCompilers, Warning, TEXT("No Shader changes found."));
-			}
-		}
-		else if (FCString::Stricmp(*FlagStr, TEXT("Global")) == 0)
-		{
-			FRecompileShadersTimer TestTimer(TEXT("RecompileShaders Global"));
-			RecompileGlobalShaders();
-		}
-		else if (FCString::Stricmp(*FlagStr, TEXT("Material")) == 0)
-		{
-			FString RequestedMaterialName(FParse::Token(Cmd, 0));
-			FRecompileShadersTimer TestTimer(FString::Printf(TEXT("Recompile Material %s"), *RequestedMaterialName));
+				switch (InConfig)
+				{
+				default:
+				case ERHIBindlessConfiguration::Disabled:          return TEXT("Off");
+				case ERHIBindlessConfiguration::AllShaders:        return TEXT("On");
+				case ERHIBindlessConfiguration::RayTracingShaders: return TEXT("RT");
+				}
+			};
 
-			ITargetPlatformManagerModule& TPM = GetTargetPlatformManagerRef();
-			FString TargetPlatformName(FParse::Token(Cmd, 0));
-			const ITargetPlatform* TargetPlatform = nullptr;
-			if (TargetPlatformName.Len() > 0)
-			{
-				TargetPlatform = TPM.FindTargetPlatform(TargetPlatformName);
-			}
+		FStringBuilderBase Builder;
+		Builder.Append("BindlessR").Append(GetBindlessString(UE::ShaderCompiler::GetBindlessResourcesConfiguration(ShaderFormatName)));
+		Builder.Append("_");
+		Builder.Append("BindlessS").Append(GetBindlessString(UE::ShaderCompiler::GetBindlessSamplersConfiguration(ShaderFormatName)));
 
-			bool bMaterialFound = false;
-			for( TObjectIterator<UMaterialInterface> It; It; ++It )
-			{
-				UMaterialInterface* Material = *It;
-				if( Material && Material->GetName() == RequestedMaterialName)
-				{
-					bMaterialFound = true;
+		Input.DebugGroupName = Builder.ToString() / Input.DebugGroupName;
+	}
 
-					// <Pre/Post>EditChange will force a re-creation of the resource,
-					// in turn recompiling the shader.
-					if (TargetPlatform)
-					{
-						Material->BeginCacheForCookedPlatformData(TargetPlatform);
-						while (!Material->IsCachedCookedPlatformDataLoaded(TargetPlatform))
-						{
-							FPlatformProcess::Sleep(0.1f);
-							GShaderCompilingManager->ProcessAsyncResults(false, false);
-						}
-						Material->ClearCachedCookedPlatformData(TargetPlatform);
-					}
-					else
-					{
-						Material->PreEditChange(nullptr);
-						Material->PostEditChange();
-					}
+	// Verify FShaderCompilerInput's file paths are consistent. 
+	#if DO_CHECK
+		check(CheckVirtualShaderFilePath(Input.VirtualSourceFilePath));
 
-					break;
-				}
-			}
+		checkf(FPaths::GetExtension(Input.VirtualSourceFilePath) == TEXT("usf"),
+			TEXT("Incorrect virtual shader path extension for shader file to compile '%s': Only .usf files should be "
+				 "compiled. .ush file are meant to be included only."),
+			*Input.VirtualSourceFilePath);
 
-			if (!bMaterialFound)
-			{
-				TestTimer.Stop(false);
-				UE_LOG(LogShaderCompilers, Warning, TEXT("Couldn't find Material %s!"), *RequestedMaterialName);
-			}
+		for (const auto& Entry : Input.Environment.IncludeVirtualPathToContentsMap)
+		{
+			ValidateShaderFilePath(Entry.Key, Input.VirtualSourceFilePath);
 		}
-		else if (FCString::Stricmp(*FlagStr, TEXT("All")) == 0)
+
+		for (const auto& Entry : Input.Environment.IncludeVirtualPathToSharedContentsMap)
 		{
-			FRecompileShadersTimer TestTimer(TEXT("RecompileShaders"));
-			RecompileGlobalShaders();
+			ValidateShaderFilePath(Entry.Key, Input.VirtualSourceFilePath);
+		}
+	#endif
 
-			FMaterialUpdateContext UpdateContext(0);
-			for( TObjectIterator<UMaterial> It; It; ++It )
+	if (ShaderPipelineType)
+	{
+		Input.DebugGroupName = Input.DebugGroupName / ShaderPipelineType->GetName();
+	}
+	
+	if (VFType)
+	{
+		FString VFName = VFType->GetName();
+		if (GDumpShaderDebugInfoShort)
+		{
+			// Shorten vertex factory name
+			if (VFName[0] == TCHAR('F') || VFName[0] == TCHAR('T'))
 			{
-				UMaterial* Material = *It;
-				if( Material )
-				{
-					UE_LOG(LogShaderCompilers, Log, TEXT("recompiling [%s]"),*Material->GetFullName());
-					UpdateContext.AddMaterial(Material);
-
-					// <Pre/Post>EditChange will force a re-creation of the resource,
-					// in turn recompiling the shader.
-					Material->PreEditChange(nullptr);
-					Material->PostEditChange();
-				}
+				VFName.RemoveAt(0);
 			}
+			VFName.ReplaceInline(TEXT("VertexFactory"), TEXT("VF"));
+			VFName.ReplaceInline(TEXT("GPUSkinAPEXCloth"), TEXT("APEX"));
+			VFName.ReplaceInline(TEXT("true"), TEXT("_1"));
+			VFName.ReplaceInline(TEXT("false"), TEXT("_0"));
 		}
-		else if (FCString::Stricmp(*FlagStr, TEXT("listtypes")) == 0)
-		{
-			ListAllShaderTypes();
-		}
-		else
+		Input.DebugGroupName = Input.DebugGroupName / VFName;
+	}
+	
+	{
+		FString ShaderTypeName = ShaderType->GetName();
+		if (GDumpShaderDebugInfoShort)
 		{
-			TArray<const FShaderType*> ShaderTypes = FShaderType::GetShaderTypesByFilename(*FlagStr);
-			TArray<const FShaderPipelineType*> ShaderPipelineTypes = FShaderPipelineType::GetShaderPipelineTypesByFilename(*FlagStr);
-			if (ShaderTypes.Num() > 0 || ShaderPipelineTypes.Num() > 0)
+			// Shorten known types
+			if (ShaderTypeName[0] == TCHAR('F') || ShaderTypeName[0] == TCHAR('T'))
 			{
-				FRecompileShadersTimer TestTimer(TEXT("RecompileShaders SingleShader"));
-				
-				UpdateReferencedUniformBufferNames(ShaderTypes, {}, ShaderPipelineTypes);
-
-				UMaterialInterface::IterateOverActiveFeatureLevels([&](ERHIFeatureLevel::Type InFeatureLevel) {
-					auto ShaderPlatform = GShaderPlatformForFeatureLevel[InFeatureLevel];
-					BeginRecompileGlobalShaders(ShaderTypes, ShaderPipelineTypes, ShaderPlatform);
-					FinishRecompileGlobalShaders();
-				});
+				ShaderTypeName.RemoveAt(0);
 			}
 		}
+		Input.DebugGroupName = Input.DebugGroupName / ShaderTypeName / FString::Printf(TEXT("%i"), PermutationId);
+		
+		if (GDumpShaderDebugInfoShort)
+		{
+			Input.DebugGroupName.ReplaceInline(TEXT("BasePass"), TEXT("BP"));
+			Input.DebugGroupName.ReplaceInline(TEXT("ForForward"), TEXT("Fwd"));
+			Input.DebugGroupName.ReplaceInline(TEXT("Shadow"), TEXT("Shdw"));
+			Input.DebugGroupName.ReplaceInline(TEXT("LightMap"), TEXT("LM"));
+			Input.DebugGroupName.ReplaceInline(TEXT("EHeightFogFeature==E_"), TEXT(""));
+			Input.DebugGroupName.ReplaceInline(TEXT("Capsule"), TEXT("Caps"));
+			Input.DebugGroupName.ReplaceInline(TEXT("Movable"), TEXT("Mov"));
+			Input.DebugGroupName.ReplaceInline(TEXT("Culling"), TEXT("Cull"));
+			Input.DebugGroupName.ReplaceInline(TEXT("Atmospheric"), TEXT("Atm"));
+			Input.DebugGroupName.ReplaceInline(TEXT("Atmosphere"), TEXT("Atm"));
+			Input.DebugGroupName.ReplaceInline(TEXT("Exponential"), TEXT("Exp"));
+			Input.DebugGroupName.ReplaceInline(TEXT("Ambient"), TEXT("Amb"));
+			Input.DebugGroupName.ReplaceInline(TEXT("Perspective"), TEXT("Persp"));
+			Input.DebugGroupName.ReplaceInline(TEXT("Occlusion"), TEXT("Occ"));
+			Input.DebugGroupName.ReplaceInline(TEXT("Position"), TEXT("Pos"));
+			Input.DebugGroupName.ReplaceInline(TEXT("Skylight"), TEXT("Sky"));
+			Input.DebugGroupName.ReplaceInline(TEXT("LightingPolicy"), TEXT("LP"));
+			Input.DebugGroupName.ReplaceInline(TEXT("TranslucentLighting"), TEXT("TranslLight"));
+			Input.DebugGroupName.ReplaceInline(TEXT("Translucency"), TEXT("Transl"));
+			Input.DebugGroupName.ReplaceInline(TEXT("DistanceField"), TEXT("DistFiel"));
+			Input.DebugGroupName.ReplaceInline(TEXT("Indirect"), TEXT("Ind"));
+			Input.DebugGroupName.ReplaceInline(TEXT("Cached"), TEXT("Cach"));
+			Input.DebugGroupName.ReplaceInline(TEXT("Inject"), TEXT("Inj"));
+			Input.DebugGroupName.ReplaceInline(TEXT("Visualization"), TEXT("Viz"));
+			Input.DebugGroupName.ReplaceInline(TEXT("Instanced"), TEXT("Inst"));
+			Input.DebugGroupName.ReplaceInline(TEXT("Evaluate"), TEXT("Eval"));
+			Input.DebugGroupName.ReplaceInline(TEXT("Landscape"), TEXT("Land"));
+			Input.DebugGroupName.ReplaceInline(TEXT("Dynamic"), TEXT("Dyn"));
+			Input.DebugGroupName.ReplaceInline(TEXT("Vertex"), TEXT("Vtx"));
+			Input.DebugGroupName.ReplaceInline(TEXT("Output"), TEXT("Out"));
+			Input.DebugGroupName.ReplaceInline(TEXT("Directional"), TEXT("Dir"));
+			Input.DebugGroupName.ReplaceInline(TEXT("Irradiance"), TEXT("Irr"));
+			Input.DebugGroupName.ReplaceInline(TEXT("Deferred"), TEXT("Def"));
+			Input.DebugGroupName.ReplaceInline(TEXT("true"), TEXT("_1"));
+			Input.DebugGroupName.ReplaceInline(TEXT("false"), TEXT("_0"));
+			Input.DebugGroupName.ReplaceInline(TEXT("PROPAGATE_AO"), TEXT("AO"));
+			Input.DebugGroupName.ReplaceInline(TEXT("PROPAGATE_SECONDARY_OCCLUSION"), TEXT("SEC_OCC"));
+			Input.DebugGroupName.ReplaceInline(TEXT("PROPAGATE_MULTIPLE_BOUNCES"), TEXT("MULT_BOUNC"));
+			Input.DebugGroupName.ReplaceInline(TEXT("LOCAL_LIGHTS_DISABLED"), TEXT("NoLL"));
+			Input.DebugGroupName.ReplaceInline(TEXT("LOCAL_LIGHTS_ENABLED"), TEXT("LL"));
+			Input.DebugGroupName.ReplaceInline(TEXT("LOCAL_LIGHTS_PREPASS_ENABLED"), TEXT("LLPP"));
+			Input.DebugGroupName.ReplaceInline(TEXT("PostProcess"), TEXT("Post"));
+			Input.DebugGroupName.ReplaceInline(TEXT("AntiAliasing"), TEXT("AA"));
+			Input.DebugGroupName.ReplaceInline(TEXT("Mobile"), TEXT("Mob"));
+			Input.DebugGroupName.ReplaceInline(TEXT("Linear"), TEXT("Lin"));
+			Input.DebugGroupName.ReplaceInline(TEXT("INT32_MAX"), TEXT("IMAX"));
+			Input.DebugGroupName.ReplaceInline(TEXT("Policy"), TEXT("Pol"));
+			Input.DebugGroupName.ReplaceInline(TEXT("EAtmRenderFlag==E_"), TEXT(""));
+		}
+	}
 
-		GWarn->EndSlowTask();
-
-		return true;
+	// Setup the debug info path if requested, or if this is a global shader and shader development mode is enabled
+	Input.DumpDebugInfoPath.Empty();
+	if (GShaderCompilingManager->GetDumpShaderDebugInfo() == FShaderCompilingManager::EDumpShaderDebugInfo::Always)
+	{
+		Input.DumpDebugInfoPath = GShaderCompilingManager->CreateShaderDebugInfoPath(Input);
 	}
 
-	UE_LOG(LogShaderCompilers, Warning, TEXT("Invalid parameter. \n"
-											 "Options are: \n"
-											 "    'Changed'             Recompile just the shaders that have source file changes.\n"
-											 "    'Global'              Recompile just the global shaders.\n"
-											 "    'Material [name]'     Recompile all the shaders for a single material.\n"
-											 "    'Listtypes'           List all the shader type and vertex factory type class names and their source file path.  Can be used to find shader file names to be used with `recompileshaders [shaderfilename]`.\n"
-											 "    'All'                 Recompile all materials.\n"
-											 "    [shaderfilename]      Compile the single shader associated with a specific filename.\n"
-											 ));
-#endif // WITH_EDITOR
+	// Add the appropriate definitions for the shader frequency.
+	{
+		SET_SHADER_DEFINE(Input.Environment, PIXELSHADER,				Target.Frequency == SF_Pixel);
+		SET_SHADER_DEFINE(Input.Environment, VERTEXSHADER,				Target.Frequency == SF_Vertex);
+		SET_SHADER_DEFINE(Input.Environment, MESHSHADER,				Target.Frequency == SF_Mesh);
+		SET_SHADER_DEFINE(Input.Environment, AMPLIFICATIONSHADER,		Target.Frequency == SF_Amplification);
+		SET_SHADER_DEFINE(Input.Environment, GEOMETRYSHADER,			Target.Frequency == SF_Geometry);
+		SET_SHADER_DEFINE(Input.Environment, COMPUTESHADER,				Target.Frequency == SF_Compute);
+		SET_SHADER_DEFINE(Input.Environment, RAYCALLABLESHADER,			Target.Frequency == SF_RayCallable);
+		SET_SHADER_DEFINE(Input.Environment, RAYHITGROUPSHADER,			Target.Frequency == SF_RayHitGroup);
+		SET_SHADER_DEFINE(Input.Environment, RAYGENSHADER,				Target.Frequency == SF_RayGen);
+		SET_SHADER_DEFINE(Input.Environment, RAYMISSSHADER,				Target.Frequency == SF_RayMiss);
+		SET_SHADER_DEFINE(Input.Environment, WORKGRAPHROOTSHADER,		Target.Frequency == SF_WorkGraphRoot);
+		SET_SHADER_DEFINE(Input.Environment, WORKGRAPHCOMPUTESHADER,	Target.Frequency == SF_WorkGraphComputeNode);
+	}
 
-	return true;
-}
+	SET_SHADER_DEFINE(Input.Environment, FORWARD_SHADING_FORCES_SKYLIGHT_CUBEMAPS_BLENDING, ForwardShadingForcesSkyLightCubemapBlending(ShaderPlatform) ? 1 : 0);
 
-#if WITH_EDITOR
+	// Enables HLSL 2021
+	uint32 EnablesHLSL2021ByDefault = FDataDrivenShaderPlatformInfo::GetEnablesHLSL2021ByDefault(EShaderPlatform(Target.Platform));
+	if((EnablesHLSL2021ByDefault == uint32(1) && DebugGroupName == TEXT("Global")) ||
+		EnablesHLSL2021ByDefault == uint32(2) ||
+		Target.Frequency == SF_RayGen ||			// We want to make sure that function overloads follow c++ rules for FRayDesc.
+		Target.Frequency == SF_RayHitGroup)
+	{
+		Input.Environment.CompilerFlags.Add(CFLAG_HLSL2021);
+	}
 
-static void PrepareGlobalShaderCompileJob(EShaderPlatform Platform,
-	EShaderPermutationFlags PermutationFlags,
-	const FShaderPipelineType* ShaderPipeline,
-	FShaderCompileJob* NewJob)
-{
-	const FShaderCompileJobKey& Key = NewJob->Key;
-	const FGlobalShaderType* ShaderType = Key.ShaderType->AsGlobalShaderType();
+	// #defines get stripped out by the preprocessor without this. We can override with this
+	SET_SHADER_DEFINE(Input.Environment, COMPILER_DEFINE, TEXT("#define"));
 
-	FShaderCompilerEnvironment& ShaderEnvironment = NewJob->Input.Environment;
+	if (FSceneInterface::GetShadingPath(GetMaxSupportedFeatureLevel(ShaderPlatform)) == EShadingPath::Deferred)
+	{
+		SET_SHADER_DEFINE(Input.Environment, SHADING_PATH_DEFERRED, 1);
+	}
 
-	UE_LOG(LogShaders, Verbose, TEXT("	%s (permutation %d)"), ShaderType->GetName(), Key.PermutationId);
-	COOK_STAT(GlobalShaderCookStats::ShadersCompiled++);
+	const bool bUsingMobileRenderer = FSceneInterface::GetShadingPath(GetMaxSupportedFeatureLevel(ShaderPlatform)) == EShadingPath::Mobile;
+	if (bUsingMobileRenderer)
+	{
+		SET_SHADER_DEFINE_AND_COMPILE_ARGUMENT(Input.Environment, SHADING_PATH_MOBILE, true);
+		
+		const bool bMobileDeferredShading = IsMobileDeferredShadingEnabled((EShaderPlatform)Target.Platform);
+		SET_SHADER_DEFINE_AND_COMPILE_ARGUMENT(Input.Environment, MOBILE_DEFERRED_SHADING, bMobileDeferredShading);
 
-	// Allow the shader type to modify the compile environment.
-	ShaderType->SetupCompileEnvironment(Platform, Key.PermutationId, PermutationFlags, ShaderEnvironment);
+		if (bMobileDeferredShading)
+		{
+			bool bGLESDeferredShading = Target.Platform == SP_OPENGL_ES3_1_ANDROID;
+			SET_SHADER_DEFINE(Input.Environment, USE_GLES_FBF_DEFERRED, bGLESDeferredShading ? 1 : 0);
+			SET_SHADER_DEFINE(Input.Environment, MOBILE_EXTENDED_GBUFFER, MobileUsesExtenedGBuffer((EShaderPlatform)Target.Platform) ? 1 : 0);
+		}
 
-	static FString GlobalName(TEXT("Global"));
+		SET_SHADER_DEFINE(Input.Environment, USE_SCENE_DEPTH_AUX, MobileRequiresSceneDepthAux(ShaderPlatform) ? 1 : 0);
 
-	NewJob->bErrorsAreLikelyToBeCode = true;
-	NewJob->bIsGlobalShader = true;
-	NewJob->bIsDefaultMaterial = false;
+		static FShaderPlatformCachedIniValue<bool> EnableCullBeforeFetchIniValue(TEXT("r.CullBeforeFetch"));
+		if (EnableCullBeforeFetchIniValue.Get((EShaderPlatform)Target.Platform) == 1)
+		{
+			Input.Environment.CompilerFlags.Add(CFLAG_CullBeforeFetch);
+		}
 
-	// Compile the shader environment passed in with the shader type's source code.
-	::GlobalBeginCompileShader(
-		GlobalName,
-		nullptr,
-		ShaderType,
-		ShaderPipeline,
-		Key.PermutationId,
-		ShaderType->GetShaderFilename(),
-		ShaderType->GetFunctionName(),
-		FShaderTarget(ShaderType->GetFrequency(), Platform),
-		NewJob->Input
-	);
-}
+		static FShaderPlatformCachedIniValue<bool> EnableWarpCullingIniValue(TEXT("r.WarpCulling"));
+		if (EnableWarpCullingIniValue.Get((EShaderPlatform)Target.Platform) == 1)
+		{
+			Input.Environment.CompilerFlags.Add(CFLAG_WarpCulling);
+		}
+	}
 
-void FGlobalShaderTypeCompiler::BeginCompileShader(const FGlobalShaderType* ShaderType, int32 PermutationId, EShaderPlatform Platform, EShaderPermutationFlags PermutationFlags, TArray<FShaderCommonCompileJobPtr>& NewJobs)
-{
-	// Global shaders are always high priority (often need to block on completion)
-	FShaderCompileJob* NewJob = GShaderCompilingManager->PrepareShaderCompileJob(GlobalShaderMapId, FShaderCompileJobKey(ShaderType, nullptr, PermutationId), EShaderCompileJobPriority::High);
-	if (NewJob)
+	if (RenderRectLightsAsSpotLights(GetMaxSupportedFeatureLevel(ShaderPlatform)))
 	{
-		PrepareGlobalShaderCompileJob(Platform, PermutationFlags, nullptr, NewJob);
-		NewJobs.Add(FShaderCommonCompileJobPtr(NewJob));
+		SET_SHADER_DEFINE(Input.Environment, RECT_LIGHT_AS_SPOTLIGHT, 1);
 	}
-}
-
-void FGlobalShaderTypeCompiler::BeginCompileShaderPipeline(EShaderPlatform Platform, EShaderPermutationFlags PermutationFlags, const FShaderPipelineType* ShaderPipeline, TArray<FShaderCommonCompileJobPtr>& NewJobs)
-{
-	check(ShaderPipeline);
-	UE_LOG(LogShaders, Verbose, TEXT("	Pipeline: %s"), ShaderPipeline->GetName());
 
-	// Add all the jobs as individual first, then add the dependencies into a pipeline job
-	FShaderPipelineCompileJob* NewPipelineJob = GShaderCompilingManager->PreparePipelineCompileJob(GlobalShaderMapId, FShaderPipelineCompileJobKey(ShaderPipeline, nullptr, kUniqueShaderPermutationId), EShaderCompileJobPriority::High);
-	if (NewPipelineJob)
+	if (ShaderPlatform == SP_VULKAN_ES3_1_ANDROID || ShaderPlatform == SP_VULKAN_SM5_ANDROID)
 	{
-		for (FShaderCompileJob* StageJob : NewPipelineJob->StageJobs)
+		bool bIsStripReflect = true;
+		GConfig->GetBool(TEXT("/Script/AndroidRuntimeSettings.AndroidRuntimeSettings"), TEXT("bStripShaderReflection"), bIsStripReflect, GEngineIni);
+		if (!bIsStripReflect)
 		{
-			PrepareGlobalShaderCompileJob(Platform, PermutationFlags, ShaderPipeline, StageJob);
+			Input.Environment.SetCompileArgument(TEXT("STRIP_REFLECT_ANDROID"), false);
 		}
-		NewJobs.Add(FShaderCommonCompileJobPtr(NewPipelineJob));
 	}
-}
 
-FShader* FGlobalShaderTypeCompiler::FinishCompileShader(const FGlobalShaderType* ShaderType, const FShaderCompileJob& CurrentJob, const FShaderPipelineType* ShaderPipelineType)
-{
-	FShader* Shader = nullptr;
-	if (CurrentJob.bSucceeded)
+	static const auto CVarViewHasTileOffsetData = IConsoleManager::Get().FindTConsoleVariableDataInt(TEXT("r.ViewHasTileOffsetData"));
+	const bool bViewHasTileOffsetData = CVarViewHasTileOffsetData->GetValueOnAnyThread() != 0;
+	SET_SHADER_DEFINE_AND_COMPILE_ARGUMENT(Input.Environment, VIEW_HAS_TILEOFFSET_DATA, bViewHasTileOffsetData);
+
+	static const auto CVarPrimitiveHasTileOffsetData = IConsoleManager::Get().FindTConsoleVariableDataInt(TEXT("r.PrimitiveHasTileOffsetData"));
+	const bool bPrimitiveHasTileOffsetData = CVarPrimitiveHasTileOffsetData->GetValueOnAnyThread() != 0;
+	SET_SHADER_DEFINE_AND_COMPILE_ARGUMENT(Input.Environment, PRIMITIVE_HAS_TILEOFFSET_DATA, bPrimitiveHasTileOffsetData);
+
+	// Set VR definitions
+	if (ShaderFrequencyNeedsInstancedStereoMods(ShaderType))
 	{
-		EShaderPlatform Platform = CurrentJob.Input.Target.GetPlatform();
-		FGlobalShaderMapSection* Section = GGlobalShaderMap[Platform]->FindOrAddSection(ShaderType);
+		const UE::StereoRenderUtils::FStereoShaderAspects Aspects(ShaderPlatform);
 
-		Section->GetResourceCode()->AddShaderCompilerOutput(CurrentJob.Output, CurrentJob.Key.ToString());
+		SET_SHADER_DEFINE_AND_COMPILE_ARGUMENT(Input.Environment, INSTANCED_STEREO, Aspects.IsInstancedStereoEnabled());
+		SET_SHADER_DEFINE_AND_COMPILE_ARGUMENT(Input.Environment, MULTI_VIEW, Aspects.IsInstancedMultiViewportEnabled());
+		SET_SHADER_DEFINE_AND_COMPILE_ARGUMENT(Input.Environment, MOBILE_MULTI_VIEW, Aspects.IsMobileMultiViewEnabled());
 
-		if (ShaderPipelineType && !ShaderPipelineType->ShouldOptimizeUnusedOutputs(CurrentJob.Input.Target.GetPlatform()))
+		// Throw a warning if we are silently disabling ISR due to missing platform support (but don't have MMV enabled).
+		static const auto CVarInstancedStereo = IConsoleManager::Get().FindTConsoleVariableDataInt(TEXT("vr.InstancedStereo"));
+		const bool bIsInstancedStereoEnabledInSettings = CVarInstancedStereo ? (CVarInstancedStereo->GetValueOnAnyThread() != 0) : false;
+		static const auto CVarMultiview = IConsoleManager::Get().FindTConsoleVariableDataInt(TEXT("vr.MobileMultiView"));
+		const bool bIsMultiviewEnabledInSettings = CVarMultiview ? (CVarMultiview->GetValueOnAnyThread() != 0) : false;
+		bool bWarningIssued = false;
+		// warn if ISR was enabled in settings, but aspects show that it's not enabled AND we don't use Mobile MultiView as an alternative
+		if (bIsInstancedStereoEnabledInSettings && !Aspects.IsInstancedStereoEnabled() && !(bIsMultiviewEnabledInSettings && Aspects.IsMobileMultiViewEnabled()) && !GShaderCompilingManager->AreWarningsSuppressed(ShaderPlatform))
 		{
-			// If sharing shaders in this pipeline, remove it from the type/id so it uses the one in the shared shadermap list
-			ShaderPipelineType = nullptr;
+			UE_LOG(LogShaderCompilers, Warning, TEXT("Instanced stereo rendering is not supported for %s shader platform."), *ShaderFormatName.ToString());
+			bWarningIssued = true;
 		}
-
-		// Create the global shader map hash
-		FSHAHash GlobalShaderMapHash;
+		// Warn if MMV was enabled in settings, but aspects show that it's not enabled AND we don't use Instanced Stereo as an alternative
+		if (bIsMultiviewEnabledInSettings && !Aspects.IsMobileMultiViewEnabled() && !(bIsInstancedStereoEnabledInSettings && Aspects.IsInstancedStereoEnabled()) && !GShaderCompilingManager->AreWarningsSuppressed(ShaderPlatform))
+		{
+			UE_LOG(LogShaderCompilers, Warning, TEXT("Multiview rendering is not supported for %s shader platform."), *ShaderFormatName.ToString());
+			bWarningIssued = true;
+		}
+		if (bWarningIssued)
 		{
-			FSHA1 HashState;
-			const TCHAR* GlobalShaderString = TEXT("GlobalShaderMap");
-			HashState.UpdateWithString(GlobalShaderString, FCString::Strlen(GlobalShaderString));
-			HashState.Final();
-			HashState.GetHash(&GlobalShaderMapHash.Hash[0]);
+			GShaderCompilingManager->SuppressWarnings(ShaderPlatform);
 		}
+	}
+	else
+	{
+		SET_SHADER_DEFINE_AND_COMPILE_ARGUMENT(Input.Environment, INSTANCED_STEREO, false);
+		SET_SHADER_DEFINE_AND_COMPILE_ARGUMENT(Input.Environment, MULTI_VIEW, 0);
+		SET_SHADER_DEFINE_AND_COMPILE_ARGUMENT(Input.Environment, MOBILE_MULTI_VIEW, false);
+	}
+
+	// Reserve space in maps to prevent reallocation and rehashing in AddUniformBufferIncludesToEnvironment -- plus one at the end is for GeneratedInstancedStereo.ush
+	const int32 UniformBufferReserveNum = Input.Environment.UniformBufferMap.Num() + ShaderType->GetReferencedUniformBuffers().Num() + (VFType ? VFType->GetReferencedUniformBuffers().Num() : 0) + 1;
+	Input.Environment.UniformBufferMap.Reserve(UniformBufferReserveNum);
+	Input.Environment.IncludeVirtualPathToSharedContentsMap.Reserve(UniformBufferReserveNum);
+
+	ShaderType->AddUniformBufferIncludesToEnvironment(Input.Environment, ShaderPlatform);
 
-		Shader = ShaderType->ConstructCompiled(FGlobalShaderType::CompiledShaderInitializerType(ShaderType, nullptr, CurrentJob.Key.PermutationId, CurrentJob.Output, GlobalShaderMapHash, ShaderPipelineType, nullptr));
-		CurrentJob.Output.ParameterMap.VerifyBindingsAreComplete(ShaderType->GetName(), CurrentJob.Output.Target, CurrentJob.Key.VFType);
+	if (VFType)
+	{
+		VFType->AddUniformBufferIncludesToEnvironment(Input.Environment, ShaderPlatform);
 	}
 
-	return Shader;
-}
+	// Add generated instanced stereo code (this code also generates ViewState, so needed not just for ISR)
+	{
+		// this function may be called on multiple threads, so protect the storage
+		FScopeLock GeneratedInstancedCodeLock(&GCachedGeneratedInstancedStereoCodeLock);
 
-#endif // WITH_EDITOR
+		FThreadSafeSharedAnsiStringPtr* Existing = GCachedGeneratedInstancedStereoCode.Find(ShaderPlatform);
+		FThreadSafeSharedAnsiStringPtr CachedCodePtr = Existing ? *Existing : nullptr;
+		if (!CachedCodePtr.IsValid())
+		{
+			FString CachedCode;
+			GenerateInstancedStereoCode(CachedCode, ShaderPlatform);
 
-#if WITH_EDITORONLY_DATA
-namespace ShaderCompilerUtil
-{
-	FOnGlobalShadersCompilation GOnGlobalShdersCompilationDelegate;
-}
+			CachedCodePtr = MakeShareable(new TArray<ANSICHAR>());
+			ShaderConvertAndStripComments(CachedCode, *CachedCodePtr);
 
-FOnGlobalShadersCompilation& GetOnGlobalShaderCompilation()
-{
-	return ShaderCompilerUtil::GOnGlobalShdersCompilationDelegate;
-}
-#endif // WITH_EDITORONLY_DATA
+			GCachedGeneratedInstancedStereoCode.Add(ShaderPlatform, CachedCodePtr);
+		}
 
-/**
-* Makes sure all global shaders are loaded and/or compiled for the passed in platform.
-* Note: if compilation is needed, this only kicks off the compile.
-*
-* @param	Platform	Platform to verify global shaders for
-*/
-void VerifyGlobalShaders(EShaderPlatform Platform, const ITargetPlatform* TargetPlatform, bool bLoadedFromCacheFile, const TArray<const FShaderType*>* OutdatedShaderTypes, const TArray<const FShaderPipelineType*>* OutdatedShaderPipelineTypes)
-{
-	SCOPED_LOADTIMER(VerifyGlobalShaders);
+		Input.Environment.IncludeVirtualPathToSharedContentsMap.Add(TEXT("/Engine/Generated/GeneratedInstancedStereo.ush"), CachedCodePtr);
+	}
 
-	check(IsInGameThread());
-	check(!FPlatformProperties::IsServerOnly());
-	check(GGlobalShaderMap[Platform]);
+	{
+		// Check if the compile environment explicitly wants to force optimization
+		const bool bForceOptimization = Input.Environment.CompilerFlags.Contains(CFLAG_ForceOptimization);
 
-	UE_LOG(LogMaterial, Verbose, TEXT("Verifying Global Shaders for %s (%s)"), *LegacyShaderPlatformToShaderFormat(Platform).ToString(), *ShaderCompiler::GetTargetPlatformName(TargetPlatform));
+		if (!bForceOptimization && !ShouldOptimizeShaders(ShaderFormatName))
+		{
+			Input.Environment.CompilerFlags.Add(CFLAG_Debug);
+		}
+	}
 
-	// Ensure that the global shader map contains all global shader types.
-	FGlobalShaderMap* GlobalShaderMap = GetGlobalShaderMap(Platform);
-	const bool bEmptyMap = GlobalShaderMap->IsEmpty();
-	if (bEmptyMap)
+	// Extra data (names, etc)
+	if (ShouldEnableExtraShaderData(ShaderFormatName))
 	{
-		UE_LOG(LogShaders, Log, TEXT("	Empty global shader map, recompiling all global shaders"));
+		Input.Environment.CompilerFlags.Add(CFLAG_ExtraShaderData);
 	}
 
-	FPlatformTypeLayoutParameters LayoutParams;
-	LayoutParams.InitializeForPlatform(TargetPlatform);
-	EShaderPermutationFlags PermutationFlags = GetShaderPermutationFlags(LayoutParams);
+	// Symbols
+	if (ShouldGenerateShaderSymbols(ShaderFormatName))
+	{
+		Input.Environment.CompilerFlags.Add(CFLAG_GenerateSymbols);
+	}
+	if (ShouldGenerateShaderSymbolsInfo(ShaderFormatName))
+	{
+		Input.Environment.CompilerFlags.Add(CFLAG_GenerateSymbolsInfo);
+	}
+
+	// Are symbols based on source or results
+	if (ShouldAllowUniqueShaderSymbols(ShaderFormatName))
+	{
+		Input.Environment.CompilerFlags.Add(CFLAG_AllowUniqueSymbols);
+	}
+
+	if (CVarShaderFastMath.GetValueOnAnyThread() == 0)
+	{
+		Input.Environment.CompilerFlags.Add(CFLAG_NoFastMath);
+	}
+    
+	if (bUsingMobileRenderer)
+    {
+        static IConsoleVariable* CVar = IConsoleManager::Get().FindConsoleVariable(TEXT("r.Mobile.FloatPrecisionMode"));
+        Input.Environment.FullPrecisionInPS |= CVar ? (CVar->GetInt() == EMobileFloatPrecisionMode::Full) : false;
+    }
+	
+	{
+		int32 FlowControl = CVarShaderFlowControl.GetValueOnAnyThread();
+		switch (FlowControl)
+		{
+			case 2:
+				Input.Environment.CompilerFlags.Add(CFLAG_AvoidFlowControl);
+				break;
+			case 1:
+				Input.Environment.CompilerFlags.Add(CFLAG_PreferFlowControl);
+				break;
+			case 0:
+				// Fallback to nothing...
+			default:
+				break;
+		}
+	}
+
+	{
+		static const auto CVar = IConsoleManager::Get().FindConsoleVariable(TEXT("r.Shaders.Validation"));
+		if (CVar && CVar->GetInt() == 0)
+		{
+			Input.Environment.CompilerFlags.Add(CFLAG_SkipValidation);
+		}
+	}
 
-	// if the target is the current platform, then we are not cooking for another platform, in which case we want to use
-	// the loaded permutation flags that are in the shader map (or the current platform's permutation if it wasn't loaded, 
-	// see the FShaderMapBase constructor)
-	if (bLoadedFromCacheFile)
 	{
-		PermutationFlags = GlobalShaderMap->GetFirstSection()->GetPermutationFlags();
+		SET_SHADER_DEFINE(Input.Environment, DO_CHECK, GSShaderCheckLevel > 0 ? 1 : 0);
+		SET_SHADER_DEFINE(Input.Environment, DO_GUARD_SLOW, GSShaderCheckLevel > 1 ? 1 : 0);
 	}
 
-	bool bErrorOnMissing = bLoadedFromCacheFile;
-	if (FPlatformProperties::RequiresCookedData())
 	{
-		// We require all shaders to exist on cooked platforms because we can't compile them.
-		bErrorOnMissing = true;
+		static FShaderPlatformCachedIniValue<int32> CVarWarningsAsErrorsPerPlatform(TEXT("r.Shaders.WarningsAsErrors"));
+		const int WarnLevel = CVarWarningsAsErrorsPerPlatform.Get(ShaderPlatform);
+		if ((WarnLevel == 1 && ShaderType->GetTypeForDynamicCast() == FShaderType::EShaderTypeForDynamicCast::Global) || WarnLevel > 1)
+		{
+			Input.Environment.CompilerFlags.Add(CFLAG_WarningsAsErrors);
+		}
 	}
 
-#if WITH_EDITOR
-	// All jobs, single & pipeline
-	TArray<FShaderCommonCompileJobPtr> GlobalShaderJobs;
+	if (UseRemoveUnsedInterpolators((EShaderPlatform)Target.Platform) && !IsOpenGLPlatform((EShaderPlatform)Target.Platform))
+	{
+		Input.Environment.CompilerFlags.Add(CFLAG_ForceRemoveUnusedInterpolators);
+	}
+	
+	if (IsD3DPlatform((EShaderPlatform)Target.Platform) && IsPCPlatform((EShaderPlatform)Target.Platform))
+	{
 
-	// Add the single jobs first
-	TMap<TShaderTypePermutation<const FShaderType>, FShaderCompileJob*> SharedShaderJobs;
-#endif // WITH_EDITOR
 
-	for (TLinkedList<FShaderType*>::TIterator ShaderTypeIt(FShaderType::GetTypeList()); ShaderTypeIt; ShaderTypeIt.Next())
-	{
-		FGlobalShaderType* GlobalShaderType = ShaderTypeIt->GetGlobalShaderType();
-		if (!GlobalShaderType)
+		if (CVarD3DCheckedForTypedUAVs.GetValueOnAnyThread() == 0)
 		{
-			continue;
+			Input.Environment.CompilerFlags.Add(CFLAG_AllowTypedUAVLoads);
 		}
 
-		int32 PermutationCountToCompile = 0;
-		for (int32 PermutationId = 0; PermutationId < GlobalShaderType->GetPermutationCount(); PermutationId++)
 		{
-			if (GlobalShaderType->ShouldCompilePermutation(Platform, PermutationId, PermutationFlags))
+			static const auto CVar = IConsoleManager::Get().FindConsoleVariable(TEXT("r.D3D.CheckedForTypedUAVs"));
+			if (CVar && CVar->GetInt() == 0)
 			{
-				bool bOutdated = OutdatedShaderTypes && OutdatedShaderTypes->Contains(GlobalShaderType);
-				TShaderRef<FShader> GlobalShader = GlobalShaderMap->GetShader(GlobalShaderType, PermutationId);
-				if (bOutdated || !GlobalShader.IsValid())
-				{
-					if (bErrorOnMissing)
-					{
-                        if (IsMetalPlatform(GMaxRHIShaderPlatform))
-                        {
-                            check(IsInGameThread());
-                            FMessageDialog::Open(EAppMsgType::Ok, LOCTEXT("NoGlobalShader_Error", "Missing shader permutation. Please make sure cooking was successful and refer to Engine log for details."));
-                        }
-                            UE_LOG(LogShaders, Fatal, TEXT("Missing global shader %s's permutation %i, Please make sure cooking was successful."),
-                                GlobalShaderType->GetName(), PermutationId);
-                    }
-					else
-					{
-#if WITH_EDITOR
-						if (OutdatedShaderTypes)
-						{
-							// Remove old shader, if it exists
-							GlobalShaderMap->RemoveShaderTypePermutaion(GlobalShaderType, PermutationId);
-						}
-
-						// Compile this global shader type.
-						FGlobalShaderTypeCompiler::BeginCompileShader(GlobalShaderType, PermutationId, Platform, PermutationFlags, GlobalShaderJobs);
-						//TShaderTypePermutation<const FShaderType> ShaderTypePermutation(GlobalShaderType, PermutationId);
-						//check(!SharedShaderJobs.Find(ShaderTypePermutation));
-						//SharedShaderJobs.Add(ShaderTypePermutation, Job);
-						PermutationCountToCompile++;
-#endif // WITH_EDITOR
-					}
-				}
+				Input.Environment.CompilerFlags.Add(CFLAG_AllowTypedUAVLoads);
 			}
 		}
+	}
 
-		int32 PermutationCountLimit = 832;	// Nanite culling as of today (2022-01-11) can go up to 832 permutations
-		if (Substrate::IsSubstrateEnabled())
+	if (IsMetalPlatform((EShaderPlatform)Target.Platform))
+	{
+		if (CVarShaderZeroInitialise.GetValueOnAnyThread() != 0)
 		{
-			// SUBSTRATE_TODO reduce the number of permutation of FDeferredLightPS.
-			PermutationCountLimit = 1304;	// FDeferredLightPS as of today (2023-12-04)
+			Input.Environment.CompilerFlags.Add(CFLAG_ZeroInitialise);
 		}
-		ensureMsgf(
-			PermutationCountToCompile <= PermutationCountLimit,
-			TEXT("Global shader %s has %i permutations: probably more than it needs."),
-			GlobalShaderType->GetName(), PermutationCountToCompile);
 
-		if (!bEmptyMap && PermutationCountToCompile > 0)
+		if (CVarShaderBoundsChecking.GetValueOnAnyThread() != 0)
 		{
-			UE_LOG(LogShaders, Log, TEXT("	%s (%i out of %i)"),
-				GlobalShaderType->GetName(), PermutationCountToCompile, GlobalShaderType->GetPermutationCount());
+			Input.Environment.CompilerFlags.Add(CFLAG_BoundsChecking);
 		}
-	}
-
-	// Now the pipeline jobs; if it's a shareable pipeline, do not add duplicate jobs
-	for (TLinkedList<FShaderPipelineType*>::TIterator ShaderPipelineIt(FShaderPipelineType::GetTypeList()); ShaderPipelineIt; ShaderPipelineIt.Next())
-	{
-		const FShaderPipelineType* Pipeline = *ShaderPipelineIt;
-		if (Pipeline->IsGlobalTypePipeline())
+		
+		// Check whether we can compile metal shaders to bytecode - avoids poisoning the DDC
+		static ITargetPlatformManagerModule& TPM = GetTargetPlatformManagerRef();
+		const IShaderFormat* Compiler = TPM.FindShaderFormat(ShaderFormatName);
+		static const bool bCanCompileOfflineMetalShaders = Compiler && Compiler->CanCompileBinaryShaders();
+		if (!bCanCompileOfflineMetalShaders)
 		{
-			if (FGlobalShaderType::ShouldCompilePipeline(Pipeline, Platform, PermutationFlags)
-				&& (!GlobalShaderMap->HasShaderPipeline(Pipeline) || (OutdatedShaderPipelineTypes && OutdatedShaderPipelineTypes->Contains(Pipeline))))
-			{
-				if (OutdatedShaderPipelineTypes)
-				{
-					// Remove old pipeline
-					GlobalShaderMap->RemoveShaderPipelineType(Pipeline);
-				}
-
-				if (bErrorOnMissing)
-				{
-					UE_LOG(LogShaders, Fatal, TEXT("Missing global shader pipeline %s, Please make sure cooking was successful."), Pipeline->GetName());
-				}
-				else
-				{
-#if WITH_EDITOR
-					if (!bEmptyMap)
-					{
-						UE_LOG(LogShaders, Log, TEXT("	%s"), Pipeline->GetName());
-					}
-
-					if (Pipeline->ShouldOptimizeUnusedOutputs(Platform))
-					{
-						// Make a pipeline job with all the stages
-						FGlobalShaderTypeCompiler::BeginCompileShaderPipeline(Platform, PermutationFlags, Pipeline, GlobalShaderJobs);
-					}
-					else
-					{
-						// If sharing shaders amongst pipelines, add this pipeline as a dependency of an existing individual job
-						for (const FShaderType* ShaderType : Pipeline->GetStages())
-						{
-							TShaderTypePermutation<const FShaderType> ShaderTypePermutation(ShaderType, kUniqueShaderPermutationId);
+			Input.Environment.CompilerFlags.Add(CFLAG_Debug);
+		}
+		
+		// Shaders built for archiving - for Metal that requires compiling the code in a different way so that we can strip it later
+		bool bArchive = false;
+		GConfig->GetBool(TEXT("/Script/UnrealEd.ProjectPackagingSettings"), TEXT("bSharedMaterialNativeLibraries"), bArchive, GGameIni);
+		if (bCanCompileOfflineMetalShaders && bArchive)
+		{
+			Input.Environment.CompilerFlags.Add(CFLAG_Archive);
+		}
+		
+		{
+			uint32 ShaderVersion = RHIGetMetalShaderLanguageVersion(EShaderPlatform(Target.Platform));
+			Input.Environment.SetCompileArgument(TEXT("SHADER_LANGUAGE_VERSION"), ShaderVersion);
+			
+			bool bAllowFastIntrinsics = false;
+			bool bForceFloats = false;
+			int32 IndirectArgumentTier = 0;
+			bool bEnableMathOptimisations = true;
+            bool bSupportAppleA8 = false;
+			bool bMetalOptimizeForSize = false;
 
-							FShaderCompileJob** Job = SharedShaderJobs.Find(ShaderTypePermutation);
-							checkf(Job, TEXT("Couldn't find existing shared job for global shader %s on pipeline %s!"), ShaderType->GetName(), Pipeline->GetName());
-							auto* SingleJob = (*Job)->GetSingleShaderJob();
-							check(SingleJob);
-							auto& SharedPipelinesInJob = SingleJob->SharingPipelines.FindOrAdd(nullptr);
-							check(!SharedPipelinesInJob.Contains(Pipeline));
-							SharedPipelinesInJob.Add(Pipeline);
-						}
-					}
-#endif // WITH_EDITOR
-				}
+			if (IsPCPlatform(EShaderPlatform(Target.Platform)))
+			{
+				GConfig->GetBool(TEXT("/Script/MacTargetPlatform.MacTargetSettings"), TEXT("UseFastIntrinsics"), bAllowFastIntrinsics, GEngineIni);
+				GConfig->GetBool(TEXT("/Script/MacTargetPlatform.MacTargetSettings"), TEXT("EnableMathOptimisations"), bEnableMathOptimisations, GEngineIni);
+				GConfig->GetInt(TEXT("/Script/MacTargetPlatform.MacTargetSettings"), TEXT("IndirectArgumentTier"), IndirectArgumentTier, GEngineIni);
+				GConfig->GetBool(TEXT("/Script/MacTargetPlatform.MacTargetSettings"), TEXT("MetalOptimizeForSize"), bMetalOptimizeForSize, GEngineIni);
+                
+                // No half precision support on MacOS at the moment
+                bForceFloats = true;
+			}
+			else
+			{
+				GConfig->GetBool(TEXT("/Script/IOSRuntimeSettings.IOSRuntimeSettings"), TEXT("UseFastIntrinsics"), bAllowFastIntrinsics, GEngineIni);
+				GConfig->GetBool(TEXT("/Script/IOSRuntimeSettings.IOSRuntimeSettings"), TEXT("EnableMathOptimisations"), bEnableMathOptimisations, GEngineIni);
+				GConfig->GetBool(TEXT("/Script/IOSRuntimeSettings.IOSRuntimeSettings"), TEXT("ForceFloats"), bForceFloats, GEngineIni);
+				GConfig->GetInt(TEXT("/Script/IOSRuntimeSettings.IOSRuntimeSettings"), TEXT("IndirectArgumentTier"), IndirectArgumentTier, GEngineIni);
+                GConfig->GetBool(TEXT("/Script/IOSRuntimeSettings.IOSRuntimeSettings"), TEXT("bSupportAppleA8"), bSupportAppleA8, GEngineIni);
+				GConfig->GetBool(TEXT("/Script/IOSRuntimeSettings.IOSRuntimeSettings"), TEXT("MetalOptimizeForSize"), bMetalOptimizeForSize, GEngineIni);
+                
+				// Force no development shaders on iOS
+				bAllowDevelopmentShaderCompile = false;
+			}
+            
+            Input.Environment.FullPrecisionInPS |= bForceFloats;
+            
+			Input.Environment.SetCompileArgument(TEXT("METAL_USE_FAST_INTRINSICS"), bAllowFastIntrinsics);
+			Input.Environment.SetCompileArgument(TEXT("METAL_INDIRECT_ARGUMENT_BUFFERS"), IndirectArgumentTier);
+            Input.Environment.SetCompileArgument(TEXT("SUPPORT_APPLE_A8"), bSupportAppleA8);
+            Input.Environment.SetCompileArgument(TEXT("METAL_OPTIMIZE_FOR_SIZE"), bMetalOptimizeForSize);
+			
+			// Same as console-variable above, but that's global and this is per-platform, per-project
+			if (!bEnableMathOptimisations)
+			{
+				Input.Environment.CompilerFlags.Add(CFLAG_NoFastMath);
 			}
 		}
 	}
 
-#if WITH_EDITOR
-	if (GlobalShaderJobs.Num() > 0)
+	if (IsAndroidPlatform(EShaderPlatform(Target.Platform)))
 	{
-		GetOnGlobalShaderCompilation().Broadcast();
-		GShaderCompilingManager->SubmitJobs(GlobalShaderJobs, "Globals");
+		// Force no development shaders on Android platforms
+		bAllowDevelopmentShaderCompile = false;
+	}
 
-		const bool bAllowAsynchronousGlobalShaderCompiling =
-			// OpenGL requires that global shader maps are compiled before attaching
-			// primitives to the scene as it must be able to find FNULLPS.
-			// TODO_OPENGL: Allow shaders to be compiled asynchronously.
-			// Metal also needs this when using RHI thread because it uses TOneColorVS very early in RHIPostInit()
-			!IsOpenGLPlatform(GMaxRHIShaderPlatform) && !IsVulkanPlatform(GMaxRHIShaderPlatform) &&
-			!IsMetalPlatform(GMaxRHIShaderPlatform) &&
-			FDataDrivenShaderPlatformInfo::GetSupportsAsyncPipelineCompilation(GMaxRHIShaderPlatform) &&
-			GShaderCompilingManager->AllowAsynchronousShaderCompiling();
+	// Mobile emulation should be defined when a PC platform is using a mobile renderer (limited to feature level ES3_1)...  eg SP_PCD3D_ES3_1,SP_VULKAN_PCES3_1,SP_METAL_MACES3_1
+	if (IsSimulatedPlatform(EShaderPlatform(Target.Platform)) && bAllowDevelopmentShaderCompile)
+	{
+		SET_SHADER_DEFINE(Input.Environment, MOBILE_EMULATION, 1);
+	}
 
-		if (!bAllowAsynchronousGlobalShaderCompiling)
-		{
-			TArray<int32> ShaderMapIds;
-			ShaderMapIds.Add(GlobalShaderMapId);
+	// Add compiler flag CFLAG_ForceDXC if DXC is enabled
+	const bool bHlslVersion2021 = Input.Environment.CompilerFlags.Contains(CFLAG_HLSL2021);
+	const bool bIsDxcEnabled = IsDxcEnabledForPlatform((EShaderPlatform)Target.Platform, bHlslVersion2021);
+	SET_SHADER_DEFINE(Input.Environment, COMPILER_DXC, bIsDxcEnabled);
+	if (bIsDxcEnabled)
+	{
+		Input.Environment.CompilerFlags.Add(CFLAG_ForceDXC);
+	}
 
-			GShaderCompilingManager->FinishCompilation(TEXT("Global"), ShaderMapIds);
+	bool bIsMobilePlatform = IsMobilePlatform((EShaderPlatform)Target.Platform);
+
+	if (bIsMobilePlatform)
+	{
+		if (IsUsingEmulatedUniformBuffers((EShaderPlatform)Target.Platform))
+		{
+			Input.Environment.CompilerFlags.Add(CFLAG_UseEmulatedUB);
 		}
 	}
-#endif // WITH_EDITOR
-}
 
-void VerifyGlobalShaders(EShaderPlatform Platform, bool bLoadedFromCacheFile, const TArray<const FShaderType*>* OutdatedShaderTypes, const TArray<const FShaderPipelineType*>* OutdatedShaderPipelineTypes)
-{
-	VerifyGlobalShaders(Platform, nullptr, bLoadedFromCacheFile, OutdatedShaderTypes, OutdatedShaderPipelineTypes);
-}
+	SET_SHADER_DEFINE(Input.Environment, HAS_INVERTED_Z_BUFFER, (bool)ERHIZBuffer::IsInverted);
+
+	if (Input.Environment.CompilerFlags.Contains(CFLAG_HLSL2021))
+	{
+		SET_SHADER_DEFINE(Input.Environment, COMPILER_SUPPORTS_HLSL2021, 1);
+	}
+
+	{
+		static const auto CVar = IConsoleManager::Get().FindTConsoleVariableDataInt(TEXT("r.ClearCoatNormal"));
+		SET_SHADER_DEFINE(Input.Environment, CLEAR_COAT_BOTTOM_NORMAL, CVar ? (CVar->GetValueOnAnyThread() != 0) && !bIsMobilePlatform : 0);
+	}
+
+	{
+		static const auto CVar = IConsoleManager::Get().FindTConsoleVariableDataInt(TEXT("r.IrisNormal"));
+		SET_SHADER_DEFINE(Input.Environment, IRIS_NORMAL, CVar ? (CVar->GetValueOnAnyThread() != 0) : 0);
+	}
 
-void PrecacheComputePipelineStatesForGlobalShaders(EShaderPlatform Platform, const ITargetPlatform* TargetPlatform)
-{
-	static IConsoleVariable* PrecacheGlobalComputeShadersCVar = IConsoleManager::Get().FindConsoleVariable(TEXT("r.PSOPrecache.GlobalComputeShaders"));
-	if (!GRHISupportsPSOPrecaching || !PipelineStateCache::IsPSOPrecachingEnabled() || PrecacheGlobalComputeShadersCVar == nullptr || PrecacheGlobalComputeShadersCVar->GetInt() == 0)
 	{
-		return;
+		static const auto CVar = IConsoleManager::Get().FindTConsoleVariableDataInt(TEXT("Compat.UseDXT5NormalMaps"));
+		SET_SHADER_DEFINE(Input.Environment, DXT5_NORMALMAPS, CVar ? (CVar->GetValueOnAnyThread() != 0) : 0);
 	}
 
-	TRACE_CPUPROFILER_EVENT_SCOPE(PrecacheComputePipelineStatesForGlobalShaders);
+	if (bAllowDevelopmentShaderCompile)
+	{
+		static const auto CVar = IConsoleManager::Get().FindTConsoleVariableDataInt(TEXT("r.CompileShadersForDevelopment"));
+		SET_SHADER_DEFINE(Input.Environment, COMPILE_SHADERS_FOR_DEVELOPMENT, CVar ? (CVar->GetValueOnAnyThread() != 0) : 0);
+	}
 
-	FPlatformTypeLayoutParameters LayoutParams;
-	LayoutParams.InitializeForPlatform(TargetPlatform);
-	EShaderPermutationFlags PermutationFlags = GetShaderPermutationFlags(LayoutParams);
+	{
+		SET_SHADER_DEFINE(Input.Environment, ALLOW_STATIC_LIGHTING, IsStaticLightingAllowed() ? 1 : 0);
+	}
 
-	FGlobalShaderMap* GlobalShaderMap = GetGlobalShaderMap(Platform);
-	
-	// some RHIs (OpenGL) can only create shaders on the Render thread. Queue the creation instead of doing it here.
-	TArray<TShaderRef<FShader>> ComputeShadersToPrecache;
-	for (TLinkedList<FShaderType*>::TIterator ShaderTypeIt(FShaderType::GetTypeList()); ShaderTypeIt; ShaderTypeIt.Next())
 	{
-		FGlobalShaderType* GlobalShaderType = ShaderTypeIt->GetGlobalShaderType();
-		if (!GlobalShaderType || GlobalShaderType->GetFrequency() != SF_Compute)
-		{
-			continue;
-		}
+		// Allow GBuffer containing a velocity target to be overridden at a higher level with GBUFFER_LAYOUT
+		bool bUsingBasePassVelocity = IsUsingBasePassVelocity((EShaderPlatform)Target.Platform);
+		SET_SHADER_DEFINE(Input.Environment, USES_BASE_PASS_VELOCITY, bUsingBasePassVelocity ? 1 : 0);
 
-		int32 ShaderPermutationPerGlobalShaderType = 0;
-		for (int32 PermutationId = 0; PermutationId < GlobalShaderType->GetPermutationCount(); PermutationId++)
+		bool bGBufferHasVelocity = bUsingBasePassVelocity;
+		if (!bGBufferHasVelocity)
 		{
-			if (GlobalShaderType->ShouldCompilePermutation(Platform, PermutationId, PermutationFlags) &&
-				GlobalShaderType->ShouldPrecachePermutation(Platform, PermutationId, PermutationFlags) == EShaderPermutationPrecacheRequest::Precached)
-			{
-				TShaderRef<FShader> GlobalShader = GlobalShaderMap->GetShader(GlobalShaderType, PermutationId);
-				if (GlobalShader.IsValid())
-				{
-					ComputeShadersToPrecache.Add(GlobalShader);
-					ShaderPermutationPerGlobalShaderType++;
-				}
-			}
+			const EGBufferLayout Layout = FShaderCompileUtilities::FetchGBufferLayout(Input.Environment);
+			bGBufferHasVelocity |= (Layout == GBL_ForceVelocity);
 		}
+		SET_SHADER_DEFINE(Input.Environment, GBUFFER_HAS_VELOCITY, bGBufferHasVelocity ? 1 : 0);
+	}
 
-		/*
-		int32 PermutationCountLimit = 300;
-		ensureMsgf(
-			ShaderPermutationPerGlobalShaderType < PermutationCountLimit,
-			TEXT("Global shader %s has %i permutations to precache: probably more than it needs."),
-			GlobalShaderType->GetName(), ShaderPermutationPerGlobalShaderType);
-		*/
+	{
+		static const auto CVar = IConsoleManager::Get().FindTConsoleVariableDataInt(TEXT("r.GBufferDiffuseSampleOcclusion"));
+		SET_SHADER_DEFINE(Input.Environment, GBUFFER_HAS_DIFFUSE_SAMPLE_OCCLUSION, CVar ? (CVar->GetValueOnAnyThread() != 0) : 1);
 	}
-	
-	if (ComputeShadersToPrecache.Num() > 0)
+
 	{
-		UE_LOG(LogShaders, Display, TEXT("Precaching %d global compute shaders"), ComputeShadersToPrecache.Num());
-		ENQUEUE_RENDER_COMMAND(PrecachePSOsForGlobalShaders)(
-			[ComputeShadersToPrecache](FRHICommandListImmediate& RHICmdList)
-			{
-				for (TShaderRef<FShader> GlobalShader : ComputeShadersToPrecache)
-				{
-					// PSO precache shaders are not required to all load correctly
-					bool bRequired = false;
-					const TCHAR* TypeName = GlobalShader.GetType()->GetName();
-					FRHIComputeShader* RHIComputeShader = static_cast<FRHIComputeShader*>(GlobalShader.GetRHIShaderBase(SF_Compute, bRequired));
-					if (RHIComputeShader)
-					{
-						PipelineStateCache::PrecacheComputePipelineState(RHIComputeShader, TypeName);
-					}
-				}
-			});
+		SET_SHADER_DEFINE(Input.Environment, SELECTIVE_BASEPASS_OUTPUTS, IsUsingSelectiveBasePassOutputs((EShaderPlatform)Target.Platform) ? 1 : 0);
 	}
-}
 
-#include "Misc/PreLoadFile.h"
-#include "Serialization/LargeMemoryReader.h"
-static FPreLoadFile GGlobalShaderPreLoadFile(*(FString(TEXT("../../../Engine")) / TEXT("GlobalShaderCache-SP_") + FPlatformProperties::IniPlatformName() + TEXT(".bin")));
+	{
+		SET_SHADER_DEFINE(Input.Environment, USE_DBUFFER, IsUsingDBuffers((EShaderPlatform)Target.Platform) ? 1 : 0);
+	}
 
-static const ITargetPlatform* GGlobalShaderTargetPlatform[SP_NumPlatforms] = { nullptr };
+	{
+		static IConsoleVariable* CVar = IConsoleManager::Get().FindConsoleVariable(TEXT("r.AllowGlobalClipPlane"));
+		SET_SHADER_DEFINE(Input.Environment, PROJECT_ALLOW_GLOBAL_CLIP_PLANE, CVar ? (CVar->GetInt() != 0) : 0);
+	}
 
-static FString GGlobalShaderCacheOverrideDirectory;
+	{
+		const bool bSupportsClipDistance = FDataDrivenShaderPlatformInfo::GetSupportsClipDistance((EShaderPlatform)Target.Platform);
+		SET_SHADER_DEFINE(Input.Environment, PLATFORM_SUPPORTS_CLIP_DISTANCE, bSupportsClipDistance ? 1u : 0u);
+	}
 
-static FString GetGlobalShaderCacheOverrideFilename(EShaderPlatform Platform)
-{
-	FString DirectoryPrefix = FPaths::EngineDir() / TEXT("OverrideGlobalShaderCache-");
+	{
+		const bool bSupportsVertexShaderSRVs = FDataDrivenShaderPlatformInfo::GetSupportsVertexShaderSRVs((EShaderPlatform)Target.Platform);
+		SET_SHADER_DEFINE(Input.Environment, PLATFORM_SUPPORTS_VERTEX_SHADER_SRVS, bSupportsVertexShaderSRVs ? 1u : 0u);
+	}
 
-	if (!GGlobalShaderCacheOverrideDirectory.IsEmpty())
 	{
-		DirectoryPrefix = GGlobalShaderCacheOverrideDirectory / TEXT("GlobalShaderCache-");
+		const bool bSupportsVertexShaderUAVs = FDataDrivenShaderPlatformInfo::GetSupportsVertexShaderUAVs((EShaderPlatform)Target.Platform) != ERHIFeatureSupport::Unsupported;
+		SET_SHADER_DEFINE(Input.Environment, PLATFORM_SUPPORTS_VERTEX_SHADER_UAVS, bSupportsVertexShaderUAVs ? 1u : 0u);
 	}
 
-	return DirectoryPrefix + FDataDrivenShaderPlatformInfo::GetName(Platform).ToString() + TEXT(".bin");
-}
+	{
+		const uint32 MaxSamplers = FDataDrivenShaderPlatformInfo::GetMaxSamplers((EShaderPlatform)Target.Platform);
+		SET_SHADER_DEFINE_AND_COMPILE_ARGUMENT(Input.Environment, PLATFORM_MAX_SAMPLERS, MaxSamplers);
+	}
 
-static FString GetGlobalShaderCacheFilename(EShaderPlatform Platform)
-{
-	return FString(TEXT("Engine")) / TEXT("GlobalShaderCache-") + FDataDrivenShaderPlatformInfo::GetName(Platform).ToString() + TEXT(".bin");
-}
+	{
+		const bool bSupportsIndependentSamplers = FDataDrivenShaderPlatformInfo::GetSupportsIndependentSamplers((EShaderPlatform)Target.Platform);
+		SET_SHADER_DEFINE(Input.Environment, SUPPORTS_INDEPENDENT_SAMPLERS, bSupportsIndependentSamplers ? 1 : 0);
+	}
 
-#if WITH_EDITOR
+	bool bForwardShading = false;
+	{
+		if (bIsMobilePlatform)
+		{
+			bForwardShading = !IsMobileDeferredShadingEnabled((EShaderPlatform)Target.Platform);
+		}
+		else if (TargetPlatform)
+		{
+			bForwardShading = TargetPlatform->UsesForwardShading();
+		}
+		else
+		{
+			static IConsoleVariable* CVarForwardShading = IConsoleManager::Get().FindConsoleVariable(TEXT("r.ForwardShading"));
+			bForwardShading = CVarForwardShading ? (CVarForwardShading->GetInt() != 0) : false;
+		}
+		SET_SHADER_DEFINE(Input.Environment, FORWARD_SHADING, bForwardShading);
+	}
 
-static FString GetGlobalShaderMapKeyString(const FGlobalShaderMapId& ShaderMapId, EShaderPlatform Platform, TArray<FShaderTypeDependency> const& Dependencies)
-{
-	TRACE_CPUPROFILER_EVENT_SCOPE(GetGlobalShaderMapKeyString);
-	FName Format = LegacyShaderPlatformToShaderFormat(Platform);
-	FString ShaderMapKeyString = Format.ToString() + TEXT("_") + FString(FString::FromInt(GetTargetPlatformManagerRef().ShaderFormatVersion(Format))) + TEXT("_");
-	ShaderMapAppendKeyString(Platform, ShaderMapKeyString);
-	ShaderMapId.AppendKeyString(ShaderMapKeyString, Dependencies);
+	{
+		if (VelocityEncodeDepth((EShaderPlatform)Target.Platform))
+		{
+			SET_SHADER_DEFINE(Input.Environment, VELOCITY_ENCODE_DEPTH, 1);
+		}
+		else
+		{
+			SET_SHADER_DEFINE(Input.Environment, VELOCITY_ENCODE_DEPTH, 0);
+		}
+	}
 
-	const EShaderPermutationFlags PermutationFlags = ShaderMapId.GetShaderPermutationFlags();
+	{
+		if (MaskedInEarlyPass((EShaderPlatform)Target.Platform))
+		{
+			SET_SHADER_DEFINE(Input.Environment, EARLY_Z_PASS_ONLY_MATERIAL_MASKING, 1);
+		}
+		else
+		{
+			SET_SHADER_DEFINE(Input.Environment, EARLY_Z_PASS_ONLY_MATERIAL_MASKING, 0);
+		}
+	}
 
-	// Construct a hash of all the environment modifications applied for each shader type & permutation
-	FMemoryHasherBlake3 Hasher;
-	for (const FShaderTypeDependency& ShaderTypeDep : Dependencies)
 	{
-		const FGlobalShaderType* GlobalShaderType = FindShaderTypeByName(ShaderTypeDep.ShaderTypeName)->AsGlobalShaderType();
-		for (int32 PermutationId = 0; PermutationId < GlobalShaderType->GetPermutationCount(); PermutationId++)
+		static IConsoleVariable* CVar = IConsoleManager::Get().FindConsoleVariable(TEXT("r.VertexFoggingForOpaque"));
+		bool bVertexFoggingForOpaque = false;
+		if (bForwardShading)
 		{
-			if (GlobalShaderType->ShouldCompilePermutation(Platform, PermutationId, PermutationFlags))
+			bVertexFoggingForOpaque = CVar ? (CVar->GetInt() != 0) : 0;
+			if (TargetPlatform)
 			{
-				FShaderCompilerEnvironment Env(Hasher);
-				GlobalShaderType->SetupCompileEnvironment(Platform, PermutationId, PermutationFlags, Env);
-				Env.SerializeEverythingButFiles(Hasher);
+				const int32 PlatformHeightFogMode = TargetPlatform->GetHeightFogModeForOpaque();
+				if (PlatformHeightFogMode == 1)
+				{
+					bVertexFoggingForOpaque = false;
+				}
+				else if (PlatformHeightFogMode == 2)
+				{
+					bVertexFoggingForOpaque = true;
+				}
 			}
 		}
+		SET_SHADER_DEFINE(Input.Environment, PROJECT_VERTEX_FOGGING_FOR_OPAQUE, bVertexFoggingForOpaque);
 	}
 
-	// * 2 for hex representation of hash; + 6 for tag/underscores
-	TStringBuilder<sizeof(TCHAR) * (sizeof(FBlake3Hash::ByteArray) * 2 + 6)> EnvHashString;
-	EnvHashString << "_EMH_" << Hasher.Finalize() << "_";
-	check(EnvHashString.GetAllocatedSize() == 0);
-	ShaderMapKeyString.Append(EnvHashString.ToView());
-
-	return FString::Printf(TEXT("%s_%s_%s"), TEXT("GSM"), *GetGlobalShaderMapDDCGuid().ToString(), *ShaderMapKeyString);
-}
-
-/** Creates a string key for the derived data cache entry for the global shader map. */
-static UE::DerivedData::FCacheKey GetGlobalShaderMapKey(const FGlobalShaderMapId& ShaderMapId, EShaderPlatform Platform, const ITargetPlatform* TargetPlatform, TArray<FShaderTypeDependency> const& Dependencies)
-{
-	const FString DataKey = GetGlobalShaderMapKeyString(ShaderMapId, Platform, Dependencies);
-	static const UE::DerivedData::FCacheBucket Bucket(ANSITEXTVIEW("GlobalShaderMap"), TEXTVIEW("GlobalShader"));
-	return {Bucket, FIoHash::HashBuffer(MakeMemoryView(FTCHARToUTF8(DataKey)))};
-}
-
-static UE::FSharedString GetGlobalShaderMapName(const FGlobalShaderMapId& ShaderMapId, EShaderPlatform Platform, const FString& Key)
-{
-	return UE::FSharedString(WriteToString<256>(TEXTVIEW("GlobalShaderMap ["), LegacyShaderPlatformToShaderFormat(Platform), TEXTVIEW(", "), Key, TEXTVIEW("]")));
-}
-#endif // WITH_EDITOR
-
-/** Saves the platform's shader map to the DDC. It is assumed that the caller will check IsComplete() first before calling the function. */
-static void SaveGlobalShaderMapToDerivedDataCache(EShaderPlatform Platform)
-{
-#if WITH_EDITOR
-	// We've finally built the global shader map, so we can count the miss as we put it in the DDC.
-	COOK_STAT(auto Timer = GlobalShaderCookStats::UsageStats.TimeSyncWork());
-
-	const ITargetPlatform* TargetPlatform = GGlobalShaderTargetPlatform[Platform];
-	TArray<uint8> SaveData;
-
-	FGlobalShaderMapId ShaderMapId(Platform, TargetPlatform);
-	// caller should prevent incomplete shadermaps to be saved
-	FGlobalShaderMap* GlobalSM = GetGlobalShaderMap(Platform);
-	for (auto const& ShaderFilenameDependencies : ShaderMapId.GetShaderFilenameToDependeciesMap())
 	{
-		FGlobalShaderMapSection* Section = GlobalSM->FindSection(ShaderFilenameDependencies.Key);
-		if (Section)
-		{
-			Section->FinalizeContent();
-
-			SaveData.Reset();
-			FMemoryWriter Ar(SaveData, true);
-			Section->Serialize(Ar);
-			COOK_STAT(Timer.AddMiss(SaveData.Num()));
+		static IConsoleVariable* CVar = IConsoleManager::Get().FindConsoleVariable(TEXT("r.Mobile.DisableVertexFog"));
+		SET_SHADER_DEFINE(Input.Environment, PROJECT_MOBILE_DISABLE_VERTEX_FOG, CVar ? (CVar->GetInt() != 0) : 0);
+	}
 
-			using namespace UE::DerivedData;
-			FCachePutValueRequest Request;
-			Request.Name = GetGlobalShaderMapName(ShaderMapId, Platform, ShaderFilenameDependencies.Key);
-			Request.Key = GetGlobalShaderMapKey(ShaderMapId, Platform, TargetPlatform, ShaderFilenameDependencies.Value);
-			Request.Value = FValue::Compress(MakeSharedBufferFromArray(MoveTemp(SaveData)));
-			FRequestOwner AsyncOwner(EPriority::Normal);
-			FRequestBarrier AsyncBarrier(AsyncOwner);
-			GetCache().PutValue({Request}, AsyncOwner);
-			AsyncOwner.KeepAlive();
-		}
+	bool bSupportLocalFogVolumes = false;
+	{
+		static IConsoleVariable* CVar = IConsoleManager::Get().FindConsoleVariable(TEXT("r.SupportLocalFogVolumes"));
+		bSupportLocalFogVolumes = CVar && CVar->GetInt() > 0;
+		SET_SHADER_DEFINE(Input.Environment, PROJECT_SUPPORTS_LOCALFOGVOLUME, (bSupportLocalFogVolumes ? 1 : 0));
 	}
-#endif // WITH_EDITOR
-}
 
-/** Saves the global shader map as a file for the target platform. */
-FString SaveGlobalShaderFile(EShaderPlatform Platform, FString SavePath, class ITargetPlatform* TargetPlatform)
-{
-	FGlobalShaderMap* GlobalShaderMap = GetGlobalShaderMap(Platform);
+	{
+		static IConsoleVariable* CVar = IConsoleManager::Get().FindConsoleVariable(TEXT("r.LocalFogVolume.ApplyOnTranslucent"));
+		const bool bLocalFogVolumesApplyOnTranclucent = CVar && CVar->GetInt() > 0;
+		SET_SHADER_DEFINE(Input.Environment, PROJECT_LOCALFOGVOLUME_APPLYONTRANSLUCENT, ((bSupportLocalFogVolumes && bLocalFogVolumesApplyOnTranclucent) ? 1 : 0));
+	}
 
-	// Wait until all global shaders are compiled
-	if (GShaderCompilingManager)
+	bool bSupportSkyAtmosphere = false;
 	{
-		GShaderCompilingManager->ProcessAsyncResults(false, true);
+		static IConsoleVariable* CVar = IConsoleManager::Get().FindConsoleVariable(TEXT("r.SupportSkyAtmosphere"));
+		bSupportSkyAtmosphere = CVar && CVar->GetInt() != 0;
+		SET_SHADER_DEFINE(Input.Environment, PROJECT_SUPPORT_SKY_ATMOSPHERE, bSupportSkyAtmosphere ? 1 : 0);
 	}
 
-	TArray<uint8> GlobalShaderData;
 	{
-#if WITH_EDITOR
-		TOptional<FArchiveCookContext> CookContext;
-		TOptional<FArchiveCookData> CookData;
-#endif
-		FMemoryWriter MemoryWriter(GlobalShaderData, true);
+		static IConsoleVariable* CVar = IConsoleManager::Get().FindConsoleVariable(TEXT("r.SupportCloudShadowOnForwardLitTranslucent"));
+		const bool bSupportCloudShadowOnForwardLitTranslucent = CVar && CVar->GetInt() > 0;
+		SET_SHADER_DEFINE(Input.Environment, SUPPORT_CLOUD_SHADOW_ON_FORWARD_LIT_TRANSLUCENT, bSupportCloudShadowOnForwardLitTranslucent ? 1 : 0);
+	}
 
-#if WITH_EDITOR
-		if (TargetPlatform != nullptr)
-		{
-			CookContext.Emplace(nullptr /*InPackage*/, UE::Cook::ECookType::Unknown,
-				UE::Cook::ECookingDLC::Unknown, TargetPlatform);
-			CookData.Emplace(*TargetPlatform, *CookContext);
-			MemoryWriter.SetCookData(CookData.GetPtrOrNull());
-		}
-#endif // WITH_EDITOR
+	{
+		static IConsoleVariable *CVar = IConsoleManager::Get().FindConsoleVariable(TEXT("r.Water.SingleLayerWater.SupportCloudShadow"));
+		const bool bSupportCloudShadowOnSingleLayerWater = CVar && CVar->GetInt() > 0;
+		SET_SHADER_DEFINE(Input.Environment, SUPPORT_CLOUD_SHADOW_ON_SINGLE_LAYER_WATER, bSupportCloudShadowOnSingleLayerWater ? 1 : 0);
+	}
 
-		GlobalShaderMap->SaveToGlobalArchive(MemoryWriter);
+	{
+		static IConsoleVariable* CVar = IConsoleManager::Get().FindConsoleVariable(TEXT("r.FirstPersonRendering.Enabled"));
+		const bool bSupportFirstPersonRendering = CVar && CVar->GetInt() > 0;
+		SET_SHADER_DEFINE(Input.Environment, SUPPORT_FIRST_PERSON_RENDERING, bSupportFirstPersonRendering ? 1 : 0);
 	}
 
-	// make the final name
-	FString FullPath = SavePath / GetGlobalShaderCacheFilename(Platform);
-	if (!FFileHelper::SaveArrayToFile(GlobalShaderData, *FullPath))
 	{
-		UE_LOG(LogShaders, Fatal, TEXT("Could not save global shader file to '%s'"), *FullPath);
+		const bool bTranslucentUsesLightRectLights = GetTranslucentUsesLightRectLights();
+		SET_SHADER_DEFINE(Input.Environment, SUPPORT_RECTLIGHT_ON_FORWARD_LIT_TRANSLUCENT, bTranslucentUsesLightRectLights ? 1 : 0);
 	}
 
-#if WITH_EDITOR
-	if (FShaderLibraryCooker::NeedsShaderStableKeys(Platform))
 	{
-		GlobalShaderMap->SaveShaderStableKeys(Platform);
+		const bool bTranslucentUsesLightIESProfiles = GetTranslucentUsesLightIESProfiles();
+		SET_SHADER_DEFINE(Input.Environment, SUPPORT_IESPROFILE_ON_FORWARD_LIT_TRANSLUCENT, bTranslucentUsesLightIESProfiles ? 1 : 0);
 	}
-#endif // WITH_EDITOR
-	return FullPath;
-}
 
+	{
+		static IConsoleVariable* CVar = IConsoleManager::Get().FindConsoleVariable(TEXT("r.Shadow.Virtual.TranslucentQuality"));
+		const bool bHighQualityShadow = CVar && CVar->GetInt() > 0;
+		SET_SHADER_DEFINE(Input.Environment, SUPPORT_VSM_FOWARD_QUALITY, bHighQualityShadow ? 1 : 0);
+	}
 
-static inline bool ShouldCacheGlobalShaderTypeName(const FGlobalShaderType* GlobalShaderType, int32 PermutationId, const TCHAR* TypeNameSubstring, EShaderPlatform Platform, EShaderPermutationFlags PermutationFlags)
-{
-	return GlobalShaderType
-		&& (TypeNameSubstring == nullptr || (FPlatformString::Strstr(GlobalShaderType->GetName(), TypeNameSubstring) != nullptr))
-		&& GlobalShaderType->ShouldCompilePermutation(Platform, PermutationId, PermutationFlags);
-};
-
+	{
+		const bool bUseTriangleStrips = GetHairStrandsUsesTriangleStrips();
+		SET_SHADER_DEFINE(Input.Environment, USE_HAIR_TRIANGLE_STRIP, bUseTriangleStrips ? 1 : 0);
+	}
 
-bool IsGlobalShaderMapComplete(const TCHAR* TypeNameSubstring, FGlobalShaderMap* GlobalShaderMap, EShaderPlatform Platform, FString* FailureReason = nullptr)
-{
-	// look at any shadermap in the GlobalShaderMap for the permutation flags, as they will all be the same
-	if (GlobalShaderMap)
+	const bool bSubstrate = Substrate::IsSubstrateEnabled() && IsSubstrateSupportForShaderPipeline(Input);
 	{
-		const FGlobalShaderMapSection* FirstShaderMap = GlobalShaderMap->GetFirstSection();
-		if (FirstShaderMap == nullptr)
-		{
-			// if we had no sections at all, we know we aren't complete
-			return false;
-		}
-		EShaderPermutationFlags GlobalShaderPermutation = FirstShaderMap->GetPermutationFlags();
+		SET_SHADER_DEFINE(Input.Environment, SUBSTRATE_ENABLED, bSubstrate ? 1 : 0);
 
-		// Check if the individual shaders are complete
-		for (TLinkedList<FShaderType*>::TIterator ShaderTypeIt(FShaderType::GetTypeList()); ShaderTypeIt; ShaderTypeIt.Next())
+		if (bSubstrate)
 		{
-			FGlobalShaderType* GlobalShaderType = ShaderTypeIt->GetGlobalShaderType();
-			int32 PermutationCount = GlobalShaderType ? GlobalShaderType->GetPermutationCount() : 1;
-			for (int32 PermutationId = 0; PermutationId < PermutationCount; PermutationId++)
-			{
-				if (ShouldCacheGlobalShaderTypeName(GlobalShaderType, PermutationId, TypeNameSubstring, Platform, GlobalShaderPermutation))
-				{
-					if (!GlobalShaderMap->HasShader(GlobalShaderType, PermutationId))
-					{
-						if (FailureReason)
-						{
-							FString GlobalShaderTypeName = GlobalShaderType ? GlobalShaderType->GetFName().ToString() : FString(TEXT("Unknown shader type"));
-							*FailureReason = FString::Printf(TEXT("Failed to find global shader \"%s\", permutation %d"), *GlobalShaderTypeName, PermutationId);
-						}
-
-						return false;
-					}
-				}
-			}
-		}
+			const uint32 SubstrateShadingQuality = Substrate::GetShadingQuality(Target.GetPlatform());
+			SET_SHADER_DEFINE(Input.Environment, SUBSTRATE_SHADING_QUALITY, SubstrateShadingQuality);
 
-		// Then the pipelines as it may be sharing shaders
-		for (TLinkedList<FShaderPipelineType*>::TIterator ShaderPipelineIt(FShaderPipelineType::GetTypeList()); ShaderPipelineIt; ShaderPipelineIt.Next())
-		{
-			const FShaderPipelineType* Pipeline = *ShaderPipelineIt;
-			if (Pipeline->IsGlobalTypePipeline())
-			{
-				auto& Stages = Pipeline->GetStages();
-				int32 NumStagesNeeded = 0;
-				for (const FShaderType* Shader : Stages)
-				{
-					const FGlobalShaderType* GlobalShaderType = Shader->GetGlobalShaderType();
-					if (ShouldCacheGlobalShaderTypeName(GlobalShaderType, kUniqueShaderPermutationId, TypeNameSubstring, Platform, GlobalShaderPermutation))
-					{
-						++NumStagesNeeded;
-					}
-					else
-					{
-						break;
-					}
-				}
+			const bool bLowQuality = SubstrateShadingQuality > 1;
+			SET_SHADER_DEFINE(Input.Environment, USE_ACHROMATIC_BXDF_ENERGY, bLowQuality ? 1u : 0u);
 
-				if (NumStagesNeeded == Stages.Num())
-				{
-					if (!GlobalShaderMap->HasShaderPipeline(Pipeline))
-					{
-						if (FailureReason)
-						{
-							*FailureReason = FString::Printf(TEXT("Failed to find global pipeline \"%s\""), *Pipeline->GetFName().ToString());
-						}
+			const uint32 SubstrateSheenQuality = Substrate::GetSheenQuality();
+			Input.Environment.SetDefine(TEXT("SUBSTRATE_SHEEN_QUALITY"), bLowQuality ? 2 : SubstrateSheenQuality);
 
-						return false;
-					}
-				}
+			const uint32 SubstrateNormalQuality = Substrate::GetNormalQuality();
+			SET_SHADER_DEFINE(Input.Environment, SUBSTRATE_NORMAL_QUALITY, SubstrateNormalQuality);
+			if (SubstrateNormalQuality == 0)
+			{
+				SET_SHADER_DEFINE(Input.Environment, SUBSTRATE_TOP_LAYER_TYPE, TEXT("uint"));
+			}
+			else
+			{
+				SET_SHADER_DEFINE(Input.Environment, SUBSTRATE_TOP_LAYER_TYPE, TEXT("uint2"));
 			}
-		}
-	}
 
-	return true;
-}
+			const uint32 SubstrateUintPerPixel = Substrate::GetBytePerPixel(Target.GetPlatform()) / 4u;
+			SET_SHADER_DEFINE(Input.Environment, SUBSTRATE_MATERIAL_NUM_UINTS, SubstrateUintPerPixel);
 
+			const uint32 SubstrateClosurePerPixel = Substrate::GetClosurePerPixel(Target.GetPlatform());
+			SET_SHADER_DEFINE(Input.Environment, SUBSTRATE_MATERIAL_CLOSURE_COUNT, SubstrateClosurePerPixel);
 
-bool IsGlobalShaderMapComplete(const TCHAR* TypeNameSubstring)
-{
-	for (int32 i = 0; i < SP_NumPlatforms; ++i)
-	{
-		EShaderPlatform Platform = (EShaderPlatform)i;
+			const bool bSubstrateDBufferPass = Substrate::IsDBufferPassEnabled(Target.GetPlatform());
+			SET_SHADER_DEFINE(Input.Environment, SUBSTRATE_USE_DBUFFER_PASS, bSubstrateDBufferPass ? 1 : 0);
 
-		FGlobalShaderMap* GlobalShaderMap = GGlobalShaderMap[Platform];
+			const bool bSubstrateGlints = Substrate::IsGlintEnabled(Target.GetPlatform());
+			SET_SHADER_DEFINE(Input.Environment, PLATFORM_ENABLES_SUBSTRATE_GLINTS, bSubstrateGlints ? 1 : 0);
 
-		if (!IsGlobalShaderMapComplete(TypeNameSubstring, GlobalShaderMap, Platform))
+			const bool bSpecularProfileEnabled = Substrate::IsSpecularProfileEnabled(Target.GetPlatform());
+			SET_SHADER_DEFINE(Input.Environment, PLATFORM_ENABLES_SUBSTRATE_SPECULAR_PROFILE, bSpecularProfileEnabled ? 1 : 0);
+		}
+		else
 		{
-			return false;
+			// Some global uniform buffers reference this type -- so we need to have it defined in all cases
+			SET_SHADER_DEFINE(Input.Environment, SUBSTRATE_TOP_LAYER_TYPE, TEXT("uint"));
 		}
+
+		const bool bSubstrateBackCompatibility = bSubstrate && Substrate::IsBackCompatibilityEnabled();
+		SET_SHADER_DEFINE(Input.Environment, PROJECT_SUBSTRATE_BACKCOMPATIBILITY, bSubstrateBackCompatibility ? 1 : 0);
+
+		const bool bSubstrateOpaqueRoughRefrac = bSubstrate && Substrate::IsOpaqueRoughRefractionEnabled();
+		SET_SHADER_DEFINE(Input.Environment, SUBSTRATE_OPAQUE_ROUGH_REFRACTION_ENABLED, bSubstrateOpaqueRoughRefrac ? 1 : 0);
+
+		const bool bSubstrateAdvDebug = bSubstrate && Substrate::IsAdvancedVisualizationEnabled();
+		SET_SHADER_DEFINE(Input.Environment, SUBSTRATE_ADVANCED_DEBUG_ENABLED, bSubstrateAdvDebug ? 1 : 0);
 	}
 
-	return true;
-}
+	{
+		static IConsoleVariable* CVar = IConsoleManager::Get().FindConsoleVariable(TEXT("r.Material.RoughDiffuse"));
+		const bool bMaterialRoughDiffuse = CVar && CVar->GetInt() != 0;
+		const bool bSubstrateRoughDiffuse = Substrate::IsRoughDiffuseEnabled() && !Substrate::IsBackCompatibilityEnabled();
+		SET_SHADER_DEFINE(Input.Environment, MATERIAL_ROUGHDIFFUSE, (bSubstrate ? bSubstrateRoughDiffuse : bMaterialRoughDiffuse) ? 1 : 0);
+	}
 
-static bool TryLoadCookedGlobalShaderMap(EShaderPlatform Platform, FScopedSlowTask& SlowTask)
-{
-	SlowTask.EnterProgressFrame(50);
+	{
+		static IConsoleVariable* CVar = IConsoleManager::Get().FindConsoleVariable(TEXT("r.Lumen.Supported"));
+		const bool bLumenSupported = CVar->GetInt() != 0;
+		SET_SHADER_DEFINE(Input.Environment, PROJECT_SUPPORTS_LUMEN, bLumenSupported ? 1 : 0);
+	}
 
-	bool bLoadedFromCacheFile = false;
+	{
+		const bool bSupportOIT = FDataDrivenShaderPlatformInfo::GetSupportsOIT(EShaderPlatform(Target.Platform));
+		static IConsoleVariable* CVar = IConsoleManager::Get().FindConsoleVariable(TEXT("r.OIT.SortedPixels"));
+		const bool bOIT = CVar && CVar->GetInt() != 0;
+		SET_SHADER_DEFINE(Input.Environment, PROJECT_OIT, (bSupportOIT && bOIT) ? 1 : 0);
+	}
 
-	// Load from the override global shaders first, this allows us to hot reload in cooked / pak builds
-	TArray<uint8> GlobalShaderData;
-	const bool bAllowOverrideGlobalShaders = !WITH_EDITOR && !UE_BUILD_SHIPPING;
-	if (bAllowOverrideGlobalShaders)
 	{
-		FString OverrideGlobalShaderCacheFilename = GetGlobalShaderCacheOverrideFilename(Platform);
-		FPaths::MakeStandardFilename(OverrideGlobalShaderCacheFilename);
+		static IConsoleVariable* CVar = IConsoleManager::Get().FindConsoleVariable(TEXT("r.Material.EnergyConservation"));
+		const bool bMaterialEnergyConservation = CVar && CVar->GetInt() != 0;
+		SET_SHADER_DEFINE(Input.Environment, LEGACY_MATERIAL_ENERGYCONSERVATION, bMaterialEnergyConservation ? 1 : 0);
+	}
 
-		bool bFileExist = IFileManager::Get().FileExists(*OverrideGlobalShaderCacheFilename);
+	{
+		static IConsoleVariable* CVar = IConsoleManager::Get().FindConsoleVariable(TEXT("r.SupportSkyAtmosphereAffectsHeightFog"));
+		SET_SHADER_DEFINE(Input.Environment, PROJECT_SUPPORT_SKY_ATMOSPHERE_AFFECTS_HEIGHFOG, (CVar && bSupportSkyAtmosphere) ? (CVar->GetInt() != 0) : 0);
+	}
 
-		if (!bFileExist)
+	{
+		SET_SHADER_DEFINE(Input.Environment, PROJECT_EXPFOG_MATCHES_VFOG, DoesProjectSupportExpFogMatchesVolumetricFog() ? 1 : 0);
+	}
+
+	{
+		static IConsoleVariable* CVar = IConsoleManager::Get().FindConsoleVariable(TEXT("r.PostProcessing.PropagateAlpha"));
+		int32 PropagateAlpha = CVar->GetInt();
+
+		if (bIsMobilePlatform)
 		{
-			UE_LOG(LogShaders, Display, TEXT("%s doesn't exists"), *OverrideGlobalShaderCacheFilename);
+			static FShaderPlatformCachedIniValue<int32> MobilePropagateAlphaIniValue(TEXT("r.Mobile.PropagateAlpha"));
+			int MobilePropagateAlphaIniValueInt = MobilePropagateAlphaIniValue.Get((EShaderPlatform)ShaderPlatform);
+			PropagateAlpha = MobilePropagateAlphaIniValueInt > 0 ? 2 : 0;
 		}
-		else
-		{
-			bLoadedFromCacheFile = FFileHelper::LoadFileToArray(GlobalShaderData, *OverrideGlobalShaderCacheFilename, FILEREAD_Silent);
 
-			if (bLoadedFromCacheFile)
-			{
-				UE_LOG(LogShaders, Display, TEXT("%s has been loaded successfully"), *OverrideGlobalShaderCacheFilename);
-			}
-			else
-			{
-				UE_LOG(LogShaders, Error, TEXT("%s failed to load"), *OverrideGlobalShaderCacheFilename);
-			}
+		if (PropagateAlpha < 0 || PropagateAlpha > 2)
+		{
+			PropagateAlpha = 0;
 		}
+		SET_SHADER_DEFINE(Input.Environment, POST_PROCESS_ALPHA, PropagateAlpha);
 	}
 
-	// is the data already loaded?
-	int64 PreloadedSize = 0;
-	void* PreloadedData = nullptr;
-	if (!bLoadedFromCacheFile)
+	if (TargetPlatform && 
+		TargetPlatform->SupportsFeature(ETargetPlatformFeatures::NormalmapLAEncodingMode))
 	{
-		PreloadedData = GGlobalShaderPreLoadFile.TakeOwnershipOfLoadedData(&PreloadedSize);
+		SET_SHADER_DEFINE(Input.Environment, LA_NORMALMAPS, 1);
 	}
 
-	if (PreloadedData != nullptr)
-	{
-		FLargeMemoryReader MemoryReader((uint8*)PreloadedData, PreloadedSize, ELargeMemoryReaderFlags::TakeOwnership);
-		GGlobalShaderMap[Platform]->LoadFromGlobalArchive(MemoryReader);
-		bLoadedFromCacheFile = true;
-	}
-	else
+	SET_SHADER_DEFINE(Input.Environment, COLORED_LIGHT_FUNCTION_ATLAS, GetLightFunctionAtlasFormat() > 0 ? 1 : 0);
+
+	SET_SHADER_DEFINE(Input.Environment, PLATFORM_SUPPORTS_SHADER_ROOT_CONSTANTS, RHISupportsShaderRootConstants(EShaderPlatform(Target.Platform)) ? 1 : 0);
+	SET_SHADER_DEFINE(Input.Environment, PLATFORM_SUPPORTS_SHADER_BUNDLE_DISPATCH, RHISupportsShaderBundleDispatch(EShaderPlatform(Target.Platform)) ? 1 : 0);
+	SET_SHADER_DEFINE(Input.Environment, PLATFORM_SUPPORTS_RENDERTARGET_WRITE_MASK, RHISupportsRenderTargetWriteMask(EShaderPlatform(Target.Platform)) ? 1 : 0);
+	SET_SHADER_DEFINE(Input.Environment, PLATFORM_SUPPORTS_PER_PIXEL_DBUFFER_MASK, FDataDrivenShaderPlatformInfo::GetSupportsPerPixelDBufferMask(EShaderPlatform(Target.Platform)) ? 1 : 0);
+	SET_SHADER_DEFINE(Input.Environment, PLATFORM_SUPPORTS_DISTANCE_FIELDS, DoesPlatformSupportDistanceFields(EShaderPlatform(Target.Platform)) ? 1 : 0);
+	SET_SHADER_DEFINE(Input.Environment, PLATFORM_SUPPORTS_MESH_SHADERS_TIER0, RHISupportsMeshShadersTier0(EShaderPlatform(Target.Platform)) ? 1 : 0);
+	SET_SHADER_DEFINE(Input.Environment, PLATFORM_SUPPORTS_MESH_SHADERS_TIER1, RHISupportsMeshShadersTier1(EShaderPlatform(Target.Platform)) ? 1 : 0);
+	SET_SHADER_DEFINE(Input.Environment, PLATFORM_ALLOW_SCENE_DATA_COMPRESSED_TRANSFORMS, FDataDrivenShaderPlatformInfo::GetSupportSceneDataCompressedTransforms(EShaderPlatform(Target.Platform)) ? 1 : 0);
+	SET_SHADER_DEFINE(Input.Environment, PLATFORM_SUPPORTS_BUFFER_LOAD_TYPE_CONVERSION, RHISupportsBufferLoadTypeConversion(ShaderPlatform) ? 1 : 0);
+	SET_SHADER_DEFINE(Input.Environment, PLATFORM_SUPPORTS_UNIFORM_BUFFER_OBJECTS, FDataDrivenShaderPlatformInfo::GetSupportsUniformBufferObjects(EShaderPlatform(Target.Platform)) ? 1 : 0);
+	SET_SHADER_DEFINE(Input.Environment, PLATFORM_SUPPORTS_RAY_TRACING_HIGH_END_EFFECTS, FDataDrivenShaderPlatformInfo::GetSupportsHighEndRayTracingEffects(EShaderPlatform(Target.Platform)) ? 1 : 0);
+	SET_SHADER_DEFINE(Input.Environment, COMPILER_SUPPORTS_BARYCENTRIC_INTRINSICS, FDataDrivenShaderPlatformInfo::GetSupportsBarycentricsIntrinsics(EShaderPlatform(Target.Platform)));
+	SET_SHADER_DEFINE(Input.Environment, PLATFORM_SUPPORTS_BARYCENTRICS_SEMANTIC, FDataDrivenShaderPlatformInfo::GetSupportsBarycentricsSemantic(EShaderPlatform(Target.Platform)) != ERHIFeatureSupport::Unsupported);
+
+	bool bEnableBindlessMacro = false;
+	if (RHIGetBindlessSupport(ShaderPlatform) != ERHIBindlessSupport::Unsupported && !Input.Environment.CompilerFlags.Contains(CFLAG_ForceBindful))
 	{
-		FString GlobalShaderCacheFilename = FPaths::GetRelativePathToRoot() / GetGlobalShaderCacheFilename(Platform);
-		FPaths::MakeStandardFilename(GlobalShaderCacheFilename);
-		if (!bLoadedFromCacheFile)
+		const bool bIsRaytracingShader = IsRayTracingShaderFrequency(Input.Target.GetFrequency());
+
+		const ERHIBindlessConfiguration ResourcesConfig = UE::ShaderCompiler::GetBindlessResourcesConfiguration(ShaderFormatName);
+		const ERHIBindlessConfiguration SamplersConfig = UE::ShaderCompiler::GetBindlessSamplersConfiguration(ShaderFormatName);
+
+		if (ResourcesConfig == ERHIBindlessConfiguration::AllShaders || (ResourcesConfig == ERHIBindlessConfiguration::RayTracingShaders && bIsRaytracingShader))
 		{
-			bLoadedFromCacheFile = FFileHelper::LoadFileToArray(GlobalShaderData, *GlobalShaderCacheFilename, FILEREAD_Silent);
+			bEnableBindlessMacro = true;
+			Input.Environment.CompilerFlags.Add(CFLAG_BindlessResources);
+			SET_SHADER_DEFINE(Input.Environment, ENABLE_BINDLESS_RESOURCES, true);
 		}
 
-		if (bLoadedFromCacheFile)
+		if (SamplersConfig == ERHIBindlessConfiguration::AllShaders || (SamplersConfig == ERHIBindlessConfiguration::RayTracingShaders && bIsRaytracingShader))
 		{
-			FMemoryReader MemoryReader(GlobalShaderData);
-			GGlobalShaderMap[Platform]->LoadFromGlobalArchive(MemoryReader);
+			bEnableBindlessMacro = true;
+			Input.Environment.CompilerFlags.Add(CFLAG_BindlessSamplers);
+			SET_SHADER_DEFINE(Input.Environment, ENABLE_BINDLESS_SAMPLERS, true);
 		}
 	}
 
-	return bLoadedFromCacheFile;
-}
-
-void CompileGlobalShaderMap(EShaderPlatform Platform, const ITargetPlatform* TargetPlatform, bool bRefreshShaderMap)
-{
-	LLM_SCOPE_RENDER_RESOURCE(TEXT("GlobalShaderMap"));
+	SET_SHADER_DEFINE(Input.Environment, PLATFORM_SUPPORTS_BINDLESS, bEnableBindlessMacro);
 
-	// No global shaders needed on dedicated server or clients that use NullRHI. Note that cook commandlet needs to have them, even if it is not allowed to render otherwise.
-	if (FPlatformProperties::IsServerOnly() || (!IsRunningCommandlet() && !FApp::CanEverRender()))
+	if (CVarShadersRemoveDeadCode.GetValueOnAnyThread())
 	{
-		if (!GGlobalShaderMap[Platform])
-		{
-			GGlobalShaderMap[Platform] = new FGlobalShaderMap(Platform);
-		}
-		return;
+		Input.Environment.CompilerFlags.Add(CFLAG_RemoveDeadCode);
 	}
 
-	if (bRefreshShaderMap || GGlobalShaderTargetPlatform[Platform] != TargetPlatform)
+	if (CVarDisableSourceStripping.GetValueOnAnyThread())
 	{
-		// defer the deletion the current global shader map, delete the previous one if it is still valid
-		delete GGlobalShaderMap_DeferredDeleteCopy[Platform];	// deleting null is Okay
-		GGlobalShaderMap_DeferredDeleteCopy[Platform] = GGlobalShaderMap[Platform];
-		GGlobalShaderMap[Platform] = nullptr;
-
-		GGlobalShaderTargetPlatform[Platform] = TargetPlatform;
+		Input.Environment.CompilerFlags.Add(CFLAG_DisableSourceStripping);
+	}
 
-		// make sure we look for updated shader source files
-		FlushShaderFileCache();
+	{
+		static const auto CVar = IConsoleManager::Get().FindConsoleVariable(TEXT("r.VT.AnisotropicFiltering"));
+		SET_SHADER_DEFINE(Input.Environment, VIRTUAL_TEXTURE_ANISOTROPIC_FILTERING, CVar ? (CVar->GetInt() != 0) : 0);
+		
+		if (bIsMobilePlatform)
+		{
+			static FShaderPlatformCachedIniValue<bool> CVarVTMobileManualTrilinearFiltering(TEXT("r.VT.Mobile.ManualTrilinearFiltering"));
+			SET_SHADER_DEFINE(Input.Environment, VIRTUAL_TEXTURE_MANUAL_TRILINEAR_FILTERING, (CVarVTMobileManualTrilinearFiltering.Get(Target.GetPlatform()) ? 1 : 0));
+		}
 	}
 
-#if WITH_ODSC
-	// First try to load the global shader map with ODSC if it's connected. TryLoadGlobalShaders will set GGlobalShaderMap[Platform]
-	if (!GGlobalShaderMap[Platform] && FODSCManager::IsODSCActive())
+	if (bIsMobilePlatform)
 	{
-		UE_LOG(LogShaders, Display, TEXT("Trying to load global shaders from ODSC ..."));
-		GODSCManager->TryLoadGlobalShaders(Platform);
-		UE_LOG(LogShaders, Display, TEXT("Global shaders from ODSC: %s"), (GGlobalShaderMap[Platform] != nullptr) ? TEXT("success") : TEXT("failed"));
+		const bool bMobileMovableSpotlightShadowsEnabled = IsMobileMovableSpotlightShadowsEnabled(Target.GetPlatform());
+		SET_SHADER_DEFINE(Input.Environment, PROJECT_MOBILE_ENABLE_MOVABLE_SPOTLIGHT_SHADOWS, bMobileMovableSpotlightShadowsEnabled ? 1 : 0);
 	}
-#endif
 
-	// If the global shader map hasn't been created yet, create it.
-	if (!GGlobalShaderMap[Platform])
 	{
-		DECLARE_SCOPE_CYCLE_COUNTER(TEXT("GetGlobalShaderMap"), STAT_GetGlobalShaderMap, STATGROUP_LoadTime);
-		// GetGlobalShaderMap is called the first time during startup in the main thread.
-		check(IsInGameThread());
+		using namespace UE::Color;
+		const bool bWorkingColorSpaceIsSRGB = FColorSpace::GetWorking().IsSRGB();
+		SET_SHADER_DEFINE(Input.Environment, WORKING_COLOR_SPACE_IS_SRGB, bWorkingColorSpaceIsSRGB ? 1 : 0);
+		
+		// We limit matrix definitions below to WORKING_COLOR_SPACE_IS_SRGB == 0.
+		if (!bWorkingColorSpaceIsSRGB)
+		{
+			const TCHAR MatrixFormat[] = TEXT("float3x3(%0.10f, %0.10f, %0.10f, %0.10f, %0.10f, %0.10f, %0.10f, %0.10f, %0.10f)");
+			const FColorSpace& WorkingColorSpace = FColorSpace::GetWorking();
 
-		FScopedSlowTask SlowTask(70, LOCTEXT("CreateGlobalShaderMap", "Creating Global Shader Map..."));
+			// Note that we transpose the matrices during print since color matrices are usually pre-multiplied.
+			const FMatrix44d& ToXYZ = WorkingColorSpace.GetRgbToXYZ();
+			Input.Environment.SetDefine(
+				TEXT("WORKING_COLOR_SPACE_RGB_TO_XYZ_MAT"),
+				FString::Printf(MatrixFormat,
+					ToXYZ.M[0][0], ToXYZ.M[1][0], ToXYZ.M[2][0],
+					ToXYZ.M[0][1], ToXYZ.M[1][1], ToXYZ.M[2][1],
+					ToXYZ.M[0][2], ToXYZ.M[1][2], ToXYZ.M[2][2]));
 
-		// verify that all shader source files are intact
-		SlowTask.EnterProgressFrame(20, LOCTEXT("VerifyShaderSourceFiles", "Verifying Global Shader source files..."));
-		VerifyShaderSourceFiles(Platform);
+			const FMatrix44d& FromXYZ = WorkingColorSpace.GetXYZToRgb();
+			Input.Environment.SetDefine(
+				TEXT("XYZ_TO_RGB_WORKING_COLOR_SPACE_MAT"),
+				FString::Printf(MatrixFormat,
+					FromXYZ.M[0][0], FromXYZ.M[1][0], FromXYZ.M[2][0],
+					FromXYZ.M[0][1], FromXYZ.M[1][1], FromXYZ.M[2][1],
+					FromXYZ.M[0][2], FromXYZ.M[1][2], FromXYZ.M[2][2]));
 
-		GGlobalShaderMap[Platform] = new FGlobalShaderMap(Platform);
+			const FColorSpaceTransform& FromSRGB = FColorSpaceTransform::GetSRGBToWorkingColorSpace();
+			SET_SHADER_DEFINE(Input.Environment,
+				SRGB_TO_WORKING_COLOR_SPACE_MAT,
+				FString::Printf(MatrixFormat,
+					FromSRGB.M[0][0], FromSRGB.M[1][0], FromSRGB.M[2][0],
+					FromSRGB.M[0][1], FromSRGB.M[1][1], FromSRGB.M[2][1],
+					FromSRGB.M[0][2], FromSRGB.M[1][2], FromSRGB.M[2][2]));
+		}
 
-		bool bShaderMapIsBeingCompiled = false;
+		static const auto CVar = IConsoleManager::Get().FindConsoleVariable(TEXT("r.LegacyLuminanceFactors"));
+		const bool bUseLegacyLuminance = CVar && CVar->GetInt() != 0;
+		SET_SHADER_DEFINE(Input.Environment, UE_LEGACY_LUMINANCE_FACTORS, bUseLegacyLuminance ? 1 : 0);
+	}
 
-		// Try to load the global shaders from a local cache file if it exists
-		// We always try this first, even when running in the editor or if shader compiler is enabled
-		// It's always possible we'll find a cooked local cache
-		const bool bLoadedFromCacheFile = TryLoadCookedGlobalShaderMap(Platform, SlowTask);
-#if WITH_EDITOR
-		const bool bAllowShaderCompiling = !FPlatformProperties::RequiresCookedData() && AllowShaderCompiling();
-#else
-		const bool bAllowShaderCompiling = false;
-#endif
+	const double TileSize = FLargeWorldRenderScalar::GetTileSize();
+	SET_SHADER_DEFINE(Input.Environment, UE_LWC_RENDER_TILE_SIZE, (float)TileSize);
+	SET_SHADER_DEFINE(Input.Environment, UE_LWC_RENDER_TILE_SIZE_SQRT, (float)FMath::Sqrt(TileSize));
+	SET_SHADER_DEFINE(Input.Environment, UE_LWC_RENDER_TILE_SIZE_RSQRT, (float)FMath::InvSqrt(TileSize));
+	SET_SHADER_DEFINE(Input.Environment, UE_LWC_RENDER_TILE_SIZE_RCP, (float)(1.0 / TileSize));
+	SET_SHADER_DEFINE(Input.Environment, UE_LWC_RENDER_TILE_SIZE_FMOD_PI, (float)FMath::Fmod(TileSize, UE_DOUBLE_PI));
+	SET_SHADER_DEFINE(Input.Environment, UE_LWC_RENDER_TILE_SIZE_FMOD_2PI, (float)FMath::Fmod(TileSize, 2.0 * UE_DOUBLE_PI));
 
-#if WITH_EDITOR
-		if (!bLoadedFromCacheFile && bAllowShaderCompiling)
-		{
-			// If we didn't find cooked shaders, we can try loading from the DDC or compiling them if supported by the current configuration
-			FGlobalShaderMapId ShaderMapId(Platform, TargetPlatform);
+	// Add required symbols from the shader binding layout if set
+	if (Input.Environment.ShaderBindingLayout)
+	{
+		Input.Environment.ShaderBindingLayout->AddRequiredSymbols(Input.RequiredSymbols);
+	}
 
-			const int32 ShaderFilenameNum = ShaderMapId.GetShaderFilenameToDependeciesMap().Num();
-			const float ProgressStep = 25.0f / ShaderFilenameNum;
+	// Allow the target shader format to modify the shader input before we add it as a job
+	const IShaderFormat* Format = GetTargetPlatformManagerRef().FindShaderFormat(ShaderFormatName);
+	checkf(Format, TEXT("Shader format %s cannot be found"), *ShaderFormatName.ToString());
+	Format->ModifyShaderCompilerInput(Input);
 
-			// If NoShaderDDC then don't check for a material the first time we encounter it to simulate
-			// a cold DDC
-			static bool bNoShaderDDC =
-				FParse::Param(FCommandLine::Get(), TEXT("noshaderddc")) || 
-				FParse::Param(FCommandLine::Get(), TEXT("noglobalshaderddc"));
+	// Allow the GBuffer and other shader defines to cause dependend environment changes, but minimizing the #ifdef magic in the shaders, which
+	// is nearly impossible to debug when it goes wrong.
+	FShaderCompileUtilities::ApplyDerivedDefines(Input.Environment, Input.SharedEnvironment, (EShaderPlatform)Target.Platform);
+}
 
-			const bool bTempNoShaderDDC = bNoShaderDDC;
+#endif // WITH_EDITOR
+
+/** Timer class used to report information on the 'recompileshaders' console command. */
+class FRecompileShadersTimer
+{
+public:
+	FRecompileShadersTimer(const TCHAR* InInfoStr=TEXT("Test")) :
+		InfoStr(InInfoStr),
+		bAlreadyStopped(false)
+	{
+		StartTime = FPlatformTime::Seconds();
+	}
 
+	FRecompileShadersTimer(const FString& InInfoStr) :
+		InfoStr(InInfoStr),
+		bAlreadyStopped(false)
+	{
+		StartTime = FPlatformTime::Seconds();
+	}
+
+	void Stop(bool DisplayLog = true)
+	{
+		if (!bAlreadyStopped)
+		{
+			bAlreadyStopped = true;
+			EndTime = FPlatformTime::Seconds();
+			TimeElapsed = EndTime-StartTime;
+			if (DisplayLog)
 			{
-				using namespace UE::DerivedData;
+				UE_LOG(LogShaderCompilers, Warning, TEXT("		[%s] took [%.4f] s"),*InfoStr,TimeElapsed);
+			}
+		}
+	}
 
-				int32 BufferIndex = 0;
-				TArray<FCacheGetValueRequest> Requests;
+	~FRecompileShadersTimer()
+	{
+		Stop(true);
+	}
 
-				// Submit DDC requests.
-				SlowTask.EnterProgressFrame(ProgressStep, LOCTEXT("SubmitDDCRequests", "Submitting global shader DDC Requests..."));
-				for (const auto& ShaderFilenameDependencies : ShaderMapId.GetShaderFilenameToDependeciesMap())
-				{
-					FCacheGetValueRequest& Request = Requests.AddDefaulted_GetRef();
-					Request.Name = GetGlobalShaderMapName(ShaderMapId, Platform, ShaderFilenameDependencies.Key);
-					Request.Key = GetGlobalShaderMapKey(ShaderMapId, Platform, TargetPlatform, ShaderFilenameDependencies.Value);
-					Request.UserData = uint64(BufferIndex);
-					++BufferIndex;
+protected:
+	double StartTime,EndTime;
+	double TimeElapsed;
+	FString InfoStr;
+	bool bAlreadyStopped;
+};
 
-					if (UNLIKELY(ShouldDumpShaderDDCKeys()))
-					{
-						const FString DataKey = GetGlobalShaderMapKeyString(ShaderMapId, Platform, ShaderFilenameDependencies.Value);
-						// For global shaders, we dump the key multiple times (once for each shader type) so they will live on disk alongside
-						// other shader debug artifacts.
-						for (const FShaderTypeDependency& ShaderTypeDependency : ShaderFilenameDependencies.Value)
-						{
-							const FShaderType* ShaderType = FindShaderTypeByName(ShaderTypeDependency.ShaderTypeName);
-							TStringBuilder<128> GroupNameBuilder;
-							GroupNameBuilder << TEXT("Global");
-							FPathViews::Append(GroupNameBuilder, ShaderType->GetName());
-							DumpShaderDDCKeyToFile(Platform, ShaderMapId.WithEditorOnly(), GroupNameBuilder.ToString(), DataKey);
-						}
-					}
-				}
+namespace
+{
+	void ListAllShaderTypes()
+	{
+		UE_LOG(LogShaderCompilers, Display, TEXT("ShaderTypeName, Filename"));
+		for (TLinkedList<FShaderType*>::TIterator It(FShaderType::GetTypeList()); It; It.Next())
+		{
+			UE_LOG(LogShaderCompilers, Display, TEXT("%s, %s "), (*It)->GetName(), (*It)->GetShaderFilename());
+		}
 
-				int32 DDCHits = 0;
-				int32 DDCMisses = 0;
+		UE_LOG(LogShaderCompilers, Display, TEXT("VertexFactoryTypeName, Filename"));
+		for (TLinkedList<FVertexFactoryType*>::TIterator It(FVertexFactoryType::GetTypeList()); It; It.Next())
+		{
+			UE_LOG(LogShaderCompilers, Display, TEXT("%s, %s"), (*It)->GetName(), (*It)->GetShaderFilename());
+		}
+	}
 
-				// Process finished DDC requests.
-				SlowTask.EnterProgressFrame(ProgressStep, LOCTEXT("ProcessDDCRequests", "Processing global shader DDC requests..."));
-				TArray<FValue> GlobalShaderMapBuffers;
-				GlobalShaderMapBuffers.SetNum(Requests.Num());
-				{
-					COOK_STAT(auto Timer = GlobalShaderCookStats::UsageStats.TimeSyncWork());
-					COOK_STAT(Timer.TrackCyclesOnly());
-					FRequestOwner BlockingOwner(EPriority::Blocking);
-					GetCache().GetValue(Requests, BlockingOwner, [&GlobalShaderMapBuffers, &bTempNoShaderDDC](FCacheGetValueResponse&& Response)
-					{
-						if (bTempNoShaderDDC)
-						{
-							return;
-						}
+	ODSCRecompileCommand ParseRecompileCommandString(const TCHAR* CmdString, TArray<FString>& OutMaterialsToLoad, FString& OutShaderTypesToLoad, FString& OutRequestedMaterialName)
+	{
+		FString CmdName = FParse::Token(CmdString, 0);
+
+		ODSCRecompileCommand CommandType = ODSCRecompileCommand::None;
+		OutMaterialsToLoad.Empty();
+
+		if( !CmdName.IsEmpty() && FCString::Stricmp(*CmdName,TEXT("Material"))==0 )
+		{
+			CommandType = ODSCRecompileCommand::Material;
 
-						GlobalShaderMapBuffers[int32(Response.UserData)] = MoveTemp(Response.Value);
-					});
-					BlockingOwner.Wait();
+			// tell other side the material to load, by pathname
+			FString RequestedMaterialName( FParse::Token( CmdString, 0 ) );
+			OutRequestedMaterialName = RequestedMaterialName;
+			UMaterialInterface* MatchingMaterial = nullptr;
+			for (TObjectIterator<UMaterialInterface> It; It; ++It)
+			{
+				UMaterial* Material = It->GetMaterial();
+
+				if (Material && Material->GetName() == RequestedMaterialName)
+				{
+					OutMaterialsToLoad.Add(It->GetPathName());
+					MatchingMaterial = Material;
+					break;
 				}
+			}
 
-				BufferIndex = 0;
-				for (const auto& ShaderFilenameDependencies : ShaderMapId.GetShaderFilenameToDependeciesMap())
+			// Find all material instances from the requested material and 
+			// request a compile for them.
+			if (MatchingMaterial)
+			{
+				for (TObjectIterator<UMaterialInstance> It; It; ++It)
 				{
-					COOK_STAT(auto Timer = GlobalShaderCookStats::UsageStats.TimeSyncWork());
-					if (GlobalShaderMapBuffers[BufferIndex].HasData())
-					{
-						COOK_STAT(Timer.AddHit(int64(GlobalShaderMapBuffers[BufferIndex].GetRawSize())));
-						const FSharedBuffer CachedData = GlobalShaderMapBuffers[BufferIndex].GetData().Decompress();
-						FMemoryReaderView MemoryReader(CachedData);
-						GGlobalShaderMap[Platform]->AddSection(FGlobalShaderMapSection::CreateFromArchive(MemoryReader));
-						DDCHits++;
-					}
-					else
+					if (It && It->IsDependent(MatchingMaterial))
 					{
-						// it's a miss, but we haven't built anything yet. Save the counting until we actually have it built.
-						COOK_STAT(Timer.TrackCyclesOnly());
-						bShaderMapIsBeingCompiled = true;
-						DDCMisses++;
+						OutMaterialsToLoad.Add(It->GetPathName());
 					}
-					++BufferIndex;
 				}
-
-				GShaderCompilerStats->AddDDCHit(DDCHits);
-				GShaderCompilerStats->AddDDCMiss(DDCMisses);
 			}
 		}
-#endif // WITH_EDITOR
-		
-		if (!bLoadedFromCacheFile && !bAllowShaderCompiling)
+		else if (!CmdName.IsEmpty() && FCString::Stricmp(*CmdName, TEXT("Global")) == 0)
 		{
-			// Failed to load cooked shaders, and no support for compiling
-			// Handle this gracefully and exit.
-			const FString GlobalShaderCacheFilename = FPaths::GetRelativePathToRoot() / GetGlobalShaderCacheFilename(Platform);
-			const FString SandboxPath = IFileManager::Get().ConvertToAbsolutePathForExternalAppForWrite(*GlobalShaderCacheFilename);
-			// This can be too early to localize in some situations.
-			const FText Message = FText::Format(NSLOCTEXT("Engine", "GlobalShaderCacheFileMissing", "The global shader cache file '{0}' is missing.\n\nYour application is built to load COOKED content. No COOKED content was found; This usually means you did not cook content for this build.\nIt also may indicate missing cooked data for a shader platform(e.g., OpenGL under Windows): Make sure your platform's packaging settings include this Targeted RHI.\n\nAlternatively build and run the UNCOOKED version instead."), FText::FromString(SandboxPath));
-			if (FPlatformProperties::SupportsWindowedMode())
+			CommandType = ODSCRecompileCommand::Global;
+		}
+		else if (!CmdName.IsEmpty() && FCString::Stricmp(*CmdName, TEXT("Changed")) == 0)
+		{
+			CommandType = ODSCRecompileCommand::Changed;
+
+			// Compile all the shaders that have changed for the materials we have loaded.
+			for (TObjectIterator<UMaterialInterface> It; It; ++It)
 			{
-				UE_LOG(LogShaders, Error, TEXT("%s"), *Message.ToString());
-				FMessageDialog::Open(EAppMsgType::Ok, Message);
-				FPlatformMisc::RequestExit(false, TEXT("CompileGlobalShaderMap"));
-				return;
+				OutMaterialsToLoad.Add(It->GetPathName());
 			}
-			else
+		}
+		else if (FCString::Stricmp(*CmdName, TEXT("All")) == 0)
+		{
+			CommandType = ODSCRecompileCommand::Material;
+
+			// tell other side all the materials to load, by pathname
+			for (TObjectIterator<UMaterialInterface> It; It; ++It)
 			{
-				UE_LOG(LogShaders, Fatal, TEXT("%s"), *Message.ToString());
+				OutMaterialsToLoad.Add(It->GetPathName());
 			}
 		}
-
-		// If any shaders weren't loaded, compile them now.
-		VerifyGlobalShaders(Platform, TargetPlatform, bLoadedFromCacheFile);
-
-		if (GCreateShadersOnLoad && Platform == GMaxRHIShaderPlatform)
+		else if (FCString::Stricmp(*CmdName, TEXT("listtypes")) == 0)
 		{
-			GGlobalShaderMap[Platform]->BeginCreateAllShaders();
+			ListAllShaderTypes();
 		}
-
-		// While we're early in the game's startup, create certain global shaders that may be later created on random threads otherwise. 
-		if (!bShaderMapIsBeingCompiled && !GRHISupportsMultithreadedShaderCreation)
+		else
 		{
-			ENQUEUE_RENDER_COMMAND(CreateRecursiveShaders)([](FRHICommandListImmediate&)
+			CommandType = ODSCRecompileCommand::SingleShader;
+
+			OutShaderTypesToLoad = CmdName;
+
+			// tell other side which materials to load and compile the single
+			// shader for.
+			for (TObjectIterator<UMaterialInterface> It; It; ++It)
 			{
-				CreateRecursiveShaders();
-			});
+				OutMaterialsToLoad.Add(It->GetPathName());
+			}
 		}
+
+		return CommandType;
 	}
 }
 
-void CompileGlobalShaderMap(EShaderPlatform Platform, bool bRefreshShaderMap)
-{
-	CompileGlobalShaderMap(Platform, nullptr, bRefreshShaderMap);
-}
+static int32 GODSCMaterialUpdateFlags = 0;
+static FAutoConsoleVariableRef CVarODSCMaterialUpdateFlags(
+	TEXT("ODSC.MaterialUpdateFlags"),
+	GODSCMaterialUpdateFlags,
+	TEXT("Changes the material update flags when ODSC receives new shaders and needs to update the materials\n")
+	TEXT("0 (default): no additional work\n")
+	TEXT("1: Reregister all components while updating the material\n")
+	TEXT("2: Sync with the rendering thread after all the calls to RecacheUniformExpressions\n")
+	TEXT("4 (legacy): Recreates only the render state for *all* components, including the ones not changed by ODSC\n")
+);
 
-void CompileGlobalShaderMap(ERHIFeatureLevel::Type InFeatureLevel, bool bRefreshShaderMap)
+void ProcessCookOnTheFlyShaders(bool bReloadGlobalShaders, const TArray<uint8>& MeshMaterialMaps, const TArray<FString>& MaterialsToLoad, const TArray<uint8>& GlobalShaderMap)
 {
-	EShaderPlatform Platform = GShaderPlatformForFeatureLevel[InFeatureLevel];
-	CompileGlobalShaderMap(Platform, nullptr, bRefreshShaderMap);
-}
+	TRACE_CPUPROFILER_EVENT_SCOPE(ProcessCookOnTheFlyShaders);
+	check(IsInGameThread());
 
-void CompileGlobalShaderMap(bool bRefreshShaderMap)
-{
-	CompileGlobalShaderMap(GMaxRHIFeatureLevel, bRefreshShaderMap);
-}
+	bool bHasFlushed = false;
 
-void ShutdownGlobalShaderMap()
-{
-	// handle edge case where we get a shutdown before fully initialized (the globals used below are not in a valid state)
-	if (!GIsRHIInitialized)
+	auto DoFlushIfNecessary = [&bHasFlushed]() {
+		if (!bHasFlushed )
+		{
+			// now we need to refresh the RHI resources
+			FlushRenderingCommands();
+			bHasFlushed = true;
+		}
+	};
+
+	// reload the global shaders
+	if (bReloadGlobalShaders)
 	{
-		return;
+		DoFlushIfNecessary();
+
+		// Some platforms rely on global shaders to be created to implement basic RHI functionality
+		TGuardValue<int32> Guard(GCreateShadersOnLoad, 1);
+		CompileGlobalShaderMap(true);
 	}
 
-	// at the point this function is called (during the shutdown process) we do not expect any outstanding work that could potentially be still referencing
-	// global shaders, so we are not deferring the deletion (via GGlobalShaderMap_DeferredDeleteCopy) like we do during the shader recompilation.
-	EShaderPlatform Platform = GShaderPlatformForFeatureLevel[GMaxRHIFeatureLevel];
-	if (GGlobalShaderMap[Platform] != nullptr)
+	// load all the mesh material shaders if any were sent back
+	if (MeshMaterialMaps.Num() > 0)
 	{
-		GGlobalShaderMap[Platform]->ReleaseAllSections();
+		DoFlushIfNecessary();
 
-		delete GGlobalShaderMap[Platform];
-		GGlobalShaderMap[Platform] = nullptr;
+		// parse the shaders
+		FMemoryReader MemoryReader(MeshMaterialMaps, true);
+		FNameAsStringProxyArchive Ar(MemoryReader);
+
+		TArray<UMaterialInterface*> LoadedMaterials;
+		FMaterialShaderMap::LoadForRemoteRecompile(Ar, GMaxRHIShaderPlatform, LoadedMaterials);
+
+		// Only update materials if we need to.
+		if (LoadedMaterials.Num())
+		{
+			// this will stop the rendering thread, and reattach components, in the destructor
+			FMaterialUpdateContext UpdateContext(GODSCMaterialUpdateFlags);
+
+			// gather the shader maps to reattach
+			for (UMaterialInterface* Material : LoadedMaterials)
+			{
+				// ~FMaterialUpdateContext takes care of calling RecacheUniformExpressions on all MaterialInstances, no need to call it twice
+				if (Cast<UMaterialInstance>(Material) == nullptr)
+				{
+					Material->RecacheUniformExpressions(true);
+				}
+
+				UpdateContext.AddMaterialInterface(Material);
+			}
+		}
+	}
+
+	// load all the global shaders if any were sent back
+	if (GlobalShaderMap.Num() > 0)
+	{
+		DoFlushIfNecessary();
+
+		// parse the shaders
+		FMemoryReader MemoryReader(GlobalShaderMap, true);
+		FNameAsStringProxyArchive Ar(MemoryReader);
+
+		LoadGlobalShadersForRemoteRecompile(Ar, GMaxRHIShaderPlatform);
 	}
 }
 
-void ReloadGlobalShaders()
+/**
+* Forces a recompile of the global shaders.
+*/
+void RecompileGlobalShaders()
 {
-	UE_LOG(LogShaders, Display, TEXT("Reloading global shaders..."));
-
-	// Flush pending accesses to the existing global shaders.
-	FlushRenderingCommands();
+#if WITH_EDITOR
+	if (!FPlatformProperties::RequiresCookedData())
+	{
+		// Flush pending accesses to the existing global shaders.
+		FlushRenderingCommands();
 
-	UMaterialInterface::IterateOverActiveFeatureLevels(
-		[&](ERHIFeatureLevel::Type InFeatureLevel)
-		{
+		UMaterialInterface::IterateOverActiveFeatureLevels([&](ERHIFeatureLevel::Type InFeatureLevel) {
 			auto ShaderPlatform = GShaderPlatformForFeatureLevel[InFeatureLevel];
-			GetGlobalShaderMap(ShaderPlatform)->ReleaseAllSections();
-			CompileGlobalShaderMap(InFeatureLevel, true);
+			GetGlobalShaderMap(ShaderPlatform)->Empty();
 			VerifyGlobalShaders(ShaderPlatform, nullptr, false);
-		}
-	);
+		});
 
-	// Invalidate global bound shader states so they will be created with the new shaders the next time they are set (in SetGlobalBoundShaderState)
-	for (TLinkedList<FGlobalBoundShaderStateResource*>::TIterator It(FGlobalBoundShaderStateResource::GetGlobalBoundShaderStateList()); It; It.Next())
-	{
-		BeginUpdateResourceRHI(*It);
+		GShaderCompilingManager->ProcessAsyncResults(false, true);
 	}
-
-	PropagateGlobalShadersToAllPrimitives();
+#endif // WITH_EDITOR
 }
 
-static FAutoConsoleCommand CCmdReloadGlobalShaders = FAutoConsoleCommand(
-	TEXT("ReloadGlobalShaders"),
-	TEXT("Reloads the global shaders file"),
-	FConsoleCommandDelegate::CreateStatic(ReloadGlobalShaders)
-);
-
-void SetGlobalShaderCacheOverrideDirectory(const TArray<FString>& Args)
+void GetOutdatedShaderTypes(TArray<const FShaderType*>& OutdatedShaderTypes, TArray<const FShaderPipelineType*>& OutdatedShaderPipelineTypes, TArray<const FVertexFactoryType*>& OutdatedFactoryTypes)
 {
-	if (Args.Num() == 0)
+#if WITH_EDITOR
+	for (int PlatformIndex = 0; PlatformIndex < SP_NumPlatforms; ++PlatformIndex)
+	{
+		const FGlobalShaderMap* ShaderMap = GGlobalShaderMap[PlatformIndex];
+		if (ShaderMap)
+		{
+			ShaderMap->GetOutdatedTypes(OutdatedShaderTypes, OutdatedShaderPipelineTypes, OutdatedFactoryTypes);
+		}
+	}
+
+	FMaterialShaderMap::GetAllOutdatedTypes(OutdatedShaderTypes, OutdatedShaderPipelineTypes, OutdatedFactoryTypes);
+
+	for (int32 TypeIndex = 0; TypeIndex < OutdatedShaderTypes.Num(); TypeIndex++)
+	{
+		UE_LOG(LogShaders, Warning, TEXT("		Recompiling %s"), OutdatedShaderTypes[TypeIndex]->GetName());
+	}
+	for (int32 TypeIndex = 0; TypeIndex < OutdatedShaderPipelineTypes.Num(); TypeIndex++)
 	{
-		UE_LOG(LogShaders, Error, TEXT("Failed to set GGlobalShaderCacheOverrideDirectory without any arguments"));
-		return; 
+		UE_LOG(LogShaders, Warning, TEXT("		Recompiling %s"), OutdatedShaderPipelineTypes[TypeIndex]->GetName());
 	}
-	
-	GGlobalShaderCacheOverrideDirectory = Args[0];
-	UE_LOG(LogShaders, Log, TEXT("GGlobalShaderCacheOverrideDirectory = %s"), *GGlobalShaderCacheOverrideDirectory);
+	for (int32 TypeIndex = 0; TypeIndex < OutdatedFactoryTypes.Num(); TypeIndex++)
+	{
+		UE_LOG(LogShaders, Warning, TEXT("		Recompiling %s"), OutdatedFactoryTypes[TypeIndex]->GetName());
+	}
+#endif // WITH_EDITOR
 }
 
-static FAutoConsoleCommand CCmdSetGlobalShaderCacheOverrideDirectory = FAutoConsoleCommand(
-	TEXT("SetGlobalShaderCacheOverrideDirectory"),
-	TEXT("Set the directory to read the override global shader map file from."),
-	FConsoleCommandWithArgsDelegate::CreateStatic(SetGlobalShaderCacheOverrideDirectory));
-
-bool RecompileChangedShadersForPlatform(const FString& PlatformName)
+bool RecompileShaders(const TCHAR* Cmd, FOutputDevice& Ar)
 {
-	// figure out what shader platforms to recompile
-	ITargetPlatformManagerModule* TPM = GetTargetPlatformManager();
-	ITargetPlatform* TargetPlatform = TPM->FindTargetPlatform(PlatformName);
-	if (TargetPlatform == nullptr)
+	// if this platform can't compile shaders, then we try to send a message to a file/cooker server
+	if (FPlatformProperties::RequiresCookedData())
 	{
-		UE_LOG(LogShaders, Display, TEXT("Failed to find target platform module for %s"), *PlatformName);
-		return false;
+#if WITH_ODSC
+		TArray<FString> MaterialsToLoad;
+		FString ShaderTypesToLoad;
+		FString RequestedMaterialName;
+		ODSCRecompileCommand CommandType = ParseRecompileCommandString(Cmd, MaterialsToLoad, ShaderTypesToLoad, RequestedMaterialName);
+
+		ERHIFeatureLevel::Type TargetFeatureLevel = GetMaxSupportedFeatureLevel(GMaxRHIShaderPlatform);
+		const EMaterialQualityLevel::Type ActiveQualityLevel = GetCachedScalabilityCVars().MaterialQualityLevel;
+		GODSCManager->AddThreadedRequest(MaterialsToLoad, ShaderTypesToLoad, GMaxRHIShaderPlatform, TargetFeatureLevel, ActiveQualityLevel, CommandType, RequestedMaterialName);
+#endif
+		return true;
 	}
 
-	TArray<FName> DesiredShaderFormats;
-	TargetPlatform->GetAllTargetedShaderFormats(DesiredShaderFormats);
+#if WITH_EDITOR
+	FString FlagStr(FParse::Token(Cmd, 0));
+	if (FlagStr.Len() > 0)
+	{
+		TRACE_CPUPROFILER_EVENT_SCOPE(RecompileShaders);
+		GWarn->BeginSlowTask( NSLOCTEXT("ShaderCompilingManager", "BeginRecompilingShadersTask", "Recompiling shaders"), true );
 
-	// figure out which shaders are out of date
-	TArray<const FShaderType*> OutdatedShaderTypes;
-	TArray<const FVertexFactoryType*> OutdatedFactoryTypes;
-	TArray<const FShaderPipelineType*> OutdatedShaderPipelineTypes;
+		// Flush the shader file cache so that any changes to shader source files will be detected
+		FlushShaderFileCache();
+		FlushRenderingCommands();
 
-	// Pick up new changes to shader files
-	FlushShaderFileCache();
+		if (FCString::Stricmp(*FlagStr,TEXT("Changed")) == 0)
+		{
+			TArray<const FShaderType*> OutdatedShaderTypes;
+			TArray<const FVertexFactoryType*> OutdatedFactoryTypes;
+			TArray<const FShaderPipelineType*> OutdatedShaderPipelineTypes;
+			{
+				FRecompileShadersTimer SearchTimer(TEXT("Searching for changed files"));
+				GetOutdatedShaderTypes(OutdatedShaderTypes, OutdatedShaderPipelineTypes, OutdatedFactoryTypes);
+			}
 
-	GetOutdatedShaderTypes(OutdatedShaderTypes, OutdatedShaderPipelineTypes, OutdatedFactoryTypes);
-	UE_LOG(LogShaders, Display, TEXT("We found %d out of date shader types, %d outdated pipeline types, and %d out of date VF types!"), OutdatedShaderTypes.Num(), OutdatedShaderPipelineTypes.Num(), OutdatedFactoryTypes.Num());
+			if (OutdatedShaderPipelineTypes.Num() > 0 || OutdatedShaderTypes.Num() > 0 || OutdatedFactoryTypes.Num() > 0)
+			{
+				FRecompileShadersTimer TestTimer(TEXT("RecompileShaders Changed"));
 
-#if WITH_EDITOR
-	UpdateReferencedUniformBufferNames(OutdatedShaderTypes, OutdatedFactoryTypes, OutdatedShaderPipelineTypes);
-#endif
+				UpdateReferencedUniformBufferNames(OutdatedShaderTypes, OutdatedFactoryTypes, OutdatedShaderPipelineTypes);
 
-	for (int32 FormatIndex = 0; FormatIndex < DesiredShaderFormats.Num(); FormatIndex++)
-	{
-		// get the shader platform enum
-		const EShaderPlatform ShaderPlatform = ShaderFormatToLegacyShaderPlatform(DesiredShaderFormats[FormatIndex]);
+				// Kick off global shader recompiles
+				UMaterialInterface::IterateOverActiveFeatureLevels([&](ERHIFeatureLevel::Type InFeatureLevel) {
+					auto ShaderPlatform = GShaderPlatformForFeatureLevel[InFeatureLevel];
+					BeginRecompileGlobalShaders(OutdatedShaderTypes, OutdatedShaderPipelineTypes, ShaderPlatform);
+					// Block on global shader compilation. Do this for each feature level/platform compiled as otherwise global shader compile job IDs collide.
+					FinishRecompileGlobalShaders();
+				});
 
-		// Only compile for the desired platform if requested
-		// Kick off global shader recompiles
-		BeginRecompileGlobalShaders(OutdatedShaderTypes, OutdatedShaderPipelineTypes, ShaderPlatform);
+				// Kick off material shader recompiles
+				UMaterialInterface::IterateOverActiveFeatureLevels([&](ERHIFeatureLevel::Type InFeatureLevel) {
+					auto ShaderPlatform = GShaderPlatformForFeatureLevel[InFeatureLevel];
+					UMaterial::UpdateMaterialShaders(OutdatedShaderTypes, OutdatedShaderPipelineTypes, OutdatedFactoryTypes, ShaderPlatform);
+				});
 
-		// Block on global shaders
-		FinishRecompileGlobalShaders();
-#if WITH_EDITOR
-		// we only want to actually compile mesh shaders if we have out of date ones
-		if (OutdatedShaderTypes.Num() || OutdatedFactoryTypes.Num())
-		{
-			for (TObjectIterator<UMaterialInterface> It; It; ++It)
+				GWarn->StatusUpdate(0, 1, NSLOCTEXT("ShaderCompilingManager", "CompilingGlobalShaderStatus", "Compiling global shaders..."));
+			}
+			else
 			{
-				(*It)->ClearCachedCookedPlatformData(TargetPlatform);
+				UE_LOG(LogShaderCompilers, Warning, TEXT("No Shader changes found."));
 			}
 		}
-#endif // WITH_EDITOR
-	}
-
-	if (OutdatedFactoryTypes.Num() || OutdatedShaderTypes.Num())
-	{
-		return true;
-	}
-	return false;
-}
+		else if (FCString::Stricmp(*FlagStr, TEXT("Global")) == 0)
+		{
+			FRecompileShadersTimer TestTimer(TEXT("RecompileShaders Global"));
+			RecompileGlobalShaders();
+		}
+		else if (FCString::Stricmp(*FlagStr, TEXT("Material")) == 0)
+		{
+			FString RequestedMaterialName(FParse::Token(Cmd, 0));
+			FRecompileShadersTimer TestTimer(FString::Printf(TEXT("Recompile Material %s"), *RequestedMaterialName));
 
-FArchive& operator<<(FArchive& Ar, FODSCRequestPayload& Payload)
-{
-	int32 iShaderPlatform = static_cast<int32>(Payload.ShaderPlatform);
-	int32 iFeatureLevel = static_cast<int32>(Payload.FeatureLevel);
-	int32 iQualityLevel = static_cast<int32>(Payload.QualityLevel);
+			ITargetPlatformManagerModule& TPM = GetTargetPlatformManagerRef();
+			FString TargetPlatformName(FParse::Token(Cmd, 0));
+			const ITargetPlatform* TargetPlatform = nullptr;
+			if (TargetPlatformName.Len() > 0)
+			{
+				TargetPlatform = TPM.FindTargetPlatform(TargetPlatformName);
+			}
 
-	Ar << iShaderPlatform;
-	Ar << iFeatureLevel;
-	Ar << iQualityLevel;
-	Ar << Payload.MaterialName;
-	Ar << Payload.VertexFactoryName;
-	Ar << Payload.PipelineName;
-	Ar << Payload.ShaderTypeNames;
-	Ar << Payload.PermutationId;
-	Ar << Payload.RequestHash;
+			bool bMaterialFound = false;
+			for( TObjectIterator<UMaterialInterface> It; It; ++It )
+			{
+				UMaterialInterface* Material = *It;
+				if( Material && Material->GetName() == RequestedMaterialName)
+				{
+					bMaterialFound = true;
 
-	if (Ar.IsLoading())
-	{
-		Payload.ShaderPlatform = static_cast<EShaderPlatform>(iShaderPlatform);
-		Payload.FeatureLevel = static_cast<ERHIFeatureLevel::Type>(iFeatureLevel);
-		Payload.QualityLevel = static_cast<EMaterialQualityLevel::Type>(iQualityLevel);
-	}
+					// <Pre/Post>EditChange will force a re-creation of the resource,
+					// in turn recompiling the shader.
+					if (TargetPlatform)
+					{
+						Material->BeginCacheForCookedPlatformData(TargetPlatform);
+						while (!Material->IsCachedCookedPlatformDataLoaded(TargetPlatform))
+						{
+							FPlatformProcess::Sleep(0.1f);
+							GShaderCompilingManager->ProcessAsyncResults(false, false);
+						}
+						Material->ClearCachedCookedPlatformData(TargetPlatform);
+					}
+					else
+					{
+						Material->PreEditChange(nullptr);
+						Material->PostEditChange();
+					}
 
-	return Ar;
-}
+					break;
+				}
+			}
 
-FShaderRecompileData::FShaderRecompileData(const FString& InPlatformName, TArray<FString>* OutModifiedFiles, TArray<uint8>* OutMeshMaterialMaps, TArray<uint8>* OutGlobalShaderMap)
-: PlatformName(InPlatformName),
-  ModifiedFiles(OutModifiedFiles),
-  MeshMaterialMaps(OutMeshMaterialMaps),
-  GlobalShaderMap(OutGlobalShaderMap)
-{
-}
+			if (!bMaterialFound)
+			{
+				TestTimer.Stop(false);
+				UE_LOG(LogShaderCompilers, Warning, TEXT("Couldn't find Material %s!"), *RequestedMaterialName);
+			}
+		}
+		else if (FCString::Stricmp(*FlagStr, TEXT("All")) == 0)
+		{
+			FRecompileShadersTimer TestTimer(TEXT("RecompileShaders"));
+			RecompileGlobalShaders();
 
-FShaderRecompileData::FShaderRecompileData(const FString& InPlatformName, EShaderPlatform InShaderPlatform, ODSCRecompileCommand InCommandType, TArray<FString>* OutModifiedFiles, TArray<uint8>* OutMeshMaterialMaps, TArray<uint8>* OutGlobalShaderMap)
-: PlatformName(InPlatformName),
-  ShaderPlatform(InShaderPlatform),
-  ModifiedFiles(OutModifiedFiles),
-  MeshMaterialMaps(OutMeshMaterialMaps),
-  CommandType(InCommandType),
-  GlobalShaderMap(OutGlobalShaderMap)
-{
-}
+			FMaterialUpdateContext UpdateContext(0);
+			for( TObjectIterator<UMaterial> It; It; ++It )
+			{
+				UMaterial* Material = *It;
+				if( Material )
+				{
+					UE_LOG(LogShaderCompilers, Log, TEXT("recompiling [%s]"),*Material->GetFullName());
+					UpdateContext.AddMaterial(Material);
 
-FArchive& operator<<(FArchive& Ar, FShaderRecompileData& RecompileData)
-{
+					// <Pre/Post>EditChange will force a re-creation of the resource,
+					// in turn recompiling the shader.
+					Material->PreEditChange(nullptr);
+					Material->PostEditChange();
+				}
+			}
+		}
+		else if (FCString::Stricmp(*FlagStr, TEXT("listtypes")) == 0)
+		{
+			ListAllShaderTypes();
+		}
+		else
+		{
+			TArray<const FShaderType*> ShaderTypes = FShaderType::GetShaderTypesByFilename(*FlagStr);
+			TArray<const FShaderPipelineType*> ShaderPipelineTypes = FShaderPipelineType::GetShaderPipelineTypesByFilename(*FlagStr);
+			if (ShaderTypes.Num() > 0 || ShaderPipelineTypes.Num() > 0)
+			{
+				FRecompileShadersTimer TestTimer(TEXT("RecompileShaders SingleShader"));
+				
+				UpdateReferencedUniformBufferNames(ShaderTypes, {}, ShaderPipelineTypes);
 
-	int32 iShaderPlatform = static_cast<int32>(RecompileData.ShaderPlatform);
-	int32 iFeatureLevel = static_cast<int32>(RecompileData.FeatureLevel);
-	int32 iQualityLevel = static_cast<int32>(RecompileData.QualityLevel);
+				UMaterialInterface::IterateOverActiveFeatureLevels([&](ERHIFeatureLevel::Type InFeatureLevel) {
+					auto ShaderPlatform = GShaderPlatformForFeatureLevel[InFeatureLevel];
+					BeginRecompileGlobalShaders(ShaderTypes, ShaderPipelineTypes, ShaderPlatform);
+					FinishRecompileGlobalShaders();
+				});
+			}
+		}
 
-	Ar << RecompileData.MaterialsToLoad;
-	Ar << RecompileData.ShaderTypesToLoad;
-	Ar << iShaderPlatform;
-	Ar << iFeatureLevel;
-	Ar << iQualityLevel;
-	Ar << RecompileData.CommandType;
-	Ar << RecompileData.ShadersToRecompile;
+		GWarn->EndSlowTask();
 
-	if (Ar.IsLoading())
-	{
-		RecompileData.ShaderPlatform = static_cast<EShaderPlatform>(iShaderPlatform);
-		RecompileData.FeatureLevel = static_cast<ERHIFeatureLevel::Type>(iFeatureLevel);
-		RecompileData.QualityLevel = static_cast<EMaterialQualityLevel::Type>(iQualityLevel);
+		return true;
 	}
 
-	return Ar;
+	UE_LOG(LogShaderCompilers, Warning, TEXT("Invalid parameter. \n"
+											 "Options are: \n"
+											 "    'Changed'             Recompile just the shaders that have source file changes.\n"
+											 "    'Global'              Recompile just the global shaders.\n"
+											 "    'Material [name]'     Recompile all the shaders for a single material.\n"
+											 "    'Listtypes'           List all the shader type and vertex factory type class names and their source file path.  Can be used to find shader file names to be used with `recompileshaders [shaderfilename]`.\n"
+											 "    'All'                 Recompile all materials.\n"
+											 "    [shaderfilename]      Compile the single shader associated with a specific filename.\n"
+											 ));
+#endif // WITH_EDITOR
+
+	return true;
 }
 
-extern ENGINE_API const TCHAR* ODSCCmdEnumToString(ODSCRecompileCommand Cmd)
+#if WITH_EDITORONLY_DATA
+namespace ShaderCompilerUtil
 {
-	switch (Cmd)
-	{
-	case ODSCRecompileCommand::None:
-		return TEXT("None");
-	case ODSCRecompileCommand::Changed:
-		return TEXT("Change");
-	case ODSCRecompileCommand::Global:
-		return TEXT("Global");
-	case ODSCRecompileCommand::Material:
-		return TEXT("Material");
-	case ODSCRecompileCommand::SingleShader:
-		return TEXT("SingleShader");
-	case ODSCRecompileCommand::ResetMaterialCache:
-			return TEXT("ResetMaterialCache");
-	}
-	ensure(false);
-	return TEXT("Unknown");
+	FOnGlobalShadersCompilation GOnGlobalShdersCompilationDelegate;
 }
 
-#if WITH_EDITOR
+FOnGlobalShadersCompilation& GetOnGlobalShaderCompilation()
+{
+	return ShaderCompilerUtil::GOnGlobalShdersCompilationDelegate;
+}
+#endif // WITH_EDITORONLY_DATA
 
-void CompileGlobalShaderMapForRemote(
-	const TArray<const FShaderType*>& OutdatedShaderTypes, 
-	const TArray<const FShaderPipelineType*>& OutdatedShaderPipelineTypes, 
-	const EShaderPlatform ShaderPlatform, 
-	const ITargetPlatform* TargetPlatform,
-	TArray<uint8>* OutArray)
+/**
+* Makes sure all global shaders are loaded and/or compiled for the passed in platform.
+* Note: if compilation is needed, this only kicks off the compile.
+*
+* @param	Platform	Platform to verify global shaders for
+*/
+void VerifyGlobalShaders(EShaderPlatform Platform, const ITargetPlatform* TargetPlatform, bool bLoadedFromCacheFile, const TArray<const FShaderType*>* OutdatedShaderTypes, const TArray<const FShaderPipelineType*>* OutdatedShaderPipelineTypes)
 {
-	UE_LOG(LogShaders, Display, TEXT("Recompiling global shaders."));
-
-	// Kick off global shader recompiles
-	BeginRecompileGlobalShaders(OutdatedShaderTypes, OutdatedShaderPipelineTypes, ShaderPlatform, TargetPlatform);
+	SCOPED_LOADTIMER(VerifyGlobalShaders);
 
-	// Block on global shaders
-	FinishRecompileGlobalShaders();
+	check(IsInGameThread());
+	check(!FPlatformProperties::IsServerOnly());
+	check(GGlobalShaderMap[Platform]);
 
-	// Write the shader compilation info to memory, converting FName to strings
-	TOptional<FArchiveCookContext> CookContext;
-	TOptional<FArchiveCookData> CookData;
-	FMemoryWriter MemWriter(*OutArray, true);
-	FNameAsStringProxyArchive Ar(MemWriter);
+	UE_LOG(LogMaterial, Verbose, TEXT("Verifying Global Shaders for %s (%s)"), *LegacyShaderPlatformToShaderFormat(Platform).ToString(), *ShaderCompiler::GetTargetPlatformName(TargetPlatform));
 
-	if (TargetPlatform != nullptr)
+	// Ensure that the global shader map contains all global shader types.
+	FGlobalShaderMap* GlobalShaderMap = GetGlobalShaderMap(Platform);
+	const bool bEmptyMap = GlobalShaderMap->IsEmpty();
+	if (bEmptyMap)
 	{
-		CookContext.Emplace(nullptr /*InPackage*/, UE::Cook::ECookType::Unknown,
-			UE::Cook::ECookingDLC::Unknown, TargetPlatform);
-		CookData.Emplace(*TargetPlatform, *CookContext);
-		Ar.SetCookData(CookData.GetPtrOrNull());
+		UE_LOG(LogShaders, Log, TEXT("	Empty global shader map, recompiling all global shaders"));
 	}
 
-	// save out the global shader map to the byte array
-	SaveGlobalShadersForRemoteRecompile(Ar, ShaderPlatform);
-}
-
-void SaveShaderMapsForRemote(ITargetPlatform* TargetPlatform, const TMap<FString, TArray<TRefCountPtr<FMaterialShaderMap>>>& CompiledShaderMaps, TArray<uint8>* OutArray)
-{
-	// write the shader compilation info to memory, converting fnames to strings
-	TOptional<FArchiveCookContext> CookContext;
-	TOptional<FArchiveCookData> CookData;
-	FMemoryWriter MemWriter(*OutArray, true);
-	FNameAsStringProxyArchive Ar(MemWriter);
+	FPlatformTypeLayoutParameters LayoutParams;
+	LayoutParams.InitializeForPlatform(TargetPlatform);
+	EShaderPermutationFlags PermutationFlags = GetShaderPermutationFlags(LayoutParams);
 
-	if (TargetPlatform != nullptr)
+	// if the target is the current platform, then we are not cooking for another platform, in which case we want to use
+	// the loaded permutation flags that are in the shader map (or the current platform's permutation if it wasn't loaded, 
+	// see the FShaderMapBase constructor)
+	if (bLoadedFromCacheFile)
 	{
-		CookContext.Emplace(nullptr /*InPackage*/, UE::Cook::ECookType::Unknown,
-			UE::Cook::ECookingDLC::Unknown, TargetPlatform);
-		CookData.Emplace(*TargetPlatform, *CookContext);
-		Ar.SetCookData(CookData.GetPtrOrNull());
+		PermutationFlags = GlobalShaderMap->GetFirstSection()->GetPermutationFlags();
 	}
 
-	// save out the shader map to the byte array
-	FMaterialShaderMap::SaveForRemoteRecompile(Ar, CompiledShaderMaps);
-}
-
-void RecompileShadersForRemote(
-	FShaderRecompileData& Args,
-	const FString& OutputDirectory)
-{
-	TRACE_CPUPROFILER_EVENT_SCOPE(RecompileShadersForRemote);
-
-	// figure out what shader platforms to recompile
-	ITargetPlatformManagerModule* TPM = GetTargetPlatformManager();
-	ITargetPlatform* TargetPlatform = TPM->FindTargetPlatform(Args.PlatformName);
-	if (TargetPlatform == nullptr)
+	bool bErrorOnMissing = bLoadedFromCacheFile;
+	if (FPlatformProperties::RequiresCookedData())
 	{
-		UE_LOG(LogShaders, Display, TEXT("Failed to find target platform module for %s"), *Args.PlatformName);
-		return;
+		// We require all shaders to exist on cooked platforms because we can't compile them.
+		bErrorOnMissing = true;
 	}
 
-	const double StartTime = FPlatformTime::Seconds();
-
-	UE_LOG(LogShaders, Display, TEXT(""));
-	UE_LOG(LogShaders, Display, TEXT("********************************"));
-	UE_LOG(LogShaders, Display, TEXT("Received compile shader request %s."), ODSCCmdEnumToString(Args.CommandType));
-
-	const bool bPreviousState = GShaderCompilingManager->IsShaderCompilationSkipped();
-	GShaderCompilingManager->SkipShaderCompilation(false);
-
-	TArray<FName> DesiredShaderFormats;
-	TargetPlatform->GetAllTargetedShaderFormats(DesiredShaderFormats);
+#if WITH_EDITOR
+	// All jobs, single & pipeline
+	TArray<FShaderCommonCompileJobPtr> GlobalShaderJobs;
 
-	UE_LOG(LogShaders, Verbose, TEXT("Loading %d materials..."), Args.MaterialsToLoad.Num());
-	// make sure all materials the client has loaded will be processed
-	TArray<UMaterialInterface*> MaterialsToCompile;
+	// Add the single jobs first
+	TMap<TShaderTypePermutation<const FShaderType>, FShaderCompileJob*> SharedShaderJobs;
+#endif // WITH_EDITOR
 
-	for (int32 Index = 0; Index < Args.MaterialsToLoad.Num(); Index++)
+	for (TLinkedList<FShaderType*>::TIterator ShaderTypeIt(FShaderType::GetTypeList()); ShaderTypeIt; ShaderTypeIt.Next())
 	{
-		UE_LOG(LogShaders, Verbose, TEXT("   --> %s"), *Args.MaterialsToLoad[Index]);
-		MaterialsToCompile.Add(LoadObject<UMaterialInterface>(nullptr, *Args.MaterialsToLoad[Index]));
-	}
-
-	UE_LOG(LogShaders, Verbose, TEXT("  Done!"));
-
-	const uint32 StartTotalShadersCompiled = GShaderCompilerStats->GetTotalShadersCompiled();
-
-	// Pick up new changes to shader files
-	FlushShaderFileCache();
+		FGlobalShaderType* GlobalShaderType = ShaderTypeIt->GetGlobalShaderType();
+		if (!GlobalShaderType)
+		{
+			continue;
+		}
 
-	// If we have an explicit list of shaders to compile from ODSC just compile those.
-	if (Args.ShadersToRecompile.Num() && (Args.MeshMaterialMaps != nullptr))
-	{
-		TMap<FString, TArray<TRefCountPtr<FMaterialShaderMap>>> CompiledShaderMaps;
+		int32 PermutationCountToCompile = 0;
+		for (int32 PermutationId = 0; PermutationId < GlobalShaderType->GetPermutationCount(); PermutationId++)
 		{
-			// UMaterial::CompileODSCMaterialsForRemoteRecompile will call LoadObjects on the material names but doesn't keep them around. Add a GC guard to ensure we can still get them
-			// before they get unloaded, so that the whole chain UMaterial->FMaterial->FMaterialShaderMap is kept intact, and we can merge the next batch of ODSC requests
-			FGCScopeGuard NoGCScopeGuard;
-			UMaterial::CompileODSCMaterialsForRemoteRecompile(Args.ShadersToRecompile, CompiledShaderMaps);
-			if (Args.LoadedMaterialsToRecompile)
+			if (GlobalShaderType->ShouldCompilePermutation(Platform, PermutationId, PermutationFlags))
 			{
-				for (auto Iter : CompiledShaderMaps)
+				bool bOutdated = OutdatedShaderTypes && OutdatedShaderTypes->Contains(GlobalShaderType);
+				TShaderRef<FShader> GlobalShader = GlobalShaderMap->GetShader(GlobalShaderType, PermutationId);
+				if (bOutdated || !GlobalShader.IsValid())
 				{
-					TStrongObjectPtr<UMaterialInterface> MaterialInterface = TStrongObjectPtr<UMaterialInterface>(FindObject<UMaterialInterface>(nullptr, *Iter.Key));
-					if (MaterialInterface)
+					if (bErrorOnMissing)
+					{
+                        if (IsMetalPlatform(GMaxRHIShaderPlatform))
+                        {
+                            check(IsInGameThread());
+                            FMessageDialog::Open(EAppMsgType::Ok, LOCTEXT("NoGlobalShader_Error", "Missing shader permutation. Please make sure cooking was successful and refer to Engine log for details."));
+                        }
+                            UE_LOG(LogShaders, Fatal, TEXT("Missing global shader %s's permutation %i, Please make sure cooking was successful."),
+                                GlobalShaderType->GetName(), PermutationId);
+                    }
+					else
 					{
-						ResetLoaders(MaterialInterface->GetPackage());
-						Args.LoadedMaterialsToRecompile->Add(MaterialInterface);
+#if WITH_EDITOR
+						if (OutdatedShaderTypes)
+						{
+							// Remove old shader, if it exists
+							GlobalShaderMap->RemoveShaderTypePermutaion(GlobalShaderType, PermutationId);
+						}
+
+						// Compile this global shader type.
+						FGlobalShaderTypeCompiler::BeginCompileShader(GlobalShaderType, PermutationId, Platform, PermutationFlags, GlobalShaderJobs);
+						//TShaderTypePermutation<const FShaderType> ShaderTypePermutation(GlobalShaderType, PermutationId);
+						//check(!SharedShaderJobs.Find(ShaderTypePermutation));
+						//SharedShaderJobs.Add(ShaderTypePermutation, Job);
+						PermutationCountToCompile++;
+#endif // WITH_EDITOR
 					}
 				}
 			}
 		}
-		SaveShaderMapsForRemote(TargetPlatform, CompiledShaderMaps, Args.MeshMaterialMaps);
+
+		int32 PermutationCountLimit = 832;	// Nanite culling as of today (2022-01-11) can go up to 832 permutations
+		if (Substrate::IsSubstrateEnabled())
+		{
+			// SUBSTRATE_TODO reduce the number of permutation of FDeferredLightPS.
+			PermutationCountLimit = 1304;	// FDeferredLightPS as of today (2023-12-04)
+		}
+		ensureMsgf(
+			PermutationCountToCompile <= PermutationCountLimit,
+			TEXT("Global shader %s has %i permutations: probably more than it needs."),
+			GlobalShaderType->GetName(), PermutationCountToCompile);
+
+		if (!bEmptyMap && PermutationCountToCompile > 0)
+		{
+			UE_LOG(LogShaders, Log, TEXT("	%s (%i out of %i)"),
+				GlobalShaderType->GetName(), PermutationCountToCompile, GlobalShaderType->GetPermutationCount());
+		}
 	}
-	else
+
+	// Now the pipeline jobs; if it's a shareable pipeline, do not add duplicate jobs
+	for (TLinkedList<FShaderPipelineType*>::TIterator ShaderPipelineIt(FShaderPipelineType::GetTypeList()); ShaderPipelineIt; ShaderPipelineIt.Next())
 	{
-		for (int32 FormatIndex = 0; FormatIndex < DesiredShaderFormats.Num(); FormatIndex++)
+		const FShaderPipelineType* Pipeline = *ShaderPipelineIt;
+		if (Pipeline->IsGlobalTypePipeline())
 		{
-			// get the shader platform enum
-			const EShaderPlatform ShaderPlatform = ShaderFormatToLegacyShaderPlatform(DesiredShaderFormats[FormatIndex]);
-
-			// Only compile for the desired platform if requested
-			if (ShaderPlatform == Args.ShaderPlatform || Args.ShaderPlatform == SP_NumPlatforms)
+			if (FGlobalShaderType::ShouldCompilePipeline(Pipeline, Platform, PermutationFlags)
+				&& (!GlobalShaderMap->HasShaderPipeline(Pipeline) || (OutdatedShaderPipelineTypes && OutdatedShaderPipelineTypes->Contains(Pipeline))))
 			{
-				if (Args.CommandType == ODSCRecompileCommand::SingleShader &&
-					Args.ShaderTypesToLoad.Len() > 0)
+				if (OutdatedShaderPipelineTypes)
 				{
-					TArray<const FShaderType*> ShaderTypes = FShaderType::GetShaderTypesByFilename(*Args.ShaderTypesToLoad);
-					TArray<const FShaderPipelineType*> ShaderPipelineTypes = FShaderPipelineType::GetShaderPipelineTypesByFilename(*Args.ShaderTypesToLoad);
+					// Remove old pipeline
+					GlobalShaderMap->RemoveShaderPipelineType(Pipeline);
+				}
 
-					for (const FShaderType* ShaderType : ShaderTypes)
+				if (bErrorOnMissing)
+				{
+					UE_LOG(LogShaders, Fatal, TEXT("Missing global shader pipeline %s, Please make sure cooking was successful."), Pipeline->GetName());
+				}
+				else
+				{
+#if WITH_EDITOR
+					if (!bEmptyMap)
 					{
-						UE_LOG(LogShaders, Display, TEXT("\t%s..."), ShaderType->GetName());
+						UE_LOG(LogShaders, Log, TEXT("	%s"), Pipeline->GetName());
 					}
 
-					UpdateReferencedUniformBufferNames(ShaderTypes, {}, ShaderPipelineTypes);
-
-					CompileGlobalShaderMapForRemote(ShaderTypes, ShaderPipelineTypes, ShaderPlatform, TargetPlatform, Args.GlobalShaderMap);
-				}
-				else if (Args.CommandType == ODSCRecompileCommand::Global ||
-						 Args.CommandType == ODSCRecompileCommand::Changed)
-				{
-					// figure out which shaders are out of date
-					TArray<const FShaderType*> OutdatedShaderTypes;
-					TArray<const FVertexFactoryType*> OutdatedFactoryTypes;
-					TArray<const FShaderPipelineType*> OutdatedShaderPipelineTypes;
-
-					// Explicitly get outdated types for global shaders.
-					const FGlobalShaderMap* ShaderMap = GGlobalShaderMap[ShaderPlatform];
-					if (ShaderMap)
+					if (Pipeline->ShouldOptimizeUnusedOutputs(Platform))
 					{
-						ShaderMap->GetOutdatedTypes(OutdatedShaderTypes, OutdatedShaderPipelineTypes, OutdatedFactoryTypes);
+						// Make a pipeline job with all the stages
+						FGlobalShaderTypeCompiler::BeginCompileShaderPipeline(Platform, PermutationFlags, Pipeline, GlobalShaderJobs);
+					}
+					else
+					{
+						// If sharing shaders amongst pipelines, add this pipeline as a dependency of an existing individual job
+						for (const FShaderType* ShaderType : Pipeline->GetStages())
+						{
+							TShaderTypePermutation<const FShaderType> ShaderTypePermutation(ShaderType, kUniqueShaderPermutationId);
+
+							FShaderCompileJob** Job = SharedShaderJobs.Find(ShaderTypePermutation);
+							checkf(Job, TEXT("Couldn't find existing shared job for global shader %s on pipeline %s!"), ShaderType->GetName(), Pipeline->GetName());
+							auto* SingleJob = (*Job)->GetSingleShaderJob();
+							check(SingleJob);
+							auto& SharedPipelinesInJob = SingleJob->SharingPipelines.FindOrAdd(nullptr);
+							check(!SharedPipelinesInJob.Contains(Pipeline));
+							SharedPipelinesInJob.Add(Pipeline);
+						}
 					}
+#endif // WITH_EDITOR
+				}
+			}
+		}
+	}
+
+#if WITH_EDITOR
+	if (GlobalShaderJobs.Num() > 0)
+	{
+		GetOnGlobalShaderCompilation().Broadcast();
+		GShaderCompilingManager->SubmitJobs(GlobalShaderJobs, "Globals");
 
-					UE_LOG(LogShaders, Display, TEXT("\tFound %d outdated shader types."), OutdatedShaderTypes.Num() + OutdatedShaderPipelineTypes.Num());
+		const bool bAllowAsynchronousGlobalShaderCompiling =
+			// OpenGL requires that global shader maps are compiled before attaching
+			// primitives to the scene as it must be able to find FNULLPS.
+			// TODO_OPENGL: Allow shaders to be compiled asynchronously.
+			// Metal also needs this when using RHI thread because it uses TOneColorVS very early in RHIPostInit()
+			!IsOpenGLPlatform(GMaxRHIShaderPlatform) && !IsVulkanPlatform(GMaxRHIShaderPlatform) &&
+			!IsMetalPlatform(GMaxRHIShaderPlatform) &&
+			FDataDrivenShaderPlatformInfo::GetSupportsAsyncPipelineCompilation(GMaxRHIShaderPlatform) &&
+			GShaderCompilingManager->AllowAsynchronousShaderCompiling();
 
-					UpdateReferencedUniformBufferNames(OutdatedShaderTypes, OutdatedFactoryTypes, OutdatedShaderPipelineTypes);
+		if (!bAllowAsynchronousGlobalShaderCompiling)
+		{
+			TArray<int32> ShaderMapIds;
+			ShaderMapIds.Add(GlobalShaderMapId);
 
-					CompileGlobalShaderMapForRemote(OutdatedShaderTypes, OutdatedShaderPipelineTypes, ShaderPlatform, TargetPlatform, Args.GlobalShaderMap);
-				}
+			GShaderCompilingManager->FinishCompilation(TEXT("Global"), ShaderMapIds);
+		}
+	}
+#endif // WITH_EDITOR
+}
 
-				// we only want to actually compile mesh shaders if a client directly requested it
-				if ((Args.CommandType == ODSCRecompileCommand::Material || Args.CommandType == ODSCRecompileCommand::Changed) &&
-					Args.MeshMaterialMaps != nullptr)
-				{
-					TMap<FString, TArray<TRefCountPtr<FMaterialShaderMap>>> CompiledShaderMaps;
-					UMaterial::CompileMaterialsForRemoteRecompile(MaterialsToCompile, ShaderPlatform, TargetPlatform, CompiledShaderMaps);
-					SaveShaderMapsForRemote(TargetPlatform, CompiledShaderMaps, Args.MeshMaterialMaps);
-				}
+void VerifyGlobalShaders(EShaderPlatform Platform, bool bLoadedFromCacheFile, const TArray<const FShaderType*>* OutdatedShaderTypes, const TArray<const FShaderPipelineType*>* OutdatedShaderPipelineTypes)
+{
+	VerifyGlobalShaders(Platform, nullptr, bLoadedFromCacheFile, OutdatedShaderTypes, OutdatedShaderPipelineTypes);
+}
 
-				// save it out so the client can get it (and it's up to date next time), if we were sent a OutputDirectory to put it in
-				FString GlobalShaderFilename;
-				if (!OutputDirectory.IsEmpty())
-				{
-					GlobalShaderFilename = SaveGlobalShaderFile(ShaderPlatform, OutputDirectory, TargetPlatform);
-				}
+void PrecacheComputePipelineStatesForGlobalShaders(EShaderPlatform Platform, const ITargetPlatform* TargetPlatform)
+{
+	static IConsoleVariable* PrecacheGlobalComputeShadersCVar = IConsoleManager::Get().FindConsoleVariable(TEXT("r.PSOPrecache.GlobalComputeShaders"));
+	if (!GRHISupportsPSOPrecaching || !PipelineStateCache::IsPSOPrecachingEnabled() || PrecacheGlobalComputeShadersCVar == nullptr || PrecacheGlobalComputeShadersCVar->GetInt() == 0)
+	{
+		return;
+	}
+
+	TRACE_CPUPROFILER_EVENT_SCOPE(PrecacheComputePipelineStatesForGlobalShaders);
+
+	FPlatformTypeLayoutParameters LayoutParams;
+	LayoutParams.InitializeForPlatform(TargetPlatform);
+	EShaderPermutationFlags PermutationFlags = GetShaderPermutationFlags(LayoutParams);
+
+	FGlobalShaderMap* GlobalShaderMap = GetGlobalShaderMap(Platform);
+	
+	// some RHIs (OpenGL) can only create shaders on the Render thread. Queue the creation instead of doing it here.
+	TArray<TShaderRef<FShader>> ComputeShadersToPrecache;
+	for (TLinkedList<FShaderType*>::TIterator ShaderTypeIt(FShaderType::GetTypeList()); ShaderTypeIt; ShaderTypeIt.Next())
+	{
+		FGlobalShaderType* GlobalShaderType = ShaderTypeIt->GetGlobalShaderType();
+		if (!GlobalShaderType || GlobalShaderType->GetFrequency() != SF_Compute)
+		{
+			continue;
+		}
 
-				// add this to the list of files to tell the other end about
-				if (Args.ModifiedFiles && !GlobalShaderFilename.IsEmpty())
+		int32 ShaderPermutationPerGlobalShaderType = 0;
+		for (int32 PermutationId = 0; PermutationId < GlobalShaderType->GetPermutationCount(); PermutationId++)
+		{
+			if (GlobalShaderType->ShouldCompilePermutation(Platform, PermutationId, PermutationFlags) &&
+				GlobalShaderType->ShouldPrecachePermutation(Platform, PermutationId, PermutationFlags) == EShaderPermutationPrecacheRequest::Precached)
+			{
+				TShaderRef<FShader> GlobalShader = GlobalShaderMap->GetShader(GlobalShaderType, PermutationId);
+				if (GlobalShader.IsValid())
 				{
-					// need to put it in non-sandbox terms
-					FString SandboxPath(GlobalShaderFilename);
-					check(SandboxPath.StartsWith(OutputDirectory));
-					SandboxPath.ReplaceInline(*OutputDirectory, TEXT("../../../"));
-					FPaths::NormalizeFilename(SandboxPath);
-					Args.ModifiedFiles->Add(SandboxPath);
+					ComputeShadersToPrecache.Add(GlobalShader);
+					ShaderPermutationPerGlobalShaderType++;
 				}
 			}
 		}
-	}
 
-	for (UMaterialInterface* MaterialInterface : MaterialsToCompile)
+		/*
+		int32 PermutationCountLimit = 300;
+		ensureMsgf(
+			ShaderPermutationPerGlobalShaderType < PermutationCountLimit,
+			TEXT("Global shader %s has %i permutations to precache: probably more than it needs."),
+			GlobalShaderType->GetName(), ShaderPermutationPerGlobalShaderType);
+		*/
+	}
+	
+	if (ComputeShadersToPrecache.Num() > 0)
 	{
-		if (MaterialInterface)
-		{
-			ResetLoaders(MaterialInterface->GetPackage());
-		}
+		UE_LOG(LogShaders, Display, TEXT("Precaching %d global compute shaders"), ComputeShadersToPrecache.Num());
+		ENQUEUE_RENDER_COMMAND(PrecachePSOsForGlobalShaders)(
+			[ComputeShadersToPrecache](FRHICommandListImmediate& RHICmdList)
+			{
+				for (TShaderRef<FShader> GlobalShader : ComputeShadersToPrecache)
+				{
+					// PSO precache shaders are not required to all load correctly
+					bool bRequired = false;
+					const TCHAR* TypeName = GlobalShader.GetType()->GetName();
+					FRHIComputeShader* RHIComputeShader = static_cast<FRHIComputeShader*>(GlobalShader.GetRHIShaderBase(SF_Compute, bRequired));
+					if (RHIComputeShader)
+					{
+						PipelineStateCache::PrecacheComputePipelineState(RHIComputeShader, TypeName);
+					}
+				}
+			});
 	}
-	GEngine->ForceGarbageCollection(true);
-
-	UE_LOG(LogShaders, Display, TEXT(""));
-	UE_LOG(LogShaders, Display, TEXT("Compiled %u shaders in %.2f seconds."), GShaderCompilerStats->GetTotalShadersCompiled() - StartTotalShadersCompiled, FPlatformTime::Seconds() - StartTime);
-
-	// Restore compilation state.
-	GShaderCompilingManager->SkipShaderCompilation(bPreviousState);
 }
 
-void ShutdownShaderCompilers(TConstArrayView<const ITargetPlatform*> TargetPlatforms)
-{
-	ITargetPlatformManagerModule& PlatformManager = GetTargetPlatformManagerRef();
-	for (const ITargetPlatform* TargetPlatform : TargetPlatforms)
-	{
-		TArray<FName> DesiredShaderFormats;
-		TargetPlatform->GetAllTargetedShaderFormats(DesiredShaderFormats);
-		for (FName FormatName : DesiredShaderFormats)
-		{
-			const IShaderFormat* ShaderFormat = PlatformManager.FindShaderFormat(FormatName);
-			if (ShaderFormat)
-			{
-				ShaderFormat->NotifyShaderCompilersShutdown(FormatName);
-			}
-		}
-	}
-}
+#include "Misc/PreLoadFile.h"
+#include "Serialization/LargeMemoryReader.h"
+static FPreLoadFile GGlobalShaderPreLoadFile(*(FString(TEXT("../../../Engine")) / TEXT("GlobalShaderCache-SP_") + FPlatformProperties::IniPlatformName() + TEXT(".bin")));
 
-#endif // WITH_EDITOR
+const ITargetPlatform* GGlobalShaderTargetPlatform[SP_NumPlatforms] = { nullptr };
 
-void BeginRecompileGlobalShaders(const TArray<const FShaderType*>& OutdatedShaderTypes, const TArray<const FShaderPipelineType*>& OutdatedShaderPipelineTypes, EShaderPlatform ShaderPlatform, const ITargetPlatform* TargetPlatform)
-{
-#if WITH_EDITOR
-	if (!FPlatformProperties::RequiresCookedData())
-	{
-		// Flush pending accesses to the existing global shaders.
-		FlushRenderingCommands();
+static FString GGlobalShaderCacheOverrideDirectory;
 
-		// Calling CompileGlobalShaderMap will force starting the compile jobs if the map is empty (by calling VerifyGlobalShaders)
-		CompileGlobalShaderMap(ShaderPlatform, TargetPlatform, false);
-		FGlobalShaderMap* GlobalShaderMap = GetGlobalShaderMap(ShaderPlatform);
+static FString GetGlobalShaderCacheOverrideFilename(EShaderPlatform Platform)
+{
+	FString DirectoryPrefix = FPaths::EngineDir() / TEXT("OverrideGlobalShaderCache-");
 
-		// Now check if there is any work to be done wrt outdates types
-		if (OutdatedShaderTypes.Num() > 0 || OutdatedShaderPipelineTypes.Num() > 0)
-		{
-			VerifyGlobalShaders(ShaderPlatform, TargetPlatform, false, &OutdatedShaderTypes, &OutdatedShaderPipelineTypes);
-		}
+	if (!GGlobalShaderCacheOverrideDirectory.IsEmpty())
+	{
+		DirectoryPrefix = GGlobalShaderCacheOverrideDirectory / TEXT("GlobalShaderCache-");
 	}
-#endif
+
+	return DirectoryPrefix + FDataDrivenShaderPlatformInfo::GetName(Platform).ToString() + TEXT(".bin");
 }
 
-void FinishRecompileGlobalShaders()
+static FString GetGlobalShaderCacheFilename(EShaderPlatform Platform)
 {
-	// Block until global shaders have been compiled and processed
-	GShaderCompilingManager->ProcessAsyncResults(false, true);
+	return FString(TEXT("Engine")) / TEXT("GlobalShaderCache-") + FDataDrivenShaderPlatformInfo::GetName(Platform).ToString() + TEXT(".bin");
 }
 
-#if WITH_EDITOR
-
-static inline FShader* ProcessCompiledJob(FShaderCompileJob* SingleJob, const FShaderPipelineType* Pipeline, TArray<EShaderPlatform>& ShaderPlatformsProcessed, TArray<const FShaderPipelineType*>& OutSharedPipelines)
+/** Saves the global shader map as a file for the target platform. */
+FString SaveGlobalShaderFile(EShaderPlatform Platform, FString SavePath, class ITargetPlatform* TargetPlatform)
 {
-	const FGlobalShaderType* GlobalShaderType = SingleJob->Key.ShaderType->GetGlobalShaderType();
-	check(GlobalShaderType);
-	FShader* Shader = FGlobalShaderTypeCompiler::FinishCompileShader(GlobalShaderType, *SingleJob, Pipeline);
-	if (Shader)
-	{
-		// Add the new global shader instance to the global shader map if it's a shared shader
-		EShaderPlatform Platform = (EShaderPlatform)SingleJob->Input.Target.Platform;
-		if (!Pipeline || !Pipeline->ShouldOptimizeUnusedOutputs(Platform))
-		{
-			Shader = GGlobalShaderMap[Platform]->FindOrAddShader(GlobalShaderType, SingleJob->Key.PermutationId, Shader);
-			// Add this shared pipeline to the list
-			if (!Pipeline)
-			{
-				auto* JobSharedPipelines = SingleJob->SharingPipelines.Find(nullptr);
-				if (JobSharedPipelines)
-				{
-					for (auto* SharedPipeline : *JobSharedPipelines)
-					{
-						OutSharedPipelines.AddUnique(SharedPipeline);
-					}
-				}
-			}
-		}
-		ShaderPlatformsProcessed.AddUnique(Platform);
-	}
-
-	return Shader;
-};
+	FGlobalShaderMap* GlobalShaderMap = GetGlobalShaderMap(Platform);
 
-void ProcessCompiledGlobalShaders(const TArray<FShaderCommonCompileJobPtr>& CompilationResults)
-{
-	TRACE_CPUPROFILER_EVENT_SCOPE(ProcessCompiledGlobalShaders);
+	// Wait until all global shaders are compiled
+	if (GShaderCompilingManager)
+	{
+		GShaderCompilingManager->ProcessAsyncResults(false, true);
+	}
 
-	UE_LOG(LogShaders, Verbose, TEXT("Compiled %u global shaders"), CompilationResults.Num());
+	TArray<uint8> GlobalShaderData;
+	{
+#if WITH_EDITOR
+		TOptional<FArchiveCookContext> CookContext;
+		TOptional<FArchiveCookData> CookData;
+#endif
+		FMemoryWriter MemoryWriter(GlobalShaderData, true);
 
-	FShaderDiagnosticInfo ShaderDiagInfo(CompilationResults);
+#if WITH_EDITOR
+		if (TargetPlatform != nullptr)
+		{
+			CookContext.Emplace(nullptr /*InPackage*/, UE::Cook::ECookType::Unknown,
+				UE::Cook::ECookingDLC::Unknown, TargetPlatform);
+			CookData.Emplace(*TargetPlatform, *CookContext);
+			MemoryWriter.SetCookData(CookData.GetPtrOrNull());
+		}
+#endif // WITH_EDITOR
 
-	// Report unique errors for global shaders.
-	for (int32 ErrorIndex = 0; ErrorIndex < ShaderDiagInfo.UniqueErrors.Num(); ++ErrorIndex)
-	{
-		UE_LOGFMT_NSLOC(LogShaders, Error, "Shaders", "GlobalShaderCompileError", "{ErrorMessage}", 
-			("ErrorMessage", ShaderDiagInfo.UniqueErrors[ErrorIndex]));
+		GlobalShaderMap->SaveToGlobalArchive(MemoryWriter);
 	}
 
-	for (const FString& WarningString : ShaderDiagInfo.UniqueWarnings)
+	// make the final name
+	FString FullPath = SavePath / GetGlobalShaderCacheFilename(Platform);
+	if (!FFileHelper::SaveArrayToFile(GlobalShaderData, *FullPath))
 	{
-		UE_LOGFMT_NSLOC(LogShaders, Warning, "Shaders", "GlobalShaderCompileWarning", "{WarningMessage}", ("WarningMessage", WarningString));
+		UE_LOG(LogShaders, Fatal, TEXT("Could not save global shader file to '%s'"), *FullPath);
 	}
 
-	const int32 UniqueErrorCount = ShaderDiagInfo.UniqueErrors.Num();
-	if (UniqueErrorCount)
+#if WITH_EDITOR
+	if (FShaderLibraryCooker::NeedsShaderStableKeys(Platform))
 	{
-		const TCHAR* RetryMsg = TEXT(" Enable 'r.ShaderDevelopmentMode' in ConsoleVariables.ini for retries.");
-		if (AreShaderErrorsFatal())
-		{
-			UE_LOGFMT_NSLOC(LogShaders, Fatal, "Shaders", "GlobalShadersCompilationFailed", "{NumErrors} Shader compiler errors compiling GlobalShaders for platform {Platform}. {RetryMsg}",
-				("NumErrors", UniqueErrorCount),
-				("Platform", ShaderDiagInfo.TargetShaderPlatformString),
-				("RetryMsg", IsRunningCommandlet() ? TEXT("") : RetryMsg)
-			);
-		}
-		else
-		{
-			UE_LOGFMT_NSLOC(LogShaders, Error, "Shaders", "GlobalShadersCompilationFailed", "{NumErrors} Shader compiler errors compiling GlobalShaders for platform {Platform}. {RetryMsg}",
-				("NumErrors", UniqueErrorCount),
-				("Platform", ShaderDiagInfo.TargetShaderPlatformString),
-				("RetryMsg", IsRunningCommandlet() ? TEXT("") : RetryMsg)
-			);
-		}
+		GlobalShaderMap->SaveShaderStableKeys(Platform);
 	}
+#endif // WITH_EDITOR
+	return FullPath;
+}
+
+
+static inline bool ShouldCacheGlobalShaderTypeName(const FGlobalShaderType* GlobalShaderType, int32 PermutationId, const TCHAR* TypeNameSubstring, EShaderPlatform Platform, EShaderPermutationFlags PermutationFlags)
+{
+	return GlobalShaderType
+		&& (TypeNameSubstring == nullptr || (FPlatformString::Strstr(GlobalShaderType->GetName(), TypeNameSubstring) != nullptr))
+		&& GlobalShaderType->ShouldCompilePermutation(Platform, PermutationId, PermutationFlags);
+};
 
-	TArray<EShaderPlatform> ShaderPlatformsProcessed;
-	TArray<const FShaderPipelineType*> SharedPipelines;
 
-	for (int32 ResultIndex = 0; ResultIndex < CompilationResults.Num(); ResultIndex++)
+bool IsGlobalShaderMapComplete(const TCHAR* TypeNameSubstring, FGlobalShaderMap* GlobalShaderMap, EShaderPlatform Platform, FString* FailureReason = nullptr)
+{
+	// look at any shadermap in the GlobalShaderMap for the permutation flags, as they will all be the same
+	if (GlobalShaderMap)
 	{
-		const FShaderCommonCompileJob& CurrentJob = *CompilationResults[ResultIndex];
-		FShaderCompileJob* SingleJob = nullptr;
-		if ((SingleJob = (FShaderCompileJob*)CurrentJob.GetSingleShaderJob()) != nullptr)
+		const FGlobalShaderMapSection* FirstShaderMap = GlobalShaderMap->GetFirstSection();
+		if (FirstShaderMap == nullptr)
 		{
-			ProcessCompiledJob(SingleJob, nullptr, ShaderPlatformsProcessed, SharedPipelines);
+			// if we had no sections at all, we know we aren't complete
+			return false;
 		}
-		else
-		{
-			const auto* PipelineJob = CurrentJob.GetShaderPipelineJob();
-			check(PipelineJob);
+		EShaderPermutationFlags GlobalShaderPermutation = FirstShaderMap->GetPermutationFlags();
 
-			FShaderPipeline* ShaderPipeline = new FShaderPipeline(PipelineJob->Key.ShaderPipeline);
-			for (int32 Index = 0; Index < PipelineJob->StageJobs.Num(); ++Index)
+		// Check if the individual shaders are complete
+		for (TLinkedList<FShaderType*>::TIterator ShaderTypeIt(FShaderType::GetTypeList()); ShaderTypeIt; ShaderTypeIt.Next())
+		{
+			FGlobalShaderType* GlobalShaderType = ShaderTypeIt->GetGlobalShaderType();
+			int32 PermutationCount = GlobalShaderType ? GlobalShaderType->GetPermutationCount() : 1;
+			for (int32 PermutationId = 0; PermutationId < PermutationCount; PermutationId++)
 			{
-				SingleJob = PipelineJob->StageJobs[Index]->GetSingleShaderJob();
-				FShader* Shader = ProcessCompiledJob(SingleJob, PipelineJob->Key.ShaderPipeline, ShaderPlatformsProcessed, SharedPipelines);
-				ShaderPipeline->AddShader(Shader, SingleJob->Key.PermutationId);
-			}
-			ShaderPipeline->Validate(PipelineJob->Key.ShaderPipeline);
+				if (ShouldCacheGlobalShaderTypeName(GlobalShaderType, PermutationId, TypeNameSubstring, Platform, GlobalShaderPermutation))
+				{
+					if (!GlobalShaderMap->HasShader(GlobalShaderType, PermutationId))
+					{
+						if (FailureReason)
+						{
+							FString GlobalShaderTypeName = GlobalShaderType ? GlobalShaderType->GetFName().ToString() : FString(TEXT("Unknown shader type"));
+							*FailureReason = FString::Printf(TEXT("Failed to find global shader \"%s\", permutation %d"), *GlobalShaderTypeName, PermutationId);
+						}
 
-			EShaderPlatform Platform = (EShaderPlatform)PipelineJob->StageJobs[0]->GetSingleShaderJob()->Input.Target.Platform;
-			check(ShaderPipeline && !GGlobalShaderMap[Platform]->HasShaderPipeline(PipelineJob->Key.ShaderPipeline));
-			GGlobalShaderMap[Platform]->FindOrAddShaderPipeline(PipelineJob->Key.ShaderPipeline, ShaderPipeline);
+						return false;
+					}
+				}
+			}
 		}
-	}
-
-	for (int32 PlatformIndex = 0; PlatformIndex < ShaderPlatformsProcessed.Num(); PlatformIndex++)
-	{
-		EShaderPlatform Platform = ShaderPlatformsProcessed[PlatformIndex];
-		FGlobalShaderMap* GlobalShaderMap = GGlobalShaderMap[Platform];
-		const ITargetPlatform* TargetPlatform = GGlobalShaderTargetPlatform[Platform];
-
-		// Process the shader pipelines that share shaders
-		FPlatformTypeLayoutParameters LayoutParams;
-		LayoutParams.InitializeForPlatform(TargetPlatform);
-		const EShaderPermutationFlags PermutationFlags = GetShaderPermutationFlags(LayoutParams);
 
-		for (const FShaderPipelineType* ShaderPipelineType : SharedPipelines)
+		// Then the pipelines as it may be sharing shaders
+		for (TLinkedList<FShaderPipelineType*>::TIterator ShaderPipelineIt(FShaderPipelineType::GetTypeList()); ShaderPipelineIt; ShaderPipelineIt.Next())
 		{
-			check(ShaderPipelineType->IsGlobalTypePipeline());
-			if (!GlobalShaderMap->HasShaderPipeline(ShaderPipelineType))
+			const FShaderPipelineType* Pipeline = *ShaderPipelineIt;
+			if (Pipeline->IsGlobalTypePipeline())
 			{
-				auto& StageTypes = ShaderPipelineType->GetStages();
-
-				FShaderPipeline* ShaderPipeline = new FShaderPipeline(ShaderPipelineType);
-				for (int32 Index = 0; Index < StageTypes.Num(); ++Index)
+				auto& Stages = Pipeline->GetStages();
+				int32 NumStagesNeeded = 0;
+				for (const FShaderType* Shader : Stages)
 				{
-					FGlobalShaderType* GlobalShaderType = ((FShaderType*)(StageTypes[Index]))->GetGlobalShaderType();
-					if (GlobalShaderType->ShouldCompilePermutation(Platform, kUniqueShaderPermutationId, PermutationFlags))
+					const FGlobalShaderType* GlobalShaderType = Shader->GetGlobalShaderType();
+					if (ShouldCacheGlobalShaderTypeName(GlobalShaderType, kUniqueShaderPermutationId, TypeNameSubstring, Platform, GlobalShaderPermutation))
 					{
-						TShaderRef<FShader> Shader = GlobalShaderMap->GetShader(GlobalShaderType, kUniqueShaderPermutationId);
-						check(Shader.IsValid());
-						ShaderPipeline->AddShader(Shader.GetShader(), kUniqueShaderPermutationId);
+						++NumStagesNeeded;
 					}
 					else
 					{
 						break;
 					}
 				}
-				ShaderPipeline->Validate(ShaderPipelineType);
-				GlobalShaderMap->FindOrAddShaderPipeline(ShaderPipelineType, ShaderPipeline);
-			}
-		}
-
-		// at this point the new global sm is populated and we can delete the deferred copy, if any
-		delete GGlobalShaderMap_DeferredDeleteCopy[ShaderPlatformsProcessed[PlatformIndex]];	// even if it was nullptr, deleting null is Okay
-		GGlobalShaderMap_DeferredDeleteCopy[ShaderPlatformsProcessed[PlatformIndex]] = nullptr;
-
-		// Save the global shader map for any platforms that were recompiled, but only if it is complete (it can be also a subject to ODSC, perhaps unnecessarily, as we cannot use a partial global SM)
-		FGlobalShaderMapId ShaderMapId(Platform, TargetPlatform);
-		if (GlobalShaderMap->IsComplete(TargetPlatform))
-		{
-			SaveGlobalShaderMapToDerivedDataCache(ShaderPlatformsProcessed[PlatformIndex]);
 
-			if (!GRHISupportsMultithreadedShaderCreation && Platform == GMaxRHIShaderPlatform)
-			{
-				ENQUEUE_RENDER_COMMAND(CreateRecursiveShaders)([](FRHICommandListImmediate&)
+				if (NumStagesNeeded == Stages.Num())
 				{
-					CreateRecursiveShaders();
-				});
+					if (!GlobalShaderMap->HasShaderPipeline(Pipeline))
+					{
+						if (FailureReason)
+						{
+							*FailureReason = FString::Printf(TEXT("Failed to find global pipeline \"%s\""), *Pipeline->GetFName().ToString());
+						}
+
+						return false;
+					}
+				}
 			}
 		}
 	}
+
+	return true;
 }
 
-void SaveGlobalShadersForRemoteRecompile(FArchive& Ar, EShaderPlatform ShaderPlatform)
-{
-	FGlobalShaderMap* GlobalShaderMap = GetGlobalShaderMap(ShaderPlatform);
-	uint8 bIsValid = GlobalShaderMap != nullptr;
-	Ar << bIsValid;
 
-	if (GlobalShaderMap)
+bool IsGlobalShaderMapComplete(const TCHAR* TypeNameSubstring)
+{
+	for (int32 i = 0; i < SP_NumPlatforms; ++i)
 	{
-		GlobalShaderMap->SaveToGlobalArchive(Ar);
+		EShaderPlatform Platform = (EShaderPlatform)i;
+
+		FGlobalShaderMap* GlobalShaderMap = GGlobalShaderMap[Platform];
+
+		if (!IsGlobalShaderMapComplete(TypeNameSubstring, GlobalShaderMap, Platform))
+		{
+			return false;
+		}
 	}
+
+	return true;
 }
-#endif // WITH_EDITOR
 
-void LoadGlobalShadersForRemoteRecompile(FArchive& Ar, EShaderPlatform ShaderPlatform)
+static bool TryLoadCookedGlobalShaderMap(EShaderPlatform Platform, FScopedSlowTask& SlowTask)
 {
-	uint8 bIsValid = 0;
-	Ar << bIsValid;
+	SlowTask.EnterProgressFrame(50);
 
-	if (bIsValid)
+	bool bLoadedFromCacheFile = false;
+
+	// Load from the override global shaders first, this allows us to hot reload in cooked / pak builds
+	TArray<uint8> GlobalShaderData;
+	const bool bAllowOverrideGlobalShaders = !WITH_EDITOR && !UE_BUILD_SHIPPING;
+	if (bAllowOverrideGlobalShaders)
 	{
-		FlushRenderingCommands();
+		FString OverrideGlobalShaderCacheFilename = GetGlobalShaderCacheOverrideFilename(Platform);
+		FPaths::MakeStandardFilename(OverrideGlobalShaderCacheFilename);
 
-		FGlobalShaderMap* NewGlobalShaderMap = new FGlobalShaderMap(ShaderPlatform);
-		if (NewGlobalShaderMap)
-		{
-			NewGlobalShaderMap->LoadFromGlobalArchive(Ar);
+		bool bFileExist = IFileManager::Get().FileExists(*OverrideGlobalShaderCacheFilename);
 
-			FString FailureReason;
-			bool bIsNewGlobalShaderMapComplete = IsGlobalShaderMapComplete(nullptr, NewGlobalShaderMap, ShaderPlatform, &FailureReason);
+		if (!bFileExist)
+		{
+			UE_LOG(LogShaders, Display, TEXT("%s doesn't exists"), *OverrideGlobalShaderCacheFilename);
+		}
+		else
+		{
+			bLoadedFromCacheFile = FFileHelper::LoadFileToArray(GlobalShaderData, *OverrideGlobalShaderCacheFilename, FILEREAD_Silent);
 
-			if (bIsNewGlobalShaderMapComplete)
+			if (bLoadedFromCacheFile)
 			{
-				if (GGlobalShaderMap[ShaderPlatform])
-				{
-					GGlobalShaderMap[ShaderPlatform]->ReleaseAllSections();
-					delete GGlobalShaderMap[ShaderPlatform];
-					GGlobalShaderMap[ShaderPlatform] = nullptr;
-				}
-				GGlobalShaderMap[ShaderPlatform] = NewGlobalShaderMap;
-
-				VerifyGlobalShaders(ShaderPlatform, nullptr, false);
-
-				// Invalidate global bound shader states so they will be created with the new shaders the next time they are set (in SetGlobalBoundShaderState)
-				for (TLinkedList<FGlobalBoundShaderStateResource*>::TIterator It(FGlobalBoundShaderStateResource::GetGlobalBoundShaderStateList()); It; It.Next())
-				{
-					BeginUpdateResourceRHI(*It);
-				}
-
-				PropagateGlobalShadersToAllPrimitives();
+				UE_LOG(LogShaders, Display, TEXT("%s has been loaded successfully"), *OverrideGlobalShaderCacheFilename);
 			}
 			else
 			{
-				FString ErrorMessage = FString::Printf(TEXT("New global shader map is incomplete and will not be used. Reason:\n%s\n"
-													        "Please check the ODSC server log & that client/editor are compiled"), *FailureReason);
-
-				UE_LOG(LogShaderCompilers, Error, TEXT("%s"), *ErrorMessage);
-#if WITH_ODSC
-				FODSCManager::ReportODSCError(ErrorMessage);
-#endif
-				
-				delete NewGlobalShaderMap;
+				UE_LOG(LogShaders, Error, TEXT("%s failed to load"), *OverrideGlobalShaderCacheFilename);
 			}
 		}
 	}
-}
 
-TRACE_DECLARE_INT_COUNTER(Shaders_JobCacheSearchAttempts, TEXT("Shaders/JobCache/SearchAttempts"));
-TRACE_DECLARE_INT_COUNTER(Shaders_JobCacheHits, TEXT("Shaders/JobCache/Hits"));
+	// is the data already loaded?
+	int64 PreloadedSize = 0;
+	void* PreloadedData = nullptr;
+	if (!bLoadedFromCacheFile)
+	{
+		PreloadedData = GGlobalShaderPreLoadFile.TakeOwnershipOfLoadedData(&PreloadedSize);
+	}
+
+	if (PreloadedData != nullptr)
+	{
+		FLargeMemoryReader MemoryReader((uint8*)PreloadedData, PreloadedSize, ELargeMemoryReaderFlags::TakeOwnership);
+		GGlobalShaderMap[Platform]->LoadFromGlobalArchive(MemoryReader);
+		bLoadedFromCacheFile = true;
+	}
+	else
+	{
+		FString GlobalShaderCacheFilename = FPaths::GetRelativePathToRoot() / GetGlobalShaderCacheFilename(Platform);
+		FPaths::MakeStandardFilename(GlobalShaderCacheFilename);
+		if (!bLoadedFromCacheFile)
+		{
+			bLoadedFromCacheFile = FFileHelper::LoadFileToArray(GlobalShaderData, *GlobalShaderCacheFilename, FILEREAD_Silent);
+		}
 
-TRACE_DECLARE_INT_COUNTER(Shaders_JobCacheDDCRequests, TEXT("Shaders/JobCache/DDCRequests"));
-TRACE_DECLARE_INT_COUNTER(Shaders_JobCacheDDCHits, TEXT("Shaders/JobCache/DDCHits"));
-TRACE_DECLARE_MEMORY_COUNTER(Shaders_JobCacheDDCBytesReceived, TEXT("Shaders/JobCache/DDCBytesRecieved"));
-TRACE_DECLARE_MEMORY_COUNTER(Shaders_JobCacheDDCBytesSent, TEXT("Shaders/JobCache/DDCBytesSent"));
+		if (bLoadedFromCacheFile)
+		{
+			FMemoryReader MemoryReader(GlobalShaderData);
+			GGlobalShaderMap[Platform]->LoadFromGlobalArchive(MemoryReader);
+		}
+	}
 
-#if WITH_EDITOR
-namespace
-{
-	/** The FCacheBucket used with the DDC, cached to avoid recreating it for each request */
-	UE::DerivedData::FCacheBucket ShaderJobCacheDDCBucket = UE::DerivedData::FCacheBucket(ANSITEXTVIEW("FShaderJobCacheShaders"), TEXTVIEW("Shader"));
-	UE::DerivedData::FValueId ShaderJobCacheId = UE::DerivedData::FValueId::FromName("FShaderJobCacheShaderID");
+	return bLoadedFromCacheFile;
 }
-#endif
 
-FShaderJobCacheRef FShaderJobCache::FindOrAdd(const FShaderCompilerInputHash& Hash, EShaderCompileJobPriority JobPriority, const bool bCheckDDC, TPimplPtr<UE::DerivedData::FRequestOwner>& InoutRequestOwner, FJobCachedOutput*& OutCachedOutput)
+void CompileGlobalShaderMap(EShaderPlatform Platform, const ITargetPlatform* TargetPlatform, bool bRefreshShaderMap)
 {
-	LLM_SCOPE_BYTAG(ShaderCompiler);
-
-	check(ShaderCompiler::IsJobCacheEnabled());
-
-	++TotalSearchAttempts;
-	TRACE_COUNTER_INCREMENT(Shaders_JobCacheSearchAttempts);
-	TRACE_CPUPROFILER_EVENT_SCOPE(FShaderJobCache::Find);
+	LLM_SCOPE_RENDER_RESOURCE(TEXT("GlobalShaderMap"));
 
-	OutCachedOutput = nullptr;
+	// No global shaders needed on dedicated server or clients that use NullRHI. Note that cook commandlet needs to have them, even if it is not allowed to render otherwise.
+	if (FPlatformProperties::IsServerOnly() || (!IsRunningCommandlet() && !FApp::CanEverRender()))
+	{
+		if (!GGlobalShaderMap[Platform])
+		{
+			GGlobalShaderMap[Platform] = new FGlobalShaderMap(Platform);
+		}
+		return;
+	}
 
-	uint64 InputHashToJobDataSize = InputHashToJobData.GetAllocatedSize();
+	if (bRefreshShaderMap || GGlobalShaderTargetPlatform[Platform] != TargetPlatform)
+	{
+		// defer the deletion the current global shader map, delete the previous one if it is still valid
+		delete GGlobalShaderMap_DeferredDeleteCopy[Platform];	// deleting null is Okay
+		GGlobalShaderMap_DeferredDeleteCopy[Platform] = GGlobalShaderMap[Platform];
+		GGlobalShaderMap[Platform] = nullptr;
 
-	FShaderJobCacheRef JobCacheRef = InputHashToJobData.FindOrAdd(Hash);
-	FShaderJobData& JobData = GetShaderJobData(JobCacheRef);
+		GGlobalShaderTargetPlatform[Platform] = TargetPlatform;
 
-	CurrentlyAllocatedMemory += InputHashToJobData.GetAllocatedSize() - InputHashToJobDataSize;
+		// make sure we look for updated shader source files
+		FlushShaderFileCache();
+	}
 
-	if (JobData.HasOutput())
+#if WITH_ODSC
+	// First try to load the global shader map with ODSC if it's connected. TryLoadGlobalShaders will set GGlobalShaderMap[Platform]
+	if (!GGlobalShaderMap[Platform] && FODSCManager::IsODSCActive())
 	{
-		++TotalCacheHits;
-		TRACE_COUNTER_INCREMENT(Shaders_JobCacheHits);
-
-		FStoredOutput** CannedOutput = Outputs.Find(JobData.OutputHash);
-		// we should not allow a dangling input to output mapping to exist
-		checkf(CannedOutput != nullptr, TEXT("Inconsistency in FShaderJobCache - cache record for ihash %s (data 0x%p) exists, but output %s (%s) cannot be found."),
-			*LexToString(Hash), &JobData, *LexToString(JobData.OutputHash), JobData.bOutputFromDDC ? TEXT("DDC") : TEXT("Job"));
-		// update the output hit count
-		(*CannedOutput)->NumHits++;
-
-		OutCachedOutput = &(*CannedOutput)->JobOutput;
+		UE_LOG(LogShaders, Display, TEXT("Trying to load global shaders from ODSC ..."));
+		GODSCManager->TryLoadGlobalShaders(Platform);
+		UE_LOG(LogShaders, Display, TEXT("Global shaders from ODSC: %s"), (GGlobalShaderMap[Platform] != nullptr) ? TEXT("success") : TEXT("failed"));
 	}
-#if WITH_EDITOR
-	else
+#endif
+
+	// If the global shader map hasn't been created yet, create it.
+	if (!GGlobalShaderMap[Platform])
 	{
-		// If NoShaderDDC then don't check for a material the first time we encounter it to simulate a cold DDC
-		static bool bNoShaderDDC = FParse::Param(FCommandLine::Get(), TEXT("noshaderddc"));
+		DECLARE_SCOPE_CYCLE_COUNTER(TEXT("GetGlobalShaderMap"), STAT_GetGlobalShaderMap, STATGROUP_LoadTime);
+		// GetGlobalShaderMap is called the first time during startup in the main thread.
+		check(IsInGameThread());
 
-		// If we didn't find it in memory search the DDC if it's enabled.
-		// Don't search if this isn't the first job with this hash (JobInFlight already set), or there's already a request in flight.
-		const bool bCachePerShaderDDC = IsShaderJobCacheDDCEnabled() && bCheckDDC && !bNoShaderDDC;
-		if (bCachePerShaderDDC && (JobData.JobInFlight == nullptr) && !InoutRequestOwner)
-		{
-			TRACE_COUNTER_INCREMENT(Shaders_JobCacheDDCRequests);
+		FScopedSlowTask SlowTask(70, LOCTEXT("CreateGlobalShaderMap", "Creating Global Shader Map..."));
 
-			++TotalCacheDDCQueries;
+		// verify that all shader source files are intact
+		SlowTask.EnterProgressFrame(20, LOCTEXT("VerifyShaderSourceFiles", "Verifying Global Shader source files..."));
+		VerifyShaderSourceFiles(Platform);
 
-			UE::DerivedData::EPriority DerivedDataPriority;
-			UE::DerivedData::FRequestOwner* RequestOwner;
+		GGlobalShaderMap[Platform] = new FGlobalShaderMap(Platform);
 
-			static const bool PerShaderDDCAsync = CVarShaderCompilerPerShaderDDCAsync.GetValueOnAnyThread();
-			if (PerShaderDDCAsync && FGenericPlatformProcess::SupportsMultithreading())
-			{
-				if (IsRunningCookCommandlet())
-				{
-					DerivedDataPriority = UE::DerivedData::EPriority::Highest;
-				}
-				else
-				{
-					switch (JobPriority)
-					{
-					case EShaderCompileJobPriority::Low:		DerivedDataPriority = UE::DerivedData::EPriority::Low;		break;
-					case EShaderCompileJobPriority::Normal:		DerivedDataPriority = UE::DerivedData::EPriority::Normal;	break;
-					default:									DerivedDataPriority = UE::DerivedData::EPriority::Highest;	break;
-					}
-				}
-				InoutRequestOwner = MakePimpl<UE::DerivedData::FRequestOwner>(DerivedDataPriority);
-				RequestOwner = InoutRequestOwner.Get();
-			}
-			else
-			{
-				DerivedDataPriority = UE::DerivedData::EPriority::Blocking;
-				RequestOwner = new UE::DerivedData::FRequestOwner(DerivedDataPriority);
-			}
+		bool bShaderMapIsBeingCompiled = false;
 
-			UE::DerivedData::FCacheGetRequest Request;
-			Request.Name = TEXT("FShaderJobCache");
-			// Create key.
-			Request.Key.Bucket = ShaderJobCacheDDCBucket;
-			Request.Key.Hash = Hash;
-			Request.Policy = IsShaderJobCacheDDCRemotePolicyEnabled() ? UE::DerivedData::ECachePolicy::Default : UE::DerivedData::ECachePolicy::Local;
-
-			// Optionally read the cached output back to the main thread
-			FJobCachedOutput** OutCachedOutputPtr = &OutCachedOutput;
-			bool bCompletedSynchronously = false;
-			bool* bCompletedSynchronouslyPtr = &bCompletedSynchronously;
-
-			UE::DerivedData::GetCache().Get(
-				{ Request },
-				*RequestOwner,
-				[this, JobDataPtr = &JobData, OutCachedOutputPtr, DerivedDataPriority, bCompletedSynchronouslyPtr](UE::DerivedData::FCacheGetResponse&& Response)
-				{
-					if (GShaderCompilerDebugStallDDCQuery > 0)
-					{
-						FPlatformProcess::Sleep(GShaderCompilerDebugStallDDCQuery * 0.001f);
-					}
+		// Try to load the global shaders from a local cache file if it exists
+		// We always try this first, even when running in the editor or if shader compiler is enabled
+		// It's always possible we'll find a cooked local cache
+		const bool bLoadedFromCacheFile = TryLoadCookedGlobalShaderMap(Platform, SlowTask);
+#if WITH_EDITOR
+		const bool bAllowShaderCompiling = !FPlatformProperties::RequiresCookedData() && AllowShaderCompiling();
+#else
+		const bool bAllowShaderCompiling = false;
+#endif
 
-					bool bIsAsync = (DerivedDataPriority != UE::DerivedData::EPriority::Blocking);
+#if WITH_EDITOR
+		if (!bLoadedFromCacheFile && bAllowShaderCompiling)
+		{
+			// If we didn't find cooked shaders, we can try loading from the DDC or compiling them if supported by the current configuration
+			FGlobalShaderMapId ShaderMapId(Platform, TargetPlatform);
 
-					// Check thread local variable to see if we're in the submit job thread (DDC request completing synchronously), in which
-					// case we want to go through the synchronous code paths below, instead of async.
-					if (bInSubmitJobThread)
-					{
-						bIsAsync = false;
-						*bCompletedSynchronouslyPtr = true;
-					}
+			const int32 ShaderFilenameNum = ShaderMapId.GetShaderFilenameToDependeciesMap().Num();
+			const float ProgressStep = 25.0f / ShaderFilenameNum;
 
-					if (Response.Status == UE::DerivedData::EStatus::Ok)
-					{
-						// Retrieve the shared buffer containing the job output and compute the associated output hash for the result retrieved from DDC
-						// If an existing duplicate of this buffer is already registered in the Outputs map, this copy will be freed at end of scope
-						FSharedBuffer JobOutput = Response.Record.GetValue(ShaderJobCacheId).GetData().Decompress();
-						FJobOutputHash OutputHash = FBlake3::HashBuffer(JobOutput.GetData(), JobOutput.GetSize());
+			// If NoShaderDDC then don't check for a material the first time we encounter it to simulate
+			// a cold DDC
+			static bool bNoShaderDDC =
+				FParse::Param(FCommandLine::Get(), TEXT("noshaderddc")) || 
+				FParse::Param(FCommandLine::Get(), TEXT("noglobalshaderddc"));
 
-						TRACE_COUNTER_ADD(Shaders_JobCacheDDCBytesReceived, JobOutput.GetSize());
-						TRACE_COUNTER_INCREMENT(Shaders_JobCacheDDCHits);
+			const bool bTempNoShaderDDC = bNoShaderDDC;
 
-						// If we are running the cache logic async (not blocking in the main thread), we need a lock before writing to the job cache.
-						// Otherwise, the lock will already be held by the main thread (and trying to lock here would just deadlock).
-						if (bIsAsync)
-						{
-							JobLock.WriteLock();
-							check(JobDataPtr->JobInFlight);
+			{
+				using namespace UE::DerivedData;
 
-							// If job was cancelled, it will have been unlinked from PendingSubmitJobTaskJobs, and we can ignore the results.
-							if (!JobDataPtr->JobInFlight->PrevLink)
-							{
-								UE_LOG(LogShaderCompilers, Display, TEXT("Cancelled job 0x%p (data 0x%p) with pending DDC hit."), JobDataPtr->JobInFlight.GetReference(), JobDataPtr);
-								if (JobDataPtr->JobInFlight)
-								{
-#if WITH_EDITOR
-									if (JobDataPtr->JobInFlight->RequestOwner.IsValid())
-									{
-										JobDataPtr->JobInFlight->RequestOwner->KeepAlive();
-									}
-#endif
-									JobDataPtr->JobInFlight = nullptr;
-								}
-								JobLock.WriteUnlock();
-								return;
-							}
-							else
-							{
-								Unlink(*JobDataPtr->JobInFlight);		// from PendingSubmitJobTaskJobs
-							}
-						}
+				int32 BufferIndex = 0;
+				TArray<FCacheGetValueRequest> Requests;
 
-						// Add a DDC hit
-						++TotalCacheDDCHits;
+				// Submit DDC requests.
+				SlowTask.EnterProgressFrame(ProgressStep, LOCTEXT("SubmitDDCRequests", "Submitting global shader DDC Requests..."));
+				for (const auto& ShaderFilenameDependencies : ShaderMapId.GetShaderFilenameToDependeciesMap())
+				{
+					FCacheGetValueRequest& Request = Requests.AddDefaulted_GetRef();
+					Request.Name = GetGlobalShaderMapName(ShaderMapId, Platform, ShaderFilenameDependencies.Key);
+					Request.Key = GetGlobalShaderMapKey(ShaderMapId, Platform, TargetPlatform, ShaderFilenameDependencies.Value);
+					Request.UserData = uint64(BufferIndex);
+					++BufferIndex;
 
-						FStoredOutput** ExistingStoredOutput = Outputs.Find(OutputHash);
-						FStoredOutput* StoredOutput = ExistingStoredOutput ? *ExistingStoredOutput : nullptr;
-						if (StoredOutput == nullptr)
+					if (UNLIKELY(ShouldDumpShaderDDCKeys()))
+					{
+						const FString DataKey = GetGlobalShaderMapKeyString(ShaderMapId, Platform, ShaderFilenameDependencies.Value);
+						// For global shaders, we dump the key multiple times (once for each shader type) so they will live on disk alongside
+						// other shader debug artifacts.
+						for (const FShaderTypeDependency& ShaderTypeDependency : ShaderFilenameDependencies.Value)
 						{
-							// Create a new entry to store in the FShaderJobCache if one doesn't already exist for this output hash
-							StoredOutput = new FStoredOutput();
-							StoredOutput->JobOutput = JobOutput;
-							Outputs.Add(OutputHash, StoredOutput);
-							CurrentlyAllocatedMemory += StoredOutput->GetAllocatedSize();
+							const FShaderType* ShaderType = FindShaderTypeByName(ShaderTypeDependency.ShaderTypeName);
+							TStringBuilder<128> GroupNameBuilder;
+							GroupNameBuilder << TEXT("Global");
+							FPathViews::Append(GroupNameBuilder, ShaderType->GetName());
+							DumpShaderDDCKeyToFile(Platform, ShaderMapId.WithEditorOnly(), GroupNameBuilder.ToString(), DataKey);
 						}
+					}
+				}
 
-						// Increment refcount of output whether or not we created it above
-						StoredOutput->AddRef();
-
-						JobDataPtr->OutputHash = OutputHash;
-						JobDataPtr->bOutputFromDDC = true;
+				int32 DDCHits = 0;
+				int32 DDCMisses = 0;
 
-						// If async, add processed results to output.  For the synchronous case, this is handled back in the main thread.
-						if (bIsAsync)
+				// Process finished DDC requests.
+				SlowTask.EnterProgressFrame(ProgressStep, LOCTEXT("ProcessDDCRequests", "Processing global shader DDC requests..."));
+				TArray<FValue> GlobalShaderMapBuffers;
+				GlobalShaderMapBuffers.SetNum(Requests.Num());
+				{
+					COOK_STAT(auto Timer = GlobalShaderCookStats::UsageStats.TimeSyncWork());
+					COOK_STAT(Timer.TrackCyclesOnly());
+					FRequestOwner BlockingOwner(EPriority::Blocking);
+					GetCache().GetValue(Requests, BlockingOwner, [&GlobalShaderMapBuffers, &bTempNoShaderDDC](FCacheGetValueResponse&& Response)
+					{
+						if (bTempNoShaderDDC)
 						{
-							check(JobDataPtr->JobInFlight);
-							FShaderCommonCompileJobPtr Job = JobDataPtr->JobInFlight;
-
-							UE_LOG(LogShaderCompilers, UE_SHADERCACHE_LOG_LEVEL, TEXT("Found an async DDC result for job with ihash %s."), *LexToString(Job->InputHash));
-
-							// Get list of finished jobs -- JobInFlight, plus any duplicates -- and clear the job cache data
-							TArray<FShaderCommonCompileJob*> FinishedJobs;
-							FinishedJobs.Add(Job);
-								
-							FShaderCommonCompileJob* CurHead = JobDataPtr->DuplicateJobsWaitList;
-							while (CurHead)
-							{
-								FinishedJobs.Add(CurHead);
-								RemoveDuplicateJob(CurHead);
-								CurHead = CurHead->NextLink;
-							}
-							JobDataPtr->DuplicateJobsWaitList = nullptr;
-							if (JobDataPtr->JobInFlight)
-							{
-#if WITH_EDITOR
-								if (JobDataPtr->JobInFlight->RequestOwner.IsValid())
-								{
-									JobDataPtr->JobInFlight->RequestOwner->KeepAlive();
-								}
-#endif
-								JobDataPtr->JobInFlight = nullptr;
-							}
-							Job->JobCacheRef.Clear();
-
-							// Need to release the lock before calling ProcessFinishedJobs
-							JobLock.WriteUnlock();
+							return;
+						}
 
-							// Call ProcessFinishedJob on main job and duplicates
-							for (FShaderCommonCompileJob* FinishedJob : FinishedJobs)
-							{
-								FMemoryReaderView MemReader(StoredOutput->JobOutput);
-								FinishedJob->SerializeOutput(MemReader);
-								ProcessFinishedJob(FinishedJob, true);
-							}
+						GlobalShaderMapBuffers[int32(Response.UserData)] = MoveTemp(Response.Value);
+					});
+					BlockingOwner.Wait();
+				}
 
-							if (FinishedJobs.Num() > 1)
-							{
-								UE_LOG(LogShaderCompilers, UE_SHADERCACHE_LOG_LEVEL, TEXT("Processed %d outstanding jobs with the same ihash %s."), FinishedJobs.Num() - 1, *LexToString(Job->InputHash));
-							}
-						}
-						else
-						{
-							// Send results back to the main thread when running synchronous
-							*OutCachedOutputPtr = &StoredOutput->JobOutput;
-						}
+				BufferIndex = 0;
+				for (const auto& ShaderFilenameDependencies : ShaderMapId.GetShaderFilenameToDependeciesMap())
+				{
+					COOK_STAT(auto Timer = GlobalShaderCookStats::UsageStats.TimeSyncWork());
+					if (GlobalShaderMapBuffers[BufferIndex].HasData())
+					{
+						COOK_STAT(Timer.AddHit(int64(GlobalShaderMapBuffers[BufferIndex].GetRawSize())));
+						const FSharedBuffer CachedData = GlobalShaderMapBuffers[BufferIndex].GetData().Decompress();
+						FMemoryReaderView MemoryReader(CachedData);
+						GGlobalShaderMap[Platform]->AddSection(FGlobalShaderMapSection::CreateFromArchive(MemoryReader));
+						DDCHits++;
 					}
 					else
 					{
-						// If async, add job to pending queue.  For the synchronous case, this is handled back in the main thread.
-						if (bIsAsync)
-						{
-							FWriteScopeLock Locker(JobLock);
-							FShaderCommonCompileJob* Job = JobDataPtr->JobInFlight;
-							check(Job);
-
-							// If job was cancelled, it will have been unlinked from PendingSubmitJobTaskJobs, and we can ignore it.
-							if (!Job->PrevLink)
-							{
-								UE_LOG(LogShaderCompilers, Display, TEXT("Cancelled job 0x%p (data 0x%p) with pending DDC miss."), Job, JobDataPtr);
+						// it's a miss, but we haven't built anything yet. Save the counting until we actually have it built.
+						COOK_STAT(Timer.TrackCyclesOnly());
+						bShaderMapIsBeingCompiled = true;
+						DDCMisses++;
+					}
+					++BufferIndex;
+				}
 
-								if (JobDataPtr->JobInFlight)
-								{
-#if WITH_EDITOR
-									if (JobDataPtr->JobInFlight->RequestOwner.IsValid())
-									{
-										JobDataPtr->JobInFlight->RequestOwner->KeepAlive();
-									}
-#endif
-									JobDataPtr->JobInFlight = nullptr;
-								}
-								return;
-							}
-							else
-							{
-								Unlink(*Job);		// from PendingSubmitJobTaskJobs
-							}
+				GShaderCompilerStats->AddDDCHit(DDCHits);
+				GShaderCompilerStats->AddDDCMiss(DDCMisses);
+			}
+		}
+#endif // WITH_EDITOR
+		
+		if (!bLoadedFromCacheFile && !bAllowShaderCompiling)
+		{
+			// Failed to load cooked shaders, and no support for compiling
+			// Handle this gracefully and exit.
+			const FString GlobalShaderCacheFilename = FPaths::GetRelativePathToRoot() / GetGlobalShaderCacheFilename(Platform);
+			const FString SandboxPath = IFileManager::Get().ConvertToAbsolutePathForExternalAppForWrite(*GlobalShaderCacheFilename);
+			// This can be too early to localize in some situations.
+			const FText Message = FText::Format(NSLOCTEXT("Engine", "GlobalShaderCacheFileMissing", "The global shader cache file '{0}' is missing.\n\nYour application is built to load COOKED content. No COOKED content was found; This usually means you did not cook content for this build.\nIt also may indicate missing cooked data for a shader platform(e.g., OpenGL under Windows): Make sure your platform's packaging settings include this Targeted RHI.\n\nAlternatively build and run the UNCOOKED version instead."), FText::FromString(SandboxPath));
+			if (FPlatformProperties::SupportsWindowedMode())
+			{
+				UE_LOG(LogShaders, Error, TEXT("%s"), *Message.ToString());
+				FMessageDialog::Open(EAppMsgType::Ok, Message);
+				FPlatformMisc::RequestExit(false, TEXT("CompileGlobalShaderMap"));
+				return;
+			}
+			else
+			{
+				UE_LOG(LogShaders, Fatal, TEXT("%s"), *Message.ToString());
+			}
+		}
 
-							LinkJobWithPriority(*Job);
-						}
-					}
-				});
+		// If any shaders weren't loaded, compile them now.
+		VerifyGlobalShaders(Platform, TargetPlatform, bLoadedFromCacheFile);
 
-			// For blocking requests, wait on the results, and delete the request
-			if (RequestOwner->GetPriority() == UE::DerivedData::EPriority::Blocking)
-			{
-				RequestOwner->Wait();
-				delete RequestOwner;
-			}
-			else if (bCompletedSynchronously)
+		if (GCreateShadersOnLoad && Platform == GMaxRHIShaderPlatform)
+		{
+			GGlobalShaderMap[Platform]->BeginCreateAllShaders();
+		}
+
+		// While we're early in the game's startup, create certain global shaders that may be later created on random threads otherwise. 
+		if (!bShaderMapIsBeingCompiled && !GRHISupportsMultithreadedShaderCreation)
+		{
+			ENQUEUE_RENDER_COMMAND(CreateRecursiveShaders)([](FRHICommandListImmediate&)
 			{
-				// It's also possible (notably when DDC verification is enabled) for the request to have completed synchronously,
-				// in which case we can delete the TPimplPtr request owner by setting it to null.  This tells the main thread
-				// there is no async DDC request in flight, and it should handle adding the pending job to the queue, since the
-				// DDC request callback won't be handling that.
-				InoutRequestOwner = nullptr;
-			}
+				CreateRecursiveShaders();
+			});
 		}
 	}
-#endif
-
-	return JobCacheRef;
 }
 
-FShaderJobData* FShaderJobCache::Find(const FShaderCompilerInputHash& Hash)
+void CompileGlobalShaderMap(EShaderPlatform Platform, bool bRefreshShaderMap)
 {
-	check(ShaderCompiler::IsJobCacheEnabled());
-
-	return InputHashToJobData.Find(Hash);
+	CompileGlobalShaderMap(Platform, nullptr, bRefreshShaderMap);
 }
 
-/** Adds a reference to a duplicate job (to the DuplicateJobs array) */
-void FShaderJobCache::AddDuplicateJob(FShaderCommonCompileJob* DuplicateJob)
+void CompileGlobalShaderMap(ERHIFeatureLevel::Type InFeatureLevel, bool bRefreshShaderMap)
 {
-	check(DuplicateJob->JobCacheRef.DuplicateIndex == INDEX_NONE);
-
-	DuplicateJob->JobCacheRef.DuplicateIndex = DuplicateJobs.Add(DuplicateJob);
+	EShaderPlatform Platform = GShaderPlatformForFeatureLevel[InFeatureLevel];
+	CompileGlobalShaderMap(Platform, nullptr, bRefreshShaderMap);
 }
 
-/** Removes a reference to a duplicate job (from the DuplicateJobs array)  */
-void FShaderJobCache::RemoveDuplicateJob(FShaderCommonCompileJob* DuplicateJob)
+void CompileGlobalShaderMap(bool bRefreshShaderMap)
 {
-	int32 DuplicateIndex = DuplicateJob->JobCacheRef.DuplicateIndex;
-	check(DuplicateIndex >= 0 && DuplicateIndex < DuplicateJobs.Num() && DuplicateJobs[DuplicateIndex] == DuplicateJob);
-	DuplicateJob->JobCacheRef.DuplicateIndex = INDEX_NONE;
+	CompileGlobalShaderMap(GMaxRHIFeatureLevel, bRefreshShaderMap);
+}
 
-	DuplicateJobs.RemoveAtSwap(DuplicateIndex);
+void ShutdownGlobalShaderMap()
+{
+	// handle edge case where we get a shutdown before fully initialized (the globals used below are not in a valid state)
+	if (!GIsRHIInitialized)
+	{
+		return;
+	}
 
-	// After removing, we need to update the cached index of the job we swapped
-	if (DuplicateIndex < DuplicateJobs.Num())
+	// at the point this function is called (during the shutdown process) we do not expect any outstanding work that could potentially be still referencing
+	// global shaders, so we are not deferring the deletion (via GGlobalShaderMap_DeferredDeleteCopy) like we do during the shader recompilation.
+	EShaderPlatform Platform = GShaderPlatformForFeatureLevel[GMaxRHIFeatureLevel];
+	if (GGlobalShaderMap[Platform] != nullptr)
 	{
-		DuplicateJobs[DuplicateIndex]->JobCacheRef.DuplicateIndex = DuplicateIndex;
+		GGlobalShaderMap[Platform]->ReleaseAllSections();
+
+		delete GGlobalShaderMap[Platform];
+		GGlobalShaderMap[Platform] = nullptr;
 	}
 }
 
-uint64 FShaderJobCache::GetCurrentMemoryBudget() const
+void ReloadGlobalShaders()
 {
-	uint64 AbsoluteLimit = static_cast<uint64>(GShaderCompilerMaxJobCacheMemoryMB) * 1024ULL * 1024ULL;
-	uint64 RelativeLimit = FMath::Clamp(static_cast<double>(GShaderCompilerMaxJobCacheMemoryPercent), 0.0, 100.0) * (static_cast<double>(FPlatformMemory::GetPhysicalGBRam()) * 1024 * 1024 * 1024) / 100.0;
-	return FMath::Min(AbsoluteLimit, RelativeLimit);
-}
+	UE_LOG(LogShaders, Display, TEXT("Reloading global shaders..."));
 
-FShaderJobCache::FShaderJobCache(FCriticalSection& InCompileQueueSection)
-	: CompileQueueSection(InCompileQueueSection)
-{
-	FMemory::Memzero(PendingJobsHead);
-	FMemory::Memzero(NumPendingJobs);
-#if UE_SHADERCOMPILER_FIFO_JOB_EXECUTION
-	for (int32 PriorityIndex = 0; PriorityIndex < NumShaderCompileJobPriorities; PriorityIndex++)
+	// Flush pending accesses to the existing global shaders.
+	FlushRenderingCommands();
+
+	UMaterialInterface::IterateOverActiveFeatureLevels(
+		[&](ERHIFeatureLevel::Type InFeatureLevel)
+		{
+			auto ShaderPlatform = GShaderPlatformForFeatureLevel[InFeatureLevel];
+			GetGlobalShaderMap(ShaderPlatform)->ReleaseAllSections();
+			CompileGlobalShaderMap(InFeatureLevel, true);
+			VerifyGlobalShaders(ShaderPlatform, nullptr, false);
+		}
+	);
+
+	// Invalidate global bound shader states so they will be created with the new shaders the next time they are set (in SetGlobalBoundShaderState)
+	for (TLinkedList<FGlobalBoundShaderStateResource*>::TIterator It(FGlobalBoundShaderStateResource::GetGlobalBoundShaderStateList()); It; It.Next())
 	{
-		PendingJobsTail[PriorityIndex] = &PendingJobsHead[PriorityIndex];
+		BeginUpdateResourceRHI(*It);
 	}
-#endif
 
-	CurrentlyAllocatedMemory = sizeof(*this) + InputHashToJobData.GetAllocatedSize() + Outputs.GetAllocatedSize();
+	PropagateGlobalShadersToAllPrimitives();
 }
 
-FShaderJobCache::~FShaderJobCache()
+static FAutoConsoleCommand CCmdReloadGlobalShaders = FAutoConsoleCommand(
+	TEXT("ReloadGlobalShaders"),
+	TEXT("Reloads the global shaders file"),
+	FConsoleCommandDelegate::CreateStatic(ReloadGlobalShaders)
+);
+
+void SetGlobalShaderCacheOverrideDirectory(const TArray<FString>& Args)
 {
-	for (TMap<FJobOutputHash, FStoredOutput*>::TIterator Iter(Outputs); Iter; ++Iter)
+	if (Args.Num() == 0)
 	{
-		delete Iter.Value();
+		UE_LOG(LogShaders, Error, TEXT("Failed to set GGlobalShaderCacheOverrideDirectory without any arguments"));
+		return; 
 	}
+	
+	GGlobalShaderCacheOverrideDirectory = Args[0];
+	UE_LOG(LogShaders, Log, TEXT("GGlobalShaderCacheOverrideDirectory = %s"), *GGlobalShaderCacheOverrideDirectory);
 }
 
-void FShaderJobCache::AddJobOutput(FShaderJobData& JobData, const FShaderCommonCompileJob* FinishedJob, const FShaderCompilerInputHash& Hash, const FJobCachedOutput& Contents, int32 InitialHitCount, const bool bAddToDDC)
-{
-	TRACE_CPUPROFILER_EVENT_SCOPE(FShaderJobCache::Add);
-
-	if (!ShaderCompiler::IsJobCacheEnabled())
-	{
-		return;
-	}
+static FAutoConsoleCommand CCmdSetGlobalShaderCacheOverrideDirectory = FAutoConsoleCommand(
+	TEXT("SetGlobalShaderCacheOverrideDirectory"),
+	TEXT("Set the directory to read the override global shader map file from."),
+	FConsoleCommandWithArgsDelegate::CreateStatic(SetGlobalShaderCacheOverrideDirectory));
 
-	if (JobData.HasOutput() && !ShaderCompiler::IsJobCacheDebugValidateEnabled())
+bool RecompileChangedShadersForPlatform(const FString& PlatformName)
+{
+	// figure out what shader platforms to recompile
+	ITargetPlatformManagerModule* TPM = GetTargetPlatformManager();
+	ITargetPlatform* TargetPlatform = TPM->FindTargetPlatform(PlatformName);
+	if (TargetPlatform == nullptr)
 	{
-		return;
+		UE_LOG(LogShaders, Display, TEXT("Failed to find target platform module for %s"), *PlatformName);
+		return false;
 	}
 
-	FJobOutputHash OutputHash = FBlake3::HashBuffer(Contents.GetData(), Contents.GetSize());
+	TArray<FName> DesiredShaderFormats;
+	TargetPlatform->GetAllTargetedShaderFormats(DesiredShaderFormats);
 
-	if (JobData.HasOutput() && ShaderCompiler::IsJobCacheDebugValidateEnabled())
-	{
-		if (OutputHash != JobData.OutputHash)
-		{
-			TStringBuilder<1024> FinishedJobName;
-			FinishedJob->AppendDebugName(FinishedJobName);
+	// figure out which shaders are out of date
+	TArray<const FShaderType*> OutdatedShaderTypes;
+	TArray<const FVertexFactoryType*> OutdatedFactoryTypes;
+	TArray<const FShaderPipelineType*> OutdatedShaderPipelineTypes;
 
-			const FString* CachedJobName = CachedJobNames.Find(JobData.OutputHash);
-			check(CachedJobName);
-			UE_LOG(
-				LogShaderCompilers,
-				Warning,
-				TEXT("Job cache validation found output mismatch!\n")
-				TEXT("Cached job: %s\n")
-				TEXT("Original job: %s\n"),
-				**CachedJobName, FinishedJobName.ToString());
+	// Pick up new changes to shader files
+	FlushShaderFileCache();
 
-			if (GDumpShaderDebugInfo != FShaderCompilingManager::EDumpShaderDebugInfo::Always)
-			{
-				static bool bOnce = false;
-				if (!bOnce)
-				{
-					UE_LOG(
-						LogShaderCompilers,
-						Warning,
-						TEXT("Enable r.DumpShaderDebugInfo=1 to get debug info paths for the mismatching jobs instead of group names (to allow diffing debug artifacts)"));
-					bOnce = true;
-				}
-			}
-		}
-		return;
-	}
+	GetOutdatedShaderTypes(OutdatedShaderTypes, OutdatedShaderPipelineTypes, OutdatedFactoryTypes);
+	UE_LOG(LogShaders, Display, TEXT("We found %d out of date shader types, %d outdated pipeline types, and %d out of date VF types!"), OutdatedShaderTypes.Num(), OutdatedShaderPipelineTypes.Num(), OutdatedFactoryTypes.Num());
 
-	const bool bDumpCachedDebugInfo = CVarDumpShaderOutputCacheHits.GetValueOnAnyThread();
+#if WITH_EDITOR
+	UpdateReferencedUniformBufferNames(OutdatedShaderTypes, OutdatedFactoryTypes, OutdatedShaderPipelineTypes);
+#endif
 
-	// Get dump shader debug output path
-	FString InputDebugInfoPath, InputSourceFilename;
-	if (bDumpCachedDebugInfo)
+	for (int32 FormatIndex = 0; FormatIndex < DesiredShaderFormats.Num(); FormatIndex++)
 	{
-		if (const FShaderCompileJob* SingleJob = FinishedJob->GetSingleShaderJob())
+		// get the shader platform enum
+		const EShaderPlatform ShaderPlatform = ShaderFormatToLegacyShaderPlatform(DesiredShaderFormats[FormatIndex]);
+
+		// Only compile for the desired platform if requested
+		// Kick off global shader recompiles
+		BeginRecompileGlobalShaders(OutdatedShaderTypes, OutdatedShaderPipelineTypes, ShaderPlatform);
+
+		// Block on global shaders
+		FinishRecompileGlobalShaders();
+#if WITH_EDITOR
+		// we only want to actually compile mesh shaders if we have out of date ones
+		if (OutdatedShaderTypes.Num() || OutdatedFactoryTypes.Num())
 		{
-			const FShaderCompilerInput& Input = SingleJob->Input;
-			if (!Input.DumpDebugInfoPath.IsEmpty())
+			for (TObjectIterator<UMaterialInterface> It; It; ++It)
 			{
-				InputDebugInfoPath = Input.DumpDebugInfoPath;
-				InputSourceFilename = FPaths::GetBaseFilename(Input.GetSourceFilename());
+				(*It)->ClearCachedCookedPlatformData(TargetPlatform);
 			}
 		}
+#endif // WITH_EDITOR
 	}
 
-	// Cache this value for thread safety
-	bool bDiscardCacheOutputs = GShaderCompilerDebugDiscardCacheOutputs != 0;
+	if (OutdatedFactoryTypes.Num() || OutdatedShaderTypes.Num())
+	{
+		return true;
+	}
+	return false;
+}
 
-	// add the record
-	if (UNLIKELY(bDiscardCacheOutputs == false))
+extern ENGINE_API const TCHAR* ODSCCmdEnumToString(ODSCRecompileCommand Cmd)
+{
+	switch (Cmd)
 	{
-		JobData.OutputHash = OutputHash;
-		JobData.bOutputFromDDC = false;
+	case ODSCRecompileCommand::None:
+		return TEXT("None");
+	case ODSCRecompileCommand::Changed:
+		return TEXT("Change");
+	case ODSCRecompileCommand::Global:
+		return TEXT("Global");
+	case ODSCRecompileCommand::Material:
+		return TEXT("Material");
+	case ODSCRecompileCommand::SingleShader:
+		return TEXT("SingleShader");
+	case ODSCRecompileCommand::ResetMaterialCache:
+			return TEXT("ResetMaterialCache");
 	}
+	ensure(false);
+	return TEXT("Unknown");
+}
 
-	FStoredOutput** CannedOutput = Outputs.Find(OutputHash);
-	if (CannedOutput)
+void BeginRecompileGlobalShaders(const TArray<const FShaderType*>& OutdatedShaderTypes, const TArray<const FShaderPipelineType*>& OutdatedShaderPipelineTypes, EShaderPlatform ShaderPlatform, const ITargetPlatform* TargetPlatform)
+{
+#if WITH_EDITOR
+	if (!FPlatformProperties::RequiresCookedData())
 	{
-		// update the output hit count
-		int32 NumRef;
-		if (UNLIKELY(bDiscardCacheOutputs == false))
-		{
-			NumRef = (*CannedOutput)->AddRef();
-		}
-		else
-		{
-			NumRef = (*CannedOutput)->GetNumReferences();
-		}
+		// Flush pending accesses to the existing global shaders.
+		FlushRenderingCommands();
 
-		if (UNLIKELY(bDumpCachedDebugInfo))
+		// Calling CompileGlobalShaderMap will force starting the compile jobs if the map is empty (by calling VerifyGlobalShaders)
+		CompileGlobalShaderMap(ShaderPlatform, TargetPlatform, false);
+		FGlobalShaderMap* GlobalShaderMap = GetGlobalShaderMap(ShaderPlatform);
+
+		// Now check if there is any work to be done wrt outdates types
+		if (OutdatedShaderTypes.Num() > 0 || OutdatedShaderPipelineTypes.Num() > 0)
 		{
-			// Write cache hit debug file
-			const FString& CachedDebugInfoPath = (*CannedOutput)->CachedDebugInfoPath;
-			if (!CachedDebugInfoPath.IsEmpty())
-			{
-				const int32 CacheHit = NumRef - 1;
-				const FString CacheHitFilename = FString::Printf(TEXT("%s/%s.%d.cachehit"), *CachedDebugInfoPath, *InputSourceFilename, CacheHit);
-				FFileHelper::SaveStringToFile(InputDebugInfoPath, *CacheHitFilename);
-			}
+			VerifyGlobalShaders(ShaderPlatform, TargetPlatform, false, &OutdatedShaderTypes, &OutdatedShaderPipelineTypes);
 		}
 	}
-	else
-	{
-		if (UNLIKELY(bDiscardCacheOutputs == false))
-		{
-			const uint64 OutputsOriginalSize = Outputs.GetAllocatedSize();
-
-			FStoredOutput* NewStoredOutput = new FStoredOutput();
-			NewStoredOutput->NumHits = InitialHitCount;
-			NewStoredOutput->JobOutput = Contents;
-			NewStoredOutput->CachedDebugInfoPath = InputDebugInfoPath;
-			NewStoredOutput->AddRef();
-			Outputs.Add(OutputHash, NewStoredOutput);
+#endif
+}
 
-			if (ShaderCompiler::IsJobCacheDebugValidateEnabled())
-			{
-				TStringBuilder<1024> NameBuilder;
-				FinishedJob->AppendDebugName(NameBuilder);
+void FinishRecompileGlobalShaders()
+{
+	// Block until global shaders have been compiled and processed
+	GShaderCompilingManager->ProcessAsyncResults(false, true);
+}
 
-				CachedJobNames.Add(OutputHash, NameBuilder.ToString());
-			}
+void LoadGlobalShadersForRemoteRecompile(FArchive& Ar, EShaderPlatform ShaderPlatform)
+{
+	uint8 bIsValid = 0;
+	Ar << bIsValid;
 
-			CurrentlyAllocatedMemory += NewStoredOutput->GetAllocatedSize() + Outputs.GetAllocatedSize() - OutputsOriginalSize;
-		}
+	if (bIsValid)
+	{
+		FlushRenderingCommands();
 
-		if (UNLIKELY(bDumpCachedDebugInfo))
+		FGlobalShaderMap* NewGlobalShaderMap = new FGlobalShaderMap(ShaderPlatform);
+		if (NewGlobalShaderMap)
 		{
-			// Write new allocated cache file
-			if (!InputDebugInfoPath.IsEmpty())
-			{
-				const FString CacheFilename = FString::Printf(TEXT("%s/%s.bytecode"), *InputDebugInfoPath, *InputSourceFilename);
-				FFileHelper::SaveArrayToFile(TArrayView<const uint8>((const uint8*)Contents.GetData(), Contents.GetSize()), *CacheFilename);
-			}
-		}
+			NewGlobalShaderMap->LoadFromGlobalArchive(Ar);
 
-		// delete oldest cache entries if we exceed the budget
-		uint64 MemoryBudgetBytes = GetCurrentMemoryBudget();
-		if (MemoryBudgetBytes)
-		{
-			if (CurrentlyAllocatedMemory > MemoryBudgetBytes)
+			FString FailureReason;
+			bool bIsNewGlobalShaderMapComplete = IsGlobalShaderMapComplete(nullptr, NewGlobalShaderMap, ShaderPlatform, &FailureReason);
+
+			if (bIsNewGlobalShaderMapComplete)
 			{
-				TRACE_CPUPROFILER_EVENT_SCOPE(FShaderJobCache::Trim);
+				if (GGlobalShaderMap[ShaderPlatform])
+				{
+					GGlobalShaderMap[ShaderPlatform]->ReleaseAllSections();
+					delete GGlobalShaderMap[ShaderPlatform];
+					GGlobalShaderMap[ShaderPlatform] = nullptr;
+				}
+				GGlobalShaderMap[ShaderPlatform] = NewGlobalShaderMap;
 
-				uint64 TargetBudgetBytes = MemoryBudgetBytes * FMath::Clamp(GShaderCompilerJobCacheOverflowReducePercent, 0, 100) / 100;
-				uint64 MemoryBefore = CurrentlyAllocatedMemory;
+				VerifyGlobalShaders(ShaderPlatform, nullptr, false);
 
-				// Cull outputs to reach the budget target
-				CullOutputsToMemoryBudget(TargetBudgetBytes);
+				// Invalidate global bound shader states so they will be created with the new shaders the next time they are set (in SetGlobalBoundShaderState)
+				for (TLinkedList<FGlobalBoundShaderStateResource*>::TIterator It(FGlobalBoundShaderStateResource::GetGlobalBoundShaderStateList()); It; It.Next())
+				{
+					BeginUpdateResourceRHI(*It);
+				}
 
-				UE_LOG(LogShaderCompilers, Display, TEXT("Memory overflow, reduced from %.1lf to %.1lf MB."), (double)MemoryBefore / (1024 * 1024), (double)CurrentlyAllocatedMemory / (1024 * 1024));
+				PropagateGlobalShadersToAllPrimitives();
 			}
-		}
-	}
-
-#if WITH_EDITOR
-	const bool bCachePerShaderDDC = IsShaderJobCacheDDCEnabled() && bAddToDDC;
-
-	if (bCachePerShaderDDC)
-	{
-		// Create key.
-		UE::DerivedData::FCacheKey Key;
-		Key.Bucket = ShaderJobCacheDDCBucket;
-		Key.Hash = Hash;
-		UE::DerivedData::FCacheRecordBuilder RecordBuilder(Key);
-
-		RecordBuilder.AddValue(ShaderJobCacheId, FSharedBuffer::MakeView(Contents));
-
-		TRACE_COUNTER_ADD(Shaders_JobCacheDDCBytesSent, Contents.GetSize());
+			else
+			{
+				FString ErrorMessage = FString::Printf(TEXT("New global shader map is incomplete and will not be used. Reason:\n%s\n"
+													        "Please check the ODSC server log & that client/editor are compiled"), *FailureReason);
 
-		UE::DerivedData::FRequestOwner RequestOwner(UE::DerivedData::EPriority::Normal);
-		UE::DerivedData::FRequestBarrier RequestBarrier(RequestOwner);
-		RequestOwner.KeepAlive();
-		UE::DerivedData::GetCache().Put(
-			{ {{TEXT("FShaderJobCache")}, RecordBuilder.Build(), IsShaderJobCacheDDCRemotePolicyEnabled() ? UE::DerivedData::ECachePolicy::Default : UE::DerivedData::ECachePolicy::Local } },
-			RequestOwner
-		);
-	}
+				UE_LOG(LogShaderCompilers, Error, TEXT("%s"), *ErrorMessage);
+#if WITH_ODSC
+				FODSCManager::ReportODSCError(ErrorMessage);
 #endif
-}
-
-#include "Math/UnitConversion.h"
-
-/** Returns memory used by the cache*/
-uint64 FShaderJobCache::GetAllocatedMemory() const
-{
-	return CurrentlyAllocatedMemory;
-}
-
-/** Compute memory used by the cache from scratch.  Should match GetAllocatedMemory() if CurrentlyAllocateMemory is being properly updated. */
-uint64 FShaderJobCache::ComputeAllocatedMemory() const
-{
-	uint64 AllocatedSize = sizeof(FShaderJobCache) + InputHashToJobData.GetAllocatedSize() + Outputs.GetAllocatedSize();
-	for (auto OutputIter : Outputs)
-	{
-		AllocatedSize += OutputIter.Value->GetAllocatedSize();
+				
+				delete NewGlobalShaderMap;
+			}
+		}
 	}
-	return AllocatedSize;
-}
-
-void FShaderJobCache::GetStats(FShaderCompilerStats& OutStats) const
-{
-	FReadScopeLock Locker(JobLock);
-	OutStats.Counters.TotalCacheSearchAttempts = TotalSearchAttempts;
-	OutStats.Counters.TotalCacheHits = TotalCacheHits;
-	OutStats.Counters.TotalCacheDuplicates = TotalCacheDuplicates;
-	OutStats.Counters.TotalCacheDDCQueries = TotalCacheDDCQueries;
-	OutStats.Counters.TotalCacheDDCHits = TotalCacheDDCHits;
-	OutStats.Counters.UniqueCacheInputHashes = InputHashToJobData.Num();
-	OutStats.Counters.UniqueCacheOutputs = Outputs.Num();
-	OutStats.Counters.CacheMemUsed = GetAllocatedMemory();
-	OutStats.Counters.CacheMemBudget = GetCurrentMemoryBudget();
 }
 
 #undef LOCTEXT_NAMESPACE
