[Shaders] Split up ShaderCompiler.cpp into multiple source files for better readability.

diff --git a/Engine/Source/Runtime/Engine/Private/ShaderCompiler/ShaderCompilerThreadRunnable.cpp b/Engine/Source/Runtime/Engine/Private/ShaderCompiler/ShaderCompilerThreadRunnable.cpp
new file mode 100644
--- /dev/null
+++ b/Engine/Source/Runtime/Engine/Private/ShaderCompiler/ShaderCompilerThreadRunnable.cpp
@@ -0,0 +1,961 @@
+// Copyright Epic Games, Inc. All Rights Reserved.
+
+/*=============================================================================
+	ShaderCompilerThreadRunnable.cpp:
+	Implements FShaderCompileThreadRunnableBase and FShaderCompileThreadRunnable.
+=============================================================================*/
+
+#include "ShaderCompilerPrivate.h"
+
+#include "Async/ParallelFor.h"
+#include "HAL/PlatformFile.h"
+#include "HAL/PlatformFileManager.h"
+#include "Misc/ScopeTryLock.h"
+#include "ProfilingDebugging/CookStats.h"
+#include "ProfilingDebugging/DiagnosticTable.h"
+#include "ProfilingDebugging/LoadTimeTracker.h"
+#include "ProfilingDebugging/StallDetector.h"
+
+
+float GShaderCompilerTooLongIOThresholdSeconds = 0.3;
+static FAutoConsoleVariableRef CVarShaderCompilerTooLongIOThresholdSeconds(
+	TEXT("r.ShaderCompiler.TooLongIOThresholdSeconds"),
+	GShaderCompilerTooLongIOThresholdSeconds,
+	TEXT("By default, task files for SCW will be read/written sequentially, but if we ever spend more than this time (0.3s by default) doing that, we'll switch to parallel.") \
+	TEXT("We don't default to parallel writes as it increases the CPU overhead from the shader compiler."),
+	ECVF_Default);
+
+static TAutoConsoleVariable<bool> CVarCompileParallelInProcess(
+	TEXT("r.ShaderCompiler.ParallelInProcess"),
+	false,
+	TEXT("EXPERIMENTAL- If true, shader compilation will be executed in-process in parallel. Note that this will serialize if the legacy preprocessor is enabled."),
+	ECVF_ReadOnly);
+
+// Configuration to retry shader compile through workers after a worker has been abandoned
+static constexpr int32 GSingleThreadedRunsDisabled = -2;
+static constexpr int32 GSingleThreadedRunsIncreaseFactor = 8;
+static constexpr int32 GSingleThreadedRunsMaxCount = (1 << 24);
+
+
+/** Information tracked for each shader compile worker process instance. */
+struct FShaderCompileWorkerInfo
+{
+	/** Process handle of the worker app once launched.  Invalid handle means no process. */
+	FProcHandle WorkerProcess;
+
+	/** Tracks whether tasks have been issued to the worker. */
+	bool bIssuedTasksToWorker;	
+
+	/** Whether the worker has been launched for this set of tasks. */
+	bool bLaunchedWorker;
+
+	/** Tracks whether all tasks issued to the worker have been received. */
+	bool bComplete;
+
+	/** Whether this worker is available for new jobs. It will be false when shutting down the worker. */
+	bool bAvailable; 
+
+	/** Time at which the worker started the most recent batch of tasks. */
+	double StartTime;
+
+	/** Time at which the worker ended the most recent batch of tasks. */
+	double FinishTime = 0.0;
+
+	/** Jobs that this worker is responsible for compiling. */
+	TArray<FShaderCommonCompileJobPtr> QueuedJobs;
+
+	FShaderCompileWorkerInfo() :
+		bIssuedTasksToWorker(false),		
+		bLaunchedWorker(false),
+		bComplete(false),
+		bAvailable(true),
+		StartTime(0)
+	{
+	}
+
+	// warning: not virtual
+	~FShaderCompileWorkerInfo()
+	{
+		if(WorkerProcess.IsValid())
+		{
+			FPlatformProcess::TerminateProc(WorkerProcess);
+			FPlatformProcess::CloseProc(WorkerProcess);
+		}
+	}
+};
+
+FShaderCompileThreadRunnableBase::FShaderCompileThreadRunnableBase(FShaderCompilingManager* InManager)
+	: Manager(InManager)
+	, Thread(nullptr)
+	, MinPriorityIndex(0)
+	, MaxPriorityIndex(NumShaderCompileJobPriorities - 1)
+	, bForceFinish(false)
+{
+}
+void FShaderCompileThreadRunnableBase::StartThread()
+{
+	if (Manager->bAllowAsynchronousShaderCompiling && !FPlatformProperties::RequiresCookedData())
+	{
+		Thread = FRunnableThread::Create(this, GetThreadName(), 0, TPri_Normal, FPlatformAffinity::GetPoolThreadMask());
+	}
+}
+
+FShaderCompileThreadRunnable::FShaderCompileThreadRunnable(FShaderCompilingManager* InManager)
+	: FShaderCompileThreadRunnableBase(InManager)
+	, LastCheckForWorkersTime(0)
+{
+	for (uint32 WorkerIndex = 0; WorkerIndex < Manager->NumShaderCompilingThreads; WorkerIndex++)
+	{
+		WorkerInfos.Add(MakeUnique<FShaderCompileWorkerInfo>());
+	}
+}
+
+FShaderCompileThreadRunnable::~FShaderCompileThreadRunnable()
+{
+	FScopeLock WorkerScopeLock(&WorkerInfosLock);
+	WorkerInfos.Empty();
+}
+
+void FShaderCompileThreadRunnable::OnMachineResourcesChanged()
+{
+	bool bWaitForWorkersToShutdown = false;
+	{
+		FScopeLock WorkerScopeLock(&WorkerInfosLock);
+		// Set all bAvailable flags back to true
+		for (TUniquePtr< FShaderCompileWorkerInfo>& WorkerInfo : WorkerInfos)
+		{
+			WorkerInfo->bAvailable = true;
+		}
+
+		if (Manager->NumShaderCompilingThreads >= static_cast<uint32>(WorkerInfos.Num()))
+		{
+			while (static_cast<uint32>(WorkerInfos.Num()) < Manager->NumShaderCompilingThreads)
+			{
+				WorkerInfos.Add(MakeUnique<FShaderCompileWorkerInfo>());
+			}
+		}
+		else
+		{
+			for (int32 Index = 0; Index < WorkerInfos.Num(); ++Index)
+			{
+				FShaderCompileWorkerInfo& WorkerInfo = *WorkerInfos[Index];
+				bool bReadyForShutdown = WorkerInfo.QueuedJobs.Num() == 0;
+				if (bReadyForShutdown)
+				{
+					WorkerInfos.RemoveAtSwap(Index--);
+					if (WorkerInfos.Num() == Manager->NumShaderCompilingThreads)
+					{
+						break;
+					}
+				}
+			}
+			bWaitForWorkersToShutdown = Manager->NumShaderCompilingThreads < static_cast<uint32>(WorkerInfos.Num());
+			for (int32 Index = WorkerInfos.Num() - 1;
+				static_cast<uint32>(Index) >= Manager->NumShaderCompilingThreads; --Index)
+			{
+				WorkerInfos[Index]->bAvailable = false;
+			}
+		}
+	}
+	const double StartTime = FPlatformTime::Seconds();
+	constexpr float MaxDurationToWait = 60.f;
+	const double MaxTimeToWait = StartTime + MaxDurationToWait;
+	while (bWaitForWorkersToShutdown)
+	{
+		FPlatformProcess::Sleep(0.01f);
+		const double CurrentTime = FPlatformTime::Seconds();
+		if (CurrentTime > MaxTimeToWait)
+		{
+			UE_LOG(LogShaderCompilers, Warning, TEXT("OnMachineResourcesChanged timedout waiting %.0f seconds for WorkerInfos to complete. Workers will remain allocated."),
+				(float)(CurrentTime - StartTime));
+			break;
+		}
+
+		FScopeLock WorkerScopeLock(&WorkerInfosLock);
+		for (int32 Index = WorkerInfos.Num() - 1;
+			static_cast<uint32>(Index) >= Manager->NumShaderCompilingThreads; --Index)
+		{
+			FShaderCompileWorkerInfo& WorkerInfo = *WorkerInfos[Index];
+			check(!WorkerInfos[Index]->bAvailable); // It should still be set to false from when we changed it above
+			bool bReadyForShutdown = WorkerInfo.QueuedJobs.Num() == 0;
+			if (bReadyForShutdown)
+			{
+				WorkerInfos.RemoveAtSwap(Index);
+			}
+		}
+		bWaitForWorkersToShutdown = Manager->NumShaderCompilingThreads < static_cast<uint32>(WorkerInfos.Num());
+	}
+}
+
+/** Entry point for the shader compiling thread. */
+uint32 FShaderCompileThreadRunnableBase::Run()
+{
+	LLM_SCOPE_BYTAG(ShaderCompiler);
+	check(Manager->bAllowAsynchronousShaderCompiling);
+	while (!bForceFinish)
+	{
+		CompilingLoop();
+	}
+	UE_LOG(LogShaderCompilers, Display, TEXT("Shaders left to compile 0"));
+
+	return 0;
+}
+
+int32 FShaderCompileThreadRunnable::PullTasksFromQueue()
+{
+	TRACE_CPUPROFILER_EVENT_SCOPE(FShaderCompileThreadRunnable::PullTasksFromQueue);
+
+	FScopeLock WorkerScopeLock(&WorkerInfosLock); // Must be entered before CompileQueueSection
+	int32 NumActiveThreads = 0;
+	int32 NumJobsStarted[NumShaderCompileJobPriorities] = { 0 };
+	{
+		// Enter the critical section so we can access the input and output queues
+		FScopeLock Lock(&Manager->CompileQueueSection);
+
+		const int32 NumWorkersToFeed = Manager->bCompilingDuringGame ? Manager->NumShaderCompilingThreadsDuringGame : WorkerInfos.Num();
+
+		for (int32 PriorityIndex = MaxPriorityIndex; PriorityIndex >= MinPriorityIndex; --PriorityIndex)
+		{
+			int32 NumPendingJobs = Manager->AllJobs.GetNumPendingJobs((EShaderCompileJobPriority)PriorityIndex);
+			// Try to distribute the work evenly between the workers
+			const auto NumJobsPerWorker = (NumPendingJobs / NumWorkersToFeed) + 1;
+
+			for (int32 WorkerIndex = 0; WorkerIndex < WorkerInfos.Num(); WorkerIndex++)
+			{
+				FShaderCompileWorkerInfo& CurrentWorkerInfo = *WorkerInfos[WorkerIndex];
+
+				// If this worker doesn't have any queued jobs, look for more in the input queue
+				if (CurrentWorkerInfo.QueuedJobs.Num() == 0 && CurrentWorkerInfo.bAvailable && WorkerIndex < NumWorkersToFeed)
+				{
+					check(!CurrentWorkerInfo.bComplete);
+
+					NumPendingJobs = Manager->AllJobs.GetNumPendingJobs((EShaderCompileJobPriority)PriorityIndex);
+					if (NumPendingJobs > 0)
+					{
+						UE_LOG(LogShaderCompilers, Verbose, TEXT("Worker (%d/%d): shaders left to compile %i"), WorkerIndex + 1, WorkerInfos.Num(), NumPendingJobs);
+
+						int32 MaxNumJobs = 1;
+						// high priority jobs go in 1 per "batch", unless the engine is still starting up
+						if (PriorityIndex < (int32)EShaderCompileJobPriority::High || Manager->IgnoreAllThrottling())
+						{
+							MaxNumJobs = FMath::Min3(NumJobsPerWorker, NumPendingJobs, Manager->MaxShaderJobBatchSize);
+						}
+
+						NumJobsStarted[PriorityIndex] += Manager->AllJobs.GetPendingJobs(EShaderCompilerWorkerType::LocalThread, (EShaderCompileJobPriority)PriorityIndex, 1, MaxNumJobs, CurrentWorkerInfo.QueuedJobs);
+
+						// Update the worker state as having new tasks that need to be issued					
+						// don't reset worker app ID, because the shadercompileworkers don't shutdown immediately after finishing a single job queue.
+						CurrentWorkerInfo.bIssuedTasksToWorker = false;
+						CurrentWorkerInfo.bLaunchedWorker = false;
+						CurrentWorkerInfo.StartTime = FPlatformTime::Seconds();
+						NumActiveThreads++;
+
+						if (CurrentWorkerInfo.FinishTime > 0.0)
+						{
+							const double WorkerIdleTime = CurrentWorkerInfo.StartTime - CurrentWorkerInfo.FinishTime;
+							GShaderCompilerStats->RegisterLocalWorkerIdleTime(WorkerIdleTime);
+							if (Manager->bLogJobCompletionTimes)
+							{
+								UE_LOG(LogShaderCompilers, Display, TEXT("  Worker (%d/%d) started working after being idle for %fs"), WorkerIndex + 1, WorkerInfos.Num(), WorkerIdleTime);
+							}
+						}
+					}
+				}
+			}
+		}
+	}
+
+	for (int32 WorkerIndex = 0; WorkerIndex < WorkerInfos.Num(); WorkerIndex++)
+	{
+		if (WorkerInfos[WorkerIndex]->QueuedJobs.Num() > 0)
+		{
+			NumActiveThreads++;
+		}
+	}
+
+	for (int32 PriorityIndex = 0; PriorityIndex < NumShaderCompileJobPriorities; ++PriorityIndex)
+	{
+		if (NumJobsStarted[PriorityIndex] > 0)
+		{
+			UE_LOG(LogShaderCompilers, Verbose, TEXT("Started %d 'Local' shader compile jobs with '%s' priority"),
+				NumJobsStarted[PriorityIndex],
+				ShaderCompileJobPriorityToString((EShaderCompileJobPriority)PriorityIndex));
+		}
+	}
+
+	return NumActiveThreads;
+}
+
+void FShaderCompileThreadRunnable::PushCompletedJobsToManager()
+{
+	FScopeLock WorkerScopeLock(&WorkerInfosLock); // Must be entered before CompileQueueSection
+
+	for (int32 WorkerIndex = 0; WorkerIndex < WorkerInfos.Num(); WorkerIndex++)
+	{
+		FShaderCompileWorkerInfo& CurrentWorkerInfo = *WorkerInfos[WorkerIndex];
+
+		// Add completed jobs to the output queue, which is ShaderMapJobs
+		if (CurrentWorkerInfo.bComplete)
+		{
+			// Enter the critical section so we can access the input and output queues
+			FScopeLock Lock(&Manager->CompileQueueSection);
+
+			for (int32 JobIndex = 0; JobIndex < CurrentWorkerInfo.QueuedJobs.Num(); JobIndex++)
+			{
+				auto& Job = CurrentWorkerInfo.QueuedJobs[JobIndex];
+
+				Manager->ProcessFinishedJob(Job.GetReference());
+			}
+
+			const float ElapsedTime = FPlatformTime::Seconds() - CurrentWorkerInfo.StartTime;
+
+			Manager->WorkersBusyTime += ElapsedTime;
+			COOK_STAT(ShaderCompilerCookStats::AsyncCompileTimeSec += ElapsedTime);
+
+			// Log if requested or if there was an exceptionally slow batch, to see the offender easily
+			if (Manager->bLogJobCompletionTimes || ElapsedTime > 60.0f)
+			{
+				TArray<FShaderCommonCompileJobPtr> SortedJobs = CurrentWorkerInfo.QueuedJobs;
+				SortedJobs.Sort([](const FShaderCommonCompileJobPtr& JobA, const FShaderCommonCompileJobPtr& JobB)
+					{
+						const FShaderCompileJob* SingleJobA = JobA->GetSingleShaderJob();
+						const FShaderCompileJob* SingleJobB = JobB->GetSingleShaderJob();
+
+						const float TimeA = SingleJobA ? SingleJobA->Output.CompileTime : 0.0f;
+						const float TimeB = SingleJobB ? SingleJobB->Output.CompileTime : 0.0f;
+
+						return TimeA > TimeB;
+					});
+
+				FString JobNames;
+
+				for (int32 JobIndex = 0; JobIndex < SortedJobs.Num(); JobIndex++)
+				{
+					const FShaderCommonCompileJob& Job = *SortedJobs[JobIndex];
+					if (const FShaderCompileJob* SingleJob = Job.GetSingleShaderJob())
+					{
+						const TCHAR* JobName = Manager->bLogJobCompletionTimes ? *SingleJob->Input.DebugGroupName : SingleJob->Key.ShaderType->GetName();
+						JobNames += FString::Printf(TEXT("%s:%d(%s) [WorkerTime=%.3fs]"), JobName, SingleJob->Key.PermutationId, *SingleJob->Input.ShaderFormat.ToString(), SingleJob->Output.CompileTime);
+					}
+					else
+					{
+						const FShaderPipelineCompileJob* PipelineJob = Job.GetShaderPipelineJob();
+						JobNames += FString(PipelineJob->Key.ShaderPipeline->GetName());
+					}
+					if (JobIndex < SortedJobs.Num() - 1)
+					{
+						JobNames += TEXT(", ");
+					}
+				}
+
+				UE_LOG(LogShaderCompilers, Display, TEXT("Worker (%d/%d) finished batch of %u jobs in %.3fs, %s"), WorkerIndex + 1, WorkerInfos.Num(), SortedJobs.Num(), ElapsedTime, *JobNames);
+			}
+
+			CurrentWorkerInfo.FinishTime = FPlatformTime::Seconds();
+			CurrentWorkerInfo.bComplete = false;
+			CurrentWorkerInfo.QueuedJobs.Empty();
+		}
+	}
+}
+
+void FShaderCompileThreadRunnable::WriteNewTasks()
+{
+	TRACE_CPUPROFILER_EVENT_SCOPE(ShaderCompiler.WriteNewTasks);
+	FScopeLock WorkerScopeLock(&WorkerInfosLock);
+
+	// first, a quick check if anything is needed just to avoid hammering the task graph
+	bool bHasTasksToWrite = false;
+	for (int32 WorkerIndex = 0, NumWorkers = WorkerInfos.Num(); WorkerIndex < NumWorkers; ++WorkerIndex)
+	{
+		FShaderCompileWorkerInfo& CurrentWorkerInfo = *WorkerInfos[WorkerIndex];
+		if (!CurrentWorkerInfo.bIssuedTasksToWorker && CurrentWorkerInfo.QueuedJobs.Num() > 0)
+		{
+			bHasTasksToWrite = true;
+			break;
+		}
+	}
+
+	if (!bHasTasksToWrite)
+	{
+		return;
+	}
+
+
+	auto LoopBody = [this](int32 WorkerIndex)
+	{
+		// The calling thread holds the WorkerInfosLock and will not modify WorkerInfos, 
+		// so we can access it here without entering the lock
+		FShaderCompileWorkerInfo& CurrentWorkerInfo = *WorkerInfos[WorkerIndex];
+
+		// Only write tasks once
+		if (!CurrentWorkerInfo.bIssuedTasksToWorker && CurrentWorkerInfo.QueuedJobs.Num() > 0)
+		{
+			TRACE_CPUPROFILER_EVENT_SCOPE(ShaderCompiler.WriteNewTasksForWorker);
+			CurrentWorkerInfo.bIssuedTasksToWorker = true;
+
+			const FString WorkingDirectory = Manager->AbsoluteShaderBaseWorkingDirectory + FString::FromInt(WorkerIndex);
+
+			// To make sure that the process waiting for input file won't try to read it until it's ready
+			// we use a temp file name during writing.
+			FString TransferFileName;
+			do
+			{
+				FGuid Guid;
+				FPlatformMisc::CreateGuid(Guid);
+				TransferFileName = WorkingDirectory + Guid.ToString();
+			} while (IFileManager::Get().FileSize(*TransferFileName) != INDEX_NONE);
+
+			// Write out the file that the worker app is waiting for, which has all the information needed to compile the shader.
+			// 'Only' indicates that the worker should keep checking for more tasks after this one
+			FArchive* TransferFile = nullptr;
+
+			int32 RetryCount = 0;
+			// Retry over the next two seconds if we can't write out the input file
+			// Anti-virus and indexing applications can interfere and cause this write to fail
+			//@todo - switch to shared memory or some other method without these unpredictable hazards
+			while (TransferFile == nullptr && RetryCount < 2000)
+			{
+				if (RetryCount > 0)
+				{
+					FPlatformProcess::Sleep(0.01f);
+				}
+				TransferFile = IFileManager::Get().CreateFileWriter(*TransferFileName, FILEWRITE_EvenIfReadOnly);
+				RetryCount++;
+				if (TransferFile == nullptr)
+				{
+					UE_LOG(LogShaderCompilers, Warning, TEXT("Could not create the shader compiler transfer file '%s', retrying..."), *TransferFileName);
+				}
+			}
+			if (TransferFile == nullptr)
+			{
+				UE_LOG(LogShaderCompilers, Fatal, TEXT("Could not create the shader compiler transfer file '%s'."), *TransferFileName);
+			}
+			check(TransferFile);
+
+			GShaderCompilerStats->RegisterJobBatch(CurrentWorkerInfo.QueuedJobs.Num(), FShaderCompilerStats::EExecutionType::Local);
+			if (!FShaderCompileUtilities::DoWriteTasks(CurrentWorkerInfo.QueuedJobs, *TransferFile))
+			{
+				uint64 TotalDiskSpace = 0;
+				uint64 FreeDiskSpace = 0;
+				FPlatformMisc::GetDiskTotalAndFreeSpace(TransferFileName, TotalDiskSpace, FreeDiskSpace);
+				UE_LOG(LogShaderCompilers, Error, TEXT("Could not write the shader compiler transfer filename to '%s' (Free Disk Space: %llu."), *TransferFileName, FreeDiskSpace);
+			}
+			delete TransferFile;
+
+#if 0 // debugging code to dump the worker inputs
+			static FCriticalSection ArchiveLock;
+			{
+				FScopeLock Locker(&ArchiveLock);
+				static int ArchivedTransferFileNum = 0;
+				FString JobCacheDir = ShaderCompiler::IsJobCacheEnabled() ? TEXT("JobCache") : TEXT("NoJobCache");
+				FString ArchiveDir = FPaths::ProjectSavedDir() / TEXT("ArchivedWorkerInputs") / JobCacheDir;
+				FString ArchiveName = FString::Printf(TEXT("Input-%d"), ArchivedTransferFileNum++);
+				FString ArchivePath = ArchiveDir / ArchiveName;
+				if (!IFileManager::Get().Copy(*ArchivePath, *TransferFileName))
+				{
+					UE_LOG(LogInit, Error, TEXT("Could not copy file %s to %s"), *TransferFileName, *ArchivePath);
+					ensure(false);
+				}
+			}
+#endif
+
+			// Change the transfer file name to proper one
+			FString ProperTransferFileName = WorkingDirectory / TEXT("WorkerInputOnly.in");
+			if (!IFileManager::Get().Move(*ProperTransferFileName, *TransferFileName))
+			{
+				uint64 TotalDiskSpace = 0;
+				uint64 FreeDiskSpace = 0;
+				FPlatformMisc::GetDiskTotalAndFreeSpace(TransferFileName, TotalDiskSpace, FreeDiskSpace);
+				UE_LOG(LogShaderCompilers, Error, TEXT("Could not rename the shader compiler transfer filename to '%s' from '%s' (Free Disk Space: %llu)."), *ProperTransferFileName, *TransferFileName, FreeDiskSpace);
+			}
+		}
+	};
+
+	if (bParallelizeIO)
+	{
+		ParallelFor( TEXT("ShaderCompiler.WriteNewTasks.PF"), WorkerInfos.Num(),1, LoopBody, EParallelForFlags::Unbalanced);
+	}
+	else
+	{
+		double StartIOWork = FPlatformTime::Seconds();
+		for (int32 WorkerIndex = 0; WorkerIndex < WorkerInfos.Num(); WorkerIndex++)
+		{
+			LoopBody(WorkerIndex);
+		}
+
+		double IODuration = FPlatformTime::Seconds() - StartIOWork;
+		if (IODuration > GShaderCompilerTooLongIOThresholdSeconds)
+		{
+			UE_LOG(LogShaderCompilers, Display, TEXT("FShaderCompileThreadRunnable::WriteNewTasks()() took too long (%.3f seconds, threshold is %.3f s), will parallelize next time."), IODuration, GShaderCompilerTooLongIOThresholdSeconds);
+			bParallelizeIO = true;
+		}
+	}
+}
+
+bool FShaderCompileThreadRunnable::LaunchWorkersIfNeeded()
+{
+	TRACE_CPUPROFILER_EVENT_SCOPE(FShaderCompileThreadRunnable::LaunchWorkersIfNeeded);
+
+	const double CurrentTime = FPlatformTime::Seconds();
+	// Limit how often we check for workers running since IsApplicationRunning eats up some CPU time on Windows
+	const bool bCheckForWorkerRunning = (CurrentTime - LastCheckForWorkersTime > .1f);
+	bool bAbandonWorkers = false;
+	uint32_t NumberLaunched = 0;
+
+	if (bCheckForWorkerRunning)
+	{
+		LastCheckForWorkersTime = CurrentTime;
+	}
+
+	bool bClearStaleOutputs = Manager->bClearStaleWorkerOutputs;
+	Manager->bClearStaleWorkerOutputs = false;
+
+	FScopeLock WorkerScopeLock(&WorkerInfosLock);
+	for (int32 WorkerIndex = 0; WorkerIndex < WorkerInfos.Num(); WorkerIndex++)
+	{
+		FShaderCompileWorkerInfo& CurrentWorkerInfo = *WorkerInfos[WorkerIndex];
+		if (CurrentWorkerInfo.QueuedJobs.Num() == 0)
+		{
+			// Skip if nothing to do
+			// Also, use the opportunity to free OS resources by cleaning up handles of no more running processes
+			if (CurrentWorkerInfo.WorkerProcess.IsValid() && !FShaderCompilingManager::IsShaderCompilerWorkerRunning(CurrentWorkerInfo.WorkerProcess))
+			{
+				FPlatformProcess::CloseProc(CurrentWorkerInfo.WorkerProcess);
+				CurrentWorkerInfo.WorkerProcess = FProcHandle();
+			}
+			continue;
+		}
+
+		if (!CurrentWorkerInfo.WorkerProcess.IsValid() || (bCheckForWorkerRunning && !FShaderCompilingManager::IsShaderCompilerWorkerRunning(CurrentWorkerInfo.WorkerProcess)))
+		{
+			TRACE_CPUPROFILER_EVENT_SCOPE(FShaderCompileThreadRunnable::LaunchingWorkers);
+
+			// @TODO: dubious design - worker should not be launched unless we know there's more work to do.
+			bool bLaunchAgain = true;
+
+			// Detect when the worker has exited due to fatal error
+			// bLaunchedWorker check here is necessary to distinguish between 'process isn't running because it crashed' and 'process isn't running because it exited cleanly and the outputfile was already consumed'
+			if (CurrentWorkerInfo.WorkerProcess.IsValid())
+			{
+				// shader compiler exited one way or another, so clear out the stale PID.
+				int32 ReturnCode = 0;
+				FPlatformProcess::GetProcReturnCode(CurrentWorkerInfo.WorkerProcess, &ReturnCode);
+				FPlatformProcess::CloseProc(CurrentWorkerInfo.WorkerProcess);
+				CurrentWorkerInfo.WorkerProcess = FProcHandle();
+
+				if (CurrentWorkerInfo.bLaunchedWorker)
+				{
+					const FString WorkingDirectory = Manager->AbsoluteShaderBaseWorkingDirectory + FString::FromInt(WorkerIndex) + TEXT("/");
+					const FString OutputFileNameAndPath = WorkingDirectory + TEXT("WorkerOutputOnly.out");
+
+					if (FPlatformFileManager::Get().GetPlatformFile().FileExists(*OutputFileNameAndPath))
+					{
+						// If the worker is no longer running but it successfully wrote out the output, no need to assert
+						bLaunchAgain = false;
+					}
+					else
+					{
+						UE_LOG(LogShaderCompilers, Error, TEXT("ShaderCompileWorker terminated unexpectedly, return code %d! Falling back to directly compiling which will be very slow.  Thread %u."), ReturnCode, WorkerIndex);
+						LogQueuedCompileJobs(CurrentWorkerInfo.QueuedJobs, -1);
+
+						bAbandonWorkers = true;
+						break;
+					}
+				}
+			}
+
+			if (bLaunchAgain)
+			{
+				const FString WorkingDirectory = Manager->ShaderBaseWorkingDirectory + FString::FromInt(WorkerIndex) + TEXT("/");
+				FString InputFileName(TEXT("WorkerInputOnly.in"));
+				FString OutputFileName(TEXT("WorkerOutputOnly.out"));
+				
+				if (bClearStaleOutputs)
+				{
+					// Delete any potential stale output files; these can persist if we had previously abandoned running with workers due to an unexpected SCW termination.
+					IFileManager::Get().Delete(*(WorkingDirectory / OutputFileName));
+				}
+
+				// Store the handle with this thread so that we will know not to launch it again
+				CurrentWorkerInfo.WorkerProcess = Manager->LaunchWorker(WorkingDirectory, Manager->ProcessId, WorkerIndex, InputFileName, OutputFileName);
+				CurrentWorkerInfo.bLaunchedWorker = true;
+
+				NumberLaunched++;
+			}
+		}
+	}
+
+	const double FinishTime = FPlatformTime::Seconds();
+	if (NumberLaunched > 0 && (FinishTime - CurrentTime) >= 10.0)
+	{
+		UE_LOG(LogShaderCompilers, Warning, TEXT("Performance Warning: It took %f seconds to launch %d ShaderCompileWorkers"), FinishTime - CurrentTime, NumberLaunched);
+	}
+
+	return bAbandonWorkers;
+}
+
+int32 FShaderCompileThreadRunnable::ReadAvailableResults()
+{
+	TRACE_CPUPROFILER_EVENT_SCOPE(ShaderCompiler.ReadAvailableResults);
+	int32 NumProcessed = 0;
+	FScopeLock WorkerScopeLock(&WorkerInfosLock);
+
+	// first, a quick check if anything is needed just to avoid hammering the task graph
+	bool bHasQueuedJobs = false;
+	for (int32 WorkerIndex = 0, NumWorkers = WorkerInfos.Num(); WorkerIndex < NumWorkers; ++WorkerIndex)
+	{
+		if (WorkerInfos[WorkerIndex]->QueuedJobs.Num() > 0)
+		{
+			bHasQueuedJobs = true;
+			break;
+		}
+	}
+
+	if (!bHasQueuedJobs)
+	{
+		return NumProcessed;
+	}
+
+	auto LoopBody = [this, &NumProcessed](int32 WorkerIndex)
+	{
+		// The calling thread holds the WorkerInfosLock and will not modify WorkerInfos, 
+		// so we can access it here without entering the lock
+		FShaderCompileWorkerInfo& CurrentWorkerInfo = *WorkerInfos[WorkerIndex];
+
+		// Check for available result files
+		if (CurrentWorkerInfo.QueuedJobs.Num() > 0)
+		{
+			// Distributed compiles always use the same directory
+			// 'Only' indicates to the worker that it should log and continue checking for the input file after the first one is processed
+			TStringBuilder<512> OutputFileNameAndPath;
+			OutputFileNameAndPath << Manager->AbsoluteShaderBaseWorkingDirectory << WorkerIndex << TEXT("/WorkerOutputOnly.out");
+
+			// In the common case the output file will not exist, so check for existence before opening
+			// This is only a win if FileExists is faster than CreateFileReader, which it is on Windows
+			if (FPlatformFileManager::Get().GetPlatformFile().FileExists(*OutputFileNameAndPath))
+			{
+				TRACE_CPUPROFILER_EVENT_SCOPE(FShaderCompileThreadRunnable::ProcessOutputFile);
+
+				if (TUniquePtr<FArchive> OutputFile = TUniquePtr<FArchive>(IFileManager::Get().CreateFileReader(*OutputFileNameAndPath, FILEREAD_Silent)))
+				{
+					check(!CurrentWorkerInfo.bComplete);
+					FShaderCompileUtilities::DoReadTaskResults(CurrentWorkerInfo.QueuedJobs, *OutputFile);
+
+					// Close the output file.
+					OutputFile.Reset();
+
+					// Delete the output file now that we have consumed it, to avoid reading stale data on the next compile loop.
+					bool bDeletedOutput = IFileManager::Get().Delete(*OutputFileNameAndPath, true, true);
+					int32 RetryCount = 0;
+					// Retry over the next two seconds if we couldn't delete it
+					while (!bDeletedOutput && RetryCount < 200)
+					{
+						TRACE_CPUPROFILER_EVENT_SCOPE(FShaderCompileThreadRunnable::DeleteOutputFile);
+
+						FPlatformProcess::Sleep(0.01f);
+						bDeletedOutput = IFileManager::Get().Delete(*OutputFileNameAndPath, true, true);
+						RetryCount++;
+					}
+					checkf(bDeletedOutput, TEXT("Failed to delete %s!"), *OutputFileNameAndPath);
+
+					CurrentWorkerInfo.bComplete = true;
+				}
+
+				FPlatformAtomics::InterlockedIncrement(&NumProcessed);
+			}
+		}
+	};
+
+	if (bParallelizeIO)
+	{
+		ParallelFor( TEXT("ShaderCompiler.ReadAvailableResults.PF"),WorkerInfos.Num(),1, LoopBody, EParallelForFlags::Unbalanced);
+	}
+	else 
+	{
+		double StartIOWork = FPlatformTime::Seconds();
+		for (int32 WorkerIndex = 0; WorkerIndex < WorkerInfos.Num(); WorkerIndex++)
+		{
+			LoopBody(WorkerIndex);
+		}
+
+		double IODuration = FPlatformTime::Seconds() - StartIOWork;
+		if (IODuration > 0.3)
+		{
+			UE_LOG(LogShaderCompilers, Display, TEXT("FShaderCompileThreadRunnable::WriteNewTasks() took too long (%.3f seconds, threshold is %.3f s), will parallelize next time."), IODuration, GShaderCompilerTooLongIOThresholdSeconds);
+			bParallelizeIO = true;
+		}
+	}
+
+	return NumProcessed;
+}
+
+void FShaderCompileThreadRunnable::CompileDirectlyThroughDll()
+{
+	// If we aren't compiling through workers, so we can just track the serial time here.
+	COOK_STAT(FScopedDurationTimer CompileTimer (ShaderCompilerCookStats::AsyncCompileTimeSec));
+
+	FScopeLock WorkerScopeLock(&WorkerInfosLock);
+	for (int32 WorkerIndex = 0; WorkerIndex < WorkerInfos.Num(); WorkerIndex++)
+	{
+		FShaderCompileWorkerInfo& CurrentWorkerInfo = *WorkerInfos[WorkerIndex];
+
+		if (CurrentWorkerInfo.QueuedJobs.Num() > 0)
+		{
+			DumpWorkerInputs(CurrentWorkerInfo.QueuedJobs);
+
+			for (int32 JobIndex = 0; JobIndex < CurrentWorkerInfo.QueuedJobs.Num(); JobIndex++)
+			{
+				FShaderCommonCompileJob& CurrentJob = *CurrentWorkerInfo.QueuedJobs[JobIndex];
+				FShaderCompileUtilities::ExecuteShaderCompileJob(CurrentJob);
+			}
+
+			CurrentWorkerInfo.bComplete = true;
+		}
+	}
+}
+
+void FShaderCompileThreadRunnable::PrintWorkerMemoryUsageWithLockTaken()
+{
+	FPlatformProcessMemoryStats TotalMemoryStats{};
+	int32 NumValidWorkers = 0;
+	constexpr int64 Gibibyte = 1024 * 1024 * 1024;
+	for (int32 Iter = 0, End = WorkerInfos.Num(); Iter < End; Iter++)
+	{
+		const TUniquePtr<FShaderCompileWorkerInfo>& WorkerInfo = WorkerInfos[Iter];
+		FProcHandle ProcHandle = WorkerInfo->WorkerProcess;
+		if (!ProcHandle.IsValid())
+		{
+			continue;
+		}
+		FPlatformProcessMemoryStats MemoryStats;
+		if (FPlatformProcess::TryGetMemoryUsage(ProcHandle, MemoryStats))
+		{
+			NumValidWorkers++;
+			UE_LOG(LogShaderCompilers, Display,
+				TEXT("ShaderCompileWorker [%d/%d] MemoryStats:")
+				TEXT("\n\t     UsedPhysical %llu (%.2f GiB)")
+				TEXT("\n\t PeakUsedPhysical %llu (%.2f GiB)")
+				TEXT("\n\t      UsedVirtual %llu (%.2f GiB)")
+				TEXT("\n\t  PeakUsedVirtual %llu (%.2f GiB)"),
+				Iter + 1,
+				End,
+				MemoryStats.UsedPhysical, double(MemoryStats.UsedPhysical) / Gibibyte,
+				MemoryStats.PeakUsedPhysical, double(MemoryStats.PeakUsedPhysical) / Gibibyte,
+				MemoryStats.UsedVirtual, double(MemoryStats.UsedVirtual) / Gibibyte,
+				MemoryStats.PeakUsedVirtual, double(MemoryStats.PeakUsedVirtual) / Gibibyte
+			);
+			TotalMemoryStats.UsedPhysical += MemoryStats.UsedPhysical;
+			TotalMemoryStats.PeakUsedPhysical += MemoryStats.PeakUsedPhysical;
+			TotalMemoryStats.UsedVirtual += MemoryStats.PeakUsedVirtual;
+			TotalMemoryStats.PeakUsedVirtual += MemoryStats.PeakUsedVirtual;
+		}
+		LogQueuedCompileJobs(WorkerInfo->QueuedJobs, -1);
+	}
+
+	if (NumValidWorkers > 0)
+	{
+		UE_LOG(LogShaderCompilers, Display,
+			TEXT("Sum of MemoryStats for %d ShaderCompileWorker(s):")
+			TEXT("\n\t     UsedPhysical %llu (%.2f GiB)")
+			TEXT("\n\t PeakUsedPhysical %llu (%.2f GiB)")
+			TEXT("\n\t      UsedVirtual %llu (%.2f GiB)")
+			TEXT("\n\t  PeakUsedVirtual %llu (%.2f GiB)"),
+			NumValidWorkers,
+			TotalMemoryStats.UsedPhysical, double(TotalMemoryStats.UsedPhysical) / Gibibyte,
+			TotalMemoryStats.PeakUsedPhysical, double(TotalMemoryStats.PeakUsedPhysical) / Gibibyte,
+			TotalMemoryStats.UsedVirtual, double(TotalMemoryStats.UsedVirtual) / Gibibyte,
+			TotalMemoryStats.PeakUsedVirtual, double(TotalMemoryStats.PeakUsedVirtual) / Gibibyte
+		);
+	}
+}
+
+bool FShaderCompileThreadRunnable::PrintWorkerMemoryUsage(bool bAllowToWaitForLock)
+{
+	if (bAllowToWaitForLock)
+	{
+		FScopeLock WorkerScopeLock(&WorkerInfosLock);
+		PrintWorkerMemoryUsageWithLockTaken();
+		return true;
+	}
+	else
+	{
+		FScopeTryLock WorkerScopeLock(&WorkerInfosLock);
+		if (WorkerScopeLock.IsLocked())
+		{
+			PrintWorkerMemoryUsageWithLockTaken();
+			return true;
+		}
+		return false;
+	}
+}
+
+FShaderCompileMemoryUsage FShaderCompileThreadRunnable::GetExternalWorkerMemoryUsage()
+{
+	FScopeLock WorkerScopeLock(&WorkerInfosLock);
+	FShaderCompileMemoryUsage MemoryUsage{};
+	for (const TUniquePtr<FShaderCompileWorkerInfo>& WorkerInfo : WorkerInfos)
+	{
+		FProcHandle ProcHandle = WorkerInfo->WorkerProcess;
+		if (!ProcHandle.IsValid())
+		{
+			continue;
+		}
+		FPlatformProcessMemoryStats MemoryStats;
+		if (FPlatformProcess::TryGetMemoryUsage(ProcHandle, MemoryStats))
+		{
+			// Virtual memory is committed memory on Windows.
+			MemoryUsage.VirtualMemory += MemoryStats.UsedVirtual;
+			MemoryUsage.PhysicalMemory += MemoryStats.UsedPhysical;
+		}
+	}
+	return MemoryUsage;
+}
+
+int32 FShaderCompileThreadRunnable::CompilingLoop()
+{
+	if (!Manager->bAllowCompilingThroughWorkers && CVarCompileParallelInProcess.GetValueOnAnyThread())
+	{
+		int32 NumJobs = Manager->GetNumPendingJobs();
+		if (NumJobs == 0)
+		{
+			return 0;
+		}
+
+		// if -noshaderworker is specified and the experimental in-process parallel compile is enabled, submit tasks for a batch 
+		// of pending jobs to run wide in-process (and a tailing task to mark those jobs as complete in the manager)
+		TUniquePtr<TArray<FShaderCommonCompileJobPtr>> Jobs = TUniquePtr<TArray<FShaderCommonCompileJobPtr>>(new TArray<FShaderCommonCompileJobPtr>());
+		TUniquePtr<TArray<float>> JobTimes = TUniquePtr<TArray<float>>(new TArray<float>());
+		{
+			FScopeLock Lock(&Manager->CompileQueueSection);
+			for (int32 PriorityIndex = MaxPriorityIndex; PriorityIndex >= MinPriorityIndex; --PriorityIndex)
+			{
+				// Throttle how many jobs we kick per tick so we get more frequent progress updates in the UI;
+				// this doesn't seem to have much effect on overall throughput.
+				const int32 MaxNumJobs = 64;
+				Manager->AllJobs.GetPendingJobs(EShaderCompilerWorkerType::LocalThread, (EShaderCompileJobPriority)PriorityIndex, 1, MaxNumJobs, *Jobs);
+			}
+		}
+
+		JobTimes->SetNum(Jobs->Num());
+
+		TArray<UE::Tasks::FTask> CompileTasks;
+		CompileTasks.Reserve(Jobs->Num());
+		for (TArray<FShaderCommonCompileJobPtr>::SizeType JobIndex = 0; JobIndex < Jobs->Num(); ++JobIndex)
+		{
+			FShaderCommonCompileJob* Job = (*Jobs)[JobIndex];
+			float* Time = &(*JobTimes)[JobIndex];
+			CompileTasks.Add(UE::Tasks::Launch(UE_SOURCE_LOCATION, [Job, Time]()
+				{
+					const float StartTime = FPlatformTime::Seconds();
+					FShaderCompileUtilities::ExecuteShaderCompileJob(*Job);
+					*Time = FPlatformTime::Seconds() - StartTime;
+				}));
+		}
+
+		if (!Jobs->IsEmpty())
+		{
+			UE::Tasks::Launch(UE_SOURCE_LOCATION, [Jobs = MoveTemp(Jobs), JobTimes = MoveTemp(JobTimes), this]()
+				{
+					FScopeLock Lock(&Manager->CompileQueueSection);
+					for (FShaderCommonCompileJobPtr Job : *Jobs)
+					{
+						Manager->ProcessFinishedJob(Job.GetReference());
+					}
+
+					float ElapsedTime = 0.0f;
+					for (const float& JobTime : *JobTimes)
+					{
+						ElapsedTime += JobTime;
+					}
+
+					Manager->WorkersBusyTime += ElapsedTime;
+					COOK_STAT(ShaderCompilerCookStats::AsyncCompileTimeSec += ElapsedTime);
+				}, CompileTasks);
+			// num active threads is up to the task system; as long as we return non-zero
+			// this will indicate to the caller that pending jobs remain
+			return 1;
+		}
+		return 0;
+	}
+	else // compile either through worker processes or single-threaded in-process (depending on Manager->bAllowCompilingThroughWorkers)
+	{
+		// push completed jobs to Manager->ShaderMapJobs before asking for new ones, so we can free the workers now and avoid them waiting a cycle
+		PushCompletedJobsToManager();
+
+		// Grab more shader compile jobs from the input queue
+		const int32 NumActiveThreads = PullTasksFromQueue();
+
+		if (NumActiveThreads == 0 && Manager->bAllowAsynchronousShaderCompiling)
+		{
+			// Yield while there's nothing to do
+			// Note: sleep-looping is bad threading practice, wait on an event instead!
+			// The shader worker thread does it because it needs to communicate with other processes through the file system
+			FPlatformProcess::Sleep(.010f);
+		}
+
+		if (Manager->bAllowCompilingThroughWorkers)
+		{
+			// Write out the files which are input to the shader compile workers
+			WriteNewTasks();
+
+			// Launch shader compile workers if they are not already running
+			// Workers can time out when idle so they may need to be relaunched
+			bool bAbandonWorkers = LaunchWorkersIfNeeded();
+
+			if (bAbandonWorkers)
+			{
+				// Fall back to local compiles if the SCW crashed.
+				// This is nasty but needed to work around issues where message passing through files to SCW is unreliable on random PCs
+				Manager->bAllowCompilingThroughWorkers = false;
+
+				// Try to recover from abandoned workers after a certain amount of single-threaded compilations
+				if (Manager->NumSingleThreadedRunsBeforeRetry == GSingleThreadedRunsIdle)
+				{
+					// First try to recover, only run single-threaded approach once
+					Manager->NumSingleThreadedRunsBeforeRetry = 1;
+				}
+				else if (Manager->NumSingleThreadedRunsBeforeRetry > GSingleThreadedRunsMaxCount)
+				{
+					// Stop retry approach after too many retries have failed
+					Manager->NumSingleThreadedRunsBeforeRetry = GSingleThreadedRunsDisabled;
+				}
+				else
+				{
+					// Next time increase runs by factor X
+					Manager->NumSingleThreadedRunsBeforeRetry *= GSingleThreadedRunsIncreaseFactor;
+				}
+			}
+			else
+			{
+				// Read files which are outputs from the shader compile workers
+				int32 NumProcessedResults = ReadAvailableResults();
+				if (NumProcessedResults == 0)
+				{
+					// Reduce filesystem query rate while actively waiting for results.
+					FPlatformProcess::Sleep(0.1f);
+				}
+			}
+		}
+		else
+		{
+			// Execute all pending worker tasks single-threaded
+			CompileDirectlyThroughDll();
+
+			// If single-threaded mode was enabled by an abandoned worker, try to recover after the given amount of runs
+			if (Manager->NumSingleThreadedRunsBeforeRetry > 0)
+			{
+				Manager->NumSingleThreadedRunsBeforeRetry--;
+				if (Manager->NumSingleThreadedRunsBeforeRetry == 0)
+				{
+					UE_LOG(LogShaderCompilers, Display, TEXT("Retry shader compiling through workers."));
+					Manager->bClearStaleWorkerOutputs = true;
+					Manager->bAllowCompilingThroughWorkers = true;
+				}
+			}
+		}
+
+		return NumActiveThreads;
+	}
+}
+
+
diff --git a/Engine/Source/Runtime/Engine/Private/ShaderCompiler/ShaderCompilerUtil.cpp b/Engine/Source/Runtime/Engine/Private/ShaderCompiler/ShaderCompilerUtil.cpp
new file mode 100644
--- /dev/null
+++ b/Engine/Source/Runtime/Engine/Private/ShaderCompiler/ShaderCompilerUtil.cpp
@@ -0,0 +1,1037 @@
+// Copyright Epic Games, Inc. All Rights Reserved.
+
+/*=============================================================================
+	ShaderCompilerUtil.cpp: Implements FShaderCompileUtilities.
+=============================================================================*/
+
+#include "ShaderCompilerPrivate.h"
+
+#include "HAL/PlatformFile.h"
+#include "HAL/PlatformFileManager.h"
+#include "Interfaces/IShaderFormat.h"
+#include "Interfaces/ITargetPlatform.h"
+#include "Interfaces/ITargetPlatformManagerModule.h"
+#include "Misc/Compression.h"
+#include "Misc/FileHelper.h"
+#include "Misc/MessageDialog.h"
+#include "Serialization/MemoryWriter.h"
+#include "ShaderDiagnostics.h"
+
+
+static TAutoConsoleVariable<bool> CVarDebugDumpWorkerInputs(
+	TEXT("r.ShaderCompiler.DebugDumpWorkerInputs"),
+	false,
+	TEXT("If true, worker input files will be saved for each individual compile job alongside other debug data (note that r.DumpShaderDebugInfo must also be enabled for this to function)"),
+	ECVF_ReadOnly);
+
+static TAutoConsoleVariable<bool> CVarShadersPropagateLocalWorkerOOMs(
+	TEXT("r.Shaders.PropagateLocalWorkerOOMs"),
+	false,
+	TEXT("When set, out-of-memory conditions in a local shader compile worker will be treated as regular out-of-memory conditions and propagated to the main process.\n")
+	TEXT("This is useful when running in environment with hard memory limits, where it does not matter which process in particular caused us to violate the memory limit."),
+	ECVF_Default);
+
+
+static void ModalErrorOrLog(const FString& Title, const FString& Text, int64 CurrentFilePos = 0, int64 ExpectedFileSize = 0, bool bIsErrorFatal = true)
+{
+	static FThreadSafeBool bModalReported;
+
+	FString BadFile;
+	if (CurrentFilePos > ExpectedFileSize)
+	{
+		// Corrupt file
+		BadFile = FString::Printf(TEXT(" (Truncated or corrupt output file! Current file pos %lld, file size %lld)"), CurrentFilePos, ExpectedFileSize);
+	}
+
+	if (FPlatformProperties::SupportsWindowedMode() && !FApp::IsUnattended())
+	{
+		UE_LOG(LogShaderCompilers, Error, TEXT("%s\n%s"), *Text, *BadFile);
+		if (!bModalReported.AtomicSet(true))
+		{
+			// Show dialog box with error message and request exit
+			FMessageDialog::Open(EAppMsgType::Ok, FText::FromString(Text), FText::FromString(Title));
+			FPlatformMisc::RequestExit(false, TEXT("ShaderCompiler.ModalErrorOrLog"));
+		}
+		else
+		{
+			// Another thread already opened a dialog box and requests exit
+			FPlatformProcess::SleepInfinite();
+		}
+	}
+	else if (bIsErrorFatal)
+	{
+		UE_LOG(LogShaderCompilers, Fatal, TEXT("%s\n%s\n%s"), *Title, *Text, *BadFile);
+	}
+	else
+	{
+		UE_LOG(LogShaderCompilers, Error, TEXT("%s\n%s\n%s"), *Title, *Text, *BadFile);
+	}
+}
+
+static void ReissueShaderCompileJobs(const TArray<FShaderCommonCompileJob*>& SourceJobs)
+{
+	if (SourceJobs.Num())
+	{
+		TRACE_CPUPROFILER_EVENT_SCOPE(ReissueShaderCompileJobs);
+
+		TArray<FShaderCommonCompileJobPtr> ReissueJobs;
+		ReissueJobs.Reserve(SourceJobs.Num());
+		const uint32 JobId = FShaderCommonCompileJob::GetNextJobId();
+		for (const FShaderCommonCompileJob* SourceJob : SourceJobs)
+		{
+			if (const FShaderCompileJob* SingleSourceJob = SourceJob->GetSingleShaderJob())
+			{
+				if (FShaderCompileJob* ReissueJob = GShaderCompilingManager->PrepareShaderCompileJob(JobId, SingleSourceJob->Key, SingleSourceJob->Priority))
+				{
+					ReissueJob->Input = SingleSourceJob->Input;
+					ReissueJobs.Add(FShaderCommonCompileJobPtr(ReissueJob));
+				}
+			}
+			else if (const FShaderPipelineCompileJob* PipelineSourceJob = SourceJob->GetShaderPipelineJob())
+			{
+				if (FShaderPipelineCompileJob* ReissueJob = GShaderCompilingManager->PreparePipelineCompileJob(JobId, PipelineSourceJob->Key, PipelineSourceJob->Priority))
+				{
+					ReissueJob->StageJobs = PipelineSourceJob->StageJobs;
+					ReissueJobs.Add(FShaderCommonCompileJobPtr(ReissueJob));
+				}
+			}
+			else
+			{
+				checkf(0, TEXT("Reissued shader compile job is neither a single nor a pipeline job"));
+			}
+		}
+
+		GShaderCompilingManager->SubmitJobs(ReissueJobs, FString(""), FString(""));
+	}
+}
+
+// Make functions so the crash reporter can disambiguate the actual error because of the different callstacks
+namespace ShaderCompileWorkerError
+{
+	void HandleGeneralCrash(const TCHAR* ExceptionInfo, const TCHAR* Callstack)
+	{
+		GLog->Panic();
+		UE_LOG(LogShaderCompilers, Error, TEXT("ShaderCompileWorker crashed!\n%s\n%s"), ExceptionInfo, Callstack);
+	}
+
+	void HandleBadShaderFormatVersion(const TCHAR* Data)
+	{
+		ModalErrorOrLog(TEXT("ShaderCompileWorker failed"), Data);
+	}
+
+	void HandleBadInputVersion(const TCHAR* Data)
+	{
+		ModalErrorOrLog(TEXT("ShaderCompileWorker failed"), Data);
+	}
+
+	void HandleBadSingleJobHeader(const TCHAR* Data)
+	{
+		ModalErrorOrLog(TEXT("ShaderCompileWorker failed"), Data);
+	}
+
+	void HandleBadPipelineJobHeader(const TCHAR* Data)
+	{
+		ModalErrorOrLog(TEXT("ShaderCompileWorker failed"), Data);
+	}
+
+	void HandleCantDeleteInputFile(const TCHAR* Data)
+	{
+		ModalErrorOrLog(TEXT("ShaderCompileWorker failed"), Data);
+	}
+
+	void HandleCantSaveOutputFile(const TCHAR* Data)
+	{
+		ModalErrorOrLog(TEXT("ShaderCompileWorker failed"), Data);
+	}
+
+	void HandleNoTargetShaderFormatsFound(const TCHAR* Data)
+	{
+		ModalErrorOrLog(TEXT("ShaderCompileWorker failed"), Data);
+	}
+
+	void HandleCantCompileForSpecificFormat(const TCHAR* Data)
+	{
+		ModalErrorOrLog(TEXT("ShaderCompileWorker failed"), Data);
+	}
+
+	void HandleOutputFileEmpty(const TCHAR* Filename)
+	{
+		ModalErrorOrLog(TEXT("ShaderCompileWorker failed"), FString::Printf(TEXT("Output file %s size is 0. Are you out of disk space?"), Filename));
+	}
+
+	void HandleOutputFileCorrupted(const TCHAR* Filename, int64 ExpectedSize, int64 ActualSize)
+	{
+		ModalErrorOrLog(TEXT("ShaderCompileWorker failed"), FString::Printf(TEXT("Output file corrupted (expected %I64d bytes, but only got %I64d): %s"), ExpectedSize, ActualSize, Filename));
+	}
+
+	void HandleCrashInsidePlatformCompiler(const TCHAR* Data)
+	{
+		// If the crash originates from a platform compiler, the error code must have been reported and we don't have to assume a corrupted output file.
+		// In that case, don't crash the cooker with a fatal error, just report the error so the cooker can dump debug info.
+		constexpr bool bIsErrorFatal = false;
+		ModalErrorOrLog(TEXT("ShaderCompileWorker failed"), FString::Printf(TEXT("Crash inside the platform compiler:\n%s"), Data), 0, 0, bIsErrorFatal);
+	}
+
+	void HandleBadInputFile(const TCHAR* Data)
+	{
+		ModalErrorOrLog(TEXT("ShaderCompileWorker failed"), FString::Printf(TEXT("Bad-input-file exception:\n%s"), Data));
+	}
+
+	bool HandleOutOfMemory(const TCHAR* ExceptionInfo, const TCHAR* Hostname, const FPlatformMemoryStats& MemoryStats, const TArray<FShaderCommonCompileJobPtr>& QueuedJobs)
+	{
+		constexpr int64 Gibibyte = 1024 * 1024 * 1024;
+		const FString ErrorReport = FString::Printf(
+			TEXT("ShaderCompileWorker failed with out-of-memory (OOM) exception on machine \"%s\" (%s); MemoryStats:")
+			TEXT("\n\tAvailablePhysical %llu (%.2f GiB)")
+			TEXT("\n\t AvailableVirtual %llu (%.2f GiB)")
+			TEXT("\n\t     UsedPhysical %llu (%.2f GiB)")
+			TEXT("\n\t PeakUsedPhysical %llu (%.2f GiB)")
+			TEXT("\n\t      UsedVirtual %llu (%.2f GiB)")
+			TEXT("\n\t  PeakUsedVirtual %llu (%.2f GiB)"),
+			Hostname,
+			(ExceptionInfo[0] == TEXT('\0') ? TEXT("No exception information") : ExceptionInfo),
+			MemoryStats.AvailablePhysical, double(MemoryStats.AvailablePhysical) / Gibibyte,
+			MemoryStats.AvailableVirtual, double(MemoryStats.AvailableVirtual) / Gibibyte,
+			MemoryStats.UsedPhysical, double(MemoryStats.UsedPhysical) / Gibibyte,
+			MemoryStats.PeakUsedPhysical, double(MemoryStats.PeakUsedPhysical) / Gibibyte,
+			MemoryStats.UsedVirtual, double(MemoryStats.UsedVirtual) / Gibibyte,
+			MemoryStats.PeakUsedVirtual, double(MemoryStats.PeakUsedVirtual) / Gibibyte
+		);
+
+		if (GShaderCompilingManager->IsRemoteCompilingEnabled())
+		{
+			// Remote shader compiler supports re-compiling jobs on local machine
+			UE_LOG(LogShaderCompilers, Warning, TEXT("%s\nRecompile %d shader compile %s locally"), *ErrorReport, QueuedJobs.Num(), (QueuedJobs.Num() == 1 ? TEXT("job") : TEXT("jobs")));
+			return true;
+		}
+		else
+		{
+			if (CVarShadersPropagateLocalWorkerOOMs.GetValueOnAnyThread())
+			{
+				FPlatformMemory::OnOutOfMemory(0, 64);
+			}
+			ModalErrorOrLog(TEXT("ShaderCompileWorker failed"), ErrorReport);
+			return false;
+		}
+	}
+}
+
+static TMap<FString, uint32> GetFormatVersionMap()
+{
+	TMap<FString, uint32> FormatVersionMap;
+
+	const TArray<const class IShaderFormat*>& ShaderFormats = GetTargetPlatformManagerRef().GetShaderFormats();
+	check(ShaderFormats.Num());
+	for (int32 Index = 0; Index < ShaderFormats.Num(); Index++)
+	{
+		TArray<FName> OutFormats;
+		ShaderFormats[Index]->GetSupportedFormats(OutFormats);
+		check(OutFormats.Num());
+		for (int32 InnerIndex = 0; InnerIndex < OutFormats.Num(); InnerIndex++)
+		{
+			uint32 Version = ShaderFormats[Index]->GetVersion(OutFormats[InnerIndex]);
+			FormatVersionMap.Add(OutFormats[InnerIndex].ToString(), Version);
+		}
+	}
+
+	return FormatVersionMap;
+}
+
+#if 0 //UNUSED
+static int32 GetNumTotalJobs(const TArray<FShaderCommonCompileJobPtr>& Jobs)
+{
+	int32 NumJobs = 0;
+	for (int32 Index = 0; Index < Jobs.Num(); ++Index)
+	{
+		auto* PipelineJob = Jobs[Index]->GetShaderPipelineJob();
+		NumJobs += PipelineJob ? PipelineJob->StageJobs.Num() : 1;
+	}
+
+	return NumJobs;
+}
+#endif
+
+static void SplitJobsByType(const TArray<FShaderCommonCompileJobPtr>& QueuedJobs, TArray<FShaderCompileJob*>& OutQueuedSingleJobs, TArray<FShaderPipelineCompileJob*>& OutQueuedPipelineJobs)
+{
+	for (int32 Index = 0; Index < QueuedJobs.Num(); ++Index)
+	{
+		FShaderCommonCompileJobPtr CommonJob = QueuedJobs[Index];
+		if (FShaderCompileJob* SingleJob = CommonJob->GetSingleShaderJob())
+		{
+			OutQueuedSingleJobs.Add(SingleJob);
+
+		}
+		else if (FShaderPipelineCompileJob* PipelineJob = CommonJob->GetShaderPipelineJob())
+		{
+			OutQueuedPipelineJobs.Add(PipelineJob);
+		}
+		else
+		{
+			checkf(0, TEXT("FShaderCommonCompileJob::Type=%d is not a valid type for a shader compile job"), (int32)CommonJob->Type);
+		}
+	}
+}
+
+static bool DoWriteTasksInner(const TArray<FShaderCommonCompileJobPtr>& QueuedJobs, FArchive& InTransferFile, IDistributedBuildController* BuildDistributionController, bool bUseRelativePaths, bool bCompressTaskFile)
+{
+	int32 InputVersion = ShaderCompileWorkerInputVersion;
+	InTransferFile << InputVersion;
+
+	TArray<uint8> UncompressedArray;
+	FMemoryWriter TransferMemory(UncompressedArray);
+	FArchive& TransferFile = bCompressTaskFile ? TransferMemory : InTransferFile;
+	if (!bCompressTaskFile)
+	{
+		// still write NAME_None as string
+		FString FormatNone = FName(NAME_None).ToString();
+		TransferFile << FormatNone;
+	}
+
+	static TMap<FString, uint32> FormatVersionMap = GetFormatVersionMap();
+
+	TransferFile << FormatVersionMap;
+
+	TArray<FShaderCompileJob*> QueuedSingleJobs;
+	TArray<FShaderPipelineCompileJob*> QueuedPipelineJobs;
+	SplitJobsByType(QueuedJobs, QueuedSingleJobs, QueuedPipelineJobs);
+
+	TArray<TRefCountPtr<FSharedShaderCompilerEnvironment>> SharedEnvironments;
+	TArray<const FShaderParametersMetadata*> RequestShaderParameterStructures;
+
+	// gather shared environments and parameter structures, these tend to be shared between jobs
+	{
+
+		for (int32 JobIndex = 0; JobIndex < QueuedSingleJobs.Num(); JobIndex++)
+		{
+			QueuedSingleJobs[JobIndex]->Input.GatherSharedInputs(SharedEnvironments, RequestShaderParameterStructures);
+		}
+
+		for (int32 JobIndex = 0; JobIndex < QueuedPipelineJobs.Num(); JobIndex++)
+		{
+			auto* PipelineJob = QueuedPipelineJobs[JobIndex];
+			int32 NumStageJobs = PipelineJob->StageJobs.Num();
+
+			for (int32 Index = 0; Index < NumStageJobs; Index++)
+			{
+				PipelineJob->StageJobs[Index]->Input.GatherSharedInputs(SharedEnvironments, RequestShaderParameterStructures);
+			}
+		}
+
+		int32 NumSharedEnvironments = SharedEnvironments.Num();
+		TransferFile << NumSharedEnvironments;
+
+		for (int32 EnvironmentIndex = 0; EnvironmentIndex < SharedEnvironments.Num(); EnvironmentIndex++)
+		{
+			TransferFile << *SharedEnvironments[EnvironmentIndex];
+		}
+	}
+
+	// Write shader parameter structures
+	TArray<const FShaderParametersMetadata*> AllShaderParameterStructures;
+	{
+		// List all dependencies.
+		for (int32 StructId = 0; StructId < RequestShaderParameterStructures.Num(); StructId++)
+		{
+			RequestShaderParameterStructures[StructId]->IterateStructureMetadataDependencies(
+				[&](const FShaderParametersMetadata* Struct)
+			{
+				AllShaderParameterStructures.AddUnique(Struct);
+			});
+		}
+
+		// Write all shader parameter structure.
+		int32 NumParameterStructures = AllShaderParameterStructures.Num();
+		TransferFile << NumParameterStructures;
+		for (const FShaderParametersMetadata* Struct : AllShaderParameterStructures)
+		{
+			FString LayoutName = Struct->GetLayout().GetDebugName();
+			FString StructTypeName = Struct->GetStructTypeName();
+			FString ShaderVariableName = Struct->GetShaderVariableName();
+			uint8 UseCase = uint8(Struct->GetUseCase());
+			FString StructFileName = FString(ANSI_TO_TCHAR(Struct->GetFileName()));
+			int32 StructFileLine = Struct->GetFileLine();
+			uint32 Size = Struct->GetSize();
+			int32 MemberCount = Struct->GetMembers().Num();
+
+			static_assert(sizeof(UseCase) == sizeof(FShaderParametersMetadata::EUseCase), "Cast failure.");
+
+			TransferFile << LayoutName;
+			TransferFile << StructTypeName;
+			TransferFile << ShaderVariableName;
+			TransferFile << UseCase;
+			TransferFile << StructFileName;
+			TransferFile << StructFileLine;
+			TransferFile << Size;
+			TransferFile << MemberCount;
+
+			for (const FShaderParametersMetadata::FMember& Member : Struct->GetMembers())
+			{
+				FString Name = Member.GetName();
+				FString ShaderType = Member.GetShaderType();
+				int32 FileLine = Member.GetFileLine();
+				uint32 Offset = Member.GetOffset();
+				uint8 BaseType = uint8(Member.GetBaseType());
+				uint8 PrecisionModifier = uint8(Member.GetPrecision());
+				uint32 NumRows = Member.GetNumRows();
+				uint32 NumColumns = Member.GetNumColumns();
+				uint32 NumElements = Member.GetNumElements();
+				int32 StructMetadataIndex = INDEX_NONE;
+				if (Member.GetStructMetadata())
+				{
+					StructMetadataIndex = AllShaderParameterStructures.Find(Member.GetStructMetadata());
+					check(StructMetadataIndex != INDEX_NONE);
+				}
+
+				static_assert(sizeof(BaseType) == sizeof(EUniformBufferBaseType), "Cast failure.");
+				static_assert(sizeof(PrecisionModifier) == sizeof(EShaderPrecisionModifier::Type), "Cast failure.");
+
+				TransferFile << Name;
+				TransferFile << ShaderType;
+				TransferFile << FileLine;
+				TransferFile << Offset;
+				TransferFile << BaseType;
+				TransferFile << PrecisionModifier;
+				TransferFile << NumRows;
+				TransferFile << NumColumns;
+				TransferFile << NumElements;
+				TransferFile << StructMetadataIndex;
+			}
+		}
+	}
+
+	// Write individual shader jobs
+	{
+		int32 SingleJobHeader = ShaderCompileWorkerSingleJobHeader;
+		TransferFile << SingleJobHeader;
+
+		int32 NumBatches = QueuedSingleJobs.Num();
+		TransferFile << NumBatches;
+
+		// Serialize all the batched jobs
+		for (int32 JobIndex = 0; JobIndex < QueuedSingleJobs.Num(); JobIndex++)
+		{
+			QueuedSingleJobs[JobIndex]->SerializeWorkerInput(TransferFile);
+			QueuedSingleJobs[JobIndex]->Input.SerializeSharedInputs(TransferFile, SharedEnvironments, AllShaderParameterStructures);
+		}
+	}
+
+	// Write shader pipeline jobs
+	{
+		int32 PipelineJobHeader = ShaderCompileWorkerPipelineJobHeader;
+		TransferFile << PipelineJobHeader;
+
+		int32 NumBatches = QueuedPipelineJobs.Num();
+		TransferFile << NumBatches;
+		for (int32 JobIndex = 0; JobIndex < QueuedPipelineJobs.Num(); JobIndex++)
+		{
+			auto* PipelineJob = QueuedPipelineJobs[JobIndex];
+			FString PipelineName = PipelineJob->Key.ShaderPipeline->GetName();
+			TransferFile << PipelineName;
+			int32 NumStageJobs = PipelineJob->StageJobs.Num();
+			TransferFile << NumStageJobs;
+			for (int32 Index = 0; Index < NumStageJobs; Index++)
+			{
+				PipelineJob->StageJobs[Index]->SerializeWorkerInput(TransferFile);
+				PipelineJob->StageJobs[Index]->Input.SerializeSharedInputs(TransferFile, SharedEnvironments, AllShaderParameterStructures);
+			}
+		}
+	}
+
+	if (bCompressTaskFile)
+	{
+		TransferFile.Close();
+
+		FName CompressionFormatToUse = NAME_LZ4;
+
+		FString FormatName = CompressionFormatToUse.ToString();
+		InTransferFile << FormatName;
+
+		// serialize uncompressed data size
+		int32 UncompressedDataSize = UncompressedArray.Num();
+		checkf(UncompressedDataSize != 0, TEXT("Did not write any data to the task file for the compression."));
+		InTransferFile << UncompressedDataSize;
+
+		// not using SerializeCompressed because it splits into smaller chunks
+		int32 CompressedSizeBound = FCompression::CompressMemoryBound(CompressionFormatToUse, static_cast<int32>(UncompressedDataSize));
+		TArray<uint8> CompressedBuffer;
+		CompressedBuffer.SetNumUninitialized(CompressedSizeBound);
+
+		int32 ActualCompressedSize = CompressedSizeBound;
+		bool bSucceeded = FCompression::CompressMemory(CompressionFormatToUse, CompressedBuffer.GetData(), ActualCompressedSize, UncompressedArray.GetData(), UncompressedDataSize, COMPRESS_BiasSpeed);
+		checkf(ActualCompressedSize <= CompressedSizeBound, TEXT("Compressed size was larger than the bound - we stomped the memory."));
+		CompressedBuffer.SetNum(ActualCompressedSize, EAllowShrinking::No);
+
+		InTransferFile << CompressedBuffer;
+		UE_LOG(LogShaderCompilers, Verbose, TEXT("Compressed the task file from %d bytes to %d bytes (%.2f%% savings)"), UncompressedDataSize, ActualCompressedSize,
+			100.0 * (UncompressedDataSize - ActualCompressedSize) / static_cast<double>(UncompressedDataSize));
+	}
+
+	return InTransferFile.Close();
+}
+
+const TCHAR* DebugWorkerInputFileName = TEXT("DebugSCW.in");
+const TCHAR* DebugWorkerOutputFileName = TEXT("DebugSCW.out");
+
+static FString CreateShaderCompilerWorkerDebugCommandLine(FString DebugWorkerInputFilePath)
+{
+	// 0 is parent PID, pass zero TTL and KeepInput to make SCW process the single job then exit without deleting the input file
+	return FString::Printf(TEXT("\"%s\" 0 \"DebugSCW\" %s %s -TimeToLive=0.0f -KeepInput"),
+		*DebugWorkerInputFilePath, // working directory for SCW
+		DebugWorkerInputFileName,
+		DebugWorkerOutputFileName);
+}
+
+void DumpWorkerInputs(TConstArrayView<FShaderCommonCompileJobPtr> QueuedJobs)
+{
+	if (CVarDebugDumpWorkerInputs.GetValueOnAnyThread())
+	{
+		for (const FShaderCommonCompileJobPtr& CommonJob : QueuedJobs)
+		{
+			FShaderPipelineCompileJob* PipelineJob = CommonJob->GetShaderPipelineJob();
+			FString DebugWorkerInputFilePath;
+			if (PipelineJob)
+			{
+				// for pipeline jobs, write out the worker input for the whole pipeline, but only for the first stage
+				// would be better to put in a parent folder probably...
+				DebugWorkerInputFilePath = PipelineJob->StageJobs[0]->Input.DumpDebugInfoPath;
+			}
+			else
+			{
+				DebugWorkerInputFilePath = CommonJob->GetSingleShaderJob()->Input.DumpDebugInfoPath;
+			}
+			if (!DebugWorkerInputFilePath.IsEmpty())
+			{
+				TArray<FShaderCommonCompileJobPtr> SingleJobArray;
+				SingleJobArray.Add(CommonJob);
+
+				FArchive* DebugWorkerInputFileWriter = IFileManager::Get().CreateFileWriter(*(DebugWorkerInputFilePath / DebugWorkerInputFileName), FILEWRITE_NoFail);
+				DoWriteTasksInner(
+					SingleJobArray,
+					*DebugWorkerInputFileWriter,
+					nullptr,	// Don't pass a IDistributedBuildController, this is only used for conversion to relative paths which we do not want for debug files
+					false,		// As above, use absolute paths not relative
+					true);		// Always compress the debug files; they are rather large so this saves some disk space
+				DebugWorkerInputFileWriter->Close();
+				delete DebugWorkerInputFileWriter;
+
+				FFileHelper::SaveStringToFile(
+					CreateShaderCompilerWorkerDebugCommandLine(DebugWorkerInputFilePath),
+					*(DebugWorkerInputFilePath / TEXT("DebugCompileArgs.txt")));
+			}
+		}
+	}
+}
+
+// Serialize Queued Job information
+bool FShaderCompileUtilities::DoWriteTasks(const TArray<FShaderCommonCompileJobPtr>& QueuedJobs, FArchive& InTransferFile, IDistributedBuildController* BuildDistributionController, bool bUseRelativePaths, bool bCompressTaskFile)
+{
+	DumpWorkerInputs(QueuedJobs);
+
+	return DoWriteTasksInner(QueuedJobs, InTransferFile, BuildDistributionController, bUseRelativePaths, bCompressTaskFile);
+}
+
+static bool ReadSingleJob(FShaderCompileJob* CurrentJob, FArchive& WorkerOutputFileReader)
+{
+	TRACE_CPUPROFILER_EVENT_SCOPE(ReadSingleJob);
+
+	check(!CurrentJob->bFinalized);
+	CurrentJob->bFinalized = true;
+
+	// Deserialize the shader compilation output.
+	CurrentJob->SerializeWorkerOutput(WorkerOutputFileReader);
+
+	// The job should already have a non-zero output hash
+	checkf(CurrentJob->Output.OutputHash != FSHAHash() || !CurrentJob->bSucceeded, TEXT("OutputHash for a successful job was not set in the shader compile worker!"));
+
+
+	// Support dumping debug info for only failed compilations or those with warnings
+	if (GShaderCompilingManager->ShouldRecompileToDumpShaderDebugInfo(*CurrentJob))
+	{
+		// Build debug info path and create the directory if it doesn't already exist
+		CurrentJob->Input.DumpDebugInfoPath = GShaderCompilingManager->CreateShaderDebugInfoPath(CurrentJob->Input);
+
+		// We failed to compile a shader, we will retry and dump the shader source. so increment the number of shader sources we've dumped so far.
+		GShaderCompilingManager->IncrementNumDumpedShaderSources();
+
+		return true;
+	}
+
+	return false;
+}
+
+static const TCHAR* GetCompileJobSuccessText(FShaderCompileJob* SingleJob)
+{
+	if (SingleJob)
+	{
+		return SingleJob->Output.bSucceeded ? TEXT("Succeeded") : TEXT("Failed");
+	}
+	return TEXT("");
+}
+
+void LogQueuedCompileJobs(const TArray<FShaderCommonCompileJobPtr>& QueuedJobs, int32 NumProcessedJobs)
+{
+	if (NumProcessedJobs == -1)
+	{
+		UE_LOG(LogShaderCompilers, Error, TEXT("SCW %d Queued Jobs, Unknown number of processed jobs!"), QueuedJobs.Num());
+	}
+	else
+	{
+		UE_LOG(LogShaderCompilers, Error, TEXT("SCW %d Queued Jobs, Finished %d single jobs"), QueuedJobs.Num(), NumProcessedJobs);
+	}
+
+	for (int32 Index = 0; Index < QueuedJobs.Num(); ++Index)
+	{
+		if (FShaderCompileJob* SingleJob = QueuedJobs[Index]->GetSingleShaderJob())
+		{
+			UE_LOG(LogShaderCompilers, Error, TEXT("Job %d [Single] %s: %s"), Index, GetCompileJobSuccessText(SingleJob), *GetSingleJobCompilationDump(SingleJob));
+		}
+		else
+		{
+			FShaderPipelineCompileJob* PipelineJob = QueuedJobs[Index]->GetShaderPipelineJob();
+			UE_LOG(LogShaderCompilers, Error, TEXT("Job %d: Pipeline %s "), Index, PipelineJob->Key.ShaderPipeline->GetName());
+			for (int32 JobIndex = 0; JobIndex < PipelineJob->StageJobs.Num(); ++JobIndex)
+			{
+				FShaderCompileJob* StageJob = PipelineJob->StageJobs[JobIndex]->GetSingleShaderJob();
+				UE_LOG(LogShaderCompilers, Error, TEXT("PipelineJob %d %s: %s"), JobIndex, GetCompileJobSuccessText(StageJob), *GetSingleJobCompilationDump(StageJob));
+			}
+		}
+	}
+
+	// Force a log flush so we can track the crash before the cooker potentially crashes before the output shows up
+	GLog->Flush();
+}
+
+// Disable optimization for this crash handler to get full access to the entire stack frame when debugging a crash dump
+UE_DISABLE_OPTIMIZATION_SHIP
+static bool HandleWorkerCrash(const TArray<FShaderCommonCompileJobPtr>& QueuedJobs, FArchive& OutputFile, int32 OutputVersion, int64 FileSize, FSCWErrorCode::ECode ErrorCode, int32 NumProcessedJobs, int32 CallstackLength, int32 ExceptionInfoLength, int32 HostnameLength)
+{
+	TArray<TCHAR> Callstack;
+	Callstack.AddUninitialized(CallstackLength + 1);
+	OutputFile.Serialize(Callstack.GetData(), CallstackLength * sizeof(TCHAR));
+	Callstack[CallstackLength] = TEXT('\0');
+
+	TArray<TCHAR> ExceptionInfo;
+	ExceptionInfo.AddUninitialized(ExceptionInfoLength + 1);
+	OutputFile.Serialize(ExceptionInfo.GetData(), ExceptionInfoLength * sizeof(TCHAR));
+	ExceptionInfo[ExceptionInfoLength] = TEXT('\0');
+
+	TArray<TCHAR> Hostname;
+	Hostname.AddUninitialized(HostnameLength + 1);
+	OutputFile.Serialize(Hostname.GetData(), HostnameLength * sizeof(TCHAR));
+	Hostname[HostnameLength] = TEXT('\0');
+
+	// Read available and used physical memory from worker machine on OOM error
+	FPlatformMemoryStats MemoryStats;
+	if (ErrorCode == FSCWErrorCode::OutOfMemory)
+	{
+		OutputFile
+			<< MemoryStats.AvailablePhysical
+			<< MemoryStats.AvailableVirtual
+			<< MemoryStats.UsedPhysical
+			<< MemoryStats.PeakUsedPhysical
+			<< MemoryStats.UsedVirtual
+			<< MemoryStats.PeakUsedVirtual
+			;
+	}
+
+	// Store primary job information onto stack to make it part of a crash dump
+	static const int32 MaxNumCharsForSourcePaths = 8192;
+	int32 JobInputSourcePathsLength = 0;
+	ANSICHAR JobInputSourcePaths[MaxNumCharsForSourcePaths];
+	JobInputSourcePaths[0] = 0;
+
+	auto WriteInputSourcePathOntoStack = [&JobInputSourcePathsLength, &JobInputSourcePaths](const ANSICHAR* InputSourcePath)
+	{
+		if (InputSourcePath != nullptr && JobInputSourcePathsLength + 3 < MaxNumCharsForSourcePaths)
+		{
+			// Copy input source path into stack buffer
+			int32 InputSourcePathLength = FMath::Min(FCStringAnsi::Strlen(InputSourcePath), (MaxNumCharsForSourcePaths - JobInputSourcePathsLength - 2));
+			FMemory::Memcpy(JobInputSourcePaths + JobInputSourcePathsLength, InputSourcePath, InputSourcePathLength);
+
+			// Write newline character and put NUL character at the end
+			JobInputSourcePathsLength += InputSourcePathLength;
+			JobInputSourcePaths[JobInputSourcePathsLength] = TEXT('\n');
+			++JobInputSourcePathsLength;
+			JobInputSourcePaths[JobInputSourcePathsLength] = 0;
+		}
+	};
+
+	auto StoreInputDebugInfo = [&WriteInputSourcePathOntoStack, &JobInputSourcePathsLength, &JobInputSourcePaths](const FShaderCompilerInput& Input)
+	{
+		FString DebugInfo = FString::Printf(TEXT("%s:%s"), *Input.VirtualSourceFilePath, *Input.EntryPointName);
+		WriteInputSourcePathOntoStack(TCHAR_TO_UTF8(*DebugInfo));
+	};
+
+	for (auto CommonJob : QueuedJobs)
+	{
+		if (FShaderCompileJob* SingleJob = CommonJob->GetSingleShaderJob())
+		{
+			StoreInputDebugInfo(SingleJob->Input);
+		}
+		else if (FShaderPipelineCompileJob* PipelineJob = CommonJob->GetShaderPipelineJob())
+		{
+			for (int32 Job = 0; Job < PipelineJob->StageJobs.Num(); ++Job)
+			{
+				if (FShaderCompileJob* SingleStageJob = PipelineJob->StageJobs[Job])
+				{
+					StoreInputDebugInfo(SingleStageJob->Input);
+				}
+			}
+		}
+	}
+
+	// One entry per error code as we want to have different callstacks for crash reporter...
+	switch (ErrorCode)
+	{
+	default:
+	case FSCWErrorCode::GeneralCrash:
+		LogQueuedCompileJobs(QueuedJobs, NumProcessedJobs);
+		ShaderCompileWorkerError::HandleGeneralCrash(ExceptionInfo.GetData(), Callstack.GetData());
+		break;
+	case FSCWErrorCode::BadShaderFormatVersion:
+		ShaderCompileWorkerError::HandleBadShaderFormatVersion(ExceptionInfo.GetData());
+		break;
+	case FSCWErrorCode::BadInputVersion:
+		ShaderCompileWorkerError::HandleBadInputVersion(ExceptionInfo.GetData());
+		break;
+	case FSCWErrorCode::BadSingleJobHeader:
+		ShaderCompileWorkerError::HandleBadSingleJobHeader(ExceptionInfo.GetData());
+		break;
+	case FSCWErrorCode::BadPipelineJobHeader:
+		ShaderCompileWorkerError::HandleBadPipelineJobHeader(ExceptionInfo.GetData());
+		break;
+	case FSCWErrorCode::CantDeleteInputFile:
+		ShaderCompileWorkerError::HandleCantDeleteInputFile(ExceptionInfo.GetData());
+		break;
+	case FSCWErrorCode::CantSaveOutputFile:
+		ShaderCompileWorkerError::HandleCantSaveOutputFile(ExceptionInfo.GetData());
+		break;
+	case FSCWErrorCode::NoTargetShaderFormatsFound:
+		ShaderCompileWorkerError::HandleNoTargetShaderFormatsFound(ExceptionInfo.GetData());
+		break;
+	case FSCWErrorCode::CantCompileForSpecificFormat:
+		ShaderCompileWorkerError::HandleCantCompileForSpecificFormat(ExceptionInfo.GetData());
+		break;
+	case FSCWErrorCode::CrashInsidePlatformCompiler:
+		LogQueuedCompileJobs(QueuedJobs, NumProcessedJobs);
+		ShaderCompileWorkerError::HandleCrashInsidePlatformCompiler(ExceptionInfo.GetData());
+		break;
+	case FSCWErrorCode::BadInputFile:
+		ShaderCompileWorkerError::HandleBadInputFile(ExceptionInfo.GetData());
+		break;
+	case FSCWErrorCode::OutOfMemory:
+		return ShaderCompileWorkerError::HandleOutOfMemory(ExceptionInfo.GetData(), Hostname.GetData(), MemoryStats, QueuedJobs);
+	case FSCWErrorCode::Success:
+		// Can't get here...
+		return true;
+	}
+	return false;
+}
+UE_ENABLE_OPTIMIZATION_SHIP
+
+// Helper struct to provide consistent error report with detailed information about corrupted ShaderCompileWorker output file.
+struct FSCWOutputFileContext
+{
+	FArchive& OutputFile;
+	int64 FileSize = 0;
+
+	FSCWOutputFileContext(FArchive& OutputFile) :
+		OutputFile(OutputFile)
+	{
+	}
+
+	template <typename FmtType, typename... Types>
+	void ModalErrorOrLog(const FmtType& Format, Types&&... Args)
+	{
+		FString Text = FString::Printf(Format, Args...);
+		Text = FString::Printf(TEXT("File path: \"%s\"\n%s\nForgot to build ShaderCompileWorker or delete invalidated DerivedDataCache?"), *OutputFile.GetArchiveName(), *Text);
+		const TCHAR* Title = TEXT("Corrupted ShaderCompileWorker output file");
+		if (FileSize > 0)
+		{
+			::ModalErrorOrLog(Title, Text, OutputFile.Tell(), FileSize);
+		}
+		else
+		{
+			::ModalErrorOrLog(Title, Text, 0, 0);
+		}
+	}
+};
+
+// Process results from Worker Process.
+// Returns false if reading the tasks failed but we were able to recover from handing a crash report. In this case, all jobs must be submitted/processed again.
+FSCWErrorCode::ECode FShaderCompileUtilities::DoReadTaskResults(const TArray<FShaderCommonCompileJobPtr>& QueuedJobs, FArchive& OutputFile)
+{
+	FSCWOutputFileContext OutputFileContext(OutputFile);
+
+	if (OutputFile.TotalSize() == 0)
+	{
+		ShaderCompileWorkerError::HandleOutputFileEmpty(*OutputFile.GetArchiveName());
+	}
+
+	int32 OutputVersion = ShaderCompileWorkerOutputVersion;
+	OutputFile << OutputVersion;
+
+	if (ShaderCompileWorkerOutputVersion != OutputVersion)
+	{
+		OutputFileContext.ModalErrorOrLog(TEXT("Expecting output version %d, got %d instead!"), ShaderCompileWorkerOutputVersion, OutputVersion);
+	}
+
+	OutputFile << OutputFileContext.FileSize;
+
+	// Check for corrupted output file
+	if (OutputFileContext.FileSize > OutputFile.TotalSize())
+	{
+		ShaderCompileWorkerError::HandleOutputFileCorrupted(*OutputFile.GetArchiveName(), OutputFileContext.FileSize, OutputFile.TotalSize());
+	}
+
+	int32 ErrorCode = 0;
+	OutputFile << ErrorCode;
+
+	int32 NumProcessedJobs = 0;
+	OutputFile << NumProcessedJobs;
+
+	int32 CallstackLength = 0;
+	OutputFile << CallstackLength;
+
+	int32 ExceptionInfoLength = 0;
+	OutputFile << ExceptionInfoLength;
+
+	int32 HostnameLength = 0;
+	OutputFile << HostnameLength;
+
+	if (ErrorCode != FSCWErrorCode::Success)
+	{
+		// If worker crashed in a way we were able to recover from, return and expect the compile jobs to be reissued already
+		if (HandleWorkerCrash(QueuedJobs, OutputFile, OutputVersion, OutputFileContext.FileSize, (FSCWErrorCode::ECode)ErrorCode, NumProcessedJobs, CallstackLength, ExceptionInfoLength, HostnameLength))
+		{
+			FSCWErrorCode::Reset();
+			return (FSCWErrorCode::ECode)ErrorCode;
+		}
+	}
+
+	TArray<FShaderCompileJob*> QueuedSingleJobs;
+	TArray<FShaderPipelineCompileJob*> QueuedPipelineJobs;
+	SplitJobsByType(QueuedJobs, QueuedSingleJobs, QueuedPipelineJobs);
+	TArray<FShaderCommonCompileJob*> ReissueSourceJobs;
+
+	// Read single jobs
+	{
+		int32 SingleJobHeader = -1;
+		OutputFile << SingleJobHeader;
+		if (SingleJobHeader != ShaderCompileWorkerSingleJobHeader)
+		{
+			OutputFileContext.ModalErrorOrLog(TEXT("Expecting single job header ID 0x%08X, got 0x%08X instead!"), ShaderCompileWorkerSingleJobHeader, SingleJobHeader);
+		}
+
+		int32 NumJobs;
+		OutputFile << NumJobs;
+		if (NumJobs != QueuedSingleJobs.Num())
+		{
+			OutputFileContext.ModalErrorOrLog(TEXT("Expecting %d single %s, got %d instead!"), QueuedSingleJobs.Num(), (QueuedSingleJobs.Num() == 1 ? TEXT("job") : TEXT("jobs")), NumJobs);
+		}
+		else
+		{
+			for (int32 JobIndex = 0; JobIndex < NumJobs; JobIndex++)
+			{
+				auto* CurrentJob = QueuedSingleJobs[JobIndex];
+				if (ReadSingleJob(CurrentJob, OutputFile))
+				{
+					ReissueSourceJobs.Add(CurrentJob);
+				}
+			}
+		}
+	}
+
+	// Pipeline jobs
+	{
+		int32 PipelineJobHeader = -1;
+		OutputFile << PipelineJobHeader;
+		if (PipelineJobHeader != ShaderCompileWorkerPipelineJobHeader)
+		{
+			OutputFileContext.ModalErrorOrLog(TEXT("Expecting pipeline jobs header ID 0x%08X, got 0x%08X instead!"), ShaderCompileWorkerPipelineJobHeader, PipelineJobHeader);
+		}
+
+		int32 NumJobs;
+		OutputFile << NumJobs;
+		if (NumJobs != QueuedPipelineJobs.Num())
+		{
+			OutputFileContext.ModalErrorOrLog(TEXT("Expecting %d pipeline %s, got %d instead!"), QueuedPipelineJobs.Num(), (QueuedPipelineJobs.Num() == 1 ? TEXT("job") : TEXT("jobs")), NumJobs);
+		}
+		else
+		{
+			for (int32 JobIndex = 0; JobIndex < NumJobs; JobIndex++)
+			{
+				FShaderPipelineCompileJob* CurrentJob = QueuedPipelineJobs[JobIndex];
+
+				FString PipelineName;
+				OutputFile << PipelineName;
+				bool bSucceeded = false;
+				OutputFile << bSucceeded;
+				CurrentJob->bSucceeded = bSucceeded;
+				if (PipelineName != CurrentJob->Key.ShaderPipeline->GetName())
+				{
+					OutputFileContext.ModalErrorOrLog(TEXT("Expecting pipeline job \"%s\", got \"%s\" instead!"), CurrentJob->Key.ShaderPipeline->GetName(), *PipelineName);
+				}
+
+				check(!CurrentJob->bFinalized);
+				CurrentJob->bFinalized = true;
+
+				int32 NumStageJobs = -1;
+				OutputFile << NumStageJobs;
+
+				if (NumStageJobs != CurrentJob->StageJobs.Num())
+				{
+					OutputFileContext.ModalErrorOrLog(TEXT("Expecting %d stage pipeline %s, got %d instead!"), CurrentJob->StageJobs.Num(), (CurrentJob->StageJobs.Num() == 1 ? TEXT("job") : TEXT("jobs")), NumStageJobs);
+				}
+				else
+				{
+					for (int32 Index = 0; Index < NumStageJobs; Index++)
+					{
+						FShaderCompileJob* SingleJob = CurrentJob->StageJobs[Index];
+						// cannot reissue a single stage of a pipeline job
+						ReadSingleJob(SingleJob, OutputFile);
+					}
+				}
+			}
+		}
+	}
+	
+	// Requeue any jobs we wish to run again
+	ReissueShaderCompileJobs(ReissueSourceJobs);
+
+	return FSCWErrorCode::Success;
+}
+
+void FShaderCompileUtilities::ExecuteShaderCompileJob(FShaderCommonCompileJob& Job)
+{
+	TRACE_CPUPROFILER_EVENT_SCOPE(FShaderCompileUtilities::ExecuteShaderCompileJob);
+
+	check(!Job.bFinalized);
+
+	FString WorkingDir = FPlatformProcess::ShaderDir();
+	static ITargetPlatformManagerModule& TPM = GetTargetPlatformManagerRef();
+	auto* SingleJob = Job.GetSingleShaderJob();
+	const TArray<const IShaderFormat*> ShaderFormats = TPM.GetShaderFormats();
+	if (SingleJob)
+	{
+		const FName Format = (SingleJob->Input.ShaderFormat != NAME_None) ? SingleJob->Input.ShaderFormat : LegacyShaderPlatformToShaderFormat(EShaderPlatform(SingleJob->Input.Target.Platform));
+		CompileShader(ShaderFormats, *SingleJob, WorkingDir);
+	}
+	else
+	{
+		FShaderPipelineCompileJob* PipelineJob = Job.GetShaderPipelineJob();
+		check(PipelineJob);
+
+		EShaderPlatform Platform = (EShaderPlatform)PipelineJob->StageJobs[0]->Input.Target.Platform;
+		const FName Format = LegacyShaderPlatformToShaderFormat(Platform);
+
+		// Verify same platform on all stages
+		for (int32 Index = 1; Index < PipelineJob->StageJobs.Num(); ++Index)
+		{
+			auto SingleStage = PipelineJob->StageJobs[Index];
+			if (!SingleStage)
+			{
+				UE_LOG(LogShaderCompilers, Fatal, TEXT("Can't nest Shader Pipelines inside Shader Pipeline '%s'!"), PipelineJob->Key.ShaderPipeline->GetName());
+			}
+			else if (Platform != SingleStage->Input.Target.Platform)
+			{
+				UE_LOG(LogShaderCompilers, Fatal, TEXT("Mismatched Target Platform %s while compiling Shader Pipeline '%s'."), *Format.GetPlainNameString(), PipelineJob->Key.ShaderPipeline->GetName());
+			}
+		}
+
+		CompileShaderPipeline(ShaderFormats, PipelineJob, WorkingDir);
+	}
+
+	Job.bFinalized = true;
+}
+
+FArchive* FShaderCompileUtilities::CreateFileHelper(const FString& Filename)
+{
+	// TODO: This logic came from FShaderCompileThreadRunnable::WriteNewTasks().
+	// We can't avoid code duplication unless we refactored the local worker too.
+
+	FArchive* File = nullptr;
+	int32 RetryCount = 0;
+	// Retry over the next two seconds if we can't write out the file.
+	// Anti-virus and indexing applications can interfere and cause this to fail.
+	while (File == nullptr && RetryCount < 200)
+	{
+		if (RetryCount > 0)
+		{
+			FPlatformProcess::Sleep(0.01f);
+		}
+		File = IFileManager::Get().CreateFileWriter(*Filename, FILEWRITE_EvenIfReadOnly);
+		RetryCount++;
+	}
+	if (File == nullptr)
+	{
+		File = IFileManager::Get().CreateFileWriter(*Filename, FILEWRITE_EvenIfReadOnly | FILEWRITE_NoFail);
+	}
+	checkf(File, TEXT("Failed to create file %s!"), *Filename);
+	return File;
+}
+
+void FShaderCompileUtilities::MoveFileHelper(const FString& To, const FString& From)
+{
+	IPlatformFile& PlatformFile = FPlatformFileManager::Get().GetPlatformFile();
+
+	if (PlatformFile.FileExists(*From))
+	{
+		FString DirectoryName;
+		int32 LastSlashIndex;
+		if (To.FindLastChar('/', LastSlashIndex))
+		{
+			DirectoryName = To.Left(LastSlashIndex);
+		} else
+		{
+			DirectoryName = To;
+		}
+
+		// TODO: This logic came from FShaderCompileThreadRunnable::WriteNewTasks().
+		// We can't avoid code duplication unless we refactored the local worker too.
+
+		bool Success = false;
+		int32 RetryCount = 0;
+		// Retry over the next two seconds if we can't move the file.
+		// Anti-virus and indexing applications can interfere and cause this to fail.
+		while (!Success && RetryCount < 200)
+		{
+			if (RetryCount > 0)
+			{
+				FPlatformProcess::Sleep(0.01f);
+			}
+
+			// MoveFile does not create the directory tree, so try to do that now...
+			Success = PlatformFile.CreateDirectoryTree(*DirectoryName);
+			if (Success)
+			{
+				Success = PlatformFile.MoveFile(*To, *From);
+			}
+			RetryCount++;
+		}
+		checkf(Success, TEXT("Failed to move file %s to %s!"), *From, *To);
+	}
+}
+
+void FShaderCompileUtilities::DeleteFileHelper(const FString& Filename)
+{
+	// TODO: This logic came from FShaderCompileThreadRunnable::WriteNewTasks().
+	// We can't avoid code duplication unless we refactored the local worker too.
+
+	if (FPlatformFileManager::Get().GetPlatformFile().FileExists(*Filename))
+	{
+		bool bDeletedOutput = IFileManager::Get().Delete(*Filename, true, true);
+
+		// Retry over the next two seconds if we couldn't delete it
+		int32 RetryCount = 0;
+		while (!bDeletedOutput && RetryCount < 200)
+		{
+			FPlatformProcess::Sleep(0.01f);
+			bDeletedOutput = IFileManager::Get().Delete(*Filename, true, true);
+			RetryCount++;
+		}
+		checkf(bDeletedOutput, TEXT("Failed to delete %s!"), *Filename);
+	}
+}
+
+
