Issue a material translation error if custom primitive data is used with UI materials since it isn't supported.

- Scalar, Vector, and Texture Atlas Row Index parameters that used custom primitive data can't be used with UI materials.
- Custom Primitive Data requires a UPrimitive that has a uniform buffer for each primitives and UI rendering doens't use primitives.
- This error will prevent users from making materials that can't be used for UI rendering.
- It might make sense to limit this just to Surface domain materials

diff --git a/Engine/Source/Runtime/Engine/Private/Materials/MaterialExpressions.cpp b/Engine/Source/Runtime/Engine/Private/Materials/MaterialExpressions.cpp
--- a/Engine/Source/Runtime/Engine/Private/Materials/MaterialExpressions.cpp
+++ b/Engine/Source/Runtime/Engine/Private/Materials/MaterialExpressions.cpp
@@ -329,6 +329,8 @@ FUObjectAnnotationSparseBool GMaterialFunctionsThatNeedCoordinateCheck;
 FUObjectAnnotationSparseBool GMaterialFunctionsThatNeedCommentFix;
 FUObjectAnnotationSparseBool GMaterialFunctionsThatNeedSamplerFixup;
 FUObjectAnnotationSparseBool GMaterialFunctionsThatNeedFeatureLevelSM6Fix;
+
+static const TCHAR* CPD_UI_ErrorMessage = TEXT("Custom Primitive Data can't be used with the UI material domain.");
 #endif // #if WITH_EDITOR
 
 /** Returns whether the given expression class is allowed. */
@@ -9205,6 +9207,11 @@ int32 UMaterialExpressionVectorParameter::Compile(class FMaterialCompiler* Compi
 {
 	if (bUseCustomPrimitiveData)
 	{
+		if (Material && Material->MaterialDomain == MD_UI)
+		{
+			return CompilerError(Compiler, CPD_UI_ErrorMessage);
+		}
+
 		return Compiler->CustomPrimitiveData(PrimitiveDataIndex, MCT_Float4);
 	}
 	else
@@ -9578,6 +9585,11 @@ int32 UMaterialExpressionScalarParameter::Compile(class FMaterialCompiler* Compi
 {
 	if (bUseCustomPrimitiveData)
 	{
+		if (Material && Material->MaterialDomain == MD_UI)
+		{
+			return CompilerError(Compiler, CPD_UI_ErrorMessage);
+		}
+
 		return Compiler->CustomPrimitiveData(PrimitiveDataIndex, MCT_Float);
 	}
 	else
@@ -24824,6 +24836,11 @@ int32 UMaterialExpressionCurveAtlasRowParameter::Compile(class FMaterialCompiler
 	// Support for using the Custom Primitive Data to fetch an atlas index if that is chosen
 	if (bUseCustomPrimitiveData)
 	{
+		if (Material && Material->MaterialDomain == MD_UI)
+		{
+			return CompilerError(Compiler,CPD_UI_ErrorMessage);
+		}
+
 		Slot = Compiler->CustomPrimitiveData(PrimitiveDataIndex, MCT_Float);
 
 		if (Slot == INDEX_NONE)
