Fix tooltip on Vector Noise position input.

- The input position only accepts 3 dimensional vectors and will fail material compilation if given a 2D vector.
- Update the tooltip by removing the 2D vector.

diff --git a/Engine/Source/Runtime/Engine/Public/Materials/MaterialExpressionVectorNoise.h b/Engine/Source/Runtime/Engine/Public/Materials/MaterialExpressionVectorNoise.h
--- a/Engine/Source/Runtime/Engine/Public/Materials/MaterialExpressionVectorNoise.h
+++ b/Engine/Source/Runtime/Engine/Public/Materials/MaterialExpressionVectorNoise.h
@@ -53,7 +53,7 @@ class UMaterialExpressionVectorNoise : public UMaterialExpression
 {
 	GENERATED_UCLASS_BODY()
 
-	/** 2 to 3 dimensional vector */
+	/** 3 dimensional vector */
 	UPROPERTY()
 	FExpressionInput Position;
 
