Added support for blueprint set RealTimeCapture and RayTracedShadows.

#rnx

diff --git a/Engine/Source/Runtime/Engine/Classes/Components/LightComponent.h b/Engine/Source/Runtime/Engine/Classes/Components/LightComponent.h
--- a/Engine/Source/Runtime/Engine/Classes/Components/LightComponent.h
+++ b/Engine/Source/Runtime/Engine/Classes/Components/LightComponent.h
@@ -332,6 +332,9 @@ public:
 	UFUNCTION(BlueprintCallable, Category = "Rendering|Components|Light")
 	ENGINE_API void SetLightingChannels(bool bChannel0, bool bChannel1, bool bChannel2);
 
+	UFUNCTION(BlueprintCallable, Category="Rendering|Components|Light")
+	ENGINE_API void SetUseRayTracedDistanceFieldShadows(bool bNewValue);
+
 public:
 	/** The light's scene info. */
 	class FLightSceneProxy* SceneProxy;
@@ -479,6 +482,10 @@ protected:
 #endif // WITH_EDITOR
 	//~ End USceneComponent Interface
 
+private:
+
+	bool CanTraceDistanceFieldShadows() const;
+
 public:
 	ENGINE_API virtual void InvalidateLightingCacheDetailed(bool bInvalidateBuildEnqueuedLighting, bool bTranslationOnly) override;
 
diff --git a/Engine/Source/Runtime/Engine/Classes/Components/SkyLightComponent.h b/Engine/Source/Runtime/Engine/Classes/Components/SkyLightComponent.h
--- a/Engine/Source/Runtime/Engine/Classes/Components/SkyLightComponent.h
+++ b/Engine/Source/Runtime/Engine/Classes/Components/SkyLightComponent.h
@@ -274,6 +274,9 @@ class USkyLightComponent : public ULightComponentBase
 	UFUNCTION(BlueprintCallable, Category="Rendering|Components|SkyLight")
 	ENGINE_API void SetMinOcclusion(float InMinOcclusion);
 
+	UFUNCTION(BlueprintCallable, Category="Rendering|Components|SkyLight")
+	ENGINE_API void SetRealTimeCapture(bool bInRealTimeCapture);
+
 protected:
 	ENGINE_API virtual void OnVisibilityChanged() override;
 
diff --git a/Engine/Source/Runtime/Engine/Private/Components/LightComponent.cpp b/Engine/Source/Runtime/Engine/Private/Components/LightComponent.cpp
--- a/Engine/Source/Runtime/Engine/Private/Components/LightComponent.cpp
+++ b/Engine/Source/Runtime/Engine/Private/Components/LightComponent.cpp
@@ -605,6 +605,11 @@ void ULightComponent::PostLoad()
 	}
 }
 
+bool ULightComponent::CanTraceDistanceFieldShadows() const
+{
+	return CastShadows && CastDynamicShadows && Mobility != EComponentMobility::Static && DoesProjectSupportDistanceFields();
+}
+
 #if WITH_EDITOR
 void ULightComponent::PreSave(const class ITargetPlatform* TargetPlatform)
 {
@@ -664,7 +669,7 @@ bool ULightComponent::CanEditChange(const FProperty* InProperty) const
 		if (PropertyName == GET_MEMBER_NAME_STRING_CHECKED(ULightComponent, bUseRayTracedDistanceFieldShadows)
 			|| bIsRayStartOffset)
 		{
-			bool bCanEdit = CastShadows && CastDynamicShadows && Mobility != EComponentMobility::Static && DoesProjectSupportDistanceFields();
+			bool bCanEdit = CanTraceDistanceFieldShadows();
 
 			if (bIsRayStartOffset)
 			{
@@ -1271,6 +1276,19 @@ void ULightComponent::SetLightingChannels(bool bChannel0, bool bChannel1, bool b
 	}
 }
 
+void ULightComponent::SetUseRayTracedDistanceFieldShadows(bool bNewValue)
+{
+	// Never set to true if not supported.
+	bNewValue = bNewValue && CanTraceDistanceFieldShadows();
+
+	if (AreDynamicDataChangesAllowed() &&
+		bNewValue != bUseRayTracedDistanceFieldShadows)
+	{
+		bUseRayTracedDistanceFieldShadows = bNewValue;
+		MarkRenderStateDirty();
+	}
+}
+
 // GetDirection
 FVector ULightComponent::GetDirection() const 
 { 
diff --git a/Engine/Source/Runtime/Engine/Private/Components/SkyLightComponent.cpp b/Engine/Source/Runtime/Engine/Private/Components/SkyLightComponent.cpp
--- a/Engine/Source/Runtime/Engine/Private/Components/SkyLightComponent.cpp
+++ b/Engine/Source/Runtime/Engine/Private/Components/SkyLightComponent.cpp
@@ -1116,6 +1116,15 @@ void USkyLightComponent::SetMinOcclusion(float InMinOcclusion)
 	}
 }
 
+void USkyLightComponent::SetRealTimeCapture(bool bInRealTimeCapture)
+{
+	if (AreDynamicDataChangesAllowed()
+		&& bInRealTimeCapture != bRealTimeCapture)
+	{
+		SetRealTimeCaptureEnabled(bInRealTimeCapture);
+	}
+}
+
 bool USkyLightComponent::IsOcclusionSupported() const
 {
 	FSceneInterface* LocalScene = GetScene();
@@ -1141,6 +1150,7 @@ void USkyLightComponent::SetRealTimeCaptureEnabled(bool bNewRealTimeCaptureEnabl
 {
 	bRealTimeCapture = bNewRealTimeCaptureEnabled;
 	MarkRenderStateDirty();
+	SetCaptureIsDirty();
 }
 
 void USkyLightComponent::OnVisibilityChanged()
