PR #12116: Changed the Cloud show flag to always be accessible

diff --git a/Engine/Source/Runtime/Engine/Public/ShowFlagsValues.inl b/Engine/Source/Runtime/Engine/Public/ShowFlagsValues.inl
--- a/Engine/Source/Runtime/Engine/Public/ShowFlagsValues.inl
+++ b/Engine/Source/Runtime/Engine/Public/ShowFlagsValues.inl
@@ -293,8 +293,8 @@ SHOWFLAG_ALWAYS_ACCESSIBLE(LightShafts, SFG_LightingFeatures, NSLOCTEXT("UnrealE
 SHOWFLAG_FIXED_IN_SHIPPING(1, PostProcessMaterial, SFG_PostProcess, NSLOCTEXT("UnrealEd", "PostProcessMaterialSF", "Post Process Material"))
 /** Render Sky and Atmospheric lighting, for now SHOWFLAG_ALWAYS_ACCESSIBLE because it's exposed in SceneCapture */
 SHOWFLAG_ALWAYS_ACCESSIBLE(Atmosphere, SFG_Normal, NSLOCTEXT("UnrealEd", "AtmosphereSF", "Atmosphere"))
-/** Render Cloud */
-SHOWFLAG_FIXED_IN_SHIPPING(1, Cloud, SFG_Normal, NSLOCTEXT("UnrealEd", "CloudSF", "Cloud"))
+/** Render Cloud, for now SHOWFLAG_ALWAYS_ACCESSIBLE because it's exposed in SceneCapture */
+SHOWFLAG_ALWAYS_ACCESSIBLE(Cloud, SFG_Normal, NSLOCTEXT("UnrealEd", "CloudSF", "Cloud"))
 /** Render safe frames bars*/
 SHOWFLAG_FIXED_IN_SHIPPING(0, CameraAspectRatioBars, SFG_Advanced, NSLOCTEXT("UnrealEd", "CameraAspectRatioBarsSF", "Camera Aspect Ratio Bars"))
 /** Render safe frames */
