Fix issues in r.RayTracing.UseReferenceBasedResidency implementation
- RegisteredGeometry.Size was not correctly tracked for streamed out or evicted geometries.
- Geometries not assigned to a group weren't considered for eviction
- Too many geometries were being evicted due to incorrect logic (using RequestedSize instead of RequestedButEvictedSize)

diff --git a/Engine/Source/Runtime/Engine/Private/Rendering/RayTracingGeometryManager.cpp b/Engine/Source/Runtime/Engine/Private/Rendering/RayTracingGeometryManager.cpp
--- a/Engine/Source/Runtime/Engine/Private/Rendering/RayTracingGeometryManager.cpp
+++ b/Engine/Source/Runtime/Engine/Private/Rendering/RayTracingGeometryManager.cpp
@@ -393,11 +393,31 @@ void FRayTracingGeometryManager::RefreshRegisteredGeometry(RayTracingGeometryHan
 	{
 		FRegisteredGeometry& RegisteredGeometry = RegisteredGeometries[Handle];
 
-		if (RegisteredGeometry.Geometry->IsValid() && !RegisteredGeometry.Geometry->IsEvicted())
+		const uint32 OldSize = RegisteredGeometry.Size;
+
+		// Update size - Geometry RHI might not be valid yet (evicted or uninitialized), so calculate size using Initializer here
 		{
-			const uint32 OldSize = RegisteredGeometry.Size;
-			RegisteredGeometry.Size = RegisteredGeometry.Geometry->GetRHI()->GetSizeInfo().ResultSize;
+			bool bAllSegmentsAreValid = RegisteredGeometry.Geometry->Initializer.Segments.Num() > 0;
+			for (const FRayTracingGeometrySegment& Segment : RegisteredGeometry.Geometry->Initializer.Segments)
+			{
+				if (!Segment.VertexBuffer)
+				{
+					bAllSegmentsAreValid = false;
+					break;
+				}
+			}
+
+			RegisteredGeometry.Size = bAllSegmentsAreValid ? RHICalcRayTracingGeometrySize(RegisteredGeometry.Geometry->Initializer).ResultSize : 0;
+		}
+
+		if (AlwaysResidentGeometries.Contains(Handle))
+		{
+			TotalAlwaysResidentSize -= OldSize;
+			TotalAlwaysResidentSize += RegisteredGeometry.Size;
+		}
 
+		if (RegisteredGeometry.Geometry->IsValid() && !RegisteredGeometry.Geometry->IsEvicted())
+		{
 			bool bAlreadyInSet;
 			ResidentGeometries.Add(Handle, &bAlreadyInSet);
 
@@ -407,12 +427,6 @@ void FRayTracingGeometryManager::RefreshRegisteredGeometry(RayTracingGeometryHan
 			}
 
 			TotalResidentSize += RegisteredGeometry.Size;
-			
-			if (AlwaysResidentGeometries.Contains(Handle))
-			{
-				TotalAlwaysResidentSize -= OldSize;
-				TotalAlwaysResidentSize += RegisteredGeometry.Size;
-			}
 
 			if (RegisteredGeometry.Geometry->GroupHandle != INDEX_NONE)
 			{
@@ -423,6 +437,11 @@ void FRayTracingGeometryManager::RefreshRegisteredGeometry(RayTracingGeometryHan
 					EvictableGeometries.Add(Handle);
 				}
 			}
+			else
+			{
+				// geometries not assigned to a group (eg: dynamic geometry) are always evictable
+				EvictableGeometries.Add(Handle);
+			}
 		}
 		else
 		{
@@ -638,6 +657,8 @@ void FRayTracingGeometryManager::Tick(FRHICommandList& RHICmdList)
 					FRegisteredGeometry& RegisteredGeometry = RegisteredGeometries[GeometryHandle];
 					RegisteredGeometry.LastReferencedFrame = GFrameCounterRenderThread;
 
+					RequestedSize += RegisteredGeometry.Size;
+
 					if (RegisteredGeometry.Geometry->LODIndex >= Group.GeometryHandles.Num() - GRayTracingNumAlwaysResidentLODs)
 					{
 						checkf(!RegisteredGeometry.Geometry->IsEvicted(), TEXT("Always resident ray tracing geometry was unexpectely evicted."));
@@ -647,8 +668,6 @@ void FRayTracingGeometryManager::Tick(FRHICommandList& RHICmdList)
 						ReferencedGeometries.Add(GeometryHandle);
 						NotReferencedResidentGeometries.Remove(GeometryHandle);
 
-						RequestedSize += RegisteredGeometry.Size;
-
 						if (RegisteredGeometry.Geometry->IsEvicted())
 						{
 							RequestedButEvictedSize += RegisteredGeometry.Size;
@@ -698,7 +717,7 @@ void FRayTracingGeometryManager::Tick(FRHICommandList& RHICmdList)
 			// Step 3.2
 			// - evict geometries until we are in budget
 			int32 Index = 0;
-			while (TotalResidentSize + RequestedSize > ResidentGeometryMemoryPoolSize && Index < NotReferencedResidentGeometriesArray.Num())
+			while (TotalResidentSize + RequestedButEvictedSize > ResidentGeometryMemoryPoolSize && Index < NotReferencedResidentGeometriesArray.Num())
 			{
 				RayTracingGeometryHandle GeometryHandle = NotReferencedResidentGeometriesArray[Index];
 				FRegisteredGeometry& RegisteredGeometry = RegisteredGeometries[GeometryHandle];
