Fix ODSC requests getting lost when the connection to the ODSC server is lost
We now keep all the requests that were not properly sent / responded by the ODSC server so that we can issue them when we restore the communication to the server
Display a message on-screen when we detect we lost connection to the ODSC server

diff --git a/Engine/Source/Runtime/Engine/Private/ODSC/ODSCManager.cpp b/Engine/Source/Runtime/Engine/Private/ODSC/ODSCManager.cpp
--- a/Engine/Source/Runtime/Engine/Private/ODSC/ODSCManager.cpp
+++ b/Engine/Source/Runtime/Engine/Private/ODSC/ODSCManager.cpp
@@ -59,10 +59,12 @@ FODSCManager::FODSCManager()
 						OutMessages.Add(FCoreDelegates::EOnScreenMessageSeverity::Error, FText::FromString(LocalErrorMessage));
 					}
 
+					bool bIsConnectedToODSCServer;
 					bool bHasPendingGlobalShaders;
 					uint32 NumPendingMaterialsRecompile;
 					uint32 NumPendingMaterialsShaders;
-					if (Thread->GetPendingShaderData(bHasPendingGlobalShaders, NumPendingMaterialsRecompile, NumPendingMaterialsShaders))
+					bool bHasShaderData = Thread->GetPendingShaderData(bIsConnectedToODSCServer, bHasPendingGlobalShaders, NumPendingMaterialsRecompile, NumPendingMaterialsShaders);
+					if (bIsConnectedToODSCServer && bHasShaderData)
 					{
 						FString Message = TEXT("Recompiling shaders (");
 						if (bHasPendingGlobalShaders)
@@ -83,6 +85,10 @@ FODSCManager::FODSCManager()
 						Message += ")";
 						OutMessages.Add(FCoreDelegates::EOnScreenMessageSeverity::Info, FText::FromString(Message));
 					}
+					else if (!bIsConnectedToODSCServer)
+					{
+						OutMessages.Add(FCoreDelegates::EOnScreenMessageSeverity::Error, FText::FromString(FString::Printf(TEXT("Not connected to %s"), *Thread->GetODSCHostIP())));
+					}
 				}
 			}
 		);
diff --git a/Engine/Source/Runtime/Engine/Private/ODSC/ODSCThread.cpp b/Engine/Source/Runtime/Engine/Private/ODSC/ODSCThread.cpp
--- a/Engine/Source/Runtime/Engine/Private/ODSC/ODSCThread.cpp
+++ b/Engine/Source/Runtime/Engine/Private/ODSC/ODSCThread.cpp
@@ -168,6 +168,26 @@ bool FODSCThread::ConnectToODSCHost()
 	return CookOnTheFlyServerConnection != nullptr && CookOnTheFlyServerConnection->IsConnected();
 }
 
+bool FODSCThread::CheckODSCConnection()
+{
+	// If we have a default connection that already exists, send directly to that.
+	if ((CookOnTheFlyServerConnection == nullptr) || (!CookOnTheFlyServerConnection->IsConnected()))
+	{
+		// Losing connection when exit is requested is expected, do not try to reconnect
+		if (ExitRequest.GetValue())
+		{
+			return false;
+		}
+
+		UE_LOG(LogODSC, Display, TEXT("Detected that CookOnTheFlyServerConnection has been lost, trying again"));
+		if (!ConnectToODSCHost())
+		{
+			return false;
+		}
+	}
+	return CookOnTheFlyServerConnection != nullptr && CookOnTheFlyServerConnection->IsConnected();
+}
+
 void FODSCThread::StartThread()
 {
 	Thread = FRunnableThread::Create(this, TEXT("ODSCThread"), 128 * 1024, TPri_Normal);
@@ -368,33 +388,65 @@ void FODSCThread::Exit()
 
 void FODSCThread::Process()
 {
-	// cache all pending requests.
-	FODSCRequestPayload Payload;
-	TArray<FODSCRequestPayload> PayloadsToAggregate;
+	// cache all pending pipeline requests
 	{
+		TArray<FODSCRequestPayload> PayloadsToAggregate;
+		FODSCRequestPayload Payload;
 		while (PendingMeshMaterialThreadedRequests.Dequeue(Payload))
 		{
 			PayloadsToAggregate.Add(Payload);
 		}
+
+		if (PayloadsToAggregate.Num())
+		{
+			FODSCMessageHandler* RequestHandler = new FODSCMessageHandler(PayloadsToAggregate[0].ShaderPlatform, PayloadsToAggregate[0].FeatureLevel, PayloadsToAggregate[0].QualityLevel, ODSCRecompileCommand::Material);
+			for (const FODSCRequestPayload& payload : PayloadsToAggregate)
+			{
+				RequestHandler->AddPayload(payload);
+			}
+			PendingRequestsPipeline.Add(RequestHandler);
+		}
+	}
+
+	// cache all pending material/global requests
+	{
+		FODSCMessageHandler* Request = nullptr;
+		while (PendingMaterialThreadedRequests.Dequeue(Request))
+		{
+			PendingRequestsMaterialAndGlobal.Add(Request);
+		}
+	}
+
+	bIsConnectedToODSCServer = CheckODSCConnection();
+
+	ON_SCOPE_EXIT
+	{
+		// SendMessageToServer is synchronous, so when we're here, we know we've processed all the requests
+		WakeupEvent->Reset();
+		AllRequestsDoneEvent->Trigger();
+	};
+
+	// Early out to avoid trying to connect (and most likely fail) for every compilation request
+	if (!bIsConnectedToODSCServer)
+	{
+		return;
 	}
 
 	// cache material requests.
-	FODSCMessageHandler* Request = nullptr;
-	TArray<FODSCMessageHandler*> RequestsToStart;
+	TArray<FODSCMessageHandler*> RequestsToStart = MoveTemp(PendingRequestsMaterialAndGlobal);
 	bool bHasGlobalShaders = false;
 	uint32 NumMaterials = 0;
-	while (PendingMaterialThreadedRequests.Dequeue(Request))
+
+	for (FODSCMessageHandler* NextRequest : RequestsToStart)
 	{
-		if (Request->GetRecompileCommandType() != ODSCRecompileCommand::Material)
+		if (NextRequest->GetRecompileCommandType() != ODSCRecompileCommand::Material)
 		{
 			bHasGlobalShaders = true;
 		}
 		else
 		{
-			NumMaterials += Request->GetMaterialsToLoad().Num();
+			NumMaterials += NextRequest->GetMaterialsToLoad().Num();
 		}
-
-		RequestsToStart.Add(Request);
 	}
 
 	bHasPendingGlobalShaders.store(bHasGlobalShaders, std::memory_order_release);
@@ -404,61 +456,57 @@ void FODSCThread::Process()
 	for (FODSCMessageHandler* NextRequest : RequestsToStart)
 	{
 		// send the info, the handler will process the response (and update shaders, etc)
-		SendMessageToServer(NextRequest);
+		if (SendMessageToServer(NextRequest))
+		{
+			CompletedThreadedRequests.Enqueue(NextRequest);
+		}
+		else
+		{
+			PendingRequestsMaterialAndGlobal.Add(NextRequest);
+		}
 
-		CompletedThreadedRequests.Enqueue(NextRequest);
 	}
 
 	bHasPendingGlobalShaders.store(false, std::memory_order_release);
 	NumPendingMaterialsRecompile.store(0, std::memory_order_release);
-	NumPendingMaterialsShaders.store(PayloadsToAggregate.Num(), std::memory_order_release);
+
+	RequestsToStart = MoveTemp(PendingRequestsPipeline);
+
+	uint32 NumPipelines = 0;
+	for (FODSCMessageHandler* NextRequest : RequestsToStart)
+	{
+		NumPipelines += NextRequest->NumPayloads();
+	}
+
+	NumPendingMaterialsShaders.store(NumPipelines, std::memory_order_release);
 
 	// process any specific mesh material shader requests.
-	if (PayloadsToAggregate.Num())
+	for (FODSCMessageHandler* NextRequest : RequestsToStart)
 	{
-		FODSCMessageHandler* RequestHandler = new FODSCMessageHandler(PayloadsToAggregate[0].ShaderPlatform, PayloadsToAggregate[0].FeatureLevel, PayloadsToAggregate[0].QualityLevel, ODSCRecompileCommand::Material);
-		for (const FODSCRequestPayload& payload : PayloadsToAggregate)
+		if (SendMessageToServer(NextRequest))
 		{
-			RequestHandler->AddPayload(payload);
+			CompletedThreadedRequests.Enqueue(NextRequest);
+		}
+		else
+		{
+			PendingRequestsPipeline.Add(NextRequest);
 		}
-
-		// send the info, the handler will process the response (and update shaders, etc)
-		SendMessageToServer(RequestHandler);
-
-		CompletedThreadedRequests.Enqueue(RequestHandler);
 	}
 
 	NumPendingMaterialsShaders.store(0, std::memory_order_release);
-
-	// SendMessageToServer is synchronous, so when we're here, we know we've processed all the requests
-	WakeupEvent->Reset();
-	AllRequestsDoneEvent->Trigger();
 }
 
-void FODSCThread::SendMessageToServer(IPlatformFile::IFileServerMessageHandler* Handler)
+bool FODSCThread::SendMessageToServer(IPlatformFile::IFileServerMessageHandler* Handler)
 {
-	// If we have a default connection that already exists, send directly to that.
-	if ((CookOnTheFlyServerConnection == nullptr) || (!CookOnTheFlyServerConnection->IsConnected()))
+	if (bHasDefaultConnection)
 	{
-		if (bHasDefaultConnection)
-		{
-			IFileManager::Get().SendMessageToServer(TEXT("RecompileShaders"), Handler);
-			return;
-		}
-		else
-		{
-			// Losing connection when exit is requested is expected, do not try to reconnect
-			if (ExitRequest.GetValue())
-			{
-				return;
-			}
+		IFileManager::Get().SendMessageToServer(TEXT("RecompileShaders"), Handler);
+		return true;
+	}
 
-			UE_LOG(LogODSC, Display, TEXT("Detected that CookOnTheFlyServerConnection has been lost, trying again"));
-			if (!ConnectToODSCHost())
-			{
-				return;
-			}
-		}
+	if (!CheckODSCConnection())
+	{
+		return false;
 	}
 
 	// We don't have a default COTF connection so use our specific connection to send our command.
@@ -473,16 +521,19 @@ void FODSCThread::SendMessageToServer(IPlatformFile::IFileServerMessageHandler*
 	{
 		TUniquePtr<FArchive> Ar = Response.ReadBody();
 		Handler->ProcessResponse(*Ar);
+		return true;
 	}
 	else
 	{
 		UE_LOG(LogODSC, Display, TEXT("Received error response from CookOnTheFlyServerConnection; disconnecting"));
 		CookOnTheFlyServerConnection.Reset();
+		return false;
 	}
 }
 
-bool FODSCThread::GetPendingShaderData(bool& bOutHasPendingGlobalShaders, uint32& OutNumPendingMaterialsRecompile, uint32& OutNumPendingMaterialsShaders) const
+bool FODSCThread::GetPendingShaderData(bool& bOutIsConnectedToODSCServer, bool& bOutHasPendingGlobalShaders, uint32& OutNumPendingMaterialsRecompile, uint32& OutNumPendingMaterialsShaders) const
 {
+	bOutIsConnectedToODSCServer = bIsConnectedToODSCServer.load(std::memory_order_acquire);
 	bOutHasPendingGlobalShaders = bHasPendingGlobalShaders.load(std::memory_order_acquire);
 	OutNumPendingMaterialsRecompile = NumPendingMaterialsRecompile.load(std::memory_order_acquire);
 	OutNumPendingMaterialsShaders = NumPendingMaterialsShaders.load(std::memory_order_acquire);
diff --git a/Engine/Source/Runtime/Engine/Private/ODSC/ODSCThread.h b/Engine/Source/Runtime/Engine/Private/ODSC/ODSCThread.h
--- a/Engine/Source/Runtime/Engine/Private/ODSC/ODSCThread.h
+++ b/Engine/Source/Runtime/Engine/Private/ODSC/ODSCThread.h
@@ -48,6 +48,7 @@ public:
 	const TArray<uint8>& GetGlobalShaderMap() const;
 	bool ReloadGlobalShaders() const;
 	ODSCRecompileCommand GetRecompileCommandType() const { return RecompileCommandType; };
+	int32 NumPayloads() const { return RequestBatch.Num(); }
 
 private:
 	/** The time when this command was issued.  This isn't serialized to the cooking server. */
@@ -163,13 +164,14 @@ public:
 	*/
 	void WaitUntilAllRequestsDone();
 
-	bool GetPendingShaderData(bool& bOutHasPendingGlobalShaders, uint32& OutNumPendingMaterialsRecompile, uint32& OutNumPendingMaterialsShaders) const;
+	bool GetPendingShaderData(bool& bOutIsConnectedToODSCServer, bool& bOutHasPendingGlobalShaders, uint32& OutNumPendingMaterialsRecompile, uint32& OutNumPendingMaterialsShaders) const;
 
 	void RegisterMaterialShaderMap(const FMaterialShaderMap& MaterialShaderMap);
 
 	void ResetMaterialsODSCData(ERHIFeatureLevel::Type FeatureLevel);
 
 	bool CheckIfRequestAlreadySent(const TArray<FShaderId>& RequestShaderIds, const FString& MaterialName) const;
+	const FString& GetODSCHostIP() const { return ODSCHostIP; };
 
 protected:
 
@@ -192,6 +194,7 @@ private:
 	void Process();
 
 	bool ConnectToODSCHost();
+	bool CheckODSCConnection();
 
 	/**
 	 * Threaded requests that are waiting to be processed on the ODSC thread.
@@ -273,13 +276,17 @@ private:
 	/** Holds an event signaling when all the requests are processed*/
 	FEvent* AllRequestsDoneEvent;
 
-	void SendMessageToServer(IPlatformFile::IFileServerMessageHandler* Handler);
+	bool SendMessageToServer(IPlatformFile::IFileServerMessageHandler* Handler);
+
+	TArray<FODSCMessageHandler*> PendingRequestsMaterialAndGlobal;
+	TArray<FODSCMessageHandler*> PendingRequestsPipeline;
 
 	/** Special connection to the cooking server.  This is only used to send recompileshaders commands on. */
 	TUniquePtr<UE::Cook::ICookOnTheFlyServerConnection> CookOnTheFlyServerConnection;
 
 	FString ODSCHostIP;
 
+	std::atomic<bool> bIsConnectedToODSCServer = false;
 	std::atomic<bool> bHasPendingGlobalShaders = false;
 	std::atomic<uint32> NumPendingMaterialsRecompile = 0;
 	std::atomic<uint32> NumPendingMaterialsShaders = 0;
