Fix the upgrade of the legacy Layers material

diff --git a/Engine/Source/Runtime/Engine/Private/Materials/MaterialExpressions.cpp b/Engine/Source/Runtime/Engine/Private/Materials/MaterialExpressions.cpp
--- a/Engine/Source/Runtime/Engine/Private/Materials/MaterialExpressions.cpp
+++ b/Engine/Source/Runtime/Engine/Private/Materials/MaterialExpressions.cpp
@@ -17955,9 +17955,26 @@ bool FMaterialLayersFunctionsRuntimeData::Serialize(FArchive& Ar)
 		if (Ar.IsLoading())
 		{
 			Tree.Empty();
-			for (int32 l = 0; l < Layers.Num(); ++l)
+
+			int NumLegacyLayers = Layers.Num();
+			int NumLegacyBlends = Blends.Num();
+			// First add an extra empty Blend to match the number of layers
+			while (Blends.Num() < Layers.Num())
+			{
+				Blends.InsertDefaulted(0); // normally we should go through this only once
+			}
+
+			Layers.AddDefaulted();
+			Blends.AddDefaulted();
+
+			// Add a layer node at root 
+			FLayerNodeId LayerNodeId = Tree.AddNode({ NumLegacyLayers,  NumLegacyLayers }, -1);
+
+			// Third rebuiild the tree of nodes referencing the layers and blends
+			for (int32 l = 0; l < NumLegacyLayers; ++l)
 			{
-				Tree.AddNode({ l, l - 1 }, -1);
+				// And then ONE MORE node to be the attributes first child
+				Tree.AddNode({ l,  l}, LayerNodeId);
 			}
 		}
 	}
@@ -18414,6 +18431,28 @@ bool FMaterialLayersFunctions::ResolveParent(const FMaterialLayersFunctionsRunti
 	FMaterialLayersFunctionsEditorOnlyData& EditorOnly,
 	TArray<int32>& OutRemapLayerIndices)
 {
+#if ENABLE_MATERIAL_LAYER_PROTOTYPE
+	// This method is called right after unserialization
+	// Use this fact to fix editor only data in case of a legacy version of the LayersFunctions
+
+		// Catch case when a legacy layer has been loaded
+		// EditorOnly data need to be updated
+	if ((Runtime.Layers.Num() == Runtime.Blends.Num()) && ((EditorOnly.LayerNames.Num() + 1) == Runtime.Layers.Num()))
+	{
+		int32 NumLegacyLayers = EditorOnly.LayerNames.Num();
+		// Second add an extra Layer and Blend for each legacy layer.
+		{
+			EditorOnly.LayerStates.Add(true);
+			FText LayerName = FText::FromString(TEXT("Legacy Layers"));
+			EditorOnly.LayerNames.Add(LayerName);
+			EditorOnly.RestrictToLayerRelatives.Add(false);
+			EditorOnly.RestrictToBlendRelatives.Add(false);
+			EditorOnly.LayerGuids.Add(FGuid::NewGuid());
+			EditorOnly.LayerLinkStates.Add(EMaterialLayerLinkState::Uninitialized);
+		}
+	}
+#endif
+
  	check(EditorOnly.LayerGuids.Num() == Runtime.Layers.Num());
 	check(EditorOnly.LayerLinkStates.Num() == Runtime.Layers.Num());
 
