Improve ODSC with streaming

When serializing a new shadermap, check if one already exists that was compiled with ODSC. This will avoid using the cooked shadermap for a few frames, and then replace it with ODSC anyway.

For that, we need ODSC to keep all shadermaps in memory that we requested from the server. The shadermapid might actually have changed if things like material functions are being added/removed: the material full path is a better identifier.

Calling FMaterialResource::GetFullPath in FMaterial::TryGetShaders can be expensive though. Instead, Keep a cache of FName in ODSC to retrieve the full path from a FMaterial pointer, and remove the entry when we destroy the FMaterial

diff --git a/Engine/Source/Runtime/Engine/Private/Materials/Material.cpp b/Engine/Source/Runtime/Engine/Private/Materials/Material.cpp
--- a/Engine/Source/Runtime/Engine/Private/Materials/Material.cpp
+++ b/Engine/Source/Runtime/Engine/Private/Materials/Material.cpp
@@ -91,6 +91,10 @@
 #include "BuildSettings.h"
 #include "LocalVertexFactory.h"
 
+#if WITH_ODSC
+#include "ODSC/ODSCManager.h"
+#endif
+
 #if WITH_EDITOR
 #include "Framework/Notifications/NotificationManager.h"
 #include "MaterialCachedHLSLTree.h"
@@ -867,6 +871,19 @@ void ProcessSerializedInlineShaderMaps(UMaterialInterface* Owner, TArray<FMateri
 
 			if (bIncludeShaderMap)
 			{
+#if WITH_ODSC
+				// Check if already have something from ODSC so that we avoid using this shadermap for a few frames and making the request to replace it by something we already have
+				// If some permutations are missing, we will still perform additional requests
+				if (FODSCManager::IsODSCActive())
+				{
+					FMaterialShaderMap* ExistingShaderMap = FODSCManager::FindMaterialShaderMap(Owner->GetPathName(), LoadedShaderMap->GetShaderMapId());
+					if (ExistingShaderMap)
+					{
+						LoadedShaderMap = ExistingShaderMap;
+					}
+				}
+#endif
+
 				FMaterialResource* CurrentResource = FindOrCreateMaterialResource(OutMaterialResourcesLoaded, OwnerMaterial, OwnerMaterialInstance, LoadedFeatureLevel, LoadedQualityLevel);
 				CurrentResource->SetInlineShaderMap(LoadedShaderMap);
 			}
diff --git a/Engine/Source/Runtime/Engine/Private/Materials/MaterialShader.cpp b/Engine/Source/Runtime/Engine/Private/Materials/MaterialShader.cpp
--- a/Engine/Source/Runtime/Engine/Private/Materials/MaterialShader.cpp
+++ b/Engine/Source/Runtime/Engine/Private/Materials/MaterialShader.cpp
@@ -2098,18 +2098,31 @@ void FMaterialShaderMap::LoadForRemoteRecompile(FArchive& Ar, EShaderPlatform Sh
 
 			if (bIsValid)
 			{
-				FMaterialShaderMap* ShaderMap = new FMaterialShaderMap();
+				TRefCountPtr<FMaterialShaderMap> ShaderMap = new FMaterialShaderMap();
 
 				// serialize the id and the material shader map
 				ShaderMap->Serialize(Ar, false);
 
-				// Register in the global map
-				ShaderMap->RegisterForODSC(ShaderPlatform);
+				// If we already registered a material shadermap with the same ID, just re-use it. The material will do the same in FMaterial::CacheShaders anyway
+				TRefCountPtr<FMaterialShaderMap> ExistingShaderMap = FMaterialShaderMap::FindId(ShaderMap->GetShaderMapId(), ShaderPlatform);
+				if (ExistingShaderMap)
+				{
+					ShaderMap = ExistingShaderMap;
+				}
+				else
+				{
+					// Register in the global map
+					ShaderMap->RegisterForODSC(ShaderPlatform);
+				}
 
 				LoadedShaderMaps.Add(ShaderMap);
 			}
 		}
 
+#if WITH_ODSC
+		FODSCManager::RegisterMaterialShaderMaps(MaterialName, LoadedShaderMaps);
+#endif
+
 		UMaterialInterface* MatchingMaterial = FindObject<UMaterialInterface>(nullptr, *MaterialName);
 		if (!MatchingMaterial)
 		{
@@ -3399,7 +3412,6 @@ void FMaterialShaderMap::RegisterForODSC(EShaderPlatform InShaderPlatform)
 
 #if WITH_ODSC
 		bIsFromODSC = true;
-		FODSCManager::RegisterMaterialShaderMap(*this);
 #endif
 	}
 }
diff --git a/Engine/Source/Runtime/Engine/Private/Materials/MaterialShared.cpp b/Engine/Source/Runtime/Engine/Private/Materials/MaterialShared.cpp
--- a/Engine/Source/Runtime/Engine/Private/Materials/MaterialShared.cpp
+++ b/Engine/Source/Runtime/Engine/Private/Materials/MaterialShared.cpp
@@ -2548,6 +2548,10 @@ void FMaterial::DeleteMaterialsOnRenderThread(TArray<TRefCountPtr<FMaterial>>& M
  */
 FMaterial::~FMaterial()
 {
+#if WITH_ODSC
+	FODSCManager::UnregisterMaterialName(this);
+#endif
+
 #if WITH_EDITOR
 	check(GameThreadCompilingShaderMapId == 0u);
 	check(RenderingThreadCompilingShaderMapId == 0u);
@@ -3791,18 +3795,16 @@ bool FMaterial::TryGetShaders(const FMaterialShaderTypes& InTypes, const FVertex
 #if WITH_ODSC
 		TArray<FShaderId> RequestShaderIds;
 		TArray<FString> ShaderStageNamesToCompile;
-		FString MaterialName;
         bool bODSCRequestAlreadySent = false;
 		if (bShouldForceRecompile)
 		{
-			MaterialName = GetFullPath();
 			for (auto* ShaderType : InTypes.PipelineType->GetStages())
 			{
 				ShaderStageNamesToCompile.Add(ShaderType->GetName());
 				RequestShaderIds.Add(FShaderId(ShaderType, ShaderMap->GetShaderMapId().CookedShaderMapIdHash, InTypes.PipelineType->GetHashedName(), InVertexFactoryType, kUniqueShaderPermutationId, ShaderPlatform));
 			}
 
-			bODSCRequestAlreadySent = GODSCManager->CheckIfRequestAlreadySent(RequestShaderIds, MaterialName);
+			bODSCRequestAlreadySent = GODSCManager->CheckIfRequestAlreadySent(RequestShaderIds, this);
 			if (bODSCRequestAlreadySent)
 			{
 				bShouldForceRecompile = false;
@@ -3839,7 +3841,7 @@ bool FMaterial::TryGetShaders(const FMaterialShaderTypes& InTypes, const FVertex
 					{
 						const FString VFTypeName(InVertexFactoryType ? InVertexFactoryType->GetName() : TEXT(""));
 						const FString PipelineName(InTypes.PipelineType->GetName());
-						GODSCManager->AddThreadedShaderPipelineRequest(ShaderPlatform, GetFeatureLevel(), GetQualityLevel(), MaterialName, VFTypeName, PipelineName, 
+						GODSCManager->AddThreadedShaderPipelineRequest(ShaderPlatform, GetFeatureLevel(), GetQualityLevel(), this, VFTypeName, PipelineName, 
 						                                               ShaderStageNamesToCompile, kUniqueShaderPermutationId, RequestShaderIds);
 					}
 				}
@@ -3912,14 +3914,12 @@ bool FMaterial::TryGetShaders(const FMaterialShaderTypes& InTypes, const FVertex
 #if WITH_ODSC
 
 				TArray<FShaderId> RequestShaderIds;
-				FString MaterialName;
                 bool bODSCRequestAlreadySent = false;
 				if (bShouldForceRecompile)
 				{
-					MaterialName = GetFullPath();
 					RequestShaderIds.Add(FShaderId(ShaderType, ShaderMap->GetShaderMapId().CookedShaderMapIdHash, FHashedName(), InVertexFactoryType, PermutationId, ShaderPlatform));
 
-					bODSCRequestAlreadySent = GODSCManager->CheckIfRequestAlreadySent(RequestShaderIds, MaterialName);
+					bODSCRequestAlreadySent = GODSCManager->CheckIfRequestAlreadySent(RequestShaderIds, this);
 					if (bODSCRequestAlreadySent)
 					{
 						bShouldForceRecompile = false;
@@ -3948,7 +3948,7 @@ bool FMaterial::TryGetShaders(const FMaterialShaderTypes& InTypes, const FVertex
 							TArray<FString> ShaderStageNamesToCompile;
 							ShaderStageNamesToCompile.Add(ShaderType->GetName());
 
-							GODSCManager->AddThreadedShaderPipelineRequest(ShaderPlatform, GetFeatureLevel(), GetQualityLevel(), MaterialName, VFTypeName, PipelineName, 
+							GODSCManager->AddThreadedShaderPipelineRequest(ShaderPlatform, GetFeatureLevel(), GetQualityLevel(), this, VFTypeName, PipelineName, 
 							                                               ShaderStageNamesToCompile, PermutationId, RequestShaderIds);
 						}
 					}
diff --git a/Engine/Source/Runtime/Engine/Private/ODSC/ODSCManager.cpp b/Engine/Source/Runtime/Engine/Private/ODSC/ODSCManager.cpp
--- a/Engine/Source/Runtime/Engine/Private/ODSC/ODSCManager.cpp
+++ b/Engine/Source/Runtime/Engine/Private/ODSC/ODSCManager.cpp
@@ -194,7 +194,7 @@ void FODSCManager::AddThreadedShaderPipelineRequest(
 	EShaderPlatform ShaderPlatform,
 	ERHIFeatureLevel::Type FeatureLevel,
 	EMaterialQualityLevel::Type QualityLevel,
-	const FString& MaterialName,
+	const FMaterial* Material,
 	const FString& VertexFactoryName,
 	const FString& PipelineName,
 	const TArray<FString>& ShaderTypeNames,
@@ -203,7 +203,7 @@ void FODSCManager::AddThreadedShaderPipelineRequest(
 {
 	if (IsHandlingRequests())
 	{
-		Thread->AddShaderPipelineRequest(ShaderPlatform, FeatureLevel, QualityLevel, MaterialName, VertexFactoryName, PipelineName, ShaderTypeNames, PermutationId, RequestShaderIds);
+		Thread->AddShaderPipelineRequest(ShaderPlatform, FeatureLevel, QualityLevel, Material, VertexFactoryName, PipelineName, ShaderTypeNames, PermutationId, RequestShaderIds);
 	}
 }
 
@@ -310,21 +310,38 @@ bool FODSCManager::ShouldForceRecompileInternal(const FMaterialShaderMap* Materi
 #endif
 }
 
-void FODSCManager::RegisterMaterialShaderMap(const FMaterialShaderMap& MaterialShaderMap)
+bool FODSCManager::CheckIfRequestAlreadySent(const TArray<FShaderId>& RequestShaderIds, const FMaterial* Material) const
+{
+	if (Thread)
+	{
+		return Thread->CheckIfRequestAlreadySent(RequestShaderIds, Material);
+	}
+	return false;
+}
+
+void FODSCManager::UnregisterMaterialName(const FMaterial* Material)
 {
 	if (IsODSCActive())
 	{
-		GODSCManager->Thread->RegisterMaterialShaderMap(MaterialShaderMap);
+		GODSCManager->Thread->UnregisterMaterialName(Material);
 	}
 }
- 
-bool FODSCManager::CheckIfRequestAlreadySent(const TArray<FShaderId>& RequestShaderIds, const FString& MaterialName) const
+
+void FODSCManager::RegisterMaterialShaderMaps(const FString& MaterialName, const TArray<TRefCountPtr<FMaterialShaderMap>>& LoadedShaderMaps)
 {
-	if (Thread)
+	if (IsODSCActive())
 	{
-		return Thread->CheckIfRequestAlreadySent(RequestShaderIds, MaterialName);
+		GODSCManager->Thread->RegisterMaterialShaderMaps(MaterialName, LoadedShaderMaps);
 	}
-	return false;
+}
+
+FMaterialShaderMap* FODSCManager::FindMaterialShaderMap(const FString& MaterialName, const FMaterialShaderMapId& ShaderMapId)
+{
+	if (IsODSCActive())
+	{
+		return GODSCManager->Thread->FindMaterialShaderMap(MaterialName, ShaderMapId);
+	}
+	return nullptr;
 }
 
 void FODSCManager::TryLoadGlobalShaders(EShaderPlatform ShaderPlatform)
diff --git a/Engine/Source/Runtime/Engine/Private/ODSC/ODSCThread.cpp b/Engine/Source/Runtime/Engine/Private/ODSC/ODSCThread.cpp
--- a/Engine/Source/Runtime/Engine/Private/ODSC/ODSCThread.cpp
+++ b/Engine/Source/Runtime/Engine/Private/ODSC/ODSCThread.cpp
@@ -241,32 +241,92 @@ void FODSCThread::ResetMaterialsODSCData(ERHIFeatureLevel::Type FeatureLevel)
 }
 
 FODSCThread::FODSCShaderId::FODSCShaderId(const FShaderId& ShaderId)
-: MaterialShaderMapHash(ShaderId.MaterialShaderMapHash)
-, ShaderTypeHashedName(ShaderId.Type ? ShaderId.Type->GetHashedName() : 0)
+: ShaderTypeHashedName(ShaderId.Type ? ShaderId.Type->GetHashedName() : 0)
 , VFTypeHashedName(ShaderId.VFType ? ShaderId.VFType->GetHashedName() : 0)
 , ShaderPipelineName(ShaderId.ShaderPipelineName)
 , PermutationId(ShaderId.PermutationId)
 , Platform(ShaderId.Platform)
 {}
 
-bool FODSCThread::CheckIfRequestAlreadySent(const TArray<FShaderId>& RequestShaderIds, const FString& MaterialName) const
+bool FODSCThread::CheckIfRequestAlreadySent(const TArray<FShaderId>& RequestShaderIds, const FMaterial* Material) const
 {
 	FReadScopeLock ReadLock(RequestHashesRWLock);
+	const FName* CachedMaterialName = ODSCPointerToNames.Find((UPTRINT)Material);
+	if (CachedMaterialName == nullptr)
+	{
+		return false;
+	}
+
+	const FODSCShaderMapData* ODSCShaderMapData = RequestHashes.Find(*CachedMaterialName);
+	if (ODSCShaderMapData == nullptr)
+	{
+		return false;
+	}
+
 	for (const FShaderId& ShaderId : RequestShaderIds)
 	{
-		const FMaterialRequestsHashes* MaterialRequestHashes = RequestHashes.Find(ShaderId);
-		if (MaterialRequestHashes == nullptr)
+		if (!ODSCShaderMapData->CurrentRequests.Contains(ShaderId))
 		{
 			return false;
 		}
-			
-		if (!MaterialRequestHashes->RequestStrings.Contains(MaterialName))
+	}
+
+	return true;
+}
+
+void FODSCThread::UnregisterMaterialName(const FMaterial* Material)
+{
+	FWriteScopeLock WriteLock(RequestHashesRWLock);
+	ODSCPointerToNames.Remove((UPTRINT)Material);
+}
+
+void FODSCThread::RegisterMaterialShaderMaps(const FString& MaterialName, const TArray<TRefCountPtr<FMaterialShaderMap>>& LoadedShaderMaps)
+{
+	FWriteScopeLock WriteLock(RequestHashesRWLock);
+
+	FODSCShaderMapData& ODSCShaderMapData = RequestHashes.FindOrAdd(FName(MaterialName));
+
+	ODSCShaderMapData.MaterialShaderMaps = LoadedShaderMaps;
+
+	for (FMaterialShaderMap* MaterialShaderMap : LoadedShaderMaps)
+	{
+		TMap<FShaderId, TShaderRef<FShader>> ShadersInMap;
+		MaterialShaderMap->GetShaderList(ShadersInMap);
+		for (auto Iter : ShadersInMap)
 		{
-			return false;
+			// The shadermap we receive contains all the requests the client sent until now, so it's possible they already got removed
+			if (ODSCShaderMapData.CurrentRequests.Find(Iter.Key))
+			{
+				ODSCShaderMapData.CurrentRequests.Remove(Iter.Key);
+			}
 		}
 	}
 
-	return true;
+}
+
+FMaterialShaderMap* FODSCThread::FindMaterialShaderMap(const FString& MaterialName, const FMaterialShaderMapId& ShaderMapId) const
+{
+	FReadScopeLock ReadLock(RequestHashesRWLock);
+	const FODSCShaderMapData* ODSCShaderMapData = RequestHashes.Find(FName(MaterialName));
+	if (ODSCShaderMapData == nullptr)
+	{
+		return nullptr;
+	}
+
+	for (FMaterialShaderMap* MaterialShaderMap : ODSCShaderMapData->MaterialShaderMaps)
+	{
+		const FMaterialShaderMapId& ExistingShaderMapId = MaterialShaderMap->GetShaderMapId();
+		bool bFeatureLevelMatch = (ExistingShaderMapId.FeatureLevel == ShaderMapId.FeatureLevel);
+		bool bQualityLevelMatch = (ShaderMapId.QualityLevel == EMaterialQualityLevel::Num || ExistingShaderMapId.QualityLevel == EMaterialQualityLevel::Num
+								   || ShaderMapId.QualityLevel == ExistingShaderMapId.QualityLevel);
+
+		if (bFeatureLevelMatch && bQualityLevelMatch)
+		{
+			return MaterialShaderMap;
+		}
+	}
+
+	return nullptr;
 }
 
 void FODSCThread::AddRequest(const TArray<FString>& MaterialsToCompile, const FString& ShaderTypesToLoad, EShaderPlatform ShaderPlatform, ERHIFeatureLevel::Type FeatureLevel, EMaterialQualityLevel::Type QualityLevel, ODSCRecompileCommand RecompileCommandType)
@@ -278,7 +338,7 @@ void FODSCThread::AddShaderPipelineRequest(
 	EShaderPlatform ShaderPlatform,
 	ERHIFeatureLevel::Type FeatureLevel,
 	EMaterialQualityLevel::Type QualityLevel,
-	const FString& MaterialName,
+	const FMaterial* Material,
 	const FString& VertexFactoryName,
 	const FString& PipelineName,
 	const TArray<FString>& ShaderTypeNames,
@@ -290,11 +350,18 @@ void FODSCThread::AddShaderPipelineRequest(
 	{
 		FWriteScopeLock WriteLock(RequestHashesRWLock);
 
+		FName& CachedMaterialName = ODSCPointerToNames.FindOrAdd((UPTRINT)Material);
+		if (CachedMaterialName.IsNone())
+		{
+			CachedMaterialName = FName(Material->GetFullPath());
+		}
+		
+		FODSCShaderMapData& ODSCShaderMapData = RequestHashes.FindOrAdd(CachedMaterialName);
+
 		for (const FShaderId& ShaderId : RequestShaderIds)
 		{
-			FMaterialRequestsHashes& MaterialRequestHashes = RequestHashes.FindOrAdd(ShaderId);
 			bool bAlreadyInSet = false;
-			MaterialRequestHashes.RequestStrings.Add(MaterialName, &bAlreadyInSet);
+			ODSCShaderMapData.CurrentRequests.Add(ShaderId, &bAlreadyInSet);
 			if (!bAlreadyInSet)
 			{
 				bShouldAddRequest = true;
@@ -306,6 +373,7 @@ void FODSCThread::AddShaderPipelineRequest(
 	{
 		SCOPED_NAMED_EVENT(AddShaderPipelineRequest_AddRequest, FColor::Emerald);
 
+		FString MaterialName = Material->GetFullPath();
 		FString RequestString = (MaterialName + VertexFactoryName + PipelineName);
 		for (const auto& ShaderTypeName : ShaderTypeNames)
 		{
@@ -316,27 +384,6 @@ void FODSCThread::AddShaderPipelineRequest(
 	}
 }
 
-void FODSCThread::RegisterMaterialShaderMap(const FMaterialShaderMap& MaterialShaderMap)
-{
-	FWriteScopeLock WriteLock(RequestHashesRWLock);
-
-	TMap<FShaderId, TShaderRef<FShader>> ShadersInMap;
-	MaterialShaderMap.GetShaderList(ShadersInMap);
-	FSHAHash CookedShaderMapIdHash = MaterialShaderMap.GetShaderMapId().CookedShaderMapIdHash;
-	for (auto Iter : ShadersInMap)
-	{
-		// GetShaderList doesn't use the Cookedshadermap id
-		const FShaderId& ShaderIdSrc = Iter.Key;
-		FShaderId ShaderIdCopy(ShaderIdSrc.Type, CookedShaderMapIdHash, ShaderIdSrc.ShaderPipelineName, ShaderIdSrc.VFType, ShaderIdSrc.PermutationId, (EShaderPlatform)ShaderIdSrc.Platform);
-
-		// The shadermap we receive contains all the requests the client sent until now, so it's possible they got removed from RequestHashes already
-		if (RequestHashes.Find(ShaderIdCopy))
-		{
-			RequestHashes.Remove(ShaderIdCopy);
-		}
-    }
-}
-
 void FODSCThread::GetCompletedRequests(TArray<FODSCMessageHandler*>& OutCompletedRequests)
 {
 	check(IsInGameThread());
diff --git a/Engine/Source/Runtime/Engine/Private/ODSC/ODSCThread.h b/Engine/Source/Runtime/Engine/Private/ODSC/ODSCThread.h
--- a/Engine/Source/Runtime/Engine/Private/ODSC/ODSCThread.h
+++ b/Engine/Source/Runtime/Engine/Private/ODSC/ODSCThread.h
@@ -13,6 +13,7 @@
 class FEvent;
 class FRunnableThread;
 class FMaterialShaderMap;
+class FMaterialShaderMapId;
 
 namespace UE
 {
@@ -139,7 +140,7 @@ public:
 		EShaderPlatform ShaderPlatform,
 		ERHIFeatureLevel::Type FeatureLevel,
 		EMaterialQualityLevel::Type QualityLevel,
-		const FString& MaterialName,
+		const FMaterial* Material,
 		const FString& VertexFactoryName,
 		const FString& PipelineName,
 		const TArray<FString>& ShaderTypeNames,
@@ -166,13 +167,15 @@ public:
 
 	bool GetPendingShaderData(bool& bOutIsConnectedToODSCServer, bool& bOutHasPendingGlobalShaders, uint32& OutNumPendingMaterialsRecompile, uint32& OutNumPendingMaterialsShaders) const;
 
-	void RegisterMaterialShaderMap(const FMaterialShaderMap& MaterialShaderMap);
-
 	void ResetMaterialsODSCData(ERHIFeatureLevel::Type FeatureLevel);
 
-	bool CheckIfRequestAlreadySent(const TArray<FShaderId>& RequestShaderIds, const FString& MaterialName) const;
+	bool CheckIfRequestAlreadySent(const TArray<FShaderId>& RequestShaderIds, const FMaterial* Material) const;
 	const FString& GetODSCHostIP() const { return ODSCHostIP; };
 
+	void UnregisterMaterialName(const FMaterial* Material);
+	void RegisterMaterialShaderMaps(const FString& MaterialName, const TArray<TRefCountPtr<FMaterialShaderMap>>& LoadedShaderMaps);
+	FMaterialShaderMap* FindMaterialShaderMap(const FString& MaterialName, const FMaterialShaderMapId& ShaderMapId) const;
+
 protected:
 
 	//~ Begin FRunnable Interface
@@ -218,19 +221,12 @@ private:
 	mutable FRWLock RequestHashesRWLock;
 	FCriticalSection RequestHashCriticalSection;
 
-
-	struct FMaterialRequestsHashes
-	{
-		TSet<FString> RequestStrings;
-	};
-
 	struct FODSCShaderId
 	{
 	public:
 		inline FODSCShaderId() {}
 		FODSCShaderId(const FShaderId& ShaderId);
 
-		FSHAHash MaterialShaderMapHash;
 		FHashedName ShaderTypeHashedName = 0;
 		FHashedName VFTypeHashedName = 0;
 		FHashedName ShaderPipelineName = 0;
@@ -243,14 +239,12 @@ private:
 				GetTypeHash(Id.ShaderTypeHashedName),
 				HashCombine(GetTypeHash(Id.VFTypeHashedName),
 							HashCombine(GetTypeHash(Id.ShaderPipelineName),
-										HashCombine(GetTypeHash(Id.MaterialShaderMapHash),
-													HashCombine(GetTypeHash(Id.PermutationId), GetTypeHash(Id.Platform))))));
+										HashCombine(GetTypeHash(Id.PermutationId), GetTypeHash(Id.Platform)))));
 		}
 
 		friend bool operator==(const FODSCShaderId& X, const FODSCShaderId& Y)
 		{
 			return X.ShaderTypeHashedName == Y.ShaderTypeHashedName
-			&& X.MaterialShaderMapHash == Y.MaterialShaderMapHash
 			&& X.ShaderPipelineName == Y.ShaderPipelineName
 			&& X.VFTypeHashedName == Y.VFTypeHashedName
 			&& X.PermutationId == Y.PermutationId 
@@ -264,8 +258,19 @@ private:
 	};
 
 
-	/** Hashes for all Pending or Completed requests.  This is so we avoid making the same request multiple times. */
-	TMap<FODSCShaderId, FMaterialRequestsHashes> RequestHashes;
+	struct FODSCShaderMapData
+	{
+    	/** All the shadermaps owned by the material (quality level / feature level) */
+		TArray<TRefCountPtr<FMaterialShaderMap>> MaterialShaderMaps;
+    	/** Hashes for all Pending or Completed requests.  This is so we avoid making the same request multiple times. */
+		TSet<FODSCShaderId> CurrentRequests;
+	};
+
+    /** Requests seen for a given material name */
+	TMap<FName, FODSCShaderMapData> RequestHashes;
+
+    /** FMaterial* -> FName cache to avoid the expensive operation of calling FMaterialResource::GetFullPath and convert it to FName */
+	TMap<UPTRINT, FName> ODSCPointerToNames; 
 
 	/** Pointer to Runnable Thread */
 	FRunnableThread* Thread = nullptr;
diff --git a/Engine/Source/Runtime/Engine/Public/ODSC/ODSCManager.h b/Engine/Source/Runtime/Engine/Public/ODSC/ODSCManager.h
--- a/Engine/Source/Runtime/Engine/Public/ODSC/ODSCManager.h
+++ b/Engine/Source/Runtime/Engine/Public/ODSC/ODSCManager.h
@@ -11,6 +11,7 @@
 class FODSCThread;
 class UMaterialInstance;
 class FMaterialShaderMap;
+class FMaterialShaderMapId;
 
 /**
  * Responsible for processing shader compile responses from the ODSC Thread.
@@ -73,7 +74,7 @@ public:
 		EShaderPlatform ShaderPlatform,
 		ERHIFeatureLevel::Type FeatureLevel,
 		EMaterialQualityLevel::Type QualityLevel,
-		const FString& MaterialName,
+		const FMaterial* Material,
 		const FString& VertexFactoryName,
 		const FString& PipelineName,
 		const TArray<FString>& ShaderTypeNames,
@@ -95,7 +96,6 @@ public:
 
 	/** Returns true if we would actually add a request when calling AddThreadedShaderPipelineRequest. */
 	inline bool IsHandlingRequests() const { return Thread != nullptr; }
-	static void RegisterMaterialShaderMap(const FMaterialShaderMap& MaterialShaderMap);
 
 	static void RegisterMaterialInstance(const UMaterialInstance* MI);
 	static void UnregisterMaterialInstance(const UMaterialInstance* MI);
@@ -109,7 +109,11 @@ public:
 
 	static void ReportODSCError(const FString& InErrorMessage);
 
-	bool CheckIfRequestAlreadySent(const TArray<FShaderId>& RequestShaderIds, const FString& MaterialName) const;
+	bool CheckIfRequestAlreadySent(const TArray<FShaderId>& RequestShaderIds, const FMaterial* Material) const;
+
+	static void UnregisterMaterialName(const FMaterial* Material);
+	static void RegisterMaterialShaderMaps(const FString& MaterialName, const TArray<TRefCountPtr<FMaterialShaderMap>>& LoadedShaderMaps);
+	static FMaterialShaderMap* FindMaterialShaderMap(const FString& MaterialName, const FMaterialShaderMapId& ShaderMapId);
 
 private:
 
