Disable parameter association checks for MFs in the new layering system as they currently cause frequent asserts, but as our end goal is for MFs and MLs to be interchangeable, we will compensate for this at a later date anyway. Uses the layer prototyping macro so legacy behaviour remains unaffected.

diff --git a/Engine/Source/Runtime/Engine/Private/Materials/MaterialInstance.cpp b/Engine/Source/Runtime/Engine/Private/Materials/MaterialInstance.cpp
--- a/Engine/Source/Runtime/Engine/Private/Materials/MaterialInstance.cpp
+++ b/Engine/Source/Runtime/Engine/Private/Materials/MaterialInstance.cpp
@@ -293,7 +293,9 @@ bool FMaterialInstanceResource::GetParameterValue(EMaterialParameterType Type, c
 	// Check for hard-coded parameters
 	if (Type == EMaterialParameterType::Scalar && ParameterInfo.Name == GetSubsurfaceProfileParameterName())
 	{
+#if !ENABLE_MATERIAL_LAYER_PROTOTYPE
 		check(ParameterInfo.Association == EMaterialParameterAssociation::GlobalParameter);
+#endif
 		const USubsurfaceProfile* MySubsurfaceProfileRT = GetSubsurfaceProfileRT();
 		OutValue = GetSubsurfaceProfileId(MySubsurfaceProfileRT);
 		bResult = true;
@@ -304,7 +306,9 @@ bool FMaterialInstanceResource::GetParameterValue(EMaterialParameterType Type, c
 		{
 			if (ParameterInfo.Name == SpecularProfileAtlas::GetSpecularProfileParameterName(GetSpecularProfileRT(It)))
 			{
+#if !ENABLE_MATERIAL_LAYER_PROTOTYPE
 				check(ParameterInfo.Association == EMaterialParameterAssociation::GlobalParameter);
+#endif
 				OutValue = SpecularProfileAtlas::GetSpecularProfileId(GetSpecularProfileRT(It));
 				bResult = true;
 				break;
