[Shaders] Cleaning up remaining duplicates from splitting up ShaderCompiler.cpp. Appendix to CL 34822731.

#rnx

diff --git a/Engine/Source/Runtime/Engine/Private/ShaderCompiler/ShaderCompiler.cpp b/Engine/Source/Runtime/Engine/Private/ShaderCompiler/ShaderCompiler.cpp
--- a/Engine/Source/Runtime/Engine/Private/ShaderCompiler/ShaderCompiler.cpp
+++ b/Engine/Source/Runtime/Engine/Private/ShaderCompiler/ShaderCompiler.cpp
@@ -576,53 +576,6 @@ static bool CheckSingleJob(const FShaderCompileJob& SingleJob, TArray<FString>&
 };
 #endif // WITH_EDITOR
 
-/** Information tracked for each shader compile worker process instance. */
-struct FShaderCompileWorkerInfo
-{
-	/** Process handle of the worker app once launched.  Invalid handle means no process. */
-	FProcHandle WorkerProcess;
-
-	/** Tracks whether tasks have been issued to the worker. */
-	bool bIssuedTasksToWorker;	
-
-	/** Whether the worker has been launched for this set of tasks. */
-	bool bLaunchedWorker;
-
-	/** Tracks whether all tasks issued to the worker have been received. */
-	bool bComplete;
-
-	/** Whether this worker is available for new jobs. It will be false when shutting down the worker. */
-	bool bAvailable; 
-
-	/** Time at which the worker started the most recent batch of tasks. */
-	double StartTime;
-
-	/** Time at which the worker ended the most recent batch of tasks. */
-	double FinishTime = 0.0;
-
-	/** Jobs that this worker is responsible for compiling. */
-	TArray<FShaderCommonCompileJobPtr> QueuedJobs;
-
-	FShaderCompileWorkerInfo() :
-		bIssuedTasksToWorker(false),		
-		bLaunchedWorker(false),
-		bComplete(false),
-		bAvailable(true),
-		StartTime(0)
-	{
-	}
-
-	// warning: not virtual
-	~FShaderCompileWorkerInfo()
-	{
-		if(WorkerProcess.IsValid())
-		{
-			FPlatformProcess::TerminateProc(WorkerProcess);
-			FPlatformProcess::CloseProc(WorkerProcess);
-		}
-	}
-};
-
 FShaderCompilingManager* GShaderCompilingManager = nullptr;
 
 bool FShaderCompilingManager::AllTargetPlatformSupportsRemoteShaderCompiling()
