fix ODSC 'recompileshaders changed': since we remove the bIsFromODSC flag during that call, we need to call RegisterForODSC again if the existing shadermap is not set as an ODSC shadermap. We might end up re-requesting the same shaders otherwise

diff --git a/Engine/Source/Runtime/Engine/Private/Materials/MaterialShader.cpp b/Engine/Source/Runtime/Engine/Private/Materials/MaterialShader.cpp
--- a/Engine/Source/Runtime/Engine/Private/Materials/MaterialShader.cpp
+++ b/Engine/Source/Runtime/Engine/Private/Materials/MaterialShader.cpp
@@ -2105,11 +2105,13 @@ void FMaterialShaderMap::LoadForRemoteRecompile(FArchive& Ar, EShaderPlatform Sh
 
 				// If we already registered a material shadermap with the same ID, just re-use it. The material will do the same in FMaterial::CacheShaders anyway
 				TRefCountPtr<FMaterialShaderMap> ExistingShaderMap = FMaterialShaderMap::FindId(ShaderMap->GetShaderMapId(), ShaderPlatform);
-				if (ExistingShaderMap)
+#if WITH_ODSC
+				if (ExistingShaderMap && ExistingShaderMap->IsFromODSC())
 				{
 					ShaderMap = ExistingShaderMap;
 				}
 				else
+#endif
 				{
 					// Register in the global map
 					ShaderMap->RegisterForODSC(ShaderPlatform);
