Add fix and test case for testing MovieScene iterators are working with AutoRTFM

diff --git a/Engine/Source/Programs/AutoRTFMEngineTests/AutoRTFMEngineTests.Build.cs b/Engine/Source/Programs/AutoRTFMEngineTests/AutoRTFMEngineTests.Build.cs
--- a/Engine/Source/Programs/AutoRTFMEngineTests/AutoRTFMEngineTests.Build.cs
+++ b/Engine/Source/Programs/AutoRTFMEngineTests/AutoRTFMEngineTests.Build.cs
@@ -24,6 +24,7 @@ public class AutoRTFMEngineTests : ModuleRules
 				"MRMesh",
 				"MoviePlayer",
 				"MoviePlayerProxy",
+				"MovieScene",
 				"PreLoadScreen",
 				"ProfilerService",
 				"SessionServices",
diff --git a/Engine/Source/Runtime/Engine/Private/Tests/AutoRTFM/AutoRTFMMovieScene.cpp b/Engine/Source/Runtime/Engine/Private/Tests/AutoRTFM/AutoRTFMMovieScene.cpp
new file mode 100644
--- /dev/null
+++ b/Engine/Source/Runtime/Engine/Private/Tests/AutoRTFM/AutoRTFMMovieScene.cpp
@@ -0,0 +1,44 @@
+// Copyright Epic Games, Inc. All Rights Reserved.
+
+#include "Misc/AutomationTest.h"
+#include "Misc/App.h"
+#include "AutoRTFM/AutoRTFM.h"
+#include "EntitySystem/MovieSceneEntityManager.h"
+#include "EntitySystem/BuiltInComponentTypes.h"
+
+
+#if WITH_DEV_AUTOMATION_TESTS
+
+#define TEST_CHECK_TRUE(b) do \
+{ \
+	if (!(b)) \
+	{ \
+		FString String(FString::Printf(TEXT("FAILED: %s:%u"), TEXT(__FILE__), __LINE__)); \
+		ExecutionInfo.AddEvent(FAutomationEvent(EAutomationEventType::Info, String)); \
+		return false; \
+	} \
+} while(false)
+
+IMPLEMENT_SIMPLE_AUTOMATION_TEST(FAutoRTFMMovieSceneTests, "AutoRTFM + MovieScene", EAutomationTestFlags::EngineFilter | EAutomationTestFlags::ClientContext | EAutomationTestFlags::ServerContext | EAutomationTestFlags::CommandletContext)
+
+bool FAutoRTFMMovieSceneTests::RunTest(const FString& Parameters)
+{
+	using namespace UE::MovieScene;
+
+	FBuiltInComponentTypes* BuiltInComponents = FBuiltInComponentTypes::Get();
+	FEntityComponentFilter Filter = FEntityComponentFilter();
+
+	FEntityManager ManagerTest;
+	AutoRTFM::Transact([&]
+	{
+		ManagerTest.Contains(Filter);
+	});
+
+	TEST_CHECK_TRUE(true);
+
+	return true;
+}
+
+#undef TEST_CHECK_TRUE
+
+#endif //WITH_DEV_AUTOMATION_TESTS
diff --git a/Engine/Source/Runtime/MovieScene/Private/EntitySystem/MovieSceneEntityManager.cpp b/Engine/Source/Runtime/MovieScene/Private/EntitySystem/MovieSceneEntityManager.cpp
--- a/Engine/Source/Runtime/MovieScene/Private/EntitySystem/MovieSceneEntityManager.cpp
+++ b/Engine/Source/Runtime/MovieScene/Private/EntitySystem/MovieSceneEntityManager.cpp
@@ -9,6 +9,7 @@
 #include "Algo/Find.h"
 #include "UObject/StrongObjectPtr.h"
 #include "Containers/SortedMap.h"
+#include "AutoRTFM/AutoRTFM.h"
 
 #include "HAL/PlatformProcess.h"
 #include "Misc/FeedbackContext.h"
@@ -1049,13 +1050,25 @@ void FEntityManager::AccumulateMask(const FEntityComponentFilter& InFilter, FCom
 
 void FEntityManager::EnterIteration() const
 {
-	++IterationCount;
+	UE_AUTORTFM_OPEN(
+	{
+		++IterationCount;
+	});
+
+	AutoRTFM::PushOnAbortHandler(this, [this](){ --(this->IterationCount); });
 }
 
 void FEntityManager::ExitIteration() const
 {
 	checkSlow(static_cast<uint16>(IterationCount) > 0);
-	--IterationCount;
+
+	// We only ever call the exit after an enter so we should not have to handle the case where exit is called first in order to ++ on the iterator
+	AutoRTFM::PopOnAbortHandler(this);
+
+	UE_AUTORTFM_OPEN(
+	{
+		--IterationCount;
+	});
 }
 
 FEntityAllocation* FEntityManager::CreateEntityAllocation(const FComponentMask& EntityComponentMask, uint16 InitialCapacity, uint16 MaxCapacity, FEntityAllocation* MigrateComponentDataFrom)
