Adding optional acceleration to RepMovement, in support of option to have Characters replicating it to sim proxies. This defaults to off.   Includes Iris and non-Iris support.

diff --git a/Engine/Plugins/Runtime/ReplicationSystemTestPlugin/Source/Private/Tests/Serialization/TestRepMovementNetSerializer.cpp b/Engine/Plugins/Runtime/ReplicationSystemTestPlugin/Source/Private/Tests/Serialization/TestRepMovementNetSerializer.cpp
--- a/Engine/Plugins/Runtime/ReplicationSystemTestPlugin/Source/Private/Tests/Serialization/TestRepMovementNetSerializer.cpp
+++ b/Engine/Plugins/Runtime/ReplicationSystemTestPlugin/Source/Private/Tests/Serialization/TestRepMovementNetSerializer.cpp
@@ -11,8 +11,10 @@ UE::Net::FTestMessage& operator<<(UE::Net::FTestMessage& Message, const FRepMove
 		"AngularVelocity: " << Value.AngularVelocity << '\n' <<
 		"Location: " << Value.Location << '\n' <<
 		"Rotation: " << Value.Rotation.ToString() << '\n' <<
+		"Acceleration: " << Value.Acceleration << '\n' <<
 		"bSimulatedPhysicSleep: " << Value.bSimulatedPhysicSleep << '\n' <<
 		"bRepPhysics: " << Value.bRepPhysics << '\n' <<
+		"bRepAcceleration: " << Value.bRepAcceleration << '\n' <<
 		"ServerFrame: " << Value.ServerFrame << '\n' <<
 		"ServerPhysicsHandle: " << Value.ServerPhysicsHandle << '\n' <<
 		"LocationQuantizationLevel: " << (uint32)Value.LocationQuantizationLevel << '\n' <<
@@ -159,6 +161,26 @@ void FTestRepMovementNetSerializer::OneTimeClassInitialization()
 		Values.Add(ValueWithSimulatedPhysicSleep);
 	}
 
+	// Values with location, velocities, rotation, acceleration, ServerFrame, and physics handle
+	{
+		FRepMovement& Value = Values.Emplace_GetRef();
+		Value.LinearVelocity = FVector(11.33, -50.23, 41.05);
+		Value.AngularVelocity = FVector(599.23, -2001.05, 4010.01);
+		Value.Location = FVector(100342.05, -190200.01, 298000.43);
+		Value.Rotation = FRotator(99.16, 10.75, 0.01);
+		Value.bRepAcceleration = true;
+		Value.Acceleration = FVector(1501.11, 1.33, 10.11);
+
+		FRepMovement ValueWithServerFrame = Value;
+		ValueWithServerFrame.ServerFrame = 1908246315;
+		Values.Add(ValueWithServerFrame);
+
+		FRepMovement ValueWithPhysicsHandle = Value;
+		ValueWithPhysicsHandle.ServerPhysicsHandle = 4718;
+		Values.Add(ValueWithPhysicsHandle);
+	}
+
+
 	bIsClassInitialized = true;
 }
 
diff --git a/Engine/Source/Runtime/Core/Public/Misc/EngineNetworkCustomVersion.h b/Engine/Source/Runtime/Core/Public/Misc/EngineNetworkCustomVersion.h
--- a/Engine/Source/Runtime/Core/Public/Misc/EngineNetworkCustomVersion.h
+++ b/Engine/Source/Runtime/Core/Public/Misc/EngineNetworkCustomVersion.h
@@ -87,6 +87,8 @@ struct FEngineNetworkCustomVersion
 		CustomVersions = 32,                    // Bump version to switch to using custom versions
 		DynamicMontageSerialization = 33,		// Bump version to support dynamic montage serialization in the Gameplay Ability System
 		PredictionKeyBaseNotReplicated = 34,	// Bump version to stop FPredictionKey::Base from being replicated (it was unused).
+		RepMoveOptionalAcceleration = 35,		// Bump version to support serialization changes to RepMove for optional Acceleration
+
 
 		// -----<new versions can be added above this line>-------------------------------------------------
 		VersionPlusOne,
diff --git a/Engine/Source/Runtime/Engine/Classes/Engine/ReplicatedState.h b/Engine/Source/Runtime/Engine/Classes/Engine/ReplicatedState.h
--- a/Engine/Source/Runtime/Engine/Classes/Engine/ReplicatedState.h
+++ b/Engine/Source/Runtime/Engine/Classes/Engine/ReplicatedState.h
@@ -123,7 +123,7 @@ struct FRepMovement
 	UPROPERTY(Transient)
 	FVector LinearVelocity;
 
-	/** Velocity of rotation for component */
+	/** Velocity of rotation for component (only valid if bRepPhysics is set) */
 	UPROPERTY(Transient)
 	FVector AngularVelocity;
 	
@@ -135,6 +135,10 @@ struct FRepMovement
 	UPROPERTY(Transient)
 	FRotator Rotation;
 
+	/** Acceleration of component in world space. Only valid if bRepAcceleration is set. */
+	UPROPERTY(Transient)
+	FVector Acceleration;
+
 	/** If set, RootComponent should be sleeping. */
 	UPROPERTY(Transient)
 	uint8 bSimulatedPhysicSleep : 1;
@@ -143,6 +147,10 @@ struct FRepMovement
 	UPROPERTY(Transient)
 	uint8 bRepPhysics : 1;
 
+	/** If set, additional acceleration data will be replicated. */
+	UPROPERTY(Transient)
+	uint8 bRepAcceleration : 1;
+
 	/** Server physics step */
 	UPROPERTY(Transient)
 	int32 ServerFrame;
@@ -214,6 +222,11 @@ struct FRepMovement
 			return false;
 		}
 
+		if (bRepAcceleration && (Acceleration != Other.Acceleration))
+		{
+			return false;
+		}
+
 		if ( bSimulatedPhysicSleep != Other.bSimulatedPhysicSleep )
 		{
 			return false;
@@ -224,6 +237,11 @@ struct FRepMovement
 			return false;
 		}
 
+		if ( bRepAcceleration != Other.bRepAcceleration )
+		{
+			return false;
+		}
+
 		return true;
 	}
 
diff --git a/Engine/Source/Runtime/Engine/Classes/GameFramework/Character.h b/Engine/Source/Runtime/Engine/Classes/GameFramework/Character.h
--- a/Engine/Source/Runtime/Engine/Classes/GameFramework/Character.h
+++ b/Engine/Source/Runtime/Engine/Classes/GameFramework/Character.h
@@ -35,6 +35,12 @@ DECLARE_DYNAMIC_MULTICAST_DELEGATE_ThreeParams(FCharacterMovementUpdatedSignatur
 DECLARE_DYNAMIC_MULTICAST_DELEGATE(FCharacterReachedApexSignature);
 DECLARE_DYNAMIC_MULTICAST_DELEGATE_OneParam(FLandedSignature, const FHitResult&, Hit);
 
+// CVars
+namespace CharacterCVars
+{	
+	extern ENGINE_API int32 EnableCharacterAccelerationReplication;
+}
+
 /** Replicated data when playing a root motion montage. */
 USTRUCT()
 struct FRepRootMotionMontage
@@ -446,6 +452,9 @@ public:
 	ENGINE_API virtual void SetReplicateMovement(bool bInReplicateMovement) override;
 
 protected:
+	/** Whether this Character should include acceleration data in its replicated movement */
+	ENGINE_API virtual bool ShouldReplicateAcceleration() const { return CharacterCVars::EnableCharacterAccelerationReplication != 0; }
+
 	/** Saved translation offset of mesh. */
 	UPROPERTY()
 	FVector BaseTranslationOffset;
@@ -632,6 +641,7 @@ public:
 	ENGINE_API virtual void ClearCrossLevelReferences() override;
 	ENGINE_API virtual void PreNetReceive() override;
 	ENGINE_API virtual void PostNetReceive() override;
+	ENGINE_API virtual void GatherCurrentMovement() override;
 	ENGINE_API virtual void OnRep_ReplicatedMovement() override;
 	ENGINE_API virtual void PostNetReceiveLocationAndRotation() override;
 	ENGINE_API virtual void GetSimpleCollisionCylinder(float& CollisionRadius, float& CollisionHalfHeight) const override;
diff --git a/Engine/Source/Runtime/Engine/Private/Character.cpp b/Engine/Source/Runtime/Engine/Private/Character.cpp
--- a/Engine/Source/Runtime/Engine/Private/Character.cpp
+++ b/Engine/Source/Runtime/Engine/Private/Character.cpp
@@ -32,6 +32,18 @@ FName ACharacter::MeshComponentName(TEXT("CharacterMesh0"));
 FName ACharacter::CharacterMovementComponentName(TEXT("CharMoveComp"));
 FName ACharacter::CapsuleComponentName(TEXT("CollisionCylinder"));
 
+
+// CVars
+namespace CharacterCVars
+{
+	// Allows characters to include acceleration in the data replicated to sim proxies
+	int32 EnableCharacterAccelerationReplication = 0;
+	static FAutoConsoleVariableRef CVarEnableCharacterAccelerationReplication(
+		TEXT("p.EnableCharacterAccelerationReplication"),
+		EnableCharacterAccelerationReplication,
+		TEXT("Whether to author acceleration data with character movement replication to sim proxies."));
+}
+
 ACharacter::ACharacter(const FObjectInitializer& ObjectInitializer)
 : Super(ObjectInitializer)
 {
@@ -1253,6 +1265,26 @@ void ACharacter::OnRep_ReplicatedBasedMovement()
 	}
 }
 
+void ACharacter::GatherCurrentMovement()
+{
+	Super::GatherCurrentMovement();
+
+	if (IsReplicatingMovement())
+	{
+		FRepMovement& MutableRepMovement = GetReplicatedMovement_Mutable();
+
+		if (ShouldReplicateAcceleration()) 
+		{
+			MutableRepMovement.bRepAcceleration = true;
+			MutableRepMovement.Acceleration = CharacterMovement->GetCurrentAcceleration();
+		}
+		else
+		{
+			MutableRepMovement.bRepAcceleration = false;
+		}
+	}
+}
+
 void ACharacter::OnRep_ReplicatedMovement()
 {
 	// Skip standard position correction if we are playing root motion, OnRep_RootMotion will handle it.
diff --git a/Engine/Source/Runtime/Engine/Private/Components/CharacterMovementComponent.cpp b/Engine/Source/Runtime/Engine/Private/Components/CharacterMovementComponent.cpp
--- a/Engine/Source/Runtime/Engine/Private/Components/CharacterMovementComponent.cpp
+++ b/Engine/Source/Runtime/Engine/Private/Components/CharacterMovementComponent.cpp
@@ -6702,9 +6702,19 @@ void UCharacterMovementComponent::MoveSmooth(const FVector& InVelocity, const fl
 
 void UCharacterMovementComponent::UpdateProxyAcceleration()
 {
-	// Not currently replicated for simulated movement, but make it non-zero for animations that may want it, based on velocity.
-	Acceleration = Velocity.GetSafeNormal();
-	AnalogInputModifier = 1.0f;
+	const FRepMovement& ReplicatedMovement = CharacterOwner->GetReplicatedMovement();
+	if (ReplicatedMovement.bRepAcceleration)
+	{
+		Acceleration = ReplicatedMovement.Acceleration;
+		AnalogInputModifier = ComputeAnalogInputModifier();
+	}
+	else
+	{
+		// If acceleration isn't replicated for simulated movement, make it non-zero for animations that may want it, based on velocity.
+		// Note that this represents Acceleration with units in the range [0, 1] instead of representing cm/s^2.
+		Acceleration = Velocity.GetSafeNormal();
+		AnalogInputModifier = 1.0f;
+	}
 }
 
 bool UCharacterMovementComponent::IsWalkable(const FHitResult& Hit) const
diff --git a/Engine/Source/Runtime/Engine/Private/Engine/ReplicatedState.cpp b/Engine/Source/Runtime/Engine/Private/Engine/ReplicatedState.cpp
--- a/Engine/Source/Runtime/Engine/Private/Engine/ReplicatedState.cpp
+++ b/Engine/Source/Runtime/Engine/Private/Engine/ReplicatedState.cpp
@@ -26,8 +26,10 @@ FRepMovement::FRepMovement()
 	, AngularVelocity(ForceInit)
 	, Location(ForceInit)
 	, Rotation(ForceInit)
+	, Acceleration(ForceInit)
 	, bSimulatedPhysicSleep(false)
 	, bRepPhysics(false)
+	, bRepAcceleration(false)
 	, ServerFrame(0)
 	, LocationQuantizationLevel(EVectorQuantization::RoundWholeNumber)
 	, VelocityQuantizationLevel(EVectorQuantization::RoundWholeNumber)
@@ -128,6 +130,27 @@ bool FRepMovement::NetSerialize(FArchive& Ar, class UPackageMap* Map, bool& bOut
 		ServerPhysicsHandle = (int32)uServerPhysicsHandle;
 	}
 
+	if (Ar.EngineNetVer() >= FEngineNetworkCustomVersion::RepMoveOptionalAcceleration)
+	{
+		uint8 AccelFlags = (bRepAcceleration << 0);
+		Ar.SerializeBits(&AccelFlags, 1);
+		bRepAcceleration = (AccelFlags & (1 << 0)) ? 1 : 0;
+
+		if (bRepAcceleration)
+		{
+#if REP_MOVEMENT_DISABLE_QUANTIZATION
+			Ar << Acceleration;
+#else
+			// Note that we're using the same quantization as Velocity, since the units are commonly on the same order
+			bOutSuccess &= SerializeQuantizedVector(Ar, Acceleration, VelocityQuantizationLevel);
+#endif
+		}
+	}
+	else if (Ar.IsLoading())
+	{
+		bRepAcceleration = false;
+	}
+
 	return true;
 }
 
diff --git a/Engine/Source/Runtime/Engine/Private/RepMovementNetSerializer.cpp b/Engine/Source/Runtime/Engine/Private/RepMovementNetSerializer.cpp
--- a/Engine/Source/Runtime/Engine/Private/RepMovementNetSerializer.cpp
+++ b/Engine/Source/Runtime/Engine/Private/RepMovementNetSerializer.cpp
@@ -31,6 +31,7 @@ struct FRepMovementNetSerializer
 		uint64 LinearVelocity[4];
 		uint64 Location[4];
 		uint16 Rotation[4];
+		uint64 Acceleration[4];
 		int32 ServerFrame;
 		int32 ServerPhysicsHandle;
 
@@ -38,7 +39,8 @@ struct FRepMovementNetSerializer
 		uint16 VelocityQuantizationLevel : 2;
 		uint16 LocationQuantizationLevel : 2;
 		uint16 RotationQuantizationLevel : 1;
-		uint16 Unused : 9;
+		uint16 RepAcceleration : 1;
+		uint16 Unused : 6;
 		uint16 Padding[3];
 	};
 
@@ -189,6 +191,23 @@ void FRepMovementNetSerializer::Serialize(FNetSerializationContext& Context, con
 	{
 		WritePackedInt32(Writer, Value.ServerPhysicsHandle);
 	}
+
+	// Optional acceleration
+	{
+		// FUTURE: Consider checking EngineNetVer >= FEngineNetworkCustomVersion::RepMoveOptionalAcceleration for backwards compatibility
+		if (Writer->WriteBool(Value.RepAcceleration))
+		{
+			UE_NET_TRACE_SCOPE(Acceleration, *Writer, Context.GetTraceCollector(), ENetTraceVerbosity::Verbose);
+			const FNetSerializer* Serializer = VectorNetQuantizeNetSerializers[Value.VelocityQuantizationLevel];	// Note that we're using the same quantization as Velocity, since the numbers are typically on the same order
+			const FNetSerializerConfig* SerializerConfig = VectorNetQuantizeNetSerializerConfigs[Value.VelocityQuantizationLevel];
+
+			FNetSerializeArgs MemberArgs = Args;
+			MemberArgs.NetSerializerConfig = NetSerializerConfigParam(SerializerConfig);
+			MemberArgs.Source = NetSerializerValuePointer(&Value.Acceleration[0]);
+			Serializer->Serialize(Context, MemberArgs);
+		}
+	}
+
 }
 
 void FRepMovementNetSerializer::Deserialize(FNetSerializationContext& Context, const FNetDeserializeArgs& Args)
@@ -282,6 +301,25 @@ void FRepMovementNetSerializer::Deserialize(FNetSerializationContext& Context, c
 		TempValue.ServerPhysicsHandle = INDEX_NONE;
 	}
 
+	// Optional acceleration
+	{
+		// FUTURE: Consider checking EngineNetVer >= FEngineNetworkCustomVersion::RepMoveOptionalAcceleration for backwards compatibility
+		TempValue.RepAcceleration = Reader->ReadBool();
+
+		if (TempValue.RepAcceleration)
+		{
+			UE_NET_TRACE_SCOPE(Acceleration, *Reader, Context.GetTraceCollector(), ENetTraceVerbosity::Verbose);
+			const FNetSerializer* Serializer = VectorNetQuantizeNetSerializers[TempValue.VelocityQuantizationLevel];	// Note that we're using the same quantization as Velocity, since the numbers are typically on the same order
+			const FNetSerializerConfig* SerializerConfig = VectorNetQuantizeNetSerializerConfigs[TempValue.VelocityQuantizationLevel];
+
+			FNetDeserializeArgs MemberArgs = Args;
+			MemberArgs.NetSerializerConfig = NetSerializerConfigParam(SerializerConfig);
+			MemberArgs.Target = NetSerializerValuePointer(&TempValue.Acceleration[0]);
+			Serializer->Deserialize(Context, MemberArgs);
+		}
+	}
+
+
 	QuantizedType& Target = *reinterpret_cast<QuantizedType*>(Args.Target);
 	Target = TempValue;
 }
@@ -379,6 +417,33 @@ void FRepMovementNetSerializer::SerializeDelta(FNetSerializationContext& Context
 	{
 		SerializeIntDelta(*Writer, Value.ServerPhysicsHandle, PrevValue.ServerPhysicsHandle, DeltaBitCountTable, DeltaBitCountTableEntryCount, 32U);
 	}
+
+	// Optional acceleration
+	{
+		// FUTURE: Consider checking EngineNetVer >= FEngineNetworkCustomVersion::RepMoveOptionalAcceleration for backwards compatibility
+		Writer->WriteBool(Value.RepAcceleration);
+
+		if (Value.RepAcceleration)
+		{
+			UE_NET_TRACE_SCOPE(Acceleration, *Writer, Context.GetTraceCollector(), ENetTraceVerbosity::Verbose);
+			const FNetSerializer* Serializer = VectorNetQuantizeNetSerializers[Value.VelocityQuantizationLevel];
+			const FNetSerializerConfig* SerializerConfig = VectorNetQuantizeNetSerializerConfigs[Value.VelocityQuantizationLevel];
+
+			FNetSerializeDeltaArgs MemberArgs = Args;
+			MemberArgs.NetSerializerConfig = NetSerializerConfigParam(SerializerConfig);
+			MemberArgs.Source = NetSerializerValuePointer(&Value.Acceleration[0]);
+			MemberArgs.Prev = NetSerializerValuePointer(&PrevValue.Acceleration[0]);
+
+			if (PrevValue.RepAcceleration)
+			{
+				Serializer->SerializeDelta(Context, MemberArgs);
+			}
+			else // since previously we have no acceleration data, fully serialize it
+			{
+				Serializer->Serialize(Context, MemberArgs);
+			}
+		}
+	}
 }
 
 void FRepMovementNetSerializer::DeserializeDelta(FNetSerializationContext& Context, const FNetDeserializeDeltaArgs& Args)
@@ -486,6 +551,35 @@ void FRepMovementNetSerializer::DeserializeDelta(FNetSerializationContext& Conte
 	FixedUpFlags |= (TempValue.ServerPhysicsHandle != INDEX_NONE ? Flag_ServerPhysicsHandleIsPresent : 0U);
 	TempValue.Flags = FixedUpFlags;
 
+	// Optional acceleration
+	{
+		// FUTURE: Consider checking EngineNetVer >= FEngineNetworkCustomVersion::RepMoveOptionalAcceleration for backwards compatibility
+		UE_NET_TRACE_SCOPE(Acceleration, *Reader, Context.GetTraceCollector(), ENetTraceVerbosity::Verbose);
+		const FNetSerializer* Serializer = VectorNetQuantizeNetSerializers[TempValue.VelocityQuantizationLevel];
+		const FNetSerializerConfig* SerializerConfig = VectorNetQuantizeNetSerializerConfigs[TempValue.VelocityQuantizationLevel];
+
+		FNetDeserializeDeltaArgs MemberArgs = Args;
+		MemberArgs.NetSerializerConfig = NetSerializerConfigParam(SerializerConfig);
+		MemberArgs.Target = NetSerializerValuePointer(&TempValue.Acceleration[0]);
+		MemberArgs.Prev = NetSerializerValuePointer(&PrevValue.Acceleration[0]);
+
+		TempValue.RepAcceleration = Reader->ReadBool();
+
+		if (TempValue.RepAcceleration)
+		{
+			if (PrevValue.RepAcceleration)
+			{
+				Serializer->DeserializeDelta(Context, MemberArgs);
+			}
+			else
+			{
+				// No prior value, so we deserialize the whole value
+				Serializer->Deserialize(Context, MemberArgs);
+			}
+		}
+	}
+
+
 	QuantizedType& Target = *reinterpret_cast<QuantizedType*>(Args.Target);
 	Target = TempValue;
 }
@@ -503,6 +597,7 @@ void FRepMovementNetSerializer::Quantize(FNetSerializationContext& Context, cons
 	TempValue.RotationQuantizationLevel = uint8(Source.RotationQuantizationLevel);
 	TempValue.ServerFrame = Source.ServerFrame;
 	TempValue.ServerPhysicsHandle = Source.ServerPhysicsHandle;
+	TempValue.RepAcceleration = Source.bRepAcceleration;
 	
 	// Angular velocity. Do note that the quantized value is cleared if RepPhysics is false. Delta compression need to accommodate to this.
 	if (TempValue.Flags & Flag_RepPhysics)
@@ -553,6 +648,20 @@ void FRepMovementNetSerializer::Quantize(FNetSerializationContext& Context, cons
 		Serializer->Quantize(Context, MemberArgs);
 	}
 
+	// Optional acceleration
+	if (TempValue.RepAcceleration)
+	{
+		// Note that we're using the same quantization settings as linear velocity, since the values are typically on the same order
+		const FNetSerializer* Serializer = VectorNetQuantizeNetSerializers[TempValue.VelocityQuantizationLevel];
+		const FNetSerializerConfig* SerializerConfig = VectorNetQuantizeNetSerializerConfigs[TempValue.VelocityQuantizationLevel];
+
+		FNetQuantizeArgs MemberArgs = Args;
+		MemberArgs.NetSerializerConfig = NetSerializerConfigParam(SerializerConfig);
+		MemberArgs.Source = NetSerializerValuePointer(&Source.Acceleration);
+		MemberArgs.Target = NetSerializerValuePointer(&TempValue.Acceleration[0]);
+		Serializer->Quantize(Context, MemberArgs);
+	}
+
 	QuantizedType& Target = *reinterpret_cast<QuantizedType*>(Args.Target);
 	Target = TempValue;
 }
@@ -569,6 +678,7 @@ void FRepMovementNetSerializer::Dequantize(FNetSerializationContext& Context, co
 	Target.LocationQuantizationLevel = EVectorQuantization(Source.LocationQuantizationLevel);
 	Target.VelocityQuantizationLevel = EVectorQuantization(Source.VelocityQuantizationLevel);
 	Target.RotationQuantizationLevel = ERotatorQuantization(Source.RotationQuantizationLevel);
+	Target.bRepAcceleration = Source.RepAcceleration;
 
 	// There's no need to dequantize angular velocity unless we're replicating it.
 	if (Source.Flags & Flag_RepPhysics)
@@ -618,6 +728,20 @@ void FRepMovementNetSerializer::Dequantize(FNetSerializationContext& Context, co
 		MemberArgs.Target = NetSerializerValuePointer(&Target.Rotation);
 		Serializer->Dequantize(Context, MemberArgs);
 	}
+
+	// Optional acceleration - we do not need to dequantize it if it isn't replicated
+	if (Source.RepAcceleration)
+	{
+		const FNetSerializer* Serializer = VectorNetQuantizeNetSerializers[Source.VelocityQuantizationLevel];
+		const FNetSerializerConfig* SerializerConfig = VectorNetQuantizeNetSerializerConfigs[Source.VelocityQuantizationLevel];
+
+		FNetDequantizeArgs MemberArgs = Args;
+		MemberArgs.NetSerializerConfig = NetSerializerConfigParam(SerializerConfig);
+		MemberArgs.Source = NetSerializerValuePointer(&Source.Acceleration[0]);
+		MemberArgs.Target = NetSerializerValuePointer(&Target.Acceleration);
+		Serializer->Dequantize(Context, MemberArgs);
+	}
+
 }
 
 bool FRepMovementNetSerializer::IsEqual(FNetSerializationContext& Context, const FNetIsEqualArgs& Args)
@@ -680,6 +804,14 @@ bool FRepMovementNetSerializer::Validate(FNetSerializationContext& Context, cons
 		return false;
 	}
 
+	if (Value.bRepAcceleration)
+	{
+		if (Value.Acceleration.ContainsNaN())
+		{
+			return false;
+		}
+	}
+
 	// Grab an arbitrary RotatorNetSerializer for the rotation validation.
 	{
 		const FNetSerializer* Serializer = RotatorNetSerializers[0];
@@ -705,6 +837,8 @@ void FRepMovementNetSerializer::Apply(FNetSerializationContext& Context, const F
 	Target.bSimulatedPhysicSleep = Source.bSimulatedPhysicSleep;
 	Target.bRepPhysics = Source.bRepPhysics;
 	Target.ServerFrame = Source.ServerFrame;
+	Target.bRepAcceleration = Source.bRepAcceleration;
+
 	// FRepMovement::NetSerialize does not overwrite ServerPhysicsHandle if it wasn't replicated. Despite the member name we mimic the behavior.
 	if (Source.ServerPhysicsHandle != INDEX_NONE)
 	{
@@ -721,6 +855,12 @@ void FRepMovementNetSerializer::Apply(FNetSerializationContext& Context, const F
 		Target.AngularVelocity = Source.AngularVelocity;
 	}
 
+	// Explicitly do not overwrite Acceleration unless it was replicated.
+	if (Source.bRepAcceleration)
+	{	
+		Target.Acceleration = Source.Acceleration;
+	}
+
 	// Explicitly do not overwrite quantization levels.
 }
 
@@ -767,6 +907,11 @@ bool FRepMovementNetSerializer::FNetSerializerRegistryDelegates::QuantizedTypeMe
 		{
 			return false;
 		}
+
+		if (!ensure(sizeof(QuantizedType::Acceleration) >= VectorNetSerializer.QuantizedTypeSize && alignof(decltype(QuantizedType::Acceleration)) >= VectorNetSerializer.QuantizedTypeAlignment))
+		{
+			return false;
+		}
 	}
 
 	// Check rotator serializer requirements.
@@ -786,7 +931,7 @@ bool FRepMovementNetSerializer::FNetSerializerRegistryDelegates::IsRepMovementLa
 {
 	const UStruct* RepMovementStruct = FRepMovement::StaticStruct();
 
-	constexpr int ExpectedPropertiesSize = 112;
+	constexpr int ExpectedPropertiesSize = 136;
 	if (!ensureMsgf(RepMovementStruct->GetPropertiesSize() == ExpectedPropertiesSize, TEXT("Unexpected FRepMovement properties size. %d != %d"), RepMovementStruct->GetPropertiesSize(), ExpectedPropertiesSize))
 	{
 		return false;
@@ -798,7 +943,7 @@ bool FRepMovementNetSerializer::FNetSerializerRegistryDelegates::IsRepMovementLa
 		{
 			"LinearVelocity", "AngularVelocity", "Location", "Rotation", "bSimulatedPhysicSleep", "bRepPhysics", 
 			"LocationQuantizationLevel", "VelocityQuantizationLevel", "RotationQuantizationLevel", 
-			"ServerFrame", "ServerPhysicsHandle", 
+			"ServerFrame", "ServerPhysicsHandle", "bRepAcceleration", "Acceleration"
 		};
 		constexpr SIZE_T FieldCount = sizeof(Fields)/sizeof(Fields[0]);
 
