Fix crash in FRayTracingGeometryManager::ProcessCompletedStreamingRequests(...)
- StreamingRequests were not cancelled/reset when ray tracing geometry were released ending up with stale references.

diff --git a/Engine/Source/Runtime/Engine/Private/Rendering/RayTracingGeometryManager.cpp b/Engine/Source/Runtime/Engine/Private/Rendering/RayTracingGeometryManager.cpp
--- a/Engine/Source/Runtime/Engine/Private/Rendering/RayTracingGeometryManager.cpp
+++ b/Engine/Source/Runtime/Engine/Private/Rendering/RayTracingGeometryManager.cpp
@@ -281,6 +281,17 @@ void FRayTracingGeometryManager::ReleaseRayTracingGeometryHandle(RayTracingGeome
 
 	FRegisteredGeometry& RegisteredGeometry = RegisteredGeometries[Handle];
 
+	// Cancel associated streaming request if currently in-flight
+	if (RegisteredGeometry.StreamingRequestIndex != INDEX_NONE)
+	{
+		FStreamingRequest& StreamingRequest = StreamingRequests[RegisteredGeometry.StreamingRequestIndex];
+		check(StreamingRequest.GeometryHandle == Handle);
+
+		StreamingRequest.Reset();
+
+		RegisteredGeometry.StreamingRequestIndex = INDEX_NONE;
+	}
+
 	if (RegisteredGeometry.Geometry->GroupHandle != INDEX_NONE)
 	{
 		// if geometry was assigned to a group, clear the relevant entry so another geometry can be registered later
@@ -899,6 +910,8 @@ bool FRayTracingGeometryManager::RequestRayTracingGeometryStreamIn(FRHICommandLi
 			return false;
 		}
 
+		RegisteredGeometry.StreamingRequestIndex = NextStreamingRequestIndex;
+
 		FStreamingRequest& StreamingRequest = StreamingRequests[NextStreamingRequestIndex];
 		NextStreamingRequestIndex = (NextStreamingRequestIndex + 1) % GRayTracingStreamingMaxPendingRequests;
 		++NumStreamingRequests;
@@ -952,12 +965,20 @@ void FRayTracingGeometryManager::ProcessCompletedStreamingRequests(FRHICommandLi
 		const int32 PendingRequestIndex = (StartPendingRequestIndex + Index) % GRayTracingStreamingMaxPendingRequests;
 		FStreamingRequest& PendingRequest = StreamingRequests[PendingRequestIndex];
 
+		if (!PendingRequest.IsValid())
+		{
+			++NumCompletedRequests;
+			continue;
+		}
+
 		if (PendingRequest.Request.IsCompleted())
 		{
 			++NumCompletedRequests;
 
 			FRegisteredGeometry& RegisteredGeometry = RegisteredGeometries[PendingRequest.GeometryHandle];
 
+			RegisteredGeometry.StreamingRequestIndex = INDEX_NONE;
+
 			const FRayTracingGeometryGroup& Group = RegisteredGroups[RegisteredGeometry.Geometry->GroupHandle];
 
 			if (RegisteredGeometry.Geometry->IsEvicted() || RegisteredGeometry.Geometry->LODIndex < Group.CurrentFirstLODIdx)
@@ -1016,6 +1037,8 @@ void FRayTracingGeometryManager::ProcessCompletedStreamingRequests(FRHICommandLi
 
 				RegisteredGeometry.Geometry->RequestBuildIfNeeded(RHICmdList, ERTAccelerationStructureBuildPriority::Normal);
 			}
+
+			PendingRequest.Reset();
 		}
 		else
 		{
diff --git a/Engine/Source/Runtime/Engine/Public/Rendering/RayTracingGeometryManager.h b/Engine/Source/Runtime/Engine/Public/Rendering/RayTracingGeometryManager.h
--- a/Engine/Source/Runtime/Engine/Public/Rendering/RayTracingGeometryManager.h
+++ b/Engine/Source/Runtime/Engine/Public/Rendering/RayTracingGeometryManager.h
@@ -112,6 +112,8 @@ private:
 		uint32 StreamableDataOffset = 0;
 		uint32 StreamableDataSize = 0;
 
+		int16 StreamingRequestIndex = INDEX_NONE;
+
 		enum class FStatus : uint8
 		{
 			StreamedOut,
@@ -128,6 +130,21 @@ private:
 		FBulkDataBatchRequest Request;
 
 		RayTracingGeometryHandle GeometryHandle;
+
+		bool IsValid() const { return GeometryHandle != INDEX_NONE; }
+
+		void Reset()
+		{
+			GeometryHandle = INDEX_NONE;
+
+			if (Request.IsPending())
+			{
+				Request.Cancel();
+				Request.Wait(); // Even after calling Cancel(), we still need to Wait() before we can touch the RequestBuffer.
+			}
+			Request.Reset();
+			RequestBuffer = {};
+		}
 	};
 
 	TSparseArray<FRayTracingGeometryGroup> RegisteredGroups;
