Expose acceleration mode setting to constraint component & constraint instance

diff --git a/Engine/Source/Runtime/Engine/Classes/PhysicsEngine/ConstraintDrives.h b/Engine/Source/Runtime/Engine/Classes/PhysicsEngine/ConstraintDrives.h
--- a/Engine/Source/Runtime/Engine/Classes/PhysicsEngine/ConstraintDrives.h
+++ b/Engine/Source/Runtime/Engine/Classes/PhysicsEngine/ConstraintDrives.h
@@ -82,6 +82,9 @@ struct FLinearDriveConstraint
 
 	UPROPERTY(EditAnywhere, Category = LinearMotor)
 	FConstraintDrive ZDrive;
+	
+	UPROPERTY(EditAnywhere, Category = AngularMotor)
+	bool bAccelerationMode;
 
 #if WITH_EDITORONLY_DATA
 	PRAGMA_DISABLE_DEPRECATION_WARNINGS
@@ -112,6 +115,8 @@ private:
 	ENGINE_API void SetDriveParams(const FVector& InStiffness, const FVector& InDamping, const FVector& InForceLimit);
 	ENGINE_API void GetDriveParams(float& OutStiffness, float& OutDamping, float& OutForceLimit) const;
 	ENGINE_API void GetDriveParams(FVector& OutStiffness, FVector& OutDamping, FVector& OutForceLimit) const;
+	ENGINE_API void SetAccelerationMode(bool bAccelerationMode);
+	ENGINE_API bool GetAccelerationMode() const;
 };
 
 
@@ -145,6 +150,9 @@ struct FAngularDriveConstraint
 	UPROPERTY(EditAnywhere, Category = AngularMotor)
 	TEnumAsByte<enum EAngularDriveMode::Type> AngularDriveMode;
 
+	UPROPERTY(EditAnywhere, Category = AngularMotor)
+	bool bAccelerationMode;
+
 	ENGINE_API FAngularDriveConstraint();
 
 	bool IsOrientationDriveEnabled() const
@@ -192,4 +200,6 @@ private:
 	// Gets drive parameters in the order swing, twist, slerp
 	ENGINE_API void GetDriveParams(FVector& OutStiffness, FVector& OutDamping, FVector& OutForceLimit) const;
 	ENGINE_API void SetAngularDriveMode(EAngularDriveMode::Type DriveMode);
+	ENGINE_API void SetAccelerationMode(bool bAccelerationMode);
+	ENGINE_API bool GetAccelerationMode() const;
 };
diff --git a/Engine/Source/Runtime/Engine/Classes/PhysicsEngine/ConstraintInstance.h b/Engine/Source/Runtime/Engine/Classes/PhysicsEngine/ConstraintInstance.h
--- a/Engine/Source/Runtime/Engine/Classes/PhysicsEngine/ConstraintInstance.h
+++ b/Engine/Source/Runtime/Engine/Classes/PhysicsEngine/ConstraintInstance.h
@@ -785,6 +785,9 @@ public:
 	/** Set the linear drive's strength parameters per-axis */
 	ENGINE_API void SetLinearDriveParams(const FVector& InPositionStrength, const FVector& InVelocityStrength, const FVector& InForceLimit);
 
+	/** Set the linear drive force-mode */
+	ENGINE_API void SetLinearDriveAccelerationMode(bool bAccelerationMode);
+
 	/** Get the linear drive's strength parameters. Assumes all axes are the same so only returns the X values */
 	ENGINE_API void GetLinearDriveParams(float& OutPositionStrength, float& OutVelocityStrength, float& OutForceLimit);
 
@@ -860,6 +863,9 @@ public:
 	/** Set the angular drive mode */
 	ENGINE_API void SetAngularDriveMode(EAngularDriveMode::Type DriveMode);
 
+	/** Set the angular drive force-mode */
+	ENGINE_API void SetAngularDriveAccelerationMode(bool bAccelerationMode);
+
 	/** Set the angular drive mode */
 	EAngularDriveMode::Type GetAngularDriveMode()
 	{
diff --git a/Engine/Source/Runtime/Engine/Classes/PhysicsEngine/PhysicsConstraintComponent.h b/Engine/Source/Runtime/Engine/Classes/PhysicsEngine/PhysicsConstraintComponent.h
--- a/Engine/Source/Runtime/Engine/Classes/PhysicsEngine/PhysicsConstraintComponent.h
+++ b/Engine/Source/Runtime/Engine/Classes/PhysicsEngine/PhysicsConstraintComponent.h
@@ -190,6 +190,10 @@ public:
 	UFUNCTION(BlueprintCallable, Category = "Physics|Components|PhysicsConstraint")
 	ENGINE_API void SetAngularDriveMode(EAngularDriveMode::Type DriveMode);
 
+	/** Toggles the angular drive acceleration mode. When enabled, the acceleration of the angular drive is not affected by the inertia of the constrained objects around the drive axis. */
+	UFUNCTION(BlueprintCallable, Category = "Physics|Components|PhysicsConstraint")
+	ENGINE_API void SetAngularDriveAccelerationMode(bool bAccelerationMode);
+
 	/** Sets the target position for the linear drive. 
 	 *	@param InPosTarget		Target position
 	 */
@@ -210,6 +214,10 @@ public:
 	UFUNCTION(BlueprintCallable, Category="Physics|Components|PhysicsConstraint")
 	ENGINE_API void SetLinearDriveParams(float PositionStrength, float VelocityStrength, float InForceLimit);
 
+	/** Toggles the linear drive acceleration mode. When enabled, the acceleration of the angular drive is not affected by the inertia of the constrained objects along the drive axis. */
+	UFUNCTION(BlueprintCallable, Category = "Physics|Components|PhysicsConstraint")
+	ENGINE_API void SetLinearDriveAccelerationMode(bool bAccelerationMode);
+
 	/** Sets the target orientation for the angular drive. 
 	 *	@param InPosTarget		Target orientation
 	 */
diff --git a/Engine/Source/Runtime/Engine/Private/PhysicsEngine/ConstraintDrives.cpp b/Engine/Source/Runtime/Engine/Private/PhysicsEngine/ConstraintDrives.cpp
--- a/Engine/Source/Runtime/Engine/Private/PhysicsEngine/ConstraintDrives.cpp
+++ b/Engine/Source/Runtime/Engine/Private/PhysicsEngine/ConstraintDrives.cpp
@@ -100,6 +100,16 @@ void FLinearDriveConstraint::GetDriveParams(FVector& OutStiffness, FVector& OutD
 	OutForceLimit.Set(XDrive.MaxForce, YDrive.MaxForce, ZDrive.MaxForce);
 }
 
+void FLinearDriveConstraint::SetAccelerationMode(const bool bInAccelerationMode)
+{
+	bAccelerationMode = bInAccelerationMode;
+}
+
+bool FLinearDriveConstraint::GetAccelerationMode() const
+{
+	return bAccelerationMode;
+}
+
 void FAngularDriveConstraint::SetAngularDriveMode(EAngularDriveMode::Type DriveMode)
 {
 	AngularDriveMode = DriveMode;
@@ -134,3 +144,14 @@ void FAngularDriveConstraint::GetDriveParams(FVector& OutStiffness, FVector& Out
 	OutForceLimit.Set(SwingDrive.MaxForce, TwistDrive.Damping, TwistDrive.MaxForce);
 
 }
+
+void FAngularDriveConstraint::SetAccelerationMode(const bool bInAccelerationMode)
+{
+	bAccelerationMode = bInAccelerationMode;
+}
+
+bool FAngularDriveConstraint::GetAccelerationMode() const
+{
+	return bAccelerationMode;
+}
+
diff --git a/Engine/Source/Runtime/Engine/Private/PhysicsEngine/ConstraintInstance.cpp b/Engine/Source/Runtime/Engine/Private/PhysicsEngine/ConstraintInstance.cpp
--- a/Engine/Source/Runtime/Engine/Private/PhysicsEngine/ConstraintInstance.cpp
+++ b/Engine/Source/Runtime/Engine/Private/PhysicsEngine/ConstraintInstance.cpp
@@ -989,9 +989,18 @@ void FConstraintInstance::SetLinearDriveParams(const FVector& InSpring, const FV
 	ProfileInstance.LinearDrive.SetDriveParams(InSpring, InDamping, InForceLimit);
 
 	FPhysicsInterface::ExecuteOnUnbrokenConstraintReadWrite(ConstraintHandle, [this](const FPhysicsConstraintHandle& InUnbrokenConstraint)
-		{
-			FPhysicsInterface::UpdateLinearDrive_AssumesLocked(InUnbrokenConstraint, ProfileInstance.LinearDrive);
-		});
+	{
+		FPhysicsInterface::UpdateLinearDrive_AssumesLocked(InUnbrokenConstraint, ProfileInstance.LinearDrive);
+	});
+}
+
+void FConstraintInstance::SetLinearDriveAccelerationMode(const bool bAccelerationMode)
+{
+	ProfileInstance.LinearDrive.SetAccelerationMode(bAccelerationMode);
+	FPhysicsInterface::ExecuteOnUnbrokenConstraintReadWrite(ConstraintHandle, [this](const FPhysicsConstraintHandle& InUnbrokenConstraint)
+	{
+		FPhysicsInterface::UpdateLinearDrive_AssumesLocked(InUnbrokenConstraint, ProfileInstance.LinearDrive);
+	});
 }
 
 /** Get the linear drive's strength parameters */
@@ -1086,6 +1095,15 @@ void FConstraintInstance::SetAngularDriveParams(float InSpring, float InDamping,
 	});
 }
 
+void FConstraintInstance::SetAngularDriveAccelerationMode(const bool bAccelerationMode)
+{
+	ProfileInstance.AngularDrive.SetAccelerationMode(bAccelerationMode);
+	FPhysicsInterface::ExecuteOnUnbrokenConstraintReadWrite(ConstraintHandle, [this](const FPhysicsConstraintHandle& InUnbrokenConstraint)
+	{
+		FPhysicsInterface::UpdateAngularDrive_AssumesLocked(InUnbrokenConstraint, ProfileInstance.AngularDrive);
+	});
+}
+
 void FConstraintInstance::GetAngularDriveParams(float& OutSpring, float& OutDamping, float& OutForceLimit) const
 {
 	ProfileInstance.AngularDrive.GetDriveParams(OutSpring, OutDamping, OutForceLimit);
diff --git a/Engine/Source/Runtime/Engine/Private/PhysicsEngine/Experimental/PhysInterface_Chaos.cpp b/Engine/Source/Runtime/Engine/Private/PhysicsEngine/Experimental/PhysInterface_Chaos.cpp
--- a/Engine/Source/Runtime/Engine/Private/PhysicsEngine/Experimental/PhysInterface_Chaos.cpp
+++ b/Engine/Source/Runtime/Engine/Private/PhysicsEngine/Experimental/PhysInterface_Chaos.cpp
@@ -386,7 +386,7 @@ void FPhysInterface_Chaos::UpdateLinearDrive_AssumesLocked(const FPhysicsConstra
 				Constraint->SetLinearDriveVelocityTarget(InDriveParams.VelocityTarget);
 			}
 
-			Constraint->SetLinearDriveForceMode(Chaos::EJointForceMode::Acceleration);
+			Constraint->SetLinearDriveForceMode(InDriveParams.bAccelerationMode ? Chaos::EJointForceMode::Acceleration : Chaos::EJointForceMode::Force);
 			Constraint->SetLinearDriveStiffness(Chaos::ConstraintSettings::LinearDriveStiffnessScale() * Chaos::FVec3(InDriveParams.XDrive.Stiffness, InDriveParams.YDrive.Stiffness, InDriveParams.ZDrive.Stiffness));
 			Constraint->SetLinearDriveDamping(Chaos::ConstraintSettings::LinearDriveDampingScale() * Chaos::FVec3(InDriveParams.XDrive.Damping, InDriveParams.YDrive.Damping, InDriveParams.ZDrive.Damping));
 			Constraint->SetLinearDriveMaxForce(Chaos::FVec3(InDriveParams.XDrive.MaxForce, InDriveParams.YDrive.MaxForce, InDriveParams.ZDrive.MaxForce));
@@ -455,7 +455,7 @@ void FPhysInterface_Chaos::UpdateAngularDrive_AssumesLocked(const FPhysicsConstr
 				}
 			}
 
-			Constraint->SetAngularDriveForceMode(Chaos::EJointForceMode::Acceleration);
+			Constraint->SetAngularDriveForceMode(InDriveParams.bAccelerationMode ? Chaos::EJointForceMode::Acceleration : Chaos::EJointForceMode::Force);
 			if (InDriveParams.AngularDriveMode == EAngularDriveMode::TwistAndSwing)
 			{
 				Constraint->SetAngularDriveStiffness(Chaos::ConstraintSettings::AngularDriveStiffnessScale() * Chaos::FVec3(InDriveParams.TwistDrive.Stiffness, InDriveParams.SwingDrive.Stiffness, InDriveParams.SwingDrive.Stiffness));
diff --git a/Engine/Source/Runtime/Engine/Private/PhysicsEngine/PhysicsConstraintComponent.cpp b/Engine/Source/Runtime/Engine/Private/PhysicsEngine/PhysicsConstraintComponent.cpp
--- a/Engine/Source/Runtime/Engine/Private/PhysicsEngine/PhysicsConstraintComponent.cpp
+++ b/Engine/Source/Runtime/Engine/Private/PhysicsEngine/PhysicsConstraintComponent.cpp
@@ -701,6 +701,11 @@ void UPhysicsConstraintComponent::SetAngularDriveMode(EAngularDriveMode::Type Dr
 	ConstraintInstance.SetAngularDriveMode(DriveMode);
 }
 
+void UPhysicsConstraintComponent::SetAngularDriveAccelerationMode(bool bAccelerationMode)
+{
+	ConstraintInstance.SetAngularDriveAccelerationMode(bAccelerationMode);
+}
+
 void UPhysicsConstraintComponent::SetAngularVelocityDriveTwistAndSwing( bool bEnableTwistDrive, bool bEnableSwingDrive)
 {
 	ConstraintInstance.SetAngularVelocityDriveTwistAndSwing(bEnableTwistDrive, bEnableSwingDrive);
@@ -727,6 +732,12 @@ void UPhysicsConstraintComponent::SetLinearDriveParams( float PositionStrength,
 	ConstraintInstance.SetLinearDriveParams(PositionStrength, VelocityStrength, InForceLimit);
 }
 
+
+void UPhysicsConstraintComponent::SetLinearDriveAccelerationMode(bool bAccelerationMode)
+{
+	ConstraintInstance.SetLinearDriveAccelerationMode(bAccelerationMode);
+}
+
 void UPhysicsConstraintComponent::SetAngularOrientationTarget( const FRotator& InPosTarget )
 {
 	ConstraintInstance.SetAngularOrientationTarget(InPosTarget.Quaternion());
diff --git a/Engine/Source/Runtime/Experimental/Chaos/Private/Chaos/PBDJointConstraints.cpp b/Engine/Source/Runtime/Experimental/Chaos/Private/Chaos/PBDJointConstraints.cpp
--- a/Engine/Source/Runtime/Experimental/Chaos/Private/Chaos/PBDJointConstraints.cpp
+++ b/Engine/Source/Runtime/Experimental/Chaos/Private/Chaos/PBDJointConstraints.cpp
@@ -200,6 +200,11 @@ namespace Chaos
 		GetConstraintSettingsInternal().LinearDriveMaxForce = MaxForce;
 	}
 
+	void FPBDJointConstraintHandle::SetLinearDriveForceMode(EJointForceMode ForceMode)
+	{
+		GetConstraintSettingsInternal().LinearDriveForceMode = ForceMode;
+	}
+
 	void FPBDJointConstraintHandle::SetAngularDriveStiffness(const FVec3 Stiffness)
 	{
 		GetConstraintSettingsInternal().AngularDriveStiffness = Stiffness;
@@ -215,6 +220,11 @@ namespace Chaos
 		GetConstraintSettingsInternal().AngularDriveMaxTorque = MaxTorque;
 	}
 
+	void FPBDJointConstraintHandle::SetAngularDriveForceMode(EJointForceMode ForceMode)
+	{
+		GetConstraintSettingsInternal().AngularDriveForceMode = ForceMode;
+	}
+
 	void FPBDJointConstraintHandle::SetCollisionEnabled(const bool bCollisionEnabled)
 	{
 		GetConstraintSettingsInternal().bCollisionEnabled = bCollisionEnabled;
diff --git a/Engine/Source/Runtime/Experimental/Chaos/Public/Chaos/PBDJointConstraints.h b/Engine/Source/Runtime/Experimental/Chaos/Public/Chaos/PBDJointConstraints.h
--- a/Engine/Source/Runtime/Experimental/Chaos/Public/Chaos/PBDJointConstraints.h
+++ b/Engine/Source/Runtime/Experimental/Chaos/Public/Chaos/PBDJointConstraints.h
@@ -68,10 +68,12 @@ namespace Chaos
 		CHAOS_API void SetLinearDriveStiffness(const FVec3 Stiffness);
 		CHAOS_API void SetLinearDriveDamping(const FVec3 Damping);
 		CHAOS_API void SetLinearDriveMaxForce(const FVec3 MaxForce);
+		CHAOS_API void SetLinearDriveForceMode(EJointForceMode ForceMode);
 
 		CHAOS_API void SetAngularDriveStiffness(const FVec3 Stiffness);
 		CHAOS_API void SetAngularDriveDamping(const FVec3 Damping);
 		CHAOS_API void SetAngularDriveMaxTorque(const FVec3 MaxTorque);
+		CHAOS_API void SetAngularDriveForceMode(EJointForceMode ForceMode);
 
 		CHAOS_API void SetCollisionEnabled(const bool bCollisionEnabled);
 		CHAOS_API void SetParentInvMassScale(const FReal ParentInvMassScale);
