- Test the component is still valid when processing deferred updates
- With incremental updates the component can be registered, queued for a send render data, then destroyed by an actor chain leaving the update incorrectly dangling inside the queue

diff --git a/Engine/Source/Runtime/Engine/Private/Components/ActorComponent.cpp b/Engine/Source/Runtime/Engine/Private/Components/ActorComponent.cpp
--- a/Engine/Source/Runtime/Engine/Private/Components/ActorComponent.cpp
+++ b/Engine/Source/Runtime/Engine/Private/Components/ActorComponent.cpp
@@ -164,7 +164,11 @@ void FRegisterComponentContext::Process()
 
 	for (UPrimitiveComponent* Primitive : SendRenderDynamicDataPrimitives)
 	{
-		Primitive->SendRenderDynamicData_Concurrent();
+		// With incremetal updates the component can be registered, added to send render data queue, then destroyed by an actor chain so we must test it's still valid before sending render data.
+		if (::IsValid(Primitive))
+		{
+			Primitive->SendRenderDynamicData_Concurrent();
+		}
 	}
 	SendRenderDynamicDataPrimitives.Empty();
 }
