Fix typo where twist/swing drive max torque is used instead of slerp drive.

diff --git a/Engine/Source/Runtime/Engine/Private/PhysicsEngine/Experimental/PhysInterface_Chaos.cpp b/Engine/Source/Runtime/Engine/Private/PhysicsEngine/Experimental/PhysInterface_Chaos.cpp
--- a/Engine/Source/Runtime/Engine/Private/PhysicsEngine/Experimental/PhysInterface_Chaos.cpp
+++ b/Engine/Source/Runtime/Engine/Private/PhysicsEngine/Experimental/PhysInterface_Chaos.cpp
@@ -466,7 +466,7 @@ void FPhysInterface_Chaos::UpdateAngularDrive_AssumesLocked(const FPhysicsConstr
 			{
 				Constraint->SetAngularDriveStiffness(Chaos::ConstraintSettings::AngularDriveStiffnessScale() * Chaos::FVec3(InDriveParams.SlerpDrive.Stiffness));
 				Constraint->SetAngularDriveDamping(Chaos::ConstraintSettings::AngularDriveDampingScale() * Chaos::FVec3(InDriveParams.SlerpDrive.Damping));
-				Constraint->SetAngularDriveMaxTorque(Chaos::FVec3(InDriveParams.TwistDrive.MaxForce, InDriveParams.SwingDrive.MaxForce, InDriveParams.SwingDrive.MaxForce));
+				Constraint->SetAngularDriveMaxTorque(Chaos::FVec3(InDriveParams.SlerpDrive.MaxForce, InDriveParams.SlerpDrive.MaxForce, InDriveParams.SlerpDrive.MaxForce));
 			}
 		}
 	}
