Test tagging core for C++ AutomationTestFramework
#tests unit tests included in submission

diff --git a/Engine/Source/Runtime/Core/Private/Misc/AutomationTest.cpp b/Engine/Source/Runtime/Core/Private/Misc/AutomationTest.cpp
--- a/Engine/Source/Runtime/Core/Private/Misc/AutomationTest.cpp
+++ b/Engine/Source/Runtime/Core/Private/Misc/AutomationTest.cpp
@@ -21,7 +21,8 @@
 #include "Misc/ScopeRWLock.h"
 #include "Modules/ModuleManager.h"
 
-DEFINE_LOG_CATEGORY(LogLatentCommands)
+DEFINE_LOG_CATEGORY(LogLatentCommands);
+DEFINE_LOG_CATEGORY(LogAutomationTestFramework);
 DEFINE_LOG_CATEGORY_STATIC(LogAutomationTestStateTrace, Log, All);
 DEFINE_LOG_CATEGORY_STATIC(LogAutomationTest, Warning, All);
 
@@ -344,6 +345,43 @@ bool FAutomationTestFramework::UnregisterAutomationTest( const FString& InTestNa
 	return bRegistered;
 }
 
+bool FAutomationTestFramework::RegisterAutomationTestTags( const FString& InTestNameToRegister, const FString& InTestTagsToRegister )
+{
+	const bool bAlreadyRegistered = TestFullNameToTagDataMap.Contains( InTestNameToRegister );
+	if ( !bAlreadyRegistered )
+	{
+		LLM_SCOPE_BYNAME(TEXT("AutomationTest/Framework"));
+		TestFullNameToTagDataMap.Add(InTestNameToRegister, InTestTagsToRegister);
+	}
+	return !bAlreadyRegistered;
+}
+
+bool FAutomationTestFramework::UnregisterAutomationTestTags(const FString& InTestNameToUnregister)
+{
+	const bool bRegistered = TestFullNameToTagDataMap.Contains(InTestNameToUnregister);
+	if (bRegistered)
+	{
+		TestFullNameToTagDataMap.Remove(InTestNameToUnregister);
+	}
+	return bRegistered;
+}
+
+bool FAutomationTestFramework::RegisterComplexAutomationTestTags(const FAutomationTestBase* InTest, const FString& InTestNameToRegister, const FString& InTestTagsToRegister)
+{
+	FString FullTestName = InTest->GetBeautifiedTestName().AppendChar('.').Append(InTestNameToRegister);
+	return RegisterAutomationTestTags(FullTestName, InTestTagsToRegister);
+}
+
+FString FAutomationTestFramework::GetTagsForAutomationTest(const FString& InTestName)
+{
+	FString * FindResult = TestFullNameToTagDataMap.Find(InTestName);
+	if (FindResult)
+	{
+		return *FindResult;
+	}
+	return FString();
+}
+
 void FAutomationTestFramework::EnqueueLatentCommand(TSharedPtr<IAutomationLatentCommand> NewCommand)
 {
 	//ensure latent commands are never used within smoke tests - will only catch when smokes are exclusively requested
@@ -761,6 +799,28 @@ void FAutomationTestFramework::GetValidTestNames( TArray<FAutomationTestInfo>& T
 	}
 }
 
+bool FAutomationTestFramework::TagsMatchPattern(const FString& Tags, const FString& TagPattern) const
+{
+	check(TagFilter);
+	TagFilter->SetFilterText(FText::FromString(TagPattern));
+	return TagFilter->TestTextFilter(FBasicStringFilterExpressionContext(Tags));
+}
+
+void FAutomationTestFramework::GetTestFullNamesMatchingTagPattern(TArray<FString>& OutTestNames, const FString& TagPattern) const
+{
+	LLM_SCOPE_BYNAME(TEXT("AutomationTest/Framework"));
+	OutTestNames.Empty();
+
+	for (TMap<FString, FString>::TConstIterator TestIter(TestFullNameToTagDataMap); TestIter; ++TestIter)
+	{
+		const FString CurTags = TestIter.Value();
+		if (TagsMatchPattern(CurTags, TagPattern))
+		{
+			OutTestNames.Add(TestIter.Key());
+		}
+	}
+}
+
 bool FAutomationTestFramework::ShouldTestContent(const FString& Path) const
 {
 	static TArray<FString> TestLevelFolders;
@@ -1138,6 +1198,7 @@ FAutomationTestFramework::FAutomationTestFramework()
 	, bForceSmokeTests(false)
 	, bCaptureStack(true)
 {
+	TagFilter = MakeShared<FTextFilterExpressionEvaluator>(ETextFilterExpressionEvaluatorMode::BasicString);
 }
 
 FAutomationTestFramework::~FAutomationTestFramework()
@@ -1616,6 +1677,7 @@ void FAutomationTestBase::GenerateTestNames(TArray<FAutomationTestInfo>& TestInf
 	TArray<FString> BeautifiedNames;
 	TArray<FString> ParameterNames;
 	GetTests(BeautifiedNames, ParameterNames);
+	FAutomationTestFramework& Framework = FAutomationTestFramework::Get();
 
 	FString BeautifiedTestName = GetBeautifiedTestName();
 
@@ -1641,7 +1703,8 @@ void FAutomationTestBase::GenerateTestNames(TArray<FAutomationTestInfo>& TestInf
 			GetTestSourceFileName(CompleteTestName),
 			GetTestSourceFileLine(CompleteTestName),
 			GetTestAssetPath(ParameterNames[ParameterIndex]),
-			GetTestOpenCommand(ParameterNames[ParameterIndex])
+			GetTestOpenCommand(ParameterNames[ParameterIndex]),
+			Framework.GetTagsForAutomationTest(TestName)
 		);
 		
 		TestInfo.Add( NewTestInfo );
diff --git a/Engine/Source/Runtime/Core/Public/Misc/AutomationTest.h b/Engine/Source/Runtime/Core/Public/Misc/AutomationTest.h
--- a/Engine/Source/Runtime/Core/Public/Misc/AutomationTest.h
+++ b/Engine/Source/Runtime/Core/Public/Misc/AutomationTest.h
@@ -40,6 +40,7 @@
 #include "Misc/Guid.h"
 #include "Misc/Optional.h"
 #include "Misc/OutputDevice.h"
+#include "Misc/TextFilterExpressionEvaluator.h"
 #include "Misc/Timespan.h"
 #include "Templates/Function.h"
 #include "Templates/SharedPointer.h"
@@ -49,7 +50,10 @@
 #include <atomic>
 
 CORE_API DECLARE_LOG_CATEGORY_EXTERN(LogLatentCommands, Log, All);
+CORE_API DECLARE_LOG_CATEGORY_EXTERN(LogAutomationTestFramework, Log, All);
+
 class FAutomationTestBase;
+class FAutomationTestTags;
 
 #ifndef WITH_AUTOMATION_TESTS
 	#define WITH_AUTOMATION_TESTS (WITH_DEV_AUTOMATION_TESTS || WITH_PERF_AUTOMATION_TESTS)
@@ -272,14 +276,22 @@ public:
 	 * Constructor
 	 *
 	 * @param	InDisplayName - Name used in the UI
+	 * @param	InFullTestPath - Dot-separated pathname to hierarchically display in the UI, which should be unique
 	 * @param	InTestName - The test command string
-	 * @param	InTestFlag - Test flags
+	 * @param	InTestFlags - Test flags
+	 * @param	InNumParticipantsRequired - Number of workers to run test
 	 * @param	InParameterName - optional parameter. e.g. asset name
+	 * @param	InSourceFile - Filesystem path to the source file which defines the test
+	 * @param	InSourceFileLine - Line number in the source file which defines the test
+	 * @param	InAssetPath - Filesystem path to the UAsset file which defines the test
+	 * @param	InOpenCommand - An exec command to open the test
+	 * @param	InTestTags - Tag metadata concatenated into one string
 	 */
-	FAutomationTestInfo(const FString& InDisplayName, const FString& InFullTestPath, const FString& InTestName, const EAutomationTestFlags InTestFlags, const int32 InNumParticipantsRequired, const FString& InParameterName = FString(), const FString& InSourceFile = FString(), int32 InSourceFileLine = 0, const FString& InAssetPath = FString(), const FString& InOpenCommand = FString())
+	FAutomationTestInfo(const FString& InDisplayName, const FString& InFullTestPath, const FString& InTestName, const EAutomationTestFlags InTestFlags, const int32 InNumParticipantsRequired, const FString& InParameterName = FString(), const FString& InSourceFile = FString(), int32 InSourceFileLine = 0, const FString& InAssetPath = FString(), const FString& InOpenCommand = FString(),  const FString& InTestTags = FString())
 		: DisplayName( InDisplayName )
 		, FullTestPath( InFullTestPath )
 		, TestName( InTestName )
+		, TestTags( InTestTags )
 		, TestParameter( InParameterName )
 		, SourceFile( InSourceFile )
 		, SourceFileLine( InSourceFileLine )
@@ -333,6 +345,16 @@ public:
 		return TestName;
 	}
 
+	/**
+	 * Get the tags associated with this test.
+	 *
+	 * @return The concatenated test tags.
+	 */
+	FString GetTestTags() const
+	{
+		return TestTags;
+	}
+
 	/**
 	 * Get the type of parameter. This will be the asset name for linked assets.
 	 *
@@ -459,6 +481,9 @@ private:
 	/** Test name used to run the test */
 	FString TestName;
 
+	/** Tags defined on this test */
+	FString TestTags;
+
 	/** Parameter - e.g. an asset name or map name */
 	FString TestParameter;
 
@@ -981,6 +1006,47 @@ public:
 	 */
 	CORE_API bool UnregisterAutomationTest( const FString& InTestNameToUnregister );
 
+	/**
+	 * Register tag metadata for a test into the framework.
+	 * 
+	 * @param	InTestNameToRegister	FullName of the test to associate with tags
+	 * @param	InTestTagsToRegister	Concatenated string of tags to register
+	 * 
+	 * @return	true if tags for the test successfully registered; false if tags were already registered for the test
+	 * @see FAutomationTestBase::GetTestFullName
+	 */
+	CORE_API bool RegisterAutomationTestTags(const FString& InTestNameToRegister, const FString& InTestTagsToRegister);
+
+	/**
+	 * Unregister tags for a automation test from the framework.
+	 * 
+	 * @param	InTestNameToUnregister	FullName of the test to remove tags from
+	 * 
+	 * @return true if test tags were successfully unregistered; false if tags for a test with that name was not found in the framework.
+	 * @see FAutomationTestBase::GetTestFullName
+	 */
+	CORE_API bool UnregisterAutomationTestTags(const FString& InTestNameToUnregister);
+
+	/**
+	 * Helper to register tags for a individual test enumerated by a Complex test
+	 *
+	 * @param	InTest					The Complex test to associate with
+	 * @param	InBeautifiedTestName	The name of the individual test
+	 * @param	InTestTagsToRegister	Concatenated string of tags to register
+	 *
+	 * @return true if test tags were successfully unregistered; false if tags for a test with that name was not found in the framework.
+	 */
+	CORE_API bool RegisterComplexAutomationTestTags(const FAutomationTestBase* InTest, const FString& InBeautifiedTestName, const FString& InTestTagsToRegister);
+
+	/**
+	 * Fetch the tags associated with a test
+	 * 
+	 * @param	InTestName	Name of the test to find tags for
+	 * 
+	 * @return	a string of all concatenated tags; empty string if test is not registered
+	 */
+	CORE_API FString GetTagsForAutomationTest(const FString& InTestName);
+
 	/**
 	 * Enqueues a latent command for execution on a subsequent frame
 	 *
@@ -1072,6 +1138,17 @@ public:
 	 */
 	CORE_API void GetValidTestNames( TArray<FAutomationTestInfo>& TestInfo ) const;
 
+	/**
+	 * Collects all registered tests with associated tags matching the pattern
+	 * 
+	 * @param	OutTestNames	Array that will be modified to contain only the set of matching tests
+	 * @param	TagPattern		An "Advanced Search Syntax" query to select matching tags from the registered tests
+	 * 
+	 * @see FTextFilterExpressionEvaluator
+	 * @see FAutomationTestBase::GetTestFullName
+	 */
+	CORE_API void GetTestFullNamesMatchingTagPattern(TArray<FString>& OutTestNames, const FString& TagPattern) const;
+
 	/**
 	 * Whether the testing framework should allow content to be tested or not.  Intended to block developer directories.
 	 * @param Path - Full path to the content in question
@@ -1355,6 +1432,9 @@ private:
 	/** Mapping of automation test names to their respective object instances */
 	TMap<FString, FAutomationTestBase*> AutomationTestClassNameToInstanceMap;
 
+	/** Mapping of full test names to their registered tags */
+	TMap<FString, FString> TestFullNameToTagDataMap;
+
 	/** Queue of deferred commands */
 	TQueue< TSharedPtr<IAutomationLatentCommand> > LatentCommands;
 
@@ -1398,6 +1478,49 @@ private:
 
 	TMap<FString, FOnTestSectionEvent> OnEnteringTestSectionEvent;
 	TMap<FString, FOnTestSectionEvent> OnLeavingTestSectionEvent;
+
+	/**
+	 * Evaluate whether a tag string matches the pattern string
+	 * 
+	 * @param	Tags		Concatenated string of all tags defined for a test
+	 * @param	TagPattern	A query pattern
+	 * @see FTextFilterExpressionEvaluator
+	 */
+	bool TagsMatchPattern(const FString& Tags, const FString& TagPattern) const;
+
+	/** Cached filter context */
+	TSharedPtr<FTextFilterExpressionEvaluator> TagFilter;
+};
+
+/** Wrapper class to simplify tag registration */
+class FAutomationTestTags
+{
+public:
+	/**
+	 * Register tags for test by its fully defined name, which will unregister when this object leaves scope
+	 * 
+	 * @param	InTestFullName	The unique path-name identifying this test
+	 * @param	InTags			A string defining all concatenated tags to register
+	 * @see FAutomationTestBase::GetTestFullName
+	 */
+	FAutomationTestTags(const FString& InTestFullName, const FString& InTags)
+	{
+		LLM_SCOPE_BYNAME(TEXT("AutomationTest/Framework"));
+		TestFullName = InTestFullName;
+		bool Registered = FAutomationTestFramework::Get().RegisterAutomationTestTags( InTestFullName, InTags );
+		if (!Registered)
+		{
+			UE_LOG(LogAutomationTestFramework, Warning, TEXT("Tag registration already exists for test '%s', existing value was not overridden"), *TestFullName);
+		}
+	}
+
+	~FAutomationTestTags()
+	{
+		//ignore failure to unregister due to already being unregistered
+		FAutomationTestFramework::Get().UnregisterAutomationTestTags(TestFullName);
+	}
+private:
+	FString TestFullName;
 };
 
 /** Simple abstract base class for all automation tests */
@@ -3934,6 +4057,12 @@ public: \
 		virtual FString GetBeautifiedTestName() const override { return PrettyName; } \
 		virtual void Define() override;
 
+#define REGISTER_SIMPLE_AUTOMATION_TEST_TAGS( TClass, PrettyName, TagsString ) \
+	namespace\
+	{\
+		FAutomationTestTags TClass##AutomationTagsInstance(PrettyName, TagsString);\
+	}
+
 #if WITH_AUTOMATION_WORKER
 	#define IMPLEMENT_SIMPLE_AUTOMATION_TEST( TClass, PrettyName, TFlags ) \
 		IMPLEMENT_SIMPLE_AUTOMATION_TEST_PRIVATE(TClass, FAutomationTestBase, PrettyName, TFlags, __FILE__, __LINE__) \
diff --git a/Engine/Source/Runtime/Engine/Private/Tests/EngineAutomationTests.cpp b/Engine/Source/Runtime/Engine/Private/Tests/EngineAutomationTests.cpp
--- a/Engine/Source/Runtime/Engine/Private/Tests/EngineAutomationTests.cpp
+++ b/Engine/Source/Runtime/Engine/Private/Tests/EngineAutomationTests.cpp
@@ -2085,4 +2085,197 @@ bool FAutomationAttachment::RunTest(const FString& Parameters)
 	return false;
 }
 
+constexpr auto ExampleTag = "[TestExampleTag]";
+constexpr auto OtherTag = "[SomeOtherTag]";
+constexpr auto NegativeTag = "[DoNotWant]";
+constexpr auto NegativeAndExampleTags = "[TestExampleTag][DoNotWant]";
+constexpr auto NegativeAndExampleTagsReversed = "[TestExampleTag][DoNotWant]";
+
+constexpr auto FullTestNameTagsExist = "TestFramework.Tags.TagsExist";
+IMPLEMENT_SIMPLE_AUTOMATION_TEST(FAutomationTagsExist, FullTestNameTagsExist, EAutomationTestFlags::EditorContext | EAutomationTestFlags::EngineFilter)
+REGISTER_SIMPLE_AUTOMATION_TEST_TAGS(FAutomationTagsExist, FullTestNameTagsExist, ExampleTag) //Other tests expect ExampleTag to only be registered once, and only for this test
+bool FAutomationTagsExist::RunTest(const FString& Parameters)
+{
+	//registration above performs necessary setup
+	FString MyTags = FAutomationTestFramework::Get().GetTagsForAutomationTest(FullTestNameTagsExist);
+
+	TestEqual(TEXT("Tags statically register"), MyTags, ExampleTag);
+	return true;
+}
+
+constexpr auto FullTestNameSelect = "TestFramework.Tags.TagsAreSelectable";
+IMPLEMENT_SIMPLE_AUTOMATION_TEST(FAutomationTagsSelect, FullTestNameSelect, EAutomationTestFlags::EditorContext | EAutomationTestFlags::EngineFilter)
+REGISTER_SIMPLE_AUTOMATION_TEST_TAGS(FAutomationTagsSelect, FullTestNameSelect, NegativeAndExampleTags)
+bool FAutomationTagsSelect::RunTest(const FString& Parameters)
+{
+	TArray<FString> TestNames;
+	FString PositiveFilter("[TestExampleTag]");
+
+	FAutomationTestFramework::Get().GetTestFullNamesMatchingTagPattern(TestNames, PositiveFilter);
+
+	TestGreaterEqual("Tags get selected", TestNames.Num(), 1);
+	bool FoundThisTest = false;
+	for (FString Element : TestNames)
+	{
+		if (Element.Equals(FullTestNameSelect))
+		{
+			FoundThisTest = true;
+		}
+	}
+	TestTrue("Current test was selected", FoundThisTest);
+	return true;
+}
+
+constexpr auto FullTestNameUnion = "TestFramework.Tags.UnionSelection";
+IMPLEMENT_SIMPLE_AUTOMATION_TEST(FAutomationTagsUnion, FullTestNameUnion, EAutomationTestFlags::EditorContext | EAutomationTestFlags::EngineFilter)
+//The two tests above are used, no extra tags need to be registered
+bool FAutomationTagsUnion::RunTest(const FString& Parameters)
+{
+	TArray<FString> TestNames;
+	FString UnionFilter("[TestExampleTag] OR [SomeOtherTag]");
+
+	FAutomationTestFramework::Get().GetTestFullNamesMatchingTagPattern(TestNames, UnionFilter);
+
+	TestGreaterEqual("Tags get selected", TestNames.Num(), 1);
+	bool FoundExample = false;
+	bool FoundOther = false;
+	for (FString Element : TestNames)
+	{
+		if (Element.Equals(FullTestNameTagsExist))
+		{
+			FoundExample = true;
+		}
+		else {
+			if (Element.Equals(FullTestNameSelect))
+			{
+				FoundOther = true;
+			}
+		}
+	}
+	TestTrue("First test was selected", FoundExample);
+	TestTrue("Second test was selected", FoundOther);
+	return true;
+}
+
+constexpr auto FullTestNameNoBracket = "TestFramework.Tags.SelectWithoutBrackets";
+IMPLEMENT_SIMPLE_AUTOMATION_TEST(FAutomationTagsFilterNoBrackets, FullTestNameNoBracket, EAutomationTestFlags::EditorContext | EAutomationTestFlags::EngineFilter)
+REGISTER_SIMPLE_AUTOMATION_TEST_TAGS(FAutomationTagsFilterNoBrackets, FullTestNameNoBracket, OtherTag)
+bool FAutomationTagsFilterNoBrackets::RunTest(const FString& Parameters)
+{
+	TArray<FString> TestNames;
+	FString PositiveFilter("SomeOtherTag");
+
+	FAutomationTestFramework::Get().GetTestFullNamesMatchingTagPattern(TestNames, PositiveFilter);
+
+	TestGreaterEqual("Tags get selected", TestNames.Num(), 1);
+	bool FoundThisTest = false;
+	for (FString Element : TestNames)
+	{
+		if (Element.Equals(FullTestNameNoBracket))
+		{
+			FoundThisTest = true;
+		}
+	}
+	TestTrue("Current test was selected", FoundThisTest);
+	return true;
+}
+
+constexpr auto FullTestNameFilter = "TestFramework.Tags.TagsCanFilter";
+IMPLEMENT_SIMPLE_AUTOMATION_TEST(FAutomationTagsFilter, FullTestNameFilter, EAutomationTestFlags::EditorContext | EAutomationTestFlags::EngineFilter)
+REGISTER_SIMPLE_AUTOMATION_TEST_TAGS(FAutomationTagsFilter, FullTestNameFilter, NegativeAndExampleTags)
+bool FAutomationTagsFilter::RunTest(const FString& Parameters)
+{
+	TArray<FString> TestNames;
+	FString NegativeFilter("[TestExampleTag] AND NOT [DoNotWant]");
+
+	FAutomationTestFramework::Get().GetTestFullNamesMatchingTagPattern(TestNames, NegativeFilter);
+
+	TestEqual("One element", TestNames.Num(), 1 );
+	TestEqual("Current test is not selected", TestNames[0], FullTestNameTagsExist);
+	return true;
+}
+
+constexpr auto FullTestNameFilterReversePattern = "TestFramework.Tags.TagsFilterPatternOrderIndependent";
+IMPLEMENT_SIMPLE_AUTOMATION_TEST(FAutomationTagsFilterPatternOrder, FullTestNameFilterReversePattern, EAutomationTestFlags::EditorContext | EAutomationTestFlags::EngineFilter)
+REGISTER_SIMPLE_AUTOMATION_TEST_TAGS(FAutomationTagsFilterPatternOrder, FullTestNameFilterReversePattern, NegativeAndExampleTags)
+bool FAutomationTagsFilterPatternOrder::RunTest(const FString& Parameters)
+{
+	TArray<FString> TestNames;
+	FString NegativeFilter("NOT [DoNotWant] AND [TestExampleTag]");
+
+	FAutomationTestFramework::Get().GetTestFullNamesMatchingTagPattern(TestNames, NegativeFilter);
+
+	TestEqual("One element", TestNames.Num(), 1);
+	TestEqual("Current test is not selected", TestNames[0], FullTestNameTagsExist);
+	return true;
+}
+
+constexpr auto FullTestNameFilterReverseTags = "TestFramework.Tags.TagsFilterTagOrderIndependent";
+IMPLEMENT_SIMPLE_AUTOMATION_TEST(FAutomationTagsFilterTagOrder, FullTestNameFilterReverseTags, EAutomationTestFlags::EditorContext | EAutomationTestFlags::EngineFilter)
+REGISTER_SIMPLE_AUTOMATION_TEST_TAGS(FAutomationTagsFilterTagOrder, FullTestNameFilterReverseTags, NegativeAndExampleTagsReversed)
+bool FAutomationTagsFilterTagOrder::RunTest(const FString& Parameters)
+{
+	TArray<FString> TestNames;
+	FString NegativeFilter("[TestExampleTag] AND NOT [DoNotWant]");
+
+	FAutomationTestFramework::Get().GetTestFullNamesMatchingTagPattern(TestNames, NegativeFilter);
+
+	TestEqual("One element", TestNames.Num(), 1);
+	TestEqual("Current test is not selected", TestNames[0], FullTestNameTagsExist);
+	return true;
+}
+
+constexpr auto FullTestNameFilterReverseBoth = "TestFramework.Tags.TagsFilterReverseOrderIndependent";
+IMPLEMENT_SIMPLE_AUTOMATION_TEST(FAutomationTagsFilterBothReverseOrder, FullTestNameFilterReverseBoth, EAutomationTestFlags::EditorContext | EAutomationTestFlags::EngineFilter)
+REGISTER_SIMPLE_AUTOMATION_TEST_TAGS(FAutomationTagsFilterBothReverseOrder, FullTestNameFilterReverseBoth, NegativeAndExampleTagsReversed)
+bool FAutomationTagsFilterBothReverseOrder::RunTest(const FString& Parameters)
+{
+	TArray<FString> TestNames;
+	FString NegativeFilter("NOT [DoNotWant] AND [TestExampleTag]");
+
+	FAutomationTestFramework::Get().GetTestFullNamesMatchingTagPattern(TestNames, NegativeFilter);
+
+	TestEqual("One element", TestNames.Num(), 1);
+	TestEqual("Current test is not selected", TestNames[0], FullTestNameTagsExist);
+	return true;
+}
+
+constexpr auto ComplexTags1 = "[TestExampleTagComplex][FirstExample]";
+constexpr auto ComplexTags2 = "[TestExampleTagComplex][SecondExample]";
+constexpr auto ComplexTagTestPath = "TestFramework.Tags.Complex";
+IMPLEMENT_COMPLEX_AUTOMATION_TEST(FAutomationTagsForComplexSuite, ComplexTagTestPath, EAutomationTestFlags::EditorContext | EAutomationTestFlags::EngineFilter)
+void FAutomationTagsForComplexSuite::GetTests(TArray<FString>& OutBeautifiedNames, TArray<FString>& OutTestCommands) const
+{
+	FAutomationTestFramework& Framework = FAutomationTestFramework::Get();
+
+	OutBeautifiedNames.Add("First");
+	Framework.RegisterComplexAutomationTestTags(this, "First", ComplexTags1);
+
+	OutBeautifiedNames.Add("Second");
+	Framework.RegisterComplexAutomationTestTags(this, "Second", ComplexTags2);
+	
+	OutTestCommands = OutBeautifiedNames; //pass names as parameters
+}
+
+bool FAutomationTagsForComplexSuite::RunTest(const FString& Parameters)
+{
+	//registration in GetTests() performs necessary setup
+	FString MyName = GetTestFullName();
+	TestTrue(TEXT("Complex tests construct names as expected"), MyName.EndsWith(Parameters));
+
+	FString MyTags = FAutomationTestFramework::Get().GetTagsForAutomationTest(MyName);
+	TestFalse(TEXT("Tag is found"), MyTags.IsEmpty());
+
+	FString ExpectedTags;
+	if( Parameters.Equals("First"))
+	{
+		ExpectedTags = FString(ComplexTags1);
+	}
+	else {
+		ExpectedTags = FString(ComplexTags2);
+	}
+	TestEqual(TEXT("Tags dynamically registered for intended complex test"), MyTags, ExpectedTags);
+	return true;
+}
+
 #endif //WITH_DEV_AUTOMATION_TESTS
