Added helper methods to create view uniform buffer from descriptor struct.

diff --git a/Engine/Source/Runtime/Engine/Private/SceneView.cpp b/Engine/Source/Runtime/Engine/Private/SceneView.cpp
--- a/Engine/Source/Runtime/Engine/Private/SceneView.cpp
+++ b/Engine/Source/Runtime/Engine/Private/SceneView.cpp
@@ -2443,7 +2443,7 @@ bool FSceneView::IsInstancedStereoPass() const
 	return bIsInstancedStereoEnabled && IStereoRendering::IsStereoEyeView(*this) && IStereoRendering::IsAPrimaryView(*this);
 }
 
-FVector4f FSceneView::GetScreenPositionScaleBias(const FIntPoint& BufferSize, const FIntRect& ViewRect) const
+FVector4f FSceneView::GetScreenPositionScaleBias(const FIntPoint& BufferSize, const FIntRect& ViewRect)
 {
 	const float InvBufferSizeX = 1.0f / BufferSize.X;
 	const float InvBufferSizeY = 1.0f / BufferSize.Y;
@@ -2458,11 +2458,13 @@ FVector4f FSceneView::GetScreenPositionScaleBias(const FIntPoint& BufferSize, co
 		);
 }
 
-void FSceneView::SetupViewRectUniformBufferParameters(FViewUniformShaderParameters& ViewUniformShaderParameters,
+void SetupViewRectUniformBufferParameters(
+	FViewUniformShaderParameters& ViewUniformShaderParameters,
 	const FIntPoint& BufferSize,
 	const FIntRect& EffectiveViewRect,
 	const FViewMatrices& InViewMatrices,
-	const FViewMatrices& InPrevViewMatrices) const
+	const FViewMatrices& InPrevViewMatrices,
+	const FSetupViewUniformParametersInputs& Inputs)
 {
 	checkfSlow(EffectiveViewRect.Area() > 0, TEXT("Invalid-size EffectiveViewRect passed to CreateUniformBufferParameters [%d * %d]."), EffectiveViewRect.Width(), EffectiveViewRect.Height());
 	ensureMsgf((BufferSize.X > 0 && BufferSize.Y > 0), TEXT("Invalid-size BufferSize passed to CreateUniformBufferParameters [%d * %d]."), BufferSize.X, BufferSize.Y);
@@ -2478,7 +2480,7 @@ void FSceneView::SetupViewRectUniformBufferParameters(FViewUniformShaderParamete
 	const float InvBufferSizeX = 1.0f / BufferSize.X;
 	const float InvBufferSizeY = 1.0f / BufferSize.Y;
 
-	ViewUniformShaderParameters.ScreenPositionScaleBias = GetScreenPositionScaleBias(BufferSize, EffectiveViewRect);
+	ViewUniformShaderParameters.ScreenPositionScaleBias = FSceneView::GetScreenPositionScaleBias(BufferSize, EffectiveViewRect);
 
 	ViewUniformShaderParameters.BufferSizeAndInvSize = FVector4f(BufferSize.X, BufferSize.Y, InvBufferSizeX, InvBufferSizeY);
 	ViewUniformShaderParameters.BufferBilinearUVMinMax = FVector4f(
@@ -2488,9 +2490,9 @@ void FSceneView::SetupViewRectUniformBufferParameters(FViewUniformShaderParamete
 		InvBufferSizeY * (EffectiveViewRect.Max.Y - 0.5));
 
 	/* Texture Level-of-Detail Strategies for Real-Time Ray Tracing https://developer.nvidia.com/raytracinggems Equation 20 */
-	if(FOV != 0)
+	if(Inputs.FOV != 0)
 	{
-		float RadFOV = (UE_PI / 180.0f) * FOV;
+		float RadFOV = (UE_PI / 180.0f) * Inputs.FOV;
 		ViewUniformShaderParameters.EyeToPixelSpreadAngle = FPlatformMath::Atan((2.0f * FPlatformMath::Tan(RadFOV * 0.5f)) / BufferSize.Y);
 	}
 	else
@@ -2498,20 +2500,18 @@ void FSceneView::SetupViewRectUniformBufferParameters(FViewUniformShaderParamete
 		ViewUniformShaderParameters.EyeToPixelSpreadAngle = 0;
 	}
 
-	ViewUniformShaderParameters.MotionBlurNormalizedToPixel = FinalPostProcessSettings.MotionBlurMax * EffectiveViewRect.Width() / 100.0f;
+	ViewUniformShaderParameters.MotionBlurNormalizedToPixel = Inputs.MotionBlurMax * EffectiveViewRect.Width() / 100.0f;
 
-	#if WITH_EDITOR
-		if (Family->bNullifyWorldSpacePosition)
-		{
-			// Forces world space position to 0 and view vector to up.
-			ViewUniformShaderParameters.SVPositionToTranslatedWorld = FMatrix44f(
-				FMatrix(FPlane(0, 0, 0, 0),
-					FPlane(0, 0, 0, 0),
-					FPlane(0, 0, 0, 0),
-					FPlane(0, 0, 1, 1)));
-		}
-		else
-	#endif 
+	if (Inputs.bNullifyWorldSpacePosition)
+	{
+		// Forces world space position to 0 and view vector to up.
+		ViewUniformShaderParameters.SVPositionToTranslatedWorld = FMatrix44f(
+			FMatrix(FPlane(0, 0, 0, 0),
+				FPlane(0, 0, 0, 0),
+				FPlane(0, 0, 0, 0),
+				FPlane(0, 0, 1, 1)));
+	}
+	else
 	{
 		// setup a matrix to transform float4(SvPosition.xyz,1) directly to TranslatedWorld (quality, performance as we don't need to convert or use interpolator)
 
@@ -2543,7 +2543,7 @@ void FSceneView::SetupViewRectUniformBufferParameters(FViewUniformShaderParamete
 		ViewUniformShaderParameters.TanAndInvTanHalfFOV = TanAndInvTanFOV;
 		ViewUniformShaderParameters.PrevTanAndInvTanHalfFOV = InPrevViewMatrices.GetTanAndInvTanHalfFOV();
 
-		if (IsPerspectiveProjection())
+		if (InViewMatrices.IsPerspectiveProjection())
 		{
 			ViewUniformShaderParameters.WorldDepthToPixelWorldRadius = FVector2f(TanAndInvTanFOV.X / float(EffectiveViewRect.Width()), 0.0f);
 			ViewUniformShaderParameters.ScreenRayLengthMultiplier = FVector4f(TanAndInvTanFOV.X, TanAndInvTanFOV.Y, 0, 0);
@@ -2564,20 +2564,21 @@ void FSceneView::SetupViewRectUniformBufferParameters(FViewUniformShaderParamete
 	ViewUniformShaderParameters.ScreenToViewSpace.Z = -((ViewUniformShaderParameters.ViewRectMin.X * ViewUniformShaderParameters.ViewSizeAndInvSize.Z * 2 * FovFixX) + FovFixX);
 	ViewUniformShaderParameters.ScreenToViewSpace.W = (ViewUniformShaderParameters.ViewRectMin.Y * ViewUniformShaderParameters.ViewSizeAndInvSize.W * 2 * FovFixY) + FovFixY;
 
-	ViewUniformShaderParameters.ViewResolutionFraction = EffectiveViewRect.Width() / (float)UnscaledViewRect.Width();
+	ViewUniformShaderParameters.ViewResolutionFraction = EffectiveViewRect.Width() / (float)Inputs.UnscaledViewRect.Width();
 }
 
-void FSceneView::SetupCommonViewUniformBufferParameters(
+void SetupCommonViewUniformBufferParameters(
 	FViewUniformShaderParameters& ViewUniformShaderParameters,
 	const FIntPoint& BufferSize,
 	int32 NumMSAASamples,
 	const FIntRect& EffectiveViewRect,
 	const FViewMatrices& InViewMatrices,
-	const FViewMatrices& InPrevViewMatrices) const
+	const FViewMatrices& InPrevViewMatrices,
+	const FSetupViewUniformParametersInputs& Inputs)
 {
 	QUICK_SCOPE_CYCLE_COUNTER(STAT_SetupCommonViewUniformBufferParameters);
-	FVector4f LocalDiffuseOverrideParameter = DiffuseOverrideParameter;
-	FVector2D LocalRoughnessOverrideParameter = RoughnessOverrideParameter;
+	FVector4f LocalDiffuseOverrideParameter = Inputs.DiffuseOverrideParameter;
+	FVector2D LocalRoughnessOverrideParameter = Inputs.RoughnessOverrideParameter;
 
 #if !(UE_BUILD_SHIPPING || UE_BUILD_TEST)
 	{
@@ -2625,7 +2626,7 @@ void FSceneView::SetupCommonViewUniformBufferParameters(
 	ViewUniformShaderParameters.ViewRight = (FVector3f)InViewMatrices.GetOverriddenTranslatedViewMatrix().GetColumn(0);
 	ViewUniformShaderParameters.HMDViewNoRollUp = (FVector3f)InViewMatrices.GetHMDViewMatrixNoRoll().GetColumn(1);
 	ViewUniformShaderParameters.HMDViewNoRollRight = (FVector3f)InViewMatrices.GetHMDViewMatrixNoRoll().GetColumn(0);
-	ViewUniformShaderParameters.InvDeviceZToWorldZTransform = InvDeviceZToWorldZTransform;
+	ViewUniformShaderParameters.InvDeviceZToWorldZTransform = Inputs.InvDeviceZToWorldZTransform;
 	FDFVector4 WorldViewOriginDF { (FVector4f)(InViewMatrices.GetOverriddenInvTranslatedViewMatrix().TransformPosition(FVector(0)) - InViewMatrices.GetPreViewTranslation()) };
 	ViewUniformShaderParameters.WorldViewOriginHigh = WorldViewOriginDF.High;
 	ViewUniformShaderParameters.WorldViewOriginLow = WorldViewOriginDF.Low;
@@ -2670,18 +2671,18 @@ void FSceneView::SetupCommonViewUniformBufferParameters(
 	ViewUniformShaderParameters.RelativePrevPreViewTranslationTO = FVector3f(InPrevViewMatrices.GetPreViewTranslation() + ViewTileOffset);
 
 	// Convert global clipping plane to translated world space
-	const FPlane4f TranslatedGlobalClippingPlane(GlobalClippingPlane.TranslateBy(InViewMatrices.GetPreViewTranslation()));
+	const FPlane4f TranslatedGlobalClippingPlane(Inputs.GlobalClippingPlane.TranslateBy(InViewMatrices.GetPreViewTranslation()));
 
 	ViewUniformShaderParameters.GlobalClippingPlane = FVector4f(TranslatedGlobalClippingPlane.X, TranslatedGlobalClippingPlane.Y, TranslatedGlobalClippingPlane.Z, -TranslatedGlobalClippingPlane.W);
 
 	ViewUniformShaderParameters.FieldOfViewWideAngles = FVector2f(2.f * InViewMatrices.ComputeHalfFieldOfViewPerAxis());	// LWC_TODO: Precision loss
 	ViewUniformShaderParameters.PrevFieldOfViewWideAngles = FVector2f(2.f * InPrevViewMatrices.ComputeHalfFieldOfViewPerAxis());	// LWC_TODO: Precision loss
 	ViewUniformShaderParameters.DiffuseOverrideParameter = LocalDiffuseOverrideParameter;
-	ViewUniformShaderParameters.SpecularOverrideParameter = SpecularOverrideParameter;
-	ViewUniformShaderParameters.NormalOverrideParameter = NormalOverrideParameter;
+	ViewUniformShaderParameters.SpecularOverrideParameter = Inputs.SpecularOverrideParameter;
+	ViewUniformShaderParameters.NormalOverrideParameter = Inputs.NormalOverrideParameter;
 	ViewUniformShaderParameters.RoughnessOverrideParameter = FVector2f(LocalRoughnessOverrideParameter);	// LWC_TODO: Precision loss
 	ViewUniformShaderParameters.WorldCameraMovementSinceLastFrame = FVector3f(InViewMatrices.GetViewOrigin() - InPrevViewMatrices.GetViewOrigin());
-	ViewUniformShaderParameters.CullingSign = bReverseCulling ? -1.0f : 1.0f;
+	ViewUniformShaderParameters.CullingSign = Inputs.bReverseCulling ? -1.0f : 1.0f;
 	ViewUniformShaderParameters.NearPlane = InViewMatrices.ComputeNearPlane();
 	ViewUniformShaderParameters.OrthoFarPlane = InViewMatrices.ComputeOrthoFarPlane();
 	ViewUniformShaderParameters.MaterialTextureMipBias = 0.0f;
@@ -2730,9 +2731,9 @@ void FSceneView::SetupCommonViewUniformBufferParameters(
 		(float)InViewMatrices.GetTemporalAAJitter().X, (float)InViewMatrices.GetTemporalAAJitter().Y,
 		(float)InPrevViewMatrices.GetTemporalAAJitter().X, (float)InPrevViewMatrices.GetTemporalAAJitter().Y );
 
-	ViewUniformShaderParameters.DebugViewModeMask = Family->UseDebugViewPS() ? 1 : 0;
-	ViewUniformShaderParameters.UnlitViewmodeMask = !Family->EngineShowFlags.Lighting || Family->EngineShowFlags.PathTracing ? 1 : 0;
-	ViewUniformShaderParameters.OutOfBoundsMask = Family->EngineShowFlags.VisualizeOutOfBoundsPixels ? 1 : 0;
+	ViewUniformShaderParameters.DebugViewModeMask = Inputs.DebugViewShaderMode != DVSM_None ? 1 : 0;
+	ViewUniformShaderParameters.UnlitViewmodeMask = Inputs.EngineShowFlags && (!Inputs.EngineShowFlags->Lighting || Inputs.EngineShowFlags->PathTracing) ? 1 : 0;
+	ViewUniformShaderParameters.OutOfBoundsMask = Inputs.EngineShowFlags && Inputs.EngineShowFlags->VisualizeOutOfBoundsPixels ? 1 : 0;
 
 #if !(UE_BUILD_SHIPPING || UE_BUILD_TEST)
 	float OverrideTimeMaterialExpression = CVarOverrideTimeMaterialExpressions.GetValueOnRenderThread();
@@ -2747,35 +2748,56 @@ void FSceneView::SetupCommonViewUniformBufferParameters(
 	else
 #endif
 	{
-		ViewUniformShaderParameters.PrevFrameGameTime = Family->Time.GetWorldTimeSeconds() - Family->Time.GetDeltaWorldTimeSeconds();
-		ViewUniformShaderParameters.PrevFrameRealTime = Family->Time.GetRealTimeSeconds() - Family->Time.GetDeltaRealTimeSeconds();
-		ViewUniformShaderParameters.GameTime = Family->Time.GetWorldTimeSeconds();
-		ViewUniformShaderParameters.RealTime = Family->Time.GetRealTimeSeconds();
-		ViewUniformShaderParameters.DeltaTime = Family->Time.GetDeltaWorldTimeSeconds();
+		ViewUniformShaderParameters.PrevFrameGameTime = Inputs.Time.GetWorldTimeSeconds() - Inputs.Time.GetDeltaWorldTimeSeconds();
+		ViewUniformShaderParameters.PrevFrameRealTime = Inputs.Time.GetRealTimeSeconds() - Inputs.Time.GetDeltaRealTimeSeconds();
+		ViewUniformShaderParameters.GameTime = Inputs.Time.GetWorldTimeSeconds();
+		ViewUniformShaderParameters.RealTime = Inputs.Time.GetRealTimeSeconds();
+		ViewUniformShaderParameters.DeltaTime = Inputs.Time.GetDeltaWorldTimeSeconds();
 	}
 
 #if !(UE_BUILD_SHIPPING || UE_BUILD_TEST)
-	static FIntPoint LockedCursorPos = CursorPos;
-	if (CVarFreezeMouseCursor.GetValueOnRenderThread() == 0 && CursorPos.X >= 0 && CursorPos.Y >= 0)
+	static FIntPoint LockedCursorPos = Inputs.CursorPosition;
+	if (CVarFreezeMouseCursor.GetValueOnRenderThread() == 0 && Inputs.CursorPosition.X >= 0 && Inputs.CursorPosition.Y >= 0)
 	{
-		LockedCursorPos = CursorPos;
+		LockedCursorPos = Inputs.CursorPosition;
 	}
 	ViewUniformShaderParameters.CursorPosition = LockedCursorPos;
 #endif
 
 	ViewUniformShaderParameters.Random = FMath::Rand();
 	// FrameNumber corresponds to how many times FRendererModule::BeginRenderingViewFamilies has been called, so multi views of the same frame have incremental values.
-	ViewUniformShaderParameters.FrameNumber = Family->FrameNumber;
+	ViewUniformShaderParameters.FrameNumber = Inputs.FrameNumber;
 	// FrameCounter is incremented once per engine tick, so multi views of the same frame have the same value.
-	ViewUniformShaderParameters.FrameCounter = Family->FrameCounter;
-	ViewUniformShaderParameters.WorldIsPaused = Family->bWorldIsPaused;
+	ViewUniformShaderParameters.FrameCounter = Inputs.FrameCounter;
+	ViewUniformShaderParameters.WorldIsPaused = Inputs.bWorldIsPaused;
 	//Set bCameraCut if we switch projection type to ensure histories are updated.
-	ViewUniformShaderParameters.CameraCut = bCameraCut ? 1 : (InViewMatrices.IsPerspectiveProjection() != InPrevViewMatrices.IsPerspectiveProjection());
+	ViewUniformShaderParameters.CameraCut = Inputs.bCameraCut ? 1 : (InViewMatrices.IsPerspectiveProjection() != InPrevViewMatrices.IsPerspectiveProjection());
 
 	ViewUniformShaderParameters.MinRoughness = FMath::Clamp(CVarGlobalMinRoughnessOverride.GetValueOnRenderThread(), 0.02f, 1.0f);
 
 	//to tail call keep the order and number of parameters of the caller function
-	SetupViewRectUniformBufferParameters(ViewUniformShaderParameters, BufferSize, EffectiveViewRect, InViewMatrices, InPrevViewMatrices);
+	SetupViewRectUniformBufferParameters(ViewUniformShaderParameters, BufferSize, EffectiveViewRect, InViewMatrices, InPrevViewMatrices, Inputs);
+}
+
+void FSceneView::SetupViewRectUniformBufferParameters(
+	FViewUniformShaderParameters& ViewUniformShaderParameters,
+	const FIntPoint& InBufferSize,
+	const FIntRect& InEffectiveViewRect,
+	const FViewMatrices& InViewMatrices,
+	const FViewMatrices& InPrevViewMatrices) const
+{
+	::SetupViewRectUniformBufferParameters(ViewUniformShaderParameters, InBufferSize, InEffectiveViewRect, InViewMatrices, InPrevViewMatrices, FSetupViewUniformParametersInputs::Create(*this));
+}
+
+void FSceneView::SetupCommonViewUniformBufferParameters(
+	FViewUniformShaderParameters& ViewUniformShaderParameters,
+	const FIntPoint& BufferSize,
+	int32 NumMSAASamples,
+	const FIntRect& EffectiveViewRect,
+	const FViewMatrices& InViewMatrices,
+	const FViewMatrices& InPrevViewMatrices) const
+{
+	::SetupCommonViewUniformBufferParameters(ViewUniformShaderParameters, BufferSize, NumMSAASamples, EffectiveViewRect, InViewMatrices, InPrevViewMatrices, FSetupViewUniformParametersInputs::Create(*this));
 }
 
 bool FSceneView::HasValidEyeAdaptationTexture() const
diff --git a/Engine/Source/Runtime/Engine/Public/SceneView.h b/Engine/Source/Runtime/Engine/Public/SceneView.h
--- a/Engine/Source/Runtime/Engine/Public/SceneView.h
+++ b/Engine/Source/Runtime/Engine/Public/SceneView.h
@@ -1848,7 +1848,7 @@ public:
 		const FViewMatrices& InViewMatrices,
 		const FViewMatrices& InPrevViewMatrice) const;
 
-	ENGINE_API FVector4f GetScreenPositionScaleBias(const FIntPoint& BufferSize, const FIntRect& ViewRect) const;
+	static ENGINE_API FVector4f GetScreenPositionScaleBias(const FIntPoint& BufferSize, const FIntRect& ViewRect);
 
 	/** 
 	 * Populates the uniform buffer prameters common to all scene view use cases
@@ -2497,3 +2497,69 @@ public:
 	/** Destructor. */
 	ENGINE_API virtual ~FSceneViewFamilyContext();
 };
+
+struct FSetupViewUniformParametersInputs
+{
+	static FSetupViewUniformParametersInputs Create(const FSceneView& View)
+	{
+		return FSetupViewUniformParametersInputs
+		{
+			  .EngineShowFlags = &View.Family->EngineShowFlags
+			, .UnscaledViewRect = View.UnscaledViewRect
+			, .Time = View.Family->Time
+			, .CursorPosition = View.CursorPos
+			, .GlobalClippingPlane = View.GlobalClippingPlane
+			, .InvDeviceZToWorldZTransform = View.InvDeviceZToWorldZTransform
+			, .DiffuseOverrideParameter = View.DiffuseOverrideParameter
+			, .NormalOverrideParameter = View.NormalOverrideParameter
+			, .SpecularOverrideParameter = View.SpecularOverrideParameter
+			, .RoughnessOverrideParameter = View.RoughnessOverrideParameter
+			, .DebugViewShaderMode = View.Family->GetDebugViewShaderMode()
+			, .FrameCounter = View.Family->FrameCounter
+			, .FrameNumber = View.Family->FrameNumber
+			, .FOV = View.FOV
+			, .MotionBlurMax = View.FinalPostProcessSettings.MotionBlurMax
+#if WITH_EDITOR
+			, .bNullifyWorldSpacePosition = View.Family->bNullifyWorldSpacePosition
+#endif
+			, .bReverseCulling = View.bReverseCulling
+			, .bCameraCut = View.bCameraCut
+			, .bWorldIsPaused = View.Family->bWorldIsPaused
+		};
+	}
+
+	const FEngineShowFlags* EngineShowFlags  = nullptr;
+	FIntRect UnscaledViewRect;
+	FGameTime Time;
+	FIntPoint CursorPosition                 = FIntPoint::ZeroValue;
+	FPlane GlobalClippingPlane               = FPlane(0, 0, 0, 0);
+	FVector4f InvDeviceZToWorldZTransform    = FVector4f(0, 0, 0, 0);
+	FVector4f DiffuseOverrideParameter       = FVector4f(0, 0, 0, 1);
+	FVector4f NormalOverrideParameter        = FVector4f(0, 0, 0, 1);
+	FVector4f SpecularOverrideParameter      = FVector4f(0, 0, 0, 1);
+	FVector2D RoughnessOverrideParameter     = FVector2D(0, 1);
+	EDebugViewShaderMode DebugViewShaderMode = EDebugViewShaderMode::DVSM_None;
+	uint64 FrameCounter                      = 0;
+	uint32 FrameNumber                       = 0;
+	float FOV                                = 0.0f;
+	float MotionBlurMax                      = 0.0f;
+	bool bNullifyWorldSpacePosition          = false;
+	bool bReverseCulling                     = false;
+	bool bCameraCut                          = false;
+	bool bWorldIsPaused                      = false;
+};
+
+ENGINE_API void SetupCommonViewUniformBufferParameters(FViewUniformShaderParameters& ViewUniformShaderParameters,
+	const FIntPoint& InBufferSize,
+	int32 NumMSAASamples,
+	const FIntRect& InEffectiveViewRect,
+	const FViewMatrices& InViewMatrices,
+	const FViewMatrices& InPrevViewMatrices,
+	const FSetupViewUniformParametersInputs& Inputs);
+
+ENGINE_API void SetupViewRectUniformBufferParameters(FViewUniformShaderParameters& ViewUniformShaderParameters,
+	const FIntPoint& InBufferSize,
+	const FIntRect& InEffectiveViewRect,
+	const FViewMatrices& InViewMatrices,
+	const FViewMatrices& InPrevViewMatrice,
+	const FSetupViewUniformParametersInputs& Inputs);
\ No newline at end of file
diff --git a/Engine/Source/Runtime/Renderer/Private/ShaderBaseClasses.cpp b/Engine/Source/Runtime/Renderer/Private/ShaderBaseClasses.cpp
--- a/Engine/Source/Runtime/Renderer/Private/ShaderBaseClasses.cpp
+++ b/Engine/Source/Runtime/Renderer/Private/ShaderBaseClasses.cpp
@@ -254,16 +254,15 @@ void FMaterialShader::SetParameters(
 	FRHIBatchedShaderParameters& BatchedParameters,
 	const FMaterialRenderProxy* MaterialRenderProxy,
 	const FMaterial& Material,
-	const FSceneView& View)
+	const FSceneInterface* Scene)
 {
-	const ERHIFeatureLevel::Type FeatureLevel = View.GetFeatureLevel();
+	const ERHIFeatureLevel::Type FeatureLevel = Scene ? Scene->GetFeatureLevel() : GMaxRHIFeatureLevel;
 	FMaterialShaderMap* ShaderMap = Material.GetRenderingThreadShaderMap();
 	checkf(ShaderMap, TEXT("RenderingThreadShaderMap: %i"), ShaderMap ? 1 : 0);
 	checkf(ShaderMap->IsValidForRendering(true) && Material.GetFeatureLevel() == FeatureLevel, TEXT("IsValid:%i, MaterialFeatureLevel:%i, FeatureLevel:%i"), ShaderMap->IsValidForRendering() ? 1 : 0, (ERHIFeatureLevel::Type)Material.GetFeatureLevel(), FeatureLevel);
 
 	FUniformExpressionCache* UniformExpressionCache = &MaterialRenderProxy->UniformExpressionCache[FeatureLevel];
 	bool bUniformExpressionCacheNeedsDelete = false;
-	//bool bForceExpressionEvaluation = false;
 
 #if DO_CHECK
 	if (UniformExpressionCache->CachedUniformExpressionShaderMap != ShaderMap)
@@ -307,15 +306,12 @@ void FMaterialShader::SetParameters(
 		}
 
 		check(ParameterCollectionUniformBuffers.Num() >= ParameterCollectionsNum);
-
-		
-
 		int32 NumToSet = FMath::Min(ParameterCollectionUniformBuffers.Num(), ParameterCollections.Num());
 
 		// Find each referenced parameter collection's uniform buffer in the scene and set the parameter
 		for (int32 CollectionIndex = 0; CollectionIndex < NumToSet; CollectionIndex++)
-		{			
-			FRHIUniformBuffer* UniformBuffer = GetParameterCollectionBuffer(ParameterCollections[CollectionIndex], View.Family->Scene);
+		{
+			FRHIUniformBuffer* UniformBuffer = GetParameterCollectionBuffer(ParameterCollections[CollectionIndex], Scene);
 
 			if (!UniformBuffer)
 			{
diff --git a/Engine/Source/Runtime/Renderer/Public/MaterialShader.h b/Engine/Source/Runtime/Renderer/Public/MaterialShader.h
--- a/Engine/Source/Runtime/Renderer/Public/MaterialShader.h
+++ b/Engine/Source/Runtime/Renderer/Public/MaterialShader.h
@@ -87,12 +87,22 @@ public:
 		RHICmdList.SetBatchedShaderParameters(ShaderRHI, BatchedParameters);
 	}
 
+	/** Sets pixel parameters that are material specific but not FMeshBatch specific. */
+	void SetParameters(
+		FRHIBatchedShaderParameters& BatchedParameters,
+		const FMaterialRenderProxy* MaterialRenderProxy,
+		const FMaterial& Material,
+		const FSceneView& View)
+	{
+		SetParameters(BatchedParameters, MaterialRenderProxy, Material, View.Family->Scene);
+	}
+	
 	/** Sets pixel parameters that are material specific but not FMeshBatch specific. */
 	RENDERER_API void SetParameters(
 		FRHIBatchedShaderParameters& BatchedParameters,
 		const FMaterialRenderProxy* MaterialRenderProxy, 
 		const FMaterial& Material,
-		const FSceneView& View);
+		const FSceneInterface* Scene);
 
 	/** Sets pixel parameters that are material specific but not FMeshBatch specific. */
 	template<typename TRHIShader, typename TRHICommandList>
