[Material Editor] Fixed crash in void UMaterialInstance::AddParameterValueInternal() due to missing StaticSwitch handling case.

diff --git a/Engine/Source/Runtime/Engine/Private/Materials/MaterialInstance.cpp b/Engine/Source/Runtime/Engine/Private/Materials/MaterialInstance.cpp
--- a/Engine/Source/Runtime/Engine/Private/Materials/MaterialInstance.cpp
+++ b/Engine/Source/Runtime/Engine/Private/Materials/MaterialInstance.cpp
@@ -3801,6 +3801,7 @@ void UMaterialInstance::AddParameterValueInternal(const FMaterialParameterInfo&
 	case EMaterialParameterType::Font: FontParameterValues.Emplace(ParameterInfo, Value.Font.Value, Value.Font.Page); break;
 	case EMaterialParameterType::RuntimeVirtualTexture: RuntimeVirtualTextureParameterValues.Emplace(ParameterInfo, Value.RuntimeVirtualTexture); break;
 	case EMaterialParameterType::SparseVolumeTexture: SparseVolumeTextureParameterValues.Emplace(ParameterInfo, Value.SparseVolumeTexture); break;
+	case EMaterialParameterType::StaticSwitch: break;
 	default: checkNoEntry();
 	}
 }
