Allow Debug View shaders on ES3_1 preview platforms

diff --git a/Engine/Source/Runtime/Engine/Private/DebugViewModeHelpers.cpp b/Engine/Source/Runtime/Engine/Private/DebugViewModeHelpers.cpp
--- a/Engine/Source/Runtime/Engine/Private/DebugViewModeHelpers.cpp
+++ b/Engine/Source/Runtime/Engine/Private/DebugViewModeHelpers.cpp
@@ -131,11 +131,6 @@ bool ShouldCompileDebugViewModeShader(const FMeshMaterialShaderPermutationParame
 		return false;
 	}
 
-	if (Parameters.MaterialParameters.FeatureLevel < ERHIFeatureLevel::SM5)
-	{
-		return false;
-	}
-
 	if (!EnumHasAllFlags(Parameters.Flags, EShaderPermutationFlags::HasEditorOnlyData))
 	{
 		// Debug view shaders only in editor
