Noise material expression only accepts 3D vector as position input.  Fix tooltip and the material translator to reflect this.

- Update tooltip to just say "3D"
- Change the translator function to only accept 3D position vectors.  That way it results in a material translation error instead of a shader compiler error.

diff --git a/Engine/Source/Runtime/Engine/Private/Materials/HLSLMaterialTranslator.cpp b/Engine/Source/Runtime/Engine/Private/Materials/HLSLMaterialTranslator.cpp
--- a/Engine/Source/Runtime/Engine/Private/Materials/HLSLMaterialTranslator.cpp
+++ b/Engine/Source/Runtime/Engine/Private/Materials/HLSLMaterialTranslator.cpp
@@ -11469,7 +11469,7 @@ int32 FHLSLMaterialTranslator::Noise(int32 Position, EPositionOrigin PositionOri
 
 	return AddCodeChunk(MCT_Float, 
 		TEXT("MaterialExpressionNoise(%s,%s,%s,%s,%s,%s,%s,%s,%s,%s,%s,%s)"), 
-		*GetParameterCode(Position),
+		*CoerceParameter(Position, MCT_Float3),
 		*GetParameterCode(ScaleConst),
 		*GetParameterCode(QualityConst),
 		*GetParameterCode(NoiseFunctionConst),
diff --git a/Engine/Source/Runtime/Engine/Public/Materials/MaterialExpressionNoise.h b/Engine/Source/Runtime/Engine/Public/Materials/MaterialExpressionNoise.h
--- a/Engine/Source/Runtime/Engine/Public/Materials/MaterialExpressionNoise.h
+++ b/Engine/Source/Runtime/Engine/Public/Materials/MaterialExpressionNoise.h
@@ -60,7 +60,7 @@ class UMaterialExpressionNoise : public UMaterialExpression
 {
 	GENERATED_UCLASS_BODY()
 
-	/** 2 to 3 dimensional vector */
+	/** 3 dimensional vector */
 	UPROPERTY()
 	FExpressionInput Position;
 
