Fix for blend validation returning false in new layering system when Blend function input numbers are not exactly 2 (which conflicts with newer behaviour). Also swap out the test project's MLB_Vertical in favour of the EngineMaterial version.

#virtualized

diff --git a/Engine/Source/Runtime/Engine/Private/Materials/MaterialExpressions.cpp b/Engine/Source/Runtime/Engine/Private/Materials/MaterialExpressions.cpp
--- a/Engine/Source/Runtime/Engine/Private/Materials/MaterialExpressions.cpp
+++ b/Engine/Source/Runtime/Engine/Private/Materials/MaterialExpressions.cpp
@@ -8438,6 +8438,7 @@ bool UMaterialExpressionMaterialAttributeLayers::ValidateLayerConfiguration(FMat
 					bIsValid = false;
 				}
 			}
+#if !ENABLE_MATERIAL_LAYER_PROTOTYPE
 			else
 			{
 				TArray<UMaterialExpressionFunctionInput*> InputExpressions;
@@ -8448,6 +8449,7 @@ bool UMaterialExpressionMaterialAttributeLayers::ValidateLayerConfiguration(FMat
 					bIsValid = false;
 				}
 			}
+#endif
 
 			if (LayerStates[LayerIndex])
 			{
@@ -8479,9 +8481,13 @@ bool UMaterialExpressionMaterialAttributeLayers::ValidateLayerConfiguration(FMat
 			{
 				TArray<UMaterialExpressionFunctionInput*> InputExpressions;
 				Blend->GetAllExpressionsOfType<UMaterialExpressionFunctionInput>(InputExpressions, false);
+#if ENABLE_MATERIAL_LAYER_PROTOTYPE
+				if (InputExpressions.Num() < 2)
+#else
 				if (InputExpressions.Num() != 2)
+#endif
 				{
-					COMPILER_OR_LOG_ERROR(TEXT("Blend %i, %s, must have two MaterialAttributes inputs only."), BlendIndex, *Blend->GetName());
+					COMPILER_OR_LOG_ERROR(TEXT("Blend %i, %s, must have two MaterialAttributes inputs."), BlendIndex, *Blend->GetName());
 					bIsValid = false;
 				}
 			}
@@ -18670,9 +18676,7 @@ void FMaterialLayersFunctions::Validate(const FMaterialLayersFunctionsRuntimeDat
 {
 	if (Runtime.Layers.Num() > 0)
 	{
-#if ENABLE_MATERIAL_LAYER_PROTOTYPE
-		check(Runtime.Blends.Num() == Runtime.Layers.Num());
-#else
+#if !ENABLE_MATERIAL_LAYER_PROTOTYPE
 		check(Runtime.Blends.Num() == Runtime.Layers.Num() - 1);
 #endif // ENABLE_MATERIAL_LAYER_PROTOTYPE
 		check(Runtime.Layers.Num() == EditorOnly.LayerStates.Num());
