[Shaders] Allow shader compilation to suspend and resume workers on demand.

This change is the first step to handle OOM errors when hitting memory limitations due to large job queues in too many shader compile workers.
The primary new functions to suspend and resume workers are SuspendWorkersAndBacklogJobs() and ResumeSuspendedWorkers() in FShaderCompileThreadRunnable.
These functions manage a list of backlogged jobs that need to be picked up before any new jobs from the shader compiling manager queue.

Also generalize the use of worker input/output files:
- When suspended workers are resumed, their potentially stale output files must be deleted before they start with new work. Otherwise, those output files will be linked to the wrong compile jobs.
- Use FPaths::Combine and generalized function GetWorkingDirectoryForWorker() everywhere to access the SCW I/O files.

#rnx

diff --git a/Engine/Source/Runtime/Engine/Private/ShaderCompiler/ShaderCompilerThreadRunnable.cpp b/Engine/Source/Runtime/Engine/Private/ShaderCompiler/ShaderCompilerThreadRunnable.cpp
--- a/Engine/Source/Runtime/Engine/Private/ShaderCompiler/ShaderCompilerThreadRunnable.cpp
+++ b/Engine/Source/Runtime/Engine/Private/ShaderCompiler/ShaderCompilerThreadRunnable.cpp
@@ -36,6 +36,9 @@ static constexpr int32 GSingleThreadedRunsDisabled = -2;
 static constexpr int32 GSingleThreadedRunsIncreaseFactor = 8;
 static constexpr int32 GSingleThreadedRunsMaxCount = (1 << 24);
 
+static const TCHAR* GWorkerInputFilename = TEXT("WorkerInputOnly.in");
+static const TCHAR* GWorkerOutputFilename = TEXT("WorkerOutputOnly.out");
+
 
 /** Information tracked for each shader compile worker process instance. */
 struct FShaderCompileWorkerInfo
@@ -53,35 +56,54 @@ struct FShaderCompileWorkerInfo
 	bool bComplete;
 
 	/** Whether this worker is available for new jobs. It will be false when shutting down the worker. */
-	bool bAvailable; 
+	bool bAvailable;
 
 	/** Time at which the worker started the most recent batch of tasks. */
 	double StartTime;
 
 	/** Time at which the worker ended the most recent batch of tasks. */
-	double FinishTime = 0.0;
+	double FinishTime;
 
 	/** Jobs that this worker is responsible for compiling. */
 	TArray<FShaderCommonCompileJobPtr> QueuedJobs;
 
 	FShaderCompileWorkerInfo() :
-		bIssuedTasksToWorker(false),		
+		bIssuedTasksToWorker(false),
 		bLaunchedWorker(false),
 		bComplete(false),
 		bAvailable(true),
-		StartTime(0)
+		StartTime(0.0),
+		FinishTime(0.0)
 	{
 	}
 
 	// warning: not virtual
 	~FShaderCompileWorkerInfo()
 	{
-		if(WorkerProcess.IsValid())
+		TerminateWorkerProcess();
+	}
+
+	void TerminateWorkerProcess()
+	{
+		if (WorkerProcess.IsValid())
 		{
 			FPlatformProcess::TerminateProc(WorkerProcess);
 			FPlatformProcess::CloseProc(WorkerProcess);
+			WorkerProcess = FProcHandle();
 		}
 	}
+
+	int32 CloseWorkerProcess()
+	{
+		int32 ReturnCode = 0;
+		if (WorkerProcess.IsValid())
+		{
+			FPlatformProcess::GetProcReturnCode(WorkerProcess, &ReturnCode);
+			FPlatformProcess::CloseProc(WorkerProcess);
+			WorkerProcess = FProcHandle();
+		}
+		return ReturnCode;
+	}
 };
 
 FShaderCompileThreadRunnableBase::FShaderCompileThreadRunnableBase(FShaderCompilingManager* InManager)
@@ -205,27 +227,100 @@ int32 FShaderCompileThreadRunnable::PullTasksFromQueue()
 {
 	TRACE_CPUPROFILER_EVENT_SCOPE(FShaderCompileThreadRunnable::PullTasksFromQueue);
 
+	auto SignalWorkerTasksToBeSubmitted = [this](FShaderCompileWorkerInfo& WorkerInfo, int32 WorkerIndex)
+		{
+			// Update the worker state as having new tasks that need to be issued					
+			// don't reset worker app ID, because the shadercompileworkers don't shutdown immediately after finishing a single job queue.
+			WorkerInfo.bIssuedTasksToWorker = false;
+			WorkerInfo.bLaunchedWorker = false;
+			WorkerInfo.StartTime = FPlatformTime::Seconds();
+
+			if (WorkerInfo.FinishTime > 0.0)
+			{
+				const double WorkerIdleTime = WorkerInfo.StartTime - WorkerInfo.FinishTime;
+				GShaderCompilerStats->RegisterLocalWorkerIdleTime(WorkerIdleTime);
+				if (Manager->bLogJobCompletionTimes)
+				{
+					UE_LOG(LogShaderCompilers, Display, TEXT("  Worker (%d/%d) started working after being idle for %fs"), WorkerIndex + 1, WorkerInfos.Num(), WorkerIdleTime);
+				}
+			}
+		};
+
 	FScopeLock WorkerScopeLock(&WorkerInfosLock); // Must be entered before CompileQueueSection
+
 	int32 NumActiveThreads = 0;
 	int32 NumJobsStarted[NumShaderCompileJobPriorities] = { 0 };
 	{
 		// Enter the critical section so we can access the input and output queues
 		FScopeLock Lock(&Manager->CompileQueueSection);
 
-		const int32 NumWorkersToFeed = Manager->bCompilingDuringGame ? Manager->NumShaderCompilingThreadsDuringGame : WorkerInfos.Num();
+		const int32 NumWorkersToFeed = Manager->bCompilingDuringGame ? Manager->NumShaderCompilingThreadsDuringGame : GetNumberOfAvailableWorkersUnsafe();
+
+		// Pull tasks from backlogged queue first
+		if (!BackloggedJobs.IsEmpty())
+		{
+			// Try to distribute the work evenly between the workers
+			const int32 PriorityIndex = static_cast<int32>(EShaderCompileJobPriority::Normal);
+			const int32 NumJobsPerWorker = FMath::DivideAndRoundUp(BackloggedJobs.Num(), NumWorkersToFeed);
+
+			int32 NumWorkersToPickupBacklog = 0;
+			int32 NumPickedupBackloggedJobs = 0;
+
+			for (int32 WorkerIndex = 0; WorkerIndex < WorkerInfos.Num(); WorkerIndex++)
+			{
+				FShaderCompileWorkerInfo& CurrentWorkerInfo = *WorkerInfos[WorkerIndex];
+
+				// If this worker doesn't have any queued jobs, look for more in the input queue
+				if (CurrentWorkerInfo.QueuedJobs.Num() == 0 && CurrentWorkerInfo.bAvailable)
+				{
+					check(!CurrentWorkerInfo.bComplete);
 
+					if (BackloggedJobs.Num() > 0)
+					{
+						const int32 MaxNumJobs = FMath::Min3(NumJobsPerWorker, BackloggedJobs.Num(), Manager->MaxShaderJobBatchSize);
+
+						// Dequeue backlogged jobs and send them to worker
+						CurrentWorkerInfo.QueuedJobs.Reserve(CurrentWorkerInfo.QueuedJobs.Num() + MaxNumJobs);
+						for (int32 JobIndex = 0; JobIndex < MaxNumJobs; ++JobIndex)
+						{
+							CurrentWorkerInfo.QueuedJobs.Add(BackloggedJobs.Pop());
+						}
+						NumJobsStarted[PriorityIndex] += MaxNumJobs;
+
+						NumPickedupBackloggedJobs += MaxNumJobs;
+						NumWorkersToPickupBacklog += 1;
+
+						SignalWorkerTasksToBeSubmitted(CurrentWorkerInfo, WorkerIndex);
+					}
+				}
+			}
+
+			if (NumPickedupBackloggedJobs > 0)
+			{
+				UE_LOG(
+					LogShaderCompilers, Verbose, TEXT("Picked up %d backlogged compile %s and distributed them over %d %s"),
+					NumPickedupBackloggedJobs,
+					NumPickedupBackloggedJobs == 1 ? TEXT("job") : TEXT("jobs"),
+					NumWorkersToPickupBacklog,
+					NumWorkersToPickupBacklog == 1 ? TEXT("worker") : TEXT("workers")
+				);
+			}
+		}
+
+		// Pull tasks from compiling manager queue
 		for (int32 PriorityIndex = MaxPriorityIndex; PriorityIndex >= MinPriorityIndex; --PriorityIndex)
 		{
 			int32 NumPendingJobs = Manager->AllJobs.GetNumPendingJobs((EShaderCompileJobPriority)PriorityIndex);
+
 			// Try to distribute the work evenly between the workers
-			const auto NumJobsPerWorker = (NumPendingJobs / NumWorkersToFeed) + 1;
+			const int32 NumJobsPerWorker = FMath::DivideAndRoundUp(NumPendingJobs, NumWorkersToFeed);
 
 			for (int32 WorkerIndex = 0; WorkerIndex < WorkerInfos.Num(); WorkerIndex++)
 			{
 				FShaderCompileWorkerInfo& CurrentWorkerInfo = *WorkerInfos[WorkerIndex];
 
 				// If this worker doesn't have any queued jobs, look for more in the input queue
-				if (CurrentWorkerInfo.QueuedJobs.Num() == 0 && CurrentWorkerInfo.bAvailable && WorkerIndex < NumWorkersToFeed)
+				if (CurrentWorkerInfo.QueuedJobs.Num() == 0 && CurrentWorkerInfo.bAvailable)
 				{
 					check(!CurrentWorkerInfo.bComplete);
 
@@ -243,22 +338,7 @@ int32 FShaderCompileThreadRunnable::PullTasksFromQueue()
 
 						NumJobsStarted[PriorityIndex] += Manager->AllJobs.GetPendingJobs(EShaderCompilerWorkerType::LocalThread, (EShaderCompileJobPriority)PriorityIndex, 1, MaxNumJobs, CurrentWorkerInfo.QueuedJobs);
 
-						// Update the worker state as having new tasks that need to be issued					
-						// don't reset worker app ID, because the shadercompileworkers don't shutdown immediately after finishing a single job queue.
-						CurrentWorkerInfo.bIssuedTasksToWorker = false;
-						CurrentWorkerInfo.bLaunchedWorker = false;
-						CurrentWorkerInfo.StartTime = FPlatformTime::Seconds();
-						NumActiveThreads++;
-
-						if (CurrentWorkerInfo.FinishTime > 0.0)
-						{
-							const double WorkerIdleTime = CurrentWorkerInfo.StartTime - CurrentWorkerInfo.FinishTime;
-							GShaderCompilerStats->RegisterLocalWorkerIdleTime(WorkerIdleTime);
-							if (Manager->bLogJobCompletionTimes)
-							{
-								UE_LOG(LogShaderCompilers, Display, TEXT("  Worker (%d/%d) started working after being idle for %fs"), WorkerIndex + 1, WorkerInfos.Num(), WorkerIdleTime);
-							}
-						}
+						SignalWorkerTasksToBeSubmitted(CurrentWorkerInfo, WorkerIndex);
 					}
 				}
 			}
@@ -393,7 +473,7 @@ void FShaderCompileThreadRunnable::WriteNewTasks()
 			TRACE_CPUPROFILER_EVENT_SCOPE(ShaderCompiler.WriteNewTasksForWorker);
 			CurrentWorkerInfo.bIssuedTasksToWorker = true;
 
-			const FString WorkingDirectory = Manager->AbsoluteShaderBaseWorkingDirectory + FString::FromInt(WorkerIndex);
+			const FString WorkingDirectory = GetWorkingDirectoryForWorker(WorkerIndex);
 
 			// To make sure that the process waiting for input file won't try to read it until it's ready
 			// we use a temp file name during writing.
@@ -402,7 +482,7 @@ void FShaderCompileThreadRunnable::WriteNewTasks()
 			{
 				FGuid Guid;
 				FPlatformMisc::CreateGuid(Guid);
-				TransferFileName = WorkingDirectory + Guid.ToString();
+				TransferFileName = FPaths::Combine(WorkingDirectory, Guid.ToString());
 			} while (IFileManager::Get().FileSize(*TransferFileName) != INDEX_NONE);
 
 			// Write out the file that the worker app is waiting for, which has all the information needed to compile the shader.
@@ -460,7 +540,7 @@ void FShaderCompileThreadRunnable::WriteNewTasks()
 #endif
 
 			// Change the transfer file name to proper one
-			FString ProperTransferFileName = WorkingDirectory / TEXT("WorkerInputOnly.in");
+			FString ProperTransferFileName = FPaths::Combine(WorkingDirectory, GWorkerInputFilename);
 			if (!IFileManager::Get().Move(*ProperTransferFileName, *TransferFileName))
 			{
 				uint64 TotalDiskSpace = 0;
@@ -520,8 +600,7 @@ bool FShaderCompileThreadRunnable::LaunchWorkersIfNeeded()
 			// Also, use the opportunity to free OS resources by cleaning up handles of no more running processes
 			if (CurrentWorkerInfo.WorkerProcess.IsValid() && !FShaderCompilingManager::IsShaderCompilerWorkerRunning(CurrentWorkerInfo.WorkerProcess))
 			{
-				FPlatformProcess::CloseProc(CurrentWorkerInfo.WorkerProcess);
-				CurrentWorkerInfo.WorkerProcess = FProcHandle();
+				CurrentWorkerInfo.CloseWorkerProcess();
 			}
 			continue;
 		}
@@ -538,15 +617,11 @@ bool FShaderCompileThreadRunnable::LaunchWorkersIfNeeded()
 			if (CurrentWorkerInfo.WorkerProcess.IsValid())
 			{
 				// shader compiler exited one way or another, so clear out the stale PID.
-				int32 ReturnCode = 0;
-				FPlatformProcess::GetProcReturnCode(CurrentWorkerInfo.WorkerProcess, &ReturnCode);
-				FPlatformProcess::CloseProc(CurrentWorkerInfo.WorkerProcess);
-				CurrentWorkerInfo.WorkerProcess = FProcHandle();
+				const int32 ReturnCode = CurrentWorkerInfo.CloseWorkerProcess();
 
 				if (CurrentWorkerInfo.bLaunchedWorker)
 				{
-					const FString WorkingDirectory = Manager->AbsoluteShaderBaseWorkingDirectory + FString::FromInt(WorkerIndex) + TEXT("/");
-					const FString OutputFileNameAndPath = WorkingDirectory + TEXT("WorkerOutputOnly.out");
+					const FString OutputFileNameAndPath = FPaths::Combine(GetWorkingDirectoryForWorker(WorkerIndex), GWorkerOutputFilename);
 
 					if (FPlatformFileManager::Get().GetPlatformFile().FileExists(*OutputFileNameAndPath))
 					{
@@ -566,18 +641,17 @@ bool FShaderCompileThreadRunnable::LaunchWorkersIfNeeded()
 
 			if (bLaunchAgain)
 			{
-				const FString WorkingDirectory = Manager->ShaderBaseWorkingDirectory + FString::FromInt(WorkerIndex) + TEXT("/");
-				FString InputFileName(TEXT("WorkerInputOnly.in"));
-				FString OutputFileName(TEXT("WorkerOutputOnly.out"));
+				constexpr bool bRelativePath = true;
+				const FString WorkingDirectory = GetWorkingDirectoryForWorker(WorkerIndex, bRelativePath);
 				
 				if (bClearStaleOutputs)
 				{
 					// Delete any potential stale output files; these can persist if we had previously abandoned running with workers due to an unexpected SCW termination.
-					IFileManager::Get().Delete(*(WorkingDirectory / OutputFileName));
+					DiscardWorkerOutputFile(WorkerIndex);
 				}
 
 				// Store the handle with this thread so that we will know not to launch it again
-				CurrentWorkerInfo.WorkerProcess = Manager->LaunchWorker(WorkingDirectory, Manager->ProcessId, WorkerIndex, InputFileName, OutputFileName);
+				CurrentWorkerInfo.WorkerProcess = Manager->LaunchWorker(WorkingDirectory, Manager->ProcessId, WorkerIndex, GWorkerInputFilename, GWorkerOutputFilename);
 				CurrentWorkerInfo.bLaunchedWorker = true;
 
 				NumberLaunched++;
@@ -627,8 +701,7 @@ int32 FShaderCompileThreadRunnable::ReadAvailableResults()
 		{
 			// Distributed compiles always use the same directory
 			// 'Only' indicates to the worker that it should log and continue checking for the input file after the first one is processed
-			TStringBuilder<512> OutputFileNameAndPath;
-			OutputFileNameAndPath << Manager->AbsoluteShaderBaseWorkingDirectory << WorkerIndex << TEXT("/WorkerOutputOnly.out");
+			const FString OutputFileNameAndPath = FPaths::Combine(GetWorkingDirectoryForWorker(WorkerIndex), GWorkerOutputFilename);
 
 			// In the common case the output file will not exist, so check for existence before opening
 			// This is only a win if FileExists is faster than CreateFileReader, which it is on Windows
@@ -769,6 +842,149 @@ void FShaderCompileThreadRunnable::PrintWorkerMemoryUsageWithLockTaken()
 	}
 }
 
+int32 FShaderCompileThreadRunnable::GetNumberOfAvailableWorkersUnsafe() const
+{
+	// Don't lock WorkerScopeLock critical section here, since this function might be called inside an already locked scope, hence the "Unsafe" name
+	int32 NumAvailableWorkers = 0;
+
+	for (const TUniquePtr<FShaderCompileWorkerInfo>& WorkerInfo : this->WorkerInfos)
+	{
+		if (WorkerInfo->bAvailable)
+		{
+			++NumAvailableWorkers;
+		}
+	}
+
+	return NumAvailableWorkers;
+}
+
+int32 FShaderCompileThreadRunnable::GetNumberOfSuspendedWorkersUnsafe() const
+{
+	return WorkerInfos.Num() - GetNumberOfAvailableWorkersUnsafe();
+}
+
+int32 FShaderCompileThreadRunnable::SuspendWorkersAndBacklogJobs(int32 NumWorkersToSuspend)
+{
+	int32 NumSuspendedWorkers = 0;
+	int32 NumBackloggedJobs = 0;
+
+	// Before suspending workers, we need to know how many workers are available to ensure there is always at least one worker available
+	if (NumWorkersToSuspend > 0)
+	{
+		FScopeLock WorkerScopeLock(&WorkerInfosLock);
+		const int32 NumAvailableWorkers = GetNumberOfAvailableWorkersUnsafe();
+		NumWorkersToSuspend = FMath::Min(NumWorkersToSuspend, NumAvailableWorkers - 1);
+
+		if (NumWorkersToSuspend > 0)
+		{
+			for (int32 WorkerIndex = 0; WorkerIndex < WorkerInfos.Num(); ++WorkerIndex)
+			{
+				FShaderCompileWorkerInfo& CurrentWorkerInfo = *WorkerInfos[WorkerIndex];
+				if (CurrentWorkerInfo.bAvailable)
+				{
+					// Suspend worker: Terminate its process immediately as we want to free up system resources.
+					// Also discard its output file if it has already created one. Otherwise, this file will be linked to the wrong compile jobs.
+					CurrentWorkerInfo.bAvailable = false;
+					CurrentWorkerInfo.TerminateWorkerProcess();
+					DiscardWorkerOutputFile(WorkerIndex);
+
+					// Move its jobs into the backlog queue
+					BackloggedJobs.Reserve(BackloggedJobs.Num() + CurrentWorkerInfo.QueuedJobs.Num());
+					for (FShaderCommonCompileJobPtr& QueuedJob : CurrentWorkerInfo.QueuedJobs)
+					{
+						BackloggedJobs.Add(QueuedJob);
+					}
+					NumBackloggedJobs += CurrentWorkerInfo.QueuedJobs.Num();
+					CurrentWorkerInfo.QueuedJobs.Empty();
+
+					// No more workers to suspend? Early exit loop.
+					++NumSuspendedWorkers;
+					if (NumSuspendedWorkers == NumWorkersToSuspend)
+					{
+						break;
+					}
+				}
+			}
+		}
+	}
+
+	// Report about backlogged jobs
+	if (NumBackloggedJobs > 0)
+	{
+		UE_LOG(
+			LogShaderCompilers, Display, TEXT("Backlogged %d compile %s from %d suspended %s"),
+			NumBackloggedJobs,
+			NumBackloggedJobs == 1 ? TEXT("job") : TEXT("jobs"),
+			NumSuspendedWorkers,
+			NumSuspendedWorkers == 1 ? TEXT("worker") : TEXT("workers")
+		);
+	}
+
+	return NumSuspendedWorkers;
+}
+
+int32 FShaderCompileThreadRunnable::ResumeSuspendedWorkers(int32 NumWorkersToResume)
+{
+	int32 NumResumedWorkers = 0;
+
+	if (NumWorkersToResume > 0)
+	{
+		FScopeLock WorkerScopeLock(&WorkerInfosLock);
+		const int32 NumSuspendedWorkers = GetNumberOfSuspendedWorkersUnsafe();
+		NumWorkersToResume = FMath::Min(NumWorkersToResume, NumSuspendedWorkers);
+
+		if (NumWorkersToResume > 0)
+		{
+			for (int32 WorkerIndex = 0; WorkerIndex < WorkerInfos.Num(); ++WorkerIndex)
+			{
+				FShaderCompileWorkerInfo& CurrentWorkerInfo = *WorkerInfos[WorkerIndex];
+				if (!CurrentWorkerInfo.bAvailable)
+				{
+					// Resume worker by making it available again. It will pick up jobs next time tasks are pulled from the queue.
+					CurrentWorkerInfo.bAvailable = true;
+
+					// No more workers to suspend? Early exit loop.
+					++NumResumedWorkers;
+					if (NumResumedWorkers == NumWorkersToResume)
+					{
+						break;
+					}
+				}
+			}
+		}
+
+		if (NumResumedWorkers > 0)
+		{
+			const int32 NumRemainingSuspendedWorkers = NumSuspendedWorkers - NumResumedWorkers;
+			UE_LOG(
+				LogShaderCompilers, Display, TEXT("Resumed %d suspended %s (%d %s suspended)"),
+				NumResumedWorkers,
+				NumResumedWorkers == 1 ? TEXT("worker") : TEXT("workers"),
+				NumRemainingSuspendedWorkers,
+				NumRemainingSuspendedWorkers == 1 ? TEXT("remains") : TEXT("remain")
+			);
+		}
+	}
+
+	return NumResumedWorkers;
+}
+
+void FShaderCompileThreadRunnable::DiscardWorkerOutputFile(int32 WorkerIndex)
+{
+	// If the previously suspended worker left a stale output file, delete it now before it gets picked up and is linked to the wrong input jobs
+	const FString OutputFileNameAndPath = FPaths::Combine(GetWorkingDirectoryForWorker(WorkerIndex), GWorkerOutputFilename);
+	if (IFileManager::Get().FileExists(*OutputFileNameAndPath))
+	{
+		UE_LOG(LogShaderCompilers, Verbose, TEXT("Discard stale worker output file: %s"), *OutputFileNameAndPath);
+		IFileManager::Get().Delete(*OutputFileNameAndPath);
+	}
+}
+
+FString FShaderCompileThreadRunnable::GetWorkingDirectoryForWorker(int32 WorkerIndex, bool bRelativePath) const
+{
+	return FPaths::Combine(bRelativePath ? Manager->ShaderBaseWorkingDirectory : Manager->AbsoluteShaderBaseWorkingDirectory, FString::FromInt(WorkerIndex));
+}
+
 bool FShaderCompileThreadRunnable::PrintWorkerMemoryUsage(bool bAllowToWaitForLock)
 {
 	if (bAllowToWaitForLock)
diff --git a/Engine/Source/Runtime/Engine/Public/ShaderCompiler.h b/Engine/Source/Runtime/Engine/Public/ShaderCompiler.h
--- a/Engine/Source/Runtime/Engine/Public/ShaderCompiler.h
+++ b/Engine/Source/Runtime/Engine/Public/ShaderCompiler.h
@@ -209,6 +209,9 @@ private:
 	/** Whether to read/write files for SCW in parallel (can help situations when this takes too long for a number of reasons) */
 	bool bParallelizeIO = false;
 
+	/** List of jobs that have been backlogged when workers had to be closed due to reaching memory limits. These jobs will be picked up first before new jobs are pulled from the manager job queue. */
+	TArray<FShaderCommonCompileJobPtr> BackloggedJobs;
+
 public:
 	/** Initialization constructor. */
 	FShaderCompileThreadRunnable(class FShaderCompilingManager* InManager);
@@ -248,6 +251,30 @@ private:
 	virtual void OnMachineResourcesChanged() override;
 
 	void PrintWorkerMemoryUsageWithLockTaken();
+
+	/** Returns the number of available workers. Only call inside the critical section WorkerInfosLock. */
+	int32 GetNumberOfAvailableWorkersUnsafe() const;
+
+	/** Returns the number of suspended workers. Only call inside the critical section WorkerInfosLock. */
+	int32 GetNumberOfSuspendedWorkersUnsafe() const;
+
+	/**
+	 * Suspends the specified number of workers and moves all their compile jobs to the backlog queue.
+	 * Returns the number of workers that have been suspended. The last worker cannot be suspended.
+	 */
+	int32 SuspendWorkersAndBacklogJobs(int32 NumWorkers);
+
+	/**
+	 * Makes the specified number of workers available again after they have been suspended.
+	 * Returns the number of workers that have been resumed. If all workers were already available, the return value is 0.
+	 */
+	int32 ResumeSuspendedWorkers(int32 NumWorkers);
+
+	/** Deletes the output file of the specified worker if it exists and discards its content. This is called when a worker output is considered stale because it was previously suspended. */
+	void DiscardWorkerOutputFile(int32 WorkerIndex);
+
+	/** Returns the working directory for the specified shader compile worker. */
+	FString GetWorkingDirectoryForWorker(int32 WorkerIndex, bool bRelativePath = false) const;
 };
 
 class FShaderCompileUtilities
