Substrate - fixed ShadingModel node pin ClearCoatRoughness text not showing up correctly.
#rnx

diff --git a/Engine/Source/Runtime/Engine/Private/Materials/MaterialExpressions.cpp b/Engine/Source/Runtime/Engine/Private/Materials/MaterialExpressions.cpp
--- a/Engine/Source/Runtime/Engine/Private/Materials/MaterialExpressions.cpp
+++ b/Engine/Source/Runtime/Engine/Private/Materials/MaterialExpressions.cpp
@@ -26401,6 +26401,10 @@ FName UMaterialExpressionSubstrateShadingModels::GetInputName(int32 InputIndex)
 			{
 				return TEXT("Iris Distance");
 			}
+			else if (ShadingModelOverride == MSM_ClearCoat)
+			{
+				return TEXT("Clear Coat Roughness");
+			}
 			return TEXT("Unused");
 		}
 		return TEXT("Custom1");
