Fix AssetManager to exclude all generated maps from being assigned primary asset types, this improves on 34573294 and works for primary asset types other than the default "Map"

diff --git a/Engine/Source/Runtime/Engine/Private/AssetManager.cpp b/Engine/Source/Runtime/Engine/Private/AssetManager.cpp
--- a/Engine/Source/Runtime/Engine/Private/AssetManager.cpp
+++ b/Engine/Source/Runtime/Engine/Private/AssetManager.cpp
@@ -1637,8 +1637,7 @@ FPrimaryAssetId UAssetManager::ExtractPrimaryAssetIdFromData(const FAssetData& A
 
 	if (!FoundId.IsValid() && SuggestedType.IsValid())
 	{
-		FName SuggestedAssetName = AssetData.AssetName;
-		if (SuggestedType == MapType)
+		if (AssetData.PackageFlags & PKG_ContainsMap)
 		{
 			// UE-216073: Temporary solution to turn off WorldPartition _Generated_ streaming cells packages
 			// from being PrimaryAssets, since doing so causes duplicate PrimaryAssetId with their owning map.
@@ -1648,8 +1647,12 @@ FPrimaryAssetId UAssetManager::ExtractPrimaryAssetIdFromData(const FAssetData& A
 			{
 				return FPrimaryAssetId(UAssetManager::SuppressionType, AssetData.PackageName);
 			}
+		}
 
-			// Maps use the full package name
+		FName SuggestedAssetName = AssetData.AssetName;
+		if (SuggestedType == MapType)
+		{
+			// Maps use the full package name, this is a hack that only works properly with the default Map type defined in BaseGame.ini
 			SuggestedAssetName = AssetData.PackageName;
 		}
 
