Prevent USceneCaptureComponent functions ShowOnlyComponent and ShowOnlyActorComponents from adding duplicate components to the ShowOnlyComponents array. Matches existing behavior for the HideComponent and HideActorComponents functions with the HiddenComponents array.

#tests PIE
[RN] fix rendering/architecture

diff --git a/Engine/Source/Runtime/Engine/Private/Components/SceneCaptureComponent.cpp b/Engine/Source/Runtime/Engine/Private/Components/SceneCaptureComponent.cpp
--- a/Engine/Source/Runtime/Engine/Private/Components/SceneCaptureComponent.cpp
+++ b/Engine/Source/Runtime/Engine/Private/Components/SceneCaptureComponent.cpp
@@ -346,7 +346,7 @@ void USceneCaptureComponent::ShowOnlyComponent(UPrimitiveComponent* InComponent)
 	{
 		// Backward compatibility - set PrimitiveRenderMode to PRM_UseShowOnlyList if BP / game code tries to add a ShowOnlyComponent
 		PrimitiveRenderMode = ESceneCapturePrimitiveRenderMode::PRM_UseShowOnlyList;
-		ShowOnlyComponents.Add(InComponent);
+		ShowOnlyComponents.AddUnique(InComponent);
 	}
 }
 
@@ -360,7 +360,7 @@ void USceneCaptureComponent::ShowOnlyActorComponents(AActor* InActor, const bool
 		TInlineComponentArray<UPrimitiveComponent*> PrimitiveComponents(InActor, bIncludeFromChildActors);
 		for (UPrimitiveComponent* PrimComp : PrimitiveComponents)
 		{
-			ShowOnlyComponents.Add(PrimComp);
+			ShowOnlyComponents.AddUnique(PrimComp);
 		}
 	}
 }
