Removed first person rendering project setting so the feature is now always enabled.

#rnx

diff --git a/Engine/Shaders/Private/Definitions.usf b/Engine/Shaders/Private/Definitions.usf
--- a/Engine/Shaders/Private/Definitions.usf
+++ b/Engine/Shaders/Private/Definitions.usf
@@ -250,6 +250,10 @@
 #define FORCE_MATERIAL_FLOAT_FULL_PRECISION             0
 #endif
 
+#ifndef SUPPORT_FIRST_PERSON_RENDERING
+#define SUPPORT_FIRST_PERSON_RENDERING					1
+#endif
+
 ///////////////////////////////////////////////////////////////////////////////
 // Global Substrate material export related defines
 
diff --git a/Engine/Source/Runtime/Engine/Classes/Engine/RendererSettings.h b/Engine/Source/Runtime/Engine/Classes/Engine/RendererSettings.h
--- a/Engine/Source/Runtime/Engine/Classes/Engine/RendererSettings.h
+++ b/Engine/Source/Runtime/Engine/Classes/Engine/RendererSettings.h
@@ -1142,15 +1142,6 @@ class URendererSettings : public UDeveloperSettings
 		ConfigRestartRequired = true))
 		uint32 bSupportCloudShadowOnSingleLayerWater : 1;
 
-	/**
-	"Compiles shaders with support for rendering first person primitives with a different depth scale and field of view."
-	*/
-	UPROPERTY(config, EditAnywhere, Category = ShaderPermutationReduction, meta = (
-		ConsoleVariable = "r.FirstPersonRendering.Enabled", DisplayName = "Support First Person Rendering",
-		ToolTip = "Compiles shaders with support for rendering first person primitives with a different depth scale and field of view.",
-		ConfigRestartRequired = true))
-		uint32 bSupportFirstPersonRendering : 1;
-
 	/**
 	"Enable Substrate materials (Beta)."
 	*/
diff --git a/Engine/Source/Runtime/Engine/Private/SceneView.cpp b/Engine/Source/Runtime/Engine/Private/SceneView.cpp
--- a/Engine/Source/Runtime/Engine/Private/SceneView.cpp
+++ b/Engine/Source/Runtime/Engine/Private/SceneView.cpp
@@ -388,12 +388,6 @@ static TAutoConsoleVariable<float> CVarOverrideTimeMaterialExpressions(
 
 #endif
 
-static TAutoConsoleVariable<int32> CVarFirstPersonRenderingSupported(
-	TEXT("r.FirstPersonRendering.Enabled"),
-	0,
-	TEXT("Compiles shaders with support for rendering first person primitives with a different depth scale and field of view."),
-	ECVF_ReadOnly | ECVF_RenderThreadSafe | ECVF_GeneralShaderChange);
-
 /** Global vertex color view mode setting when SHOW_VertexColors show flag is set */
 EVertexColorViewMode::Type GVertexColorViewMode = EVertexColorViewMode::Color;
 TWeakObjectPtr<UTexture> GVertexViewModeOverrideTexture = nullptr;
diff --git a/Engine/Source/Runtime/Engine/Private/ShaderCompiler/ShaderCompiler.cpp b/Engine/Source/Runtime/Engine/Private/ShaderCompiler/ShaderCompiler.cpp
--- a/Engine/Source/Runtime/Engine/Private/ShaderCompiler/ShaderCompiler.cpp
+++ b/Engine/Source/Runtime/Engine/Private/ShaderCompiler/ShaderCompiler.cpp
@@ -3608,12 +3608,6 @@ void GlobalBeginCompileShader(
 		SET_SHADER_DEFINE(Input.Environment, SUPPORT_CLOUD_SHADOW_ON_SINGLE_LAYER_WATER, bSupportCloudShadowOnSingleLayerWater ? 1 : 0);
 	}
 
-	{
-		static IConsoleVariable* CVar = IConsoleManager::Get().FindConsoleVariable(TEXT("r.FirstPersonRendering.Enabled"));
-		const bool bSupportFirstPersonRendering = CVar && CVar->GetInt() > 0;
-		SET_SHADER_DEFINE(Input.Environment, SUPPORT_FIRST_PERSON_RENDERING, bSupportFirstPersonRendering ? 1 : 0);
-	}
-
 	{
 		const bool bTranslucentUsesLightRectLights = GetTranslucentUsesLightRectLights();
 		SET_SHADER_DEFINE(Input.Environment, SUPPORT_RECTLIGHT_ON_FORWARD_LIT_TRANSLUCENT, bTranslucentUsesLightRectLights ? 1 : 0);
diff --git a/Engine/Source/Runtime/Engine/Private/ShaderCompiler/ShaderGenerationUtil.cpp b/Engine/Source/Runtime/Engine/Private/ShaderCompiler/ShaderGenerationUtil.cpp
--- a/Engine/Source/Runtime/Engine/Private/ShaderCompiler/ShaderGenerationUtil.cpp
+++ b/Engine/Source/Runtime/Engine/Private/ShaderCompiler/ShaderGenerationUtil.cpp
@@ -72,7 +72,6 @@ void ApplyFetchEnvironmentInternal(FShaderGlobalDefines& SrcDefines, const Conta
 	FETCH_COMPILE_BOOL(PROJECT_EXPFOG_MATCHES_VFOG);
 	FETCH_COMPILE_BOOL(SUPPORT_CLOUD_SHADOW_ON_FORWARD_LIT_TRANSLUCENT);
 	FETCH_COMPILE_BOOL(SUPPORT_CLOUD_SHADOW_ON_SINGLE_LAYER_WATER);
-	FETCH_COMPILE_BOOL(SUPPORT_FIRST_PERSON_RENDERING);
 	FETCH_COMPILE_BOOL(POST_PROCESS_ALPHA);
 	FETCH_COMPILE_BOOL(PLATFORM_SUPPORTS_SHADER_ROOT_CONSTANTS);
 	FETCH_COMPILE_BOOL(PLATFORM_SUPPORTS_SHADER_BUNDLE_DISPATCH);
@@ -434,12 +433,6 @@ static FShaderGlobalDefines FetchShaderGlobalDefines(EShaderPlatform TargetPlatf
 		Ret.SUPPORT_CLOUD_SHADOW_ON_SINGLE_LAYER_WATER = bSupportCloudShadowOnSingleLayerWater ? 1 : 0;
 	}
 
-	{
-		static IConsoleVariable* CVar = IConsoleManager::Get().FindConsoleVariable(TEXT("r.FirstPersonRendering.Enabled"));
-		const bool bSupportFirstPersonRendering = CVar && CVar->GetInt() > 0;
-		Ret.SUPPORT_FIRST_PERSON_RENDERING = bSupportFirstPersonRendering ? 1 : 0;
-	}
-
 	{
 		static IConsoleVariable* CVar = IConsoleManager::Get().FindConsoleVariable(TEXT("r.PostProcessing.PropagateAlpha"));
 		int32 PropagateAlpha = CVar->GetInt();
diff --git a/Engine/Source/Runtime/RenderCore/Public/ShaderMaterial.h b/Engine/Source/Runtime/RenderCore/Public/ShaderMaterial.h
--- a/Engine/Source/Runtime/RenderCore/Public/ShaderMaterial.h
+++ b/Engine/Source/Runtime/RenderCore/Public/ShaderMaterial.h
@@ -36,7 +36,6 @@ struct FShaderGlobalDefines
 	bool PROJECT_EXPFOG_MATCHES_VFOG;
 	bool SUPPORT_CLOUD_SHADOW_ON_FORWARD_LIT_TRANSLUCENT;
 	bool SUPPORT_CLOUD_SHADOW_ON_SINGLE_LAYER_WATER;
-	bool SUPPORT_FIRST_PERSON_RENDERING;
 	bool POST_PROCESS_ALPHA;
 	bool PLATFORM_SUPPORTS_SHADER_ROOT_CONSTANTS;
 	bool PLATFORM_SUPPORTS_SHADER_BUNDLE_DISPATCH;
