HLOD - Avoid duplicating LODActors when cooking multiple targets at once and using the bSaveLODActorsToHLODPackages project option

diff --git a/Engine/Source/Runtime/Engine/Private/HLOD/HLODEngineSubsystem.cpp b/Engine/Source/Runtime/Engine/Private/HLOD/HLODEngineSubsystem.cpp
--- a/Engine/Source/Runtime/Engine/Private/HLOD/HLODEngineSubsystem.cpp
+++ b/Engine/Source/Runtime/Engine/Private/HLOD/HLODEngineSubsystem.cpp
@@ -68,7 +68,7 @@ void UHLODEngineSubsystem::RegisterRecreateLODActorsDelegates()
 		OnPostWorldInitializationDelegateHandle = FWorldDelegates::OnPostWorldInitialization.AddUObject(this, &UHLODEngineSubsystem::RecreateLODActorsForWorld);
 		OnLevelAddedToWorldDelegateHandle = FWorldDelegates::LevelAddedToWorld.AddUObject(this, &UHLODEngineSubsystem::RecreateLODActorsForLevel);
 		OnPreSaveWorlDelegateHandle = FEditorDelegates::PreSaveWorldWithContext.AddUObject(this, &UHLODEngineSubsystem::OnPreSaveWorld);
-	}	
+	}
 }
 
 void UHLODEngineSubsystem::RecreateLODActorsForWorld(UWorld* InWorld, const UWorld::InitializationValues InInitializationValues)
@@ -92,7 +92,7 @@ void UHLODEngineSubsystem::RecreateLODActorsForLevel(ULevel* InLevel, UWorld* In
 	IHierarchicalLODUtilities* Utilities = Module.GetUtilities();
 
 	// First, destroy invalid HLOD actors. If needed, they will be recreated below.
-	if (!bDisableHLODCleanupOnLoad && !GIsCookerLoadingPackage)
+	if (!bDisableHLODCleanupOnLoad)
 	{
 		CleanupHLODs(InLevel);
 	}
@@ -154,11 +154,7 @@ bool UHLODEngineSubsystem::CleanupHLOD(ALODActor* InLODActor)
 	}
 	else if (GetDefault<UHierarchicalLODSettings>()->bSaveLODActorsToHLODPackages)
 	{
-		if (!InLODActor->HasAnyFlags(RF_Transient))
-		{
-			UE_LOG(LogEngine, Warning, TEXT("Deleting non-transient LODActor %s. Rebuild HLOD & resave %s to silence warning."), *InLODActor->GetName(), *InLODActor->GetOutermost()->GetPathName());
-		}
-
+		// LODActors will be recreated from the data contained in the HLODProxy packages
 		bShouldDestroyActor = true;
 	}
 
