[Content Bundle] When replacing an actor's class, keep its current associated content bundle intact.

diff --git a/Engine/Source/Editor/UnrealEd/Private/EditorEngine.cpp b/Engine/Source/Editor/UnrealEd/Private/EditorEngine.cpp
--- a/Engine/Source/Editor/UnrealEd/Private/EditorEngine.cpp
+++ b/Engine/Source/Editor/UnrealEd/Private/EditorEngine.cpp
@@ -112,6 +112,7 @@
 #include "SceneView.h"
 #include "StaticBoundShaderState.h"
 #include "PropertyColorSettings.h"
+#include "WorldPartition/ContentBundle/ContentBundleActivationScope.h"
 
 // needed for the RemotePropagator
 #include "AudioDevice.h"
@@ -5468,6 +5469,9 @@ void UEditorEngine::ReplaceSelectedActors(UActorFactory* Factory, const FAssetDa
 
 void UEditorEngine::ReplaceActors(UActorFactory* Factory, const FAssetData& AssetData, const TArray<AActor*>& ActorsToReplace, TArray<AActor*>* OutNewActors, bool bCopySourceProperties)
 {
+	FGuid InvalidGuid;
+	FContentBundleActivationScope Activationscope(InvalidGuid);
+
 	// Cache for attachment info of all actors being converted.
 	TArray<ReattachActorsHelper::FActorAttachmentCache> AttachmentInfo;
 
diff --git a/Engine/Source/Runtime/Engine/Public/WorldPartition/ContentBundle/ContentBundleActivationScope.h b/Engine/Source/Runtime/Engine/Public/WorldPartition/ContentBundle/ContentBundleActivationScope.h
--- a/Engine/Source/Runtime/Engine/Public/WorldPartition/ContentBundle/ContentBundleActivationScope.h
+++ b/Engine/Source/Runtime/Engine/Public/WorldPartition/ContentBundle/ContentBundleActivationScope.h
@@ -7,41 +7,33 @@
 class FContentBundleEditor;
 
 #if WITH_EDITOR
-
 class FContentBundleActivationScope
 {
 public:
 	FContentBundleActivationScope(FGuid InContentBundleGuid)
 	{
-		if (InContentBundleGuid.IsValid())
-		{
-			ContentBundleGuid = InContentBundleGuid;
+		ContentBundleGuid = InContentBundleGuid;
 
-			IContentBundleEditorSubsystemInterface* ContentBundleEditorSubsystem = IContentBundleEditorSubsystemInterface::Get();
-			ContentBundleEditorSubsystem->PushContentBundleEditing();
+		IContentBundleEditorSubsystemInterface* ContentBundleEditorSubsystem = IContentBundleEditorSubsystemInterface::Get();
+		ContentBundleEditorSubsystem->PushContentBundleEditing();
 
-			if (TSharedPtr<FContentBundleEditor> ContentBundleEditor = ContentBundleEditorSubsystem->GetEditorContentBundle(ContentBundleGuid))
-			{
-				ContentBundleEditorSubsystem->ActivateContentBundleEditing(ContentBundleEditor);
-			}
+		if (TSharedPtr<FContentBundleEditor> ContentBundleEditor = ContentBundleEditorSubsystem->GetEditorContentBundle(ContentBundleGuid))
+		{
+			ContentBundleEditorSubsystem->ActivateContentBundleEditing(ContentBundleEditor);
 		}
 	}
 	~FContentBundleActivationScope()
 	{
-		if (ContentBundleGuid.IsValid())
+		IContentBundleEditorSubsystemInterface* ContentBundleEditorSubsystem = IContentBundleEditorSubsystemInterface::Get();
+		if (TSharedPtr<FContentBundleEditor> ContentBundleEditor = ContentBundleEditorSubsystem->GetEditorContentBundle(ContentBundleGuid))
 		{
-			IContentBundleEditorSubsystemInterface* ContentBundleEditorSubsystem = IContentBundleEditorSubsystemInterface::Get();
-			if (TSharedPtr<FContentBundleEditor> ContentBundleEditor = ContentBundleEditorSubsystem->GetEditorContentBundle(ContentBundleGuid))
-			{
-				ContentBundleEditorSubsystem->DeactivateContentBundleEditing(ContentBundleEditor);
-			}
-
-			ContentBundleEditorSubsystem->PopContentBundleEditing();
+			ContentBundleEditorSubsystem->DeactivateContentBundleEditing(ContentBundleEditor);
 		}
+
+		ContentBundleEditorSubsystem->PopContentBundleEditing();
 	}
 
 private:
 	FGuid ContentBundleGuid;
 };
-
 #endif
