Add a verbose log for when FlushPressedKeys is called. This is a common thing that you may want to know when it happens when debugging.

#rnx

diff --git a/Engine/Source/Runtime/Engine/Private/UserInterface/PlayerInput.cpp b/Engine/Source/Runtime/Engine/Private/UserInterface/PlayerInput.cpp
--- a/Engine/Source/Runtime/Engine/Private/UserInterface/PlayerInput.cpp
+++ b/Engine/Source/Runtime/Engine/Private/UserInterface/PlayerInput.cpp
@@ -134,6 +134,8 @@ void UPlayerInput::PostInitProperties()
 void UPlayerInput::FlushPressedKeys()
 {
 	TRACE_CPUPROFILER_EVENT_SCOPE(UPlayerInput::FlushPressedKeys);
+
+	UE_LOG(LogPlayerInput, Verbose, TEXT("[%hs] Flush pressed keys on player input object %s"), __func__, *GetNameSafe(this));
 	
 	APlayerController* PlayerController = GetOuterAPlayerController();
 	ULocalPlayer* LocalPlayer = PlayerController ? Cast<ULocalPlayer>(PlayerController->Player) : nullptr;
