Stop serializing long deprecated static mesh high res source mesh information

diff --git a/Engine/Source/Runtime/Core/Public/UObject/RenderingObjectVersion.h b/Engine/Source/Runtime/Core/Public/UObject/RenderingObjectVersion.h
--- a/Engine/Source/Runtime/Core/Public/UObject/RenderingObjectVersion.h
+++ b/Engine/Source/Runtime/Core/Public/UObject/RenderingObjectVersion.h
@@ -136,6 +136,9 @@ struct FRenderingObjectVersion
 		// Custom material expression includes are now serialized & used for DDC key generation
 		AddedMaterialExpressionIncludesHash,
 
+		// Do not serialize deprecated high res source mesh
+		DeprecatedHighResSourceMesh,
+
 		// -----<new versions can be added above this line>-------------------------------------------------
 		VersionPlusOne,
 		LatestVersion = VersionPlusOne - 1
diff --git a/Engine/Source/Runtime/Engine/Private/StaticMesh.cpp b/Engine/Source/Runtime/Engine/Private/StaticMesh.cpp
--- a/Engine/Source/Runtime/Engine/Private/StaticMesh.cpp
+++ b/Engine/Source/Runtime/Engine/Private/StaticMesh.cpp
@@ -6174,8 +6174,8 @@ void UStaticMesh::Serialize(FArchive& Ar)
 		Super::Serialize(Ar);
 	}
 
-	Ar.UsingCustomVersion(FReleaseObjectVersion::GUID);
 	Ar.UsingCustomVersion(FEditorObjectVersion::GUID);
+	Ar.UsingCustomVersion(FFrameworkObjectVersion::GUID);
 	Ar.UsingCustomVersion(FRenderingObjectVersion::GUID);
 	Ar.UsingCustomVersion(FReleaseObjectVersion::GUID);
 	Ar.UsingCustomVersion(FUE5MainStreamObjectVersion::GUID);
@@ -6223,33 +6223,30 @@ void UStaticMesh::Serialize(FArchive& Ar)
 	}
 #endif
 
-	Ar.UsingCustomVersion(FFrameworkObjectVersion::GUID);
-
 	if(Ar.IsLoading() && Ar.CustomVer(FFrameworkObjectVersion::GUID) < FFrameworkObjectVersion::UseBodySetupCollisionProfile && GetBodySetup())
 	{
 		GetBodySetup()->DefaultInstance.SetCollisionProfileName(UCollisionProfile::BlockAll_ProfileName);
 	}
 
 #if WITH_EDITORONLY_DATA
-	if( !StripFlags.IsEditorDataStripped() )
+	if( !StripFlags.IsEditorDataStripped() && Ar.IsLoading() )
 	{
-		if ( Ar.IsLoading() && Ar.UEVer() < VER_UE4_DEPRECATED_STATIC_MESH_THUMBNAIL_PROPERTIES_REMOVED )
+		if ( Ar.UEVer() < VER_UE4_DEPRECATED_STATIC_MESH_THUMBNAIL_PROPERTIES_REMOVED )
 		{
 			FRotator DummyThumbnailAngle;
 			float DummyThumbnailDistance;
 			Ar << DummyThumbnailAngle;
 			Ar << DummyThumbnailDistance;
 		}
-	}
 
-	if( !StripFlags.IsEditorDataStripped() )
-	{
-		// TODO: These should be gated with a version check, but not able to be done in this stream.
-		FString Deprecated_HighResSourceMeshName;
-		uint32 Deprecated_HighResSourceMeshCRC;
+		if (Ar.CustomVer(FRenderingObjectVersion::GUID) < FRenderingObjectVersion::DeprecatedHighResSourceMesh)
+		{
+			FString Deprecated_HighResSourceMeshName;
+			uint32 Deprecated_HighResSourceMeshCRC;
 
-		Ar << Deprecated_HighResSourceMeshName;
-		Ar << Deprecated_HighResSourceMeshCRC;
+			Ar << Deprecated_HighResSourceMeshName;
+			Ar << Deprecated_HighResSourceMeshCRC;
+		}
 	}
 #endif // #if WITH_EDITORONLY_DATA
 
