Fixed first person FOV correction factor being incorrect with projection matrices with a AspectRatioAxisConstraint other than MaintainXFOV. Fixed by computing the FOVCorrection factor when initializing FFirstPersonParameters outside of FViewMatrices. At this point we know both the scene FOV and the first person FOV, so computing their ratio is easy and does not require recovering the horizontal FOV from an arbitrary perspective projection matrix.
Also fixed first person primitives being erroneously occlusion culled. This happened because the occlusion queries do not know about the first person transform which scales the primitive towards the camera, potentially outside of its bounds. Taking this into account during occlusion culling is a bit tricky because it requires us to pass and look up additional data, which might impact OC performance. Since first person primitives are very close to the camera, they can reasonably be assumed to never be occlusion culled by the scene anyways. The simple fix then is to skip occlusion culling for first person primitives.

#rnx

diff --git a/Engine/Source/Editor/UnrealEd/Private/EditorViewportClient.cpp b/Engine/Source/Editor/UnrealEd/Private/EditorViewportClient.cpp
--- a/Engine/Source/Editor/UnrealEd/Private/EditorViewportClient.cpp
+++ b/Engine/Source/Editor/UnrealEd/Private/EditorViewportClient.cpp
@@ -1283,7 +1283,7 @@ FSceneView* FEditorViewportClient::CalcSceneView(FSceneViewFamily* ViewFamily, c
 	{
 		ViewInitOptions.bUseFieldOfViewForLOD = ControllingActorViewInfo.bUseFieldOfViewForLOD;
 		ViewInitOptions.FOV = ControllingActorViewInfo.FOV;
-		ViewInitOptions.FirstPersonParams = FFirstPersonParameters(ControllingActorViewInfo.FirstPersonFOV, ControllingActorViewInfo.FirstPersonScale, ControllingActorViewInfo.bUseFirstPersonParameters);
+		ViewInitOptions.FirstPersonParams = FFirstPersonParameters(ControllingActorViewInfo.CalculateFirstPersonFOVCorrectionFactor(), ControllingActorViewInfo.FirstPersonScale, ControllingActorViewInfo.bUseFirstPersonParameters);
 	}
 
 	ViewInitOptions.OverrideFarClippingPlaneDistance = FarPlane;
diff --git a/Engine/Source/Runtime/Engine/Classes/Camera/CameraTypes.h b/Engine/Source/Runtime/Engine/Classes/Camera/CameraTypes.h
--- a/Engine/Source/Runtime/Engine/Classes/Camera/CameraTypes.h
+++ b/Engine/Source/Runtime/Engine/Classes/Camera/CameraTypes.h
@@ -208,4 +208,10 @@ public:
 	 * but its on-screen size will look correct.
 	 */
 	ENGINE_API FVector TransformWorldToFirstPerson(const FVector& WorldPosition, bool bIgnoreFirstPersonScale) const;
+
+	/**
+	 * Correction factor to apply to the first person transform used on primitives tagged as "IsFirstPerson" to achieve a first person specific field of view.
+	 * It is computed as tan(SceneFOVRadians * 0.5) / tan(FirstPersonFOVRadians * 0.5).
+	 */
+	ENGINE_API float CalculateFirstPersonFOVCorrectionFactor() const;
 };
diff --git a/Engine/Source/Runtime/Engine/Private/Camera/CameraStackTypes.cpp b/Engine/Source/Runtime/Engine/Private/Camera/CameraStackTypes.cpp
--- a/Engine/Source/Runtime/Engine/Private/Camera/CameraStackTypes.cpp
+++ b/Engine/Source/Runtime/Engine/Private/Camera/CameraStackTypes.cpp
@@ -453,7 +453,7 @@ FVector FMinimalViewInfo::TransformWorldToFirstPerson(const FVector& WorldPositi
 		const FVector CameraRelativePosition = WorldPosition - Location;
 		const FVector ProjectedPosition = FVector::DotProduct(Forward, CameraRelativePosition) * Forward;
 		const FVector Rejection = CameraRelativePosition - ProjectedPosition;
-		const float FOVCorrectionFactor = FMath::Tan(FMath::DegreesToRadians(FOV * 0.5f)) / FMath::Tan(FMath::DegreesToRadians(FirstPersonFOV * 0.5f)) - 1.0f;
+		const float FOVCorrectionFactor = CalculateFirstPersonFOVCorrectionFactor() - 1.0f;
 		const FVector FOVCorrectedPosition = CameraRelativePosition + Rejection * FOVCorrectionFactor;
 		const FVector ScaledPosition = FOVCorrectedPosition * FirstPersonScale;
 		const FVector Result = (bIgnoreFirstPersonScale ? FOVCorrectedPosition : ScaledPosition) + Location;
@@ -461,3 +461,11 @@ FVector FMinimalViewInfo::TransformWorldToFirstPerson(const FVector& WorldPositi
 	}
 	return WorldPosition;
 }
+
+float FMinimalViewInfo::CalculateFirstPersonFOVCorrectionFactor() const
+{
+	const float HalfTanSceneFOV = FMath::Tan(FMath::DegreesToRadians(FOV) * 0.5f);
+	const float HalfTanFirstPersonFOV = FMath::Tan(FMath::DegreesToRadians(FirstPersonFOV * 0.5f));
+	const float FOVCorrectionFactor = HalfTanSceneFOV / HalfTanFirstPersonFOV;
+	return FOVCorrectionFactor;
+}
diff --git a/Engine/Source/Runtime/Engine/Private/LocalPlayer.cpp b/Engine/Source/Runtime/Engine/Private/LocalPlayer.cpp
--- a/Engine/Source/Runtime/Engine/Private/LocalPlayer.cpp
+++ b/Engine/Source/Runtime/Engine/Private/LocalPlayer.cpp
@@ -836,7 +836,7 @@ FSceneView* ULocalPlayer::CalcSceneView( class FSceneViewFamily* ViewFamily,
 	ViewInitOptions.bUseFieldOfViewForLOD = ViewInfo.bUseFieldOfViewForLOD;
 	ViewInitOptions.FOV = ViewInfo.FOV;
 	ViewInitOptions.DesiredFOV = ViewInfo.DesiredFOV;
-	ViewInitOptions.FirstPersonParams = FFirstPersonParameters(ViewInfo.FirstPersonFOV, ViewInfo.FirstPersonScale, ViewInfo.bUseFirstPersonParameters);
+	ViewInitOptions.FirstPersonParams = FFirstPersonParameters(ViewInfo.CalculateFirstPersonFOVCorrectionFactor(), ViewInfo.FirstPersonScale, ViewInfo.bUseFirstPersonParameters);
 
 	// Fill out the rest of the view init options
 	ViewInitOptions.ViewFamily = ViewFamily;
diff --git a/Engine/Source/Runtime/Engine/Private/SceneView.cpp b/Engine/Source/Runtime/Engine/Private/SceneView.cpp
--- a/Engine/Source/Runtime/Engine/Private/SceneView.cpp
+++ b/Engine/Source/Runtime/Engine/Private/SceneView.cpp
@@ -743,8 +743,8 @@ void FViewMatrices::Init(const FMinimalInitializer& Initializer)
 		//		tan(HalfFOV) / tan(HalfFirstPersonFOV)
 		// This correction factor should be independent of the actual aspect ratio used on the original projection matrix. As the FOV correction needs to be done in view space, we can conveniently
 		// also use this opportunity to scale down the first person geometry towards the camera. The geometry takes up the same solid angle from the perspective of the camera and therefore appears the
-		// same size while taking up a smaller range in the depth buffer, reducing the likelyhood of intersections with the scene.
-		const float FOVCorrectionFactor = GetTanHalfFov().X / FMath::Tan(FMath::DegreesToRadians(Initializer.FirstPersonParams.FOV) * 0.5f);
+		// same size while taking up a smaller range in the depth buffer, reducing the likelihood of intersections with the scene.
+		const float FOVCorrectionFactor = Initializer.FirstPersonParams.FOVCorrectionFactor;
 		FirstPersonTransform = TranslatedViewMatrix * FScaleMatrix(FVector(FOVCorrectionFactor, FOVCorrectionFactor, 1.0f) * Initializer.FirstPersonParams.Scale) * InvTranslatedViewMatrix;
 	}
 }
diff --git a/Engine/Source/Runtime/Engine/Public/SceneView.h b/Engine/Source/Runtime/Engine/Public/SceneView.h
--- a/Engine/Source/Runtime/Engine/Public/SceneView.h
+++ b/Engine/Source/Runtime/Engine/Public/SceneView.h
@@ -156,8 +156,8 @@ enum class ESecondaryScreenPercentageMethod
 
 struct FFirstPersonParameters
 {
-	/** The horizontal field of view (in degrees) used for primitives tagged as "IsFirstPerson". */
-	float FOV = 90.0f;
+	/** FOV correction factor applied to the first person transform used on primitives tagged as "IsFirstPerson". This should be computed as tan(SceneFOVRadians * 0.5) / tan(FirstPersonFOVRadians * 0.5). */
+	float FOVCorrectionFactor = 1.0f;
 
 	/** The scale to apply to primitives tagged as "IsFirstPerson". This is used to scale down primitives towards the camera such that they are small enough not to intersect with the scene. */
 	float Scale = 1.0f;
@@ -166,7 +166,7 @@ struct FFirstPersonParameters
 	bool bUseParameters = false;
 
 	FFirstPersonParameters() = default;
-	FFirstPersonParameters(float InFOV, float InScale, bool bInUseParameters) : FOV(InFOV), Scale(InScale), bUseParameters(bInUseParameters) {}
+	FFirstPersonParameters(float InFOVCorrectionFactor, float InScale, bool bInUseParameters) : FOVCorrectionFactor(InFOVCorrectionFactor), Scale(InScale), bUseParameters(bInUseParameters) {}
 };
 
 // Construction parameters for a FSceneView
diff --git a/Engine/Source/Runtime/Renderer/Private/PrimitiveSceneInfo.cpp b/Engine/Source/Runtime/Renderer/Private/PrimitiveSceneInfo.cpp
--- a/Engine/Source/Runtime/Renderer/Private/PrimitiveSceneInfo.cpp
+++ b/Engine/Source/Runtime/Renderer/Private/PrimitiveSceneInfo.cpp
@@ -1693,7 +1693,9 @@ void FPrimitiveSceneInfo::AddToScene(FScene* Scene, TArrayView<FPrimitiveSceneIn
 
 			// Store occlusion flags.
 			uint8 OcclusionFlags = EOcclusionFlags::None;
-			if (Proxy->CanBeOccluded())
+			// First person primitives potentially deform the geometry outside of its bounds in a view dependent way. They are very unlikely to be occluded anyways,
+			// so to avoid falsely culling them, it is better to simply don't occlusion cull them at all.
+			if (Proxy->CanBeOccluded() && !Proxy->IsFirstPerson())
 			{
 				OcclusionFlags |= EOcclusionFlags::CanBeOccluded;
 			}
