[AutoRTFM] Migrate more critical sections to using the transactionally safe variants.

diff --git a/Engine/Plugins/Developer/N10XSourceCodeAccess/Source/N10XSourceCodeAccess/Private/N10XSourceCodeAccessor.cpp b/Engine/Plugins/Developer/N10XSourceCodeAccess/Source/N10XSourceCodeAccess/Private/N10XSourceCodeAccessor.cpp
--- a/Engine/Plugins/Developer/N10XSourceCodeAccess/Source/N10XSourceCodeAccess/Private/N10XSourceCodeAccessor.cpp
+++ b/Engine/Plugins/Developer/N10XSourceCodeAccess/Source/N10XSourceCodeAccess/Private/N10XSourceCodeAccessor.cpp
@@ -4,6 +4,7 @@
 #include "Misc/Paths.h"
 #include "Misc/UProjectInfo.h"
 #include "Misc/App.h"
+#include "Misc/ScopeLock.h"
 
 #if PLATFORM_WINDOWS
 #include "Internationalization/Regex.h"
diff --git a/Engine/Plugins/Experimental/BackChannel/Source/BackChannel/Private/Protocol/BackChannelOSCConnection.cpp b/Engine/Plugins/Experimental/BackChannel/Source/BackChannel/Private/Protocol/BackChannelOSCConnection.cpp
--- a/Engine/Plugins/Experimental/BackChannel/Source/BackChannel/Private/Protocol/BackChannelOSCConnection.cpp
+++ b/Engine/Plugins/Experimental/BackChannel/Source/BackChannel/Private/Protocol/BackChannelOSCConnection.cpp
@@ -6,6 +6,7 @@
 #include "BackChannel/Transport/IBackChannelSocketConnection.h"
 #include "HAL/PlatformProcess.h"
 #include "HAL/RunnableThread.h"
+#include "Misc/ScopeLock.h"
 #include "Sockets.h"
 
 FBackChannelOSCConnection::FBackChannelOSCConnection(TSharedRef<IBackChannelSocketConnection> InConnection)
diff --git a/Engine/Source/Developer/Horde/Private/Storage/Bundles/V2/BundleWriter.cpp b/Engine/Source/Developer/Horde/Private/Storage/Bundles/V2/BundleWriter.cpp
--- a/Engine/Source/Developer/Horde/Private/Storage/Bundles/V2/BundleWriter.cpp
+++ b/Engine/Source/Developer/Horde/Private/Storage/Bundles/V2/BundleWriter.cpp
@@ -7,6 +7,7 @@
 #include "Storage/Bundles/V2/ExportHandle.h"
 #include "../../../HordePlatform.h"
 #include "Storage/ChunkedBufferWriter.h"
+#include "Misc/ScopeLock.h"
 
 //
 // Data for an export that has been written to storage, but may not have been flushed yet
diff --git a/Engine/Source/Developer/Horde/Private/Storage/Clients/MemoryStorageClient.cpp b/Engine/Source/Developer/Horde/Private/Storage/Clients/MemoryStorageClient.cpp
--- a/Engine/Source/Developer/Horde/Private/Storage/Clients/MemoryStorageClient.cpp
+++ b/Engine/Source/Developer/Horde/Private/Storage/Clients/MemoryStorageClient.cpp
@@ -3,6 +3,7 @@
 #include "Storage/Clients/MemoryStorageClient.h"
 #include "Storage/Clients/KeyValueStorageClient.h"
 #include "Storage/Blob.h"
+#include "Misc/ScopeLock.h"
 
 FMemoryStorageClient::FMemoryStorageClient()
 {
diff --git a/Engine/Source/Developer/ShaderCompilerCommon/Private/DXCWrapper.cpp b/Engine/Source/Developer/ShaderCompilerCommon/Private/DXCWrapper.cpp
--- a/Engine/Source/Developer/ShaderCompilerCommon/Private/DXCWrapper.cpp
+++ b/Engine/Source/Developer/ShaderCompilerCommon/Private/DXCWrapper.cpp
@@ -5,6 +5,7 @@
 #include "HAL/FileManager.h"
 #include "Templates/RefCounting.h"
 #include "Misc/Paths.h"
+#include "Misc/ScopeLock.h"
 #include "Hash/xxhash.h"
 
 #if PLATFORM_WINDOWS
diff --git a/Engine/Source/Programs/AutoRTFMEngineTests/Private/AutoRTFMEngineTests.cpp b/Engine/Source/Programs/AutoRTFMEngineTests/Private/AutoRTFMEngineTests.cpp
--- a/Engine/Source/Programs/AutoRTFMEngineTests/Private/AutoRTFMEngineTests.cpp
+++ b/Engine/Source/Programs/AutoRTFMEngineTests/Private/AutoRTFMEngineTests.cpp
@@ -213,6 +213,11 @@ static void PreInit()
 
 	FDelayedAutoRegisterHelper::RunAndClearDelayedAutoRegisterDelegates(EDelayedRegisterRunPhase::StatSystemReady);
 
+
+	// Ensure GCachedScalabilityCVars.bInitialized is set (required by some mesh components)
+	FTaskTagScope::SwapTag(ETaskTag::EGameThread);
+	extern void ScalabilityCVarsSinkCallback();
+	ScalabilityCVarsSinkCallback();
 }
 
 static void LoadModules()
diff --git a/Engine/Source/Runtime/Core/Private/Features/ModularFeatures.cpp b/Engine/Source/Runtime/Core/Private/Features/ModularFeatures.cpp
--- a/Engine/Source/Runtime/Core/Private/Features/ModularFeatures.cpp
+++ b/Engine/Source/Runtime/Core/Private/Features/ModularFeatures.cpp
@@ -4,7 +4,7 @@
 
 #include "CoreTypes.h"
 #include "Misc/AssertionMacros.h"
-#include "Misc/ScopeLock.h"
+#include "Misc/TransactionallySafeScopeLock.h"
 
 IModularFeatures& IModularFeatures::Get()
 {
@@ -54,7 +54,7 @@ IModularFeature* FModularFeatures::GetModularFeatureImplementation( const FName
 
 void FModularFeatures::RegisterModularFeature( const FName Type, IModularFeature* ModularFeature )
 {
-	FScopeLock ScopeLock(&ModularFeaturesMapCriticalSection);
+	FTransactionallySafeScopeLock ScopeLock(&ModularFeaturesMapCriticalSection);
 
 	ModularFeaturesMap.AddUnique( Type, ModularFeature );
 	ModularFeatureRegisteredEvent.Broadcast( Type, ModularFeature );
@@ -63,7 +63,7 @@ void FModularFeatures::RegisterModularFeature( const FName Type, IModularFeature
 
 void FModularFeatures::UnregisterModularFeature( const FName Type, IModularFeature* ModularFeature )
 {
-	FScopeLock ScopeLock(&ModularFeaturesMapCriticalSection);
+	FTransactionallySafeScopeLock ScopeLock(&ModularFeaturesMapCriticalSection);
 
 	ModularFeaturesMap.RemoveSingle( Type, ModularFeature );
 	ModularFeatureUnregisteredEvent.Broadcast( Type, ModularFeature );
diff --git a/Engine/Source/Runtime/Core/Private/Features/ModularFeatures.h b/Engine/Source/Runtime/Core/Private/Features/ModularFeatures.h
--- a/Engine/Source/Runtime/Core/Private/Features/ModularFeatures.h
+++ b/Engine/Source/Runtime/Core/Private/Features/ModularFeatures.h
@@ -6,7 +6,7 @@
 #include "Delegates/Delegate.h"
 #include "Features/IModularFeature.h"
 #include "Features/IModularFeatures.h"
-#include "HAL/CriticalSection.h"
+#include "Misc/TransactionallySafeCriticalSection.h"
 #include "UObject/NameTypes.h"
 
 /**
@@ -37,7 +37,7 @@ private:
 	TMultiMap< FName, class IModularFeature* > ModularFeaturesMap;
 
 	/** Lock modular features map so it can be used across threads, this could be a RW lock instead of a crit sec */
-	FCriticalSection ModularFeaturesMapCriticalSection;
+	FTransactionallySafeCriticalSection ModularFeaturesMapCriticalSection;
 
 	/** Event used to inform clients that a modular feature has been registered */
 	FOnModularFeatureRegistered ModularFeatureRegisteredEvent;
diff --git a/Engine/Source/Runtime/Core/Private/Misc/ConfigContext.cpp b/Engine/Source/Runtime/Core/Private/Misc/ConfigContext.cpp
--- a/Engine/Source/Runtime/Core/Private/Misc/ConfigContext.cpp
+++ b/Engine/Source/Runtime/Core/Private/Misc/ConfigContext.cpp
@@ -10,6 +10,7 @@
 #include "Misc/RemoteConfigIni.h"
 #include "Misc/Paths.h"
 #include "Misc/DataDrivenPlatformInfoRegistry.h"
+#include "Misc/ScopeLock.h"
 #include "ProfilingDebugging/CpuProfilerTrace.h"
 #include "ProfilingDebugging/AssetMetadataTrace.h"
 #include "HAL/LowLevelMemStats.h"
diff --git a/Engine/Source/Runtime/Core/Public/Delegates/DelegateAccessHandler.h b/Engine/Source/Runtime/Core/Public/Delegates/DelegateAccessHandler.h
--- a/Engine/Source/Runtime/Core/Public/Delegates/DelegateAccessHandler.h
+++ b/Engine/Source/Runtime/Core/Public/Delegates/DelegateAccessHandler.h
@@ -3,9 +3,9 @@
 #pragma once
 
 #include "CoreTypes.h"
-#include "HAL/CriticalSection.h"
 #include "Misc/MTAccessDetector.h"
-#include "Misc/ScopeLock.h"
+#include "Misc/TransactionallySafeCriticalSection.h"
+#include "Misc/TransactionallySafeScopeLock.h"
 #include "AutoRTFM/AutoRTFM.h"
 
 //#define UE_DETECT_DELEGATES_RACE_CONDITIONS 0
@@ -65,8 +65,8 @@ template<>
 class TDelegateAccessHandlerBase<FThreadSafeDelegateMode>
 {
 protected:
-	struct FReadAccessScope { FScopeLock Lock; };
-	struct FWriteAccessScope { FScopeLock Lock; };
+	struct FReadAccessScope { FTransactionallySafeScopeLock Lock; };
+	struct FWriteAccessScope { FTransactionallySafeScopeLock Lock; };
 
 	[[nodiscard]] FReadAccessScope GetReadAccessScope() const
 	{
@@ -79,7 +79,7 @@ protected:
 	}
 
 private:
-	mutable FCriticalSection Mutex;
+	mutable FTransactionallySafeCriticalSection Mutex;
 };
 
 template<>
diff --git a/Engine/Source/Runtime/CoreUObject/Public/UObject/UObjectArray.h b/Engine/Source/Runtime/CoreUObject/Public/UObject/UObjectArray.h
--- a/Engine/Source/Runtime/CoreUObject/Public/UObject/UObjectArray.h
+++ b/Engine/Source/Runtime/CoreUObject/Public/UObject/UObjectArray.h
@@ -8,6 +8,7 @@
 
 #include "HAL/ThreadSafeCounter.h"
 #include "Containers/LockFreeList.h"
+#include "Misc/ScopeLock.h"
 #include "UObject/GarbageCollectionGlobals.h"
 #include "UObject/UObjectBase.h"
 
@@ -51,7 +52,7 @@ private:
 	int32 Flags;
 public:
 	// UObject Owner Cluster Index
-	int32 ClusterRootIndex;	
+	int32 ClusterRootIndex;
 	// Weak Object Pointer Serial number associated with the object
 	int32 SerialNumber;
 	// RefCount associated with the object preventing its destruction.
diff --git a/Engine/Source/Runtime/Engine/Classes/Animation/Skeleton.h b/Engine/Source/Runtime/Engine/Classes/Animation/Skeleton.h
--- a/Engine/Source/Runtime/Engine/Classes/Animation/Skeleton.h
+++ b/Engine/Source/Runtime/Engine/Classes/Animation/Skeleton.h
@@ -12,11 +12,11 @@
 #include "UObject/Object.h"
 #include "HAL/CriticalSection.h"
 #include "Misc/Guid.h"
+#include "Misc/TransactionallySafeRWLock.h"
 #include "ReferenceSkeleton.h"
 #include "Animation/PreviewAssetAttachComponent.h"
 #include "Animation/SmartName.h"
 #include "Engine/AssetUserData.h"
-#include "HAL/CriticalSection.h"
 #include "Interfaces/Interface_AssetUserData.h"
 #include "Interfaces/Interface_PreviewMeshProvider.h"
 #include "Delegates/DelegateCombinations.h"
@@ -658,7 +658,7 @@ private:
 	TMap<TWeakObjectPtr<USkinnedAsset>, int32> SkinnedAsset2LinkupCache;
 	
 	//Use this Lock everytime you change or access SkinnedAssetLinkupCache member.
-	FRWLock SkinnedAssetLinkupCacheLock;
+	FTransactionallySafeRWLock SkinnedAssetLinkupCacheLock;
 
 	/** Runtime built mapping table between SkinnedAssets and Mesh Linkup Data*/
 	TMap<TObjectKey<USkinnedAsset>, TUniquePtr<FSkeletonToMeshLinkup>> SkinnedAssetLinkupCache;
diff --git a/Engine/Source/Runtime/Engine/Private/Animation/Skeleton.cpp b/Engine/Source/Runtime/Engine/Private/Animation/Skeleton.cpp
--- a/Engine/Source/Runtime/Engine/Private/Animation/Skeleton.cpp
+++ b/Engine/Source/Runtime/Engine/Private/Animation/Skeleton.cpp
@@ -750,7 +750,7 @@ bool USkeleton::IsCompatibleMesh(const USkinnedAsset* InSkinnedAsset, bool bDoPa
 void USkeleton::ClearCacheData()
 {
 	{
-		FRWScopeLock Lock(SkinnedAssetLinkupCacheLock, SLT_Write);
+		FTransactionallySafeWriteScopeLock Lock(SkinnedAssetLinkupCacheLock);
 		SkinnedAssetLinkupCache.Empty();
 	}
 
@@ -767,7 +767,7 @@ const FSkeletonToMeshLinkup& USkeleton::FindOrAddMeshLinkupData(const USkinnedAs
 	const TUniquePtr<FSkeletonToMeshLinkup>* SkeletonToMeshLinkupPtr = nullptr;
 
 	{
-		FRWScopeLock Lock(SkinnedAssetLinkupCacheLock, SLT_ReadOnly);
+		FTransactionallySafeReadScopeLock Lock(SkinnedAssetLinkupCacheLock);
 		SkeletonToMeshLinkupPtr = SkinnedAssetLinkupCache.Find(InSkinnedAsset);
 	}
 
@@ -782,7 +782,7 @@ const FSkeletonToMeshLinkup& USkeleton::AddMeshLinkupData(const USkinnedAsset* I
 	BuildLinkupData(InSkinnedAsset, *TmpLinkup.Get());
 
 	{
-		FRWScopeLock Lock(SkinnedAssetLinkupCacheLock, SLT_Write);
+		FTransactionallySafeWriteScopeLock Lock(SkinnedAssetLinkupCacheLock);
 		return *SkinnedAssetLinkupCache.Add(TObjectKey<USkinnedAsset>(InSkinnedAsset), MoveTemp(TmpLinkup)).Get();
 	}
 }
@@ -823,7 +823,7 @@ int32 USkeleton::GetMeshLinkupIndex(const USkinnedAsset* InSkinnedAsset)
 void USkeleton::RemoveLinkup(const USkinnedAsset* InSkinnedAsset)
 {
 	{
-		FRWScopeLock Lock(SkinnedAssetLinkupCacheLock, SLT_Write);
+		FTransactionallySafeWriteScopeLock Lock(SkinnedAssetLinkupCacheLock);
 		SkinnedAssetLinkupCache.Remove(InSkinnedAsset);
 	}
 
diff --git a/Engine/Source/Runtime/Engine/Private/Tests/AutoRTFM/AutoRTFMActorComponentTests.cpp b/Engine/Source/Runtime/Engine/Private/Tests/AutoRTFM/AutoRTFMActorComponentTests.cpp
--- a/Engine/Source/Runtime/Engine/Private/Tests/AutoRTFM/AutoRTFMActorComponentTests.cpp
+++ b/Engine/Source/Runtime/Engine/Private/Tests/AutoRTFM/AutoRTFMActorComponentTests.cpp
@@ -1,22 +1,25 @@
 // Copyright Epic Games, Inc. All Rights Reserved.
 
 #include "Misc/AutomationTest.h"
-#include "Misc/App.h"
+#include "Animation/AnimInstance.h"
+#include "Animation/Skeleton.h"
 #include "AutoRTFM/AutoRTFM.h"
 #include "AutoRTFMTestActor.h"
 #include "AutoRTFMTestAnotherActor.h"
 #include "AutoRTFMTestBodySetup.h"
+#include "AutoRTFMTestChildActorComponent.h"
 #include "AutoRTFMTestLevel.h"
 #include "AutoRTFMTestObject.h"
-#include "AutoRTFMTestChildActorComponent.h"
 #include "AutoRTFMTestPrimitiveComponent.h"
-#include "PhysicsProxy/SingleParticlePhysicsProxy.h"
-#include "Physics/Experimental/PhysScene_Chaos.h"
-#include "PBDRigidsSolver.h"
 #include "Chaos/Core.h"
+#include "Components/SkeletalMeshComponent.h"
+#include "Engine/SkeletalMesh.h"
 #include "Engine/World.h"
-
-#include <utility>
+#include "PBDRigidsSolver.h"
+#include "Physics/Experimental/PhysScene_Chaos.h"
+#include "PhysicsProxy/SingleParticlePhysicsProxy.h"
+#include "Rendering/SkeletalMeshRenderData.h"
+#include "UObject/UObjectGlobals.h"
 
 #if WITH_DEV_AUTOMATION_TESTS
 
@@ -396,6 +399,32 @@ AUTORTFM_ACTOR_COMPONENT_TEST(ChildActor)
 	TEST_CHECK_TRUE(nullptr == ChildActorComponent->GetChildActor());
 }
 
+// Test aborting a call to USkeletalMeshComponent::RegisterComponent() with an assigned skeletal
+// mesh and empty PostProcessAnimInstance.
+// See: SOL-6779
+AUTORTFM_ACTOR_COMPONENT_TEST(USkeletalMeshComponent)
+{
+	USkeleton* Skeleton = NewObject<USkeleton>();
+	USkeletalMesh* SkeletalMesh = NewObject<USkeletalMesh>();
+	SkeletalMesh->SetSkeleton(Skeleton);
+	SkeletalMesh->AllocateResourceForRendering();
+	FSkeletalMeshRenderData* RenderData = SkeletalMesh->GetResourceForRendering();
+	TRefCountPtr<FSkeletalMeshLODRenderData> LODRenderData = MakeRefCount<FSkeletalMeshLODRenderData>();
+	RenderData->LODRenderData.Add(LODRenderData);
+	USkeletalMeshComponent* SkeletalMeshComponent = NewObject<USkeletalMeshComponent>(Actor);
+	SkeletalMeshComponent->SetSkeletalMeshAsset(SkeletalMesh);
+	SkeletalMeshComponent->PostProcessAnimInstance = NewObject<UAnimInstance>(SkeletalMeshComponent);
+
+	AutoRTFM::ETransactionResult Result;
+	Result = AutoRTFM::Transact([&]
+		{
+			SkeletalMeshComponent->RegisterComponent();
+			AutoRTFM::AbortTransaction();
+		});
+
+	TEST_CHECK_TRUE(AutoRTFM::ETransactionResult::AbortedByRequest == Result);
+}
+
 }  // anonymous namespace
 
 #undef TEST_CHECK_FALSE
diff --git a/Engine/Source/Runtime/OpenColorIOWrapper/Private/OpenColorIOWrapper.cpp b/Engine/Source/Runtime/OpenColorIOWrapper/Private/OpenColorIOWrapper.cpp
--- a/Engine/Source/Runtime/OpenColorIOWrapper/Private/OpenColorIOWrapper.cpp
+++ b/Engine/Source/Runtime/OpenColorIOWrapper/Private/OpenColorIOWrapper.cpp
@@ -10,6 +10,7 @@
 #include "Containers/Map.h"
 #include "ColorManagement/ColorSpace.h"
 #include "Engine/TextureDefines.h"
+#include "Misc/ScopeLock.h"
 #include "ImageCore.h"
 #include "ImageParallelFor.h"
 
