Update Iris when AActor::NetUpdateFrequency changes after BeginReplication().

diff --git a/Engine/Source/Runtime/Engine/Classes/Engine/NetDriver.h b/Engine/Source/Runtime/Engine/Classes/Engine/NetDriver.h
--- a/Engine/Source/Runtime/Engine/Classes/Engine/NetDriver.h
+++ b/Engine/Source/Runtime/Engine/Classes/Engine/NetDriver.h
@@ -1043,6 +1043,14 @@ private:
 	/** Reset any network metrics database values at the beginning of a frame. */
 	void ResetNetworkMetrics();
 
+public:
+	/** Delegate that will notify when an actors NetUpdateFrequency UPROPERTY has changed. */
+	DECLARE_MULTICAST_DELEGATE_OneParam(FOnNetUpdateFrequencyChanged, const AActor* Actor);
+	FOnNetUpdateFrequencyChanged& GetOnNetUpdateFrequencyChanged() { return OnNetUpdateFrequencyChanged; };
+
+private:
+	FOnNetUpdateFrequencyChanged OnNetUpdateFrequencyChanged;
+
 public:
 	/** Get the value of MaxChannelsOverride cached from the net driver definition */
 	int32 GetMaxChannelsOverride() const { return MaxChannelsOverride; }
diff --git a/Engine/Source/Runtime/Engine/Private/Actor.cpp b/Engine/Source/Runtime/Engine/Private/Actor.cpp
--- a/Engine/Source/Runtime/Engine/Private/Actor.cpp
+++ b/Engine/Source/Runtime/Engine/Private/Actor.cpp
@@ -147,6 +147,27 @@ namespace ActorUtils
 
 uint32 AActor::BeginPlayCallDepth = 0;
 
+namespace ActorUtils
+{
+	template<class Function>
+	static void ForEachNetDriver(UEngine* Engine, const UWorld* const World, const Function InFunction)
+	{
+		if (Engine == nullptr || World == nullptr)
+		{
+			return;
+		}
+
+		FWorldContext* const Context = Engine->GetWorldContextFromWorld(World);
+		if (Context != nullptr)
+		{
+			for (FNamedNetDriver& Driver : Context->ActiveNetDrivers)
+			{
+				InFunction(Driver.NetDriver);
+			}
+		}
+	}
+}
+
 AActor::AActor()
 {
 	InitializeDefaults();
@@ -6235,7 +6256,17 @@ void AActor::SetNetUpdateFrequency(float Frequency)
 {
 PRAGMA_DISABLE_DEPRECATION_WARNINGS
 	NetUpdateFrequency = Frequency;
-PRAGMA_ENABLE_DEPRECATION_WARNINGS 
+PRAGMA_ENABLE_DEPRECATION_WARNINGS
+
+	UE_LOG(LogNet, Verbose, TEXT("AActor::SetNetUpdateFrequency(): %s Frequency=%f HasActorBegunPlay()=%d"), *this->GetFullName(), Frequency, HasActorBegunPlay());
+
+	ActorUtils::ForEachNetDriver(GEngine, GetWorld(), [this](UNetDriver* NetDriver)
+	{
+		if (NetDriver)
+		{
+			NetDriver->GetOnNetUpdateFrequencyChanged().Broadcast(this);
+		}
+	});
 }
 
 float AActor::GetNetUpdateFrequency() const
diff --git a/Engine/Source/Runtime/Engine/Private/Net/Iris/ReplicationSystem/ActorReplicationBridge.cpp b/Engine/Source/Runtime/Engine/Private/Net/Iris/ReplicationSystem/ActorReplicationBridge.cpp
--- a/Engine/Source/Runtime/Engine/Private/Net/Iris/ReplicationSystem/ActorReplicationBridge.cpp
+++ b/Engine/Source/Runtime/Engine/Private/Net/Iris/ReplicationSystem/ActorReplicationBridge.cpp
@@ -62,6 +62,13 @@ static FAutoConsoleVariableRef CVarEnableActorLevelChanges(
 	TEXT("When true the ActorReplicationBridge will process actors that change levels by updating the actor's level groups.")
 );
 
+static bool bEnableDynamicNetUpdateFrequency = true;
+static FAutoConsoleVariableRef CVarEnableDynamicNetUpdateFrequency(
+	TEXT("net.Iris.EnableDynamicNetUpdateFrequency"),
+	bEnableDynamicNetUpdateFrequency,
+	TEXT("When true changes to AActor::NetUpdateFrequency will be updated in Iris after being registered for replication.")
+);
+
 bool IsActorValidForIrisReplication(const AActor* Actor)
 {
 	return IsValid(Actor) && !Actor->IsActorBeingDestroyed() && !Actor->IsUnreachable();
@@ -167,6 +174,7 @@ void UActorReplicationBridge::Deinitialize()
 	if (NetDriver)
 	{
 		NetDriver->OnNetServerMaxTickRateChanged.RemoveAll(this);
+		NetDriver->GetOnNetUpdateFrequencyChanged().RemoveAll(this);
 		NetDriver = nullptr;
 	}
 	Super::Deinitialize();
@@ -992,6 +1000,7 @@ void UActorReplicationBridge::SetNetDriver(UNetDriver* const InNetDriver)
 	if (NetDriver)
 	{
 		NetDriver->OnNetServerMaxTickRateChanged.RemoveAll(this);
+		NetDriver->GetOnNetUpdateFrequencyChanged().RemoveAll(this);
 	}
 
 	NetDriver = InNetDriver;
@@ -1000,6 +1009,7 @@ void UActorReplicationBridge::SetNetDriver(UNetDriver* const InNetDriver)
 		SetMaxTickRate(static_cast<float>(FPlatformMath::Max(InNetDriver->GetNetServerMaxTickRate(), 0)));
 
 		InNetDriver->OnNetServerMaxTickRateChanged.AddUObject(this, &UActorReplicationBridge::OnMaxTickRateChanged);
+		InNetDriver->GetOnNetUpdateFrequencyChanged().AddUObject(this, &UActorReplicationBridge::OnNetUpdateFrequencyChanged);
 
 		const FName RequiredChannelName = UObjectReplicationBridgeConfig::GetConfig()->GetRequiredNetDriverChannelClassName();
 		
@@ -1346,6 +1356,19 @@ void UActorReplicationBridge::ActorChangedLevel(const AActor* Actor, const ULeve
 	AddActorToLevelGroup(Actor);
 }
 
+void UActorReplicationBridge::OnNetUpdateFrequencyChanged(const AActor* Actor)
+{
+	if (UE::Net::Private::bEnableDynamicNetUpdateFrequency)
+	{
+		const UE::Net::FNetRefHandle ActorHandle = GetReplicatedRefHandle(Actor);
+
+		if (ActorHandle != UE::Net::FNetRefHandle::GetInvalid())
+		{
+			SetPollFrequency(ActorHandle, Actor->GetNetUpdateFrequency());
+		}
+	}
+}
+
 void UActorReplicationBridge::AddActorToLevelGroup(const AActor* Actor)
 {
 	if (!UE::Net::Private::ShouldIncludeActorInLevelGroups(Actor))
diff --git a/Engine/Source/Runtime/Engine/Public/Net/Iris/ReplicationSystem/ActorReplicationBridge.h b/Engine/Source/Runtime/Engine/Public/Net/Iris/ReplicationSystem/ActorReplicationBridge.h
--- a/Engine/Source/Runtime/Engine/Public/Net/Iris/ReplicationSystem/ActorReplicationBridge.h
+++ b/Engine/Source/Runtime/Engine/Public/Net/Iris/ReplicationSystem/ActorReplicationBridge.h
@@ -91,6 +91,9 @@ public:
 	/** Updates the level group for an actor that changed levels */
 	void ActorChangedLevel(const AActor* Actor, const ULevel* PreviousLevel);
 
+	/** Called when NetUpdateFrequency has changed on the Actor. */
+	void OnNetUpdateFrequencyChanged(const AActor* Actor);
+
 	/**
 	 * Add relevant network metrics gathered since the last call to ConsumeNetMetrics.
 	 * Any periodic stat will be reset here too.
