Add CVar for push-based AController

diff --git a/Engine/Source/Runtime/Engine/Classes/GameFramework/Controller.h b/Engine/Source/Runtime/Engine/Classes/GameFramework/Controller.h
--- a/Engine/Source/Runtime/Engine/Classes/GameFramework/Controller.h
+++ b/Engine/Source/Runtime/Engine/Classes/GameFramework/Controller.h
@@ -48,6 +48,7 @@ public:
 	/** PlayerState containing replicated information about the player using this controller (only exists for players, not NPCs). */
 	UPROPERTY(replicatedUsing = OnRep_PlayerState, BlueprintReadOnly, Category=Controller)
 	TObjectPtr<APlayerState> PlayerState;
+	ENGINE_API void SetPlayerState(APlayerState* InPlayerState);
 
 	/** Actor marking where this controller spawned in. */
 	TWeakObjectPtr<class AActor> StartSpot;
@@ -68,6 +69,7 @@ private:
 	/** Pawn currently being controlled by this controller.  Use Pawn.Possess() to take control of a pawn */
 	UPROPERTY(replicatedUsing=OnRep_Pawn)
 	TObjectPtr<APawn> Pawn;
+	void SetPawn_Direct(APawn* InPawn);
 
 	/**
 	 * Used to track when pawn changes during OnRep_Pawn. 
diff --git a/Engine/Source/Runtime/Engine/Private/Controller.cpp b/Engine/Source/Runtime/Engine/Private/Controller.cpp
--- a/Engine/Source/Runtime/Engine/Private/Controller.cpp
+++ b/Engine/Source/Runtime/Engine/Private/Controller.cpp
@@ -6,6 +6,7 @@
 =============================================================================*/
 #include "GameFramework/Controller.h"
 #include "AI/NavigationSystemBase.h"
+#include "Net/Core/PushModel/PushModel.h"
 #include "Net/UnrealNetwork.h"
 #include "NetworkingDistanceConstants.h"
 #include "VisualLogger/VisualLogger.h"
@@ -42,6 +43,12 @@ namespace UE::Gameplay::CVars
 		TEXT("Controller.InvalidControlRotationMagnitude"), InvalidControlRotationMagnitude,
 		TEXT("If any component of an FRotator passed to SetControlRotation is larger than this magnitude, ignore the value. Huge values are usually from uninitialized variables and can cause NaN/Inf to propagate later."),
 		ECVF_Default);
+
+	bool bIsControllerPushBased = false;
+	static FAutoConsoleVariableRef CVarIsControllerPushBased(
+		TEXT("Controller.IsPushBased"), bIsControllerPushBased,
+		TEXT("If true, AController's replicated properties will use push-based networking, and will therefore need to be marked dirty when changed."),
+		ECVF_Default);
 }
 
 AController::AController(const FObjectInitializer& ObjectInitializer)
@@ -505,12 +512,22 @@ void AController::AddPawnTickDependency(APawn* NewPawn)
 	}
 }
 
+void AController::SetPawn_Direct(APawn* InPawn)
+{
+	if (UE::Gameplay::CVars::bIsControllerPushBased)
+	{
+		COMPARE_ASSIGN_AND_MARK_PROPERTY_DIRTY(AController, Pawn, InPawn, this);
+		return;
+	}
+
+	Pawn = InPawn;
+}
 
 void AController::SetPawn(APawn* InPawn)
 {
 	RemovePawnTickDependency(Pawn);
 
-	Pawn = InPawn;
+	SetPawn_Direct(InPawn);
 	Character = (Pawn ? Cast<ACharacter>(Pawn) : NULL);
 
 	AttachToPawn(Pawn);
@@ -523,7 +540,7 @@ void AController::SetPawnFromRep(APawn* InPawn)
 	// This function is needed to ensure OnRep_Pawn is called in the case we need to set AController::Pawn
 	// due to APawn::Controller being replicated first. See additional notes in APawn::OnRep_Controller.
 	RemovePawnTickDependency(Pawn);
-	Pawn = InPawn;
+	SetPawn_Direct(InPawn);
 	OnRep_Pawn();
 }
 
@@ -552,6 +569,17 @@ void AController::OnRep_Pawn()
 	}
 }
 
+void AController::SetPlayerState(APlayerState* InPlayerState)
+{
+	if (UE::Gameplay::CVars::bIsControllerPushBased)
+	{
+		COMPARE_ASSIGN_AND_MARK_PROPERTY_DIRTY(AController, PlayerState, InPlayerState, this);
+		return;
+	}
+
+	PlayerState = InPlayerState;
+}
+
 void AController::OnRep_PlayerState()
 {
 	if (PlayerState != NULL)
@@ -583,7 +611,7 @@ void AController::Destroyed()
 void AController::CleanupPlayerState()
 {
 	PlayerState->Destroy();
-	PlayerState = NULL;
+	SetPlayerState(NULL);
 }
 
 void AController::InstigatedAnyDamage(float Damage, const class UDamageType* DamageType, class AActor* DamagedActor, class AActor* DamageCauser)
@@ -623,7 +651,7 @@ void AController::InitPlayerState()
 				PlayerStateClassToSpawn = APlayerState::StaticClass();
 			}
 
-			PlayerState = World->SpawnActor<APlayerState>(PlayerStateClassToSpawn, SpawnInfo);
+			SetPlayerState(World->SpawnActor<APlayerState>(PlayerStateClassToSpawn, SpawnInfo));
 	
 			// force a default player name if necessary
 			if (PlayerState && PlayerState->GetPlayerName().IsEmpty())
@@ -770,8 +798,11 @@ void AController::GetLifetimeReplicatedProps( TArray< FLifetimeProperty > & OutL
 {
 	Super::GetLifetimeReplicatedProps( OutLifetimeProps );
 
-	DOREPLIFETIME( AController, PlayerState );
-	DOREPLIFETIME_CONDITION_NOTIFY(AController, Pawn, COND_None, REPNOTIFY_Always);
+	FDoRepLifetimeParams Params;
+	Params.bIsPushBased = UE::Gameplay::CVars::bIsControllerPushBased;
+	DOREPLIFETIME_WITH_PARAMS_FAST(AController, PlayerState, Params);
+	Params.RepNotifyCondition = REPNOTIFY_Always;
+	DOREPLIFETIME_WITH_PARAMS_FAST(AController, Pawn, Params);
 }
 
 bool AController::ShouldParticipateInSeamlessTravel() const
diff --git a/Engine/Source/Runtime/Engine/Private/GameMode.cpp b/Engine/Source/Runtime/Engine/Private/GameMode.cpp
--- a/Engine/Source/Runtime/Engine/Private/GameMode.cpp
+++ b/Engine/Source/Runtime/Engine/Private/GameMode.cpp
@@ -700,7 +700,7 @@ bool AGameMode::FindInactivePlayer(APlayerController* PC)
 		{
 			// found it!
 			APlayerState* OldPlayerState = PC->PlayerState;
-			PC->PlayerState = CurrentPlayerState;
+			PC->SetPlayerState(CurrentPlayerState);
 			PC->PlayerState->SetOwner(PC);
 			PC->PlayerState->SetReplicates(true);
 			PC->PlayerState->SetLifeSpan(0.0f);
diff --git a/Engine/Source/Runtime/Engine/Private/PlayerController.cpp b/Engine/Source/Runtime/Engine/Private/PlayerController.cpp
--- a/Engine/Source/Runtime/Engine/Private/PlayerController.cpp
+++ b/Engine/Source/Runtime/Engine/Private/PlayerController.cpp
@@ -1355,7 +1355,7 @@ void APlayerController::CleanupPlayerState()
 		// By default this destroys it, but games can override
 		PlayerState->OnDeactivated();
 	}
-	PlayerState = NULL;
+	SetPlayerState(NULL);
 }
 
 void APlayerController::OnActorChannelOpen(FInBunch& InBunch, UNetConnection* Connection)
@@ -3635,7 +3635,7 @@ void APlayerController::SeamlessTravelFrom(APlayerController* OldPC)
 
 		//@fixme: need a way to replace PlayerStates that doesn't cause incorrect "player left the game"/"player entered the game" messages
 		OldPC->PlayerState->Destroy();
-		OldPC->PlayerState = NULL;
+		OldPC->SetPlayerState(NULL);
 	}
 
 	// Copy seamless travel state
