Add CVar for push-based APlayerController

diff --git a/Engine/Source/Runtime/Engine/Classes/GameFramework/PlayerController.h b/Engine/Source/Runtime/Engine/Classes/GameFramework/PlayerController.h
--- a/Engine/Source/Runtime/Engine/Classes/GameFramework/PlayerController.h
+++ b/Engine/Source/Runtime/Engine/Classes/GameFramework/PlayerController.h
@@ -301,6 +301,9 @@ public:
 	UPROPERTY(replicated)
 	FRotator TargetViewRotation; 
 
+	/** Setter for TargetViewRotation. Marks it dirty if APlayerController is configured to use push-based dirtiness. */
+	ENGINE_API void SetTargetViewRotation(const FRotator& InRotation);
+
 	/** Smoothed version of TargetViewRotation to remove jerkiness from intermittent replication updates. */
 	FRotator BlendedTargetViewRotation;
 
diff --git a/Engine/Source/Runtime/Engine/Private/PlayerController.cpp b/Engine/Source/Runtime/Engine/Private/PlayerController.cpp
--- a/Engine/Source/Runtime/Engine/Private/PlayerController.cpp
+++ b/Engine/Source/Runtime/Engine/Private/PlayerController.cpp
@@ -32,6 +32,7 @@
 #include "Engine/Console.h"
 #include "Net/UnrealNetwork.h"
 #include "Net/Core/Misc/NetConditionGroupManager.h"
+#include "Net/Core/PushModel/PushModel.h"
 #include "Engine/WorldComposition.h"
 #include "Engine/LevelScriptActor.h"
 #include "GameFramework/GameNetworkManager.h"
@@ -104,6 +105,12 @@ namespace UE::Gameplay::CVars
 		TEXT("Whether to force the use of the camera as the streaming source for World Partition. By default the player pawn is used.\n")
 		TEXT("0: Use pawn as streaming source, 1: Use camera as streaming source"));
 
+	bool bIsPlayerControllerPushBased = false;
+	static FAutoConsoleVariableRef CVarIsPlayerControllerPushBased(
+		TEXT("PlayerController.IsPushBased"), bIsPlayerControllerPushBased,
+		TEXT("If true, APlayerController's replicated properties will use push-based networking, and will therefore need to be marked dirty when changed."),
+		ECVF_Default);
+
 	extern bool bAlwaysNotifyClientOnControllerChange;
 }
 
@@ -1833,8 +1840,15 @@ void APlayerController::UpdatePing(float InPing)
 
 void APlayerController::SetSpawnLocation(const FVector& NewLocation)
 {
-	SpawnLocation = NewLocation;
 	LastSpectatorSyncLocation = NewLocation;
+
+	if (UE::Gameplay::CVars::bIsPlayerControllerPushBased)
+	{
+		COMPARE_ASSIGN_AND_MARK_PROPERTY_DIRTY(APlayerController, SpawnLocation, NewLocation, this);
+		return;
+	}
+	
+	SpawnLocation = NewLocation;
 }
 
 
@@ -4876,14 +4890,28 @@ void APlayerController::SetPawn(APawn* InPawn)
 void APlayerController::GetLifetimeReplicatedProps(TArray< FLifetimeProperty > & OutLifetimeProps) const
 {
 	Super::GetLifetimeReplicatedProps(OutLifetimeProps);
-
+	
+	FDoRepLifetimeParams Params;
+	Params.bIsPushBased = UE::Gameplay::CVars::bIsPlayerControllerPushBased;
+	Params.Condition = COND_OwnerOnly;
 	// These used to only replicate if PlayerCameraManager->GetViewTargetPawn() != GetPawn()
 	// But, since they also don't update unless that condition is true, these values won't change, thus won't send
 	// This is a little less efficient, but fits into the new condition system well, and shouldn't really add much overhead
-	DOREPLIFETIME_CONDITION(APlayerController, TargetViewRotation, COND_OwnerOnly);
+	DOREPLIFETIME_WITH_PARAMS_FAST(APlayerController, TargetViewRotation, Params);
 
 	// Replicate SpawnLocation for remote spectators
-	DOREPLIFETIME_CONDITION(APlayerController, SpawnLocation, COND_OwnerOnly);
+	DOREPLIFETIME_WITH_PARAMS_FAST(APlayerController, SpawnLocation, Params);
+}
+
+void APlayerController::SetTargetViewRotation(const FRotator& InRotation)
+{
+	if (UE::Gameplay::CVars::bIsPlayerControllerPushBased)
+	{
+		COMPARE_ASSIGN_AND_MARK_PROPERTY_DIRTY(APlayerController, TargetViewRotation, InRotation, this);
+		return;
+	}
+
+	TargetViewRotation = InRotation;
 }
 
 void APlayerController::SetPlayer( UPlayer* InPlayer )
@@ -5173,7 +5201,7 @@ void APlayerController::TickActor( float DeltaSeconds, ELevelTick TickType, FAct
 			
 			if ((TargetPawn != GetPawn()) && (TargetPawn != nullptr))
 			{
-				TargetViewRotation = TargetPawn->GetViewRotation();
+				SetTargetViewRotation(TargetPawn->GetViewRotation());
 			}
 		}
 	}
