Remove the use of deprecated bulk data streaming API when streaming texture mips

#rnx

diff --git a/Engine/Source/Runtime/Engine/Private/Streaming/Texture2DStreamIn_IO.cpp b/Engine/Source/Runtime/Engine/Private/Streaming/Texture2DStreamIn_IO.cpp
--- a/Engine/Source/Runtime/Engine/Private/Streaming/Texture2DStreamIn_IO.cpp
+++ b/Engine/Source/Runtime/Engine/Private/Streaming/Texture2DStreamIn_IO.cpp
@@ -18,33 +18,15 @@ Texture2DStreamIn.cpp: Stream in helper for 2D textures using texture streaming
 #include "RenderUtils.h"
 #endif
 
-// Will be removed once the batch API has proven to work
-bool GBatchTextureMipIoRequests = true;
-static FAutoConsoleVariableRef CVarBatchTextureMipIoRequests(
-	TEXT("r.BatchTextureMipIoRequests"),
-	GBatchTextureMipIoRequests,
-	TEXT("Whether to batch texture mip I/O requests"));
-
 FTexture2DStreamIn_IO::FTexture2DStreamIn_IO(UTexture2D* InTexture, bool InPrioritizedIORequest)
 	: FTexture2DStreamIn(InTexture)
 	, bPrioritizedIORequest(InPrioritizedIORequest)
-	, bBatchIORequest(GBatchTextureMipIoRequests)
 
 {
-	if (!bBatchIORequest)
-	{
-		IORequests.AddZeroed(ResourceState.MaxNumLODs);
-	}
 }
 
 FTexture2DStreamIn_IO::~FTexture2DStreamIn_IO()
 {
-#if DO_CHECK
-	for (const IBulkDataIORequest* IORequest : IORequests)
-	{
-		check(!IORequest);
-	}
-#endif
 }
 
 static void ValidateMipBulkDataSize(const UTexture2D& Texture, int32 MipSizeX, int32 MipSizeY, int32 MipIndex, int64& BulkDataSize)
@@ -66,93 +48,8 @@ static void ValidateMipBulkDataSize(const UTexture2D& Texture, int32 MipSizeX, i
 
 void FTexture2DStreamIn_IO::SetIORequests(const FContext& Context)
 {
-	if (bBatchIORequest)
-	{
-		const int32 BatchCount = CurrentFirstLODIdx - PendingFirstLODIdx;
-		FBulkDataBatchRequest::FBatchBuilder Batch = FBulkDataBatchRequest::NewBatch(BatchCount);
-		for (int32 MipIndex = PendingFirstLODIdx; MipIndex < CurrentFirstLODIdx && !IsCancelled(); ++MipIndex)
-		{
-			const FTexture2DMipMap& MipMap = *Context.MipsView[MipIndex];
-			check(MipData[MipIndex].Data != nullptr);
-
-			int64 BulkDataSize = MipMap.BulkData.GetBulkDataSize();
-			if (BulkDataSize > 0)
-			{
-				// Validate buffer size for the mip, so we don't overrun it on streaming
-				// note: MipData[] should have size
-				// ValidateMipBulkDataSize only does anything on Android
-				ValidateMipBulkDataSize(*Context.Texture, MipMap.SizeX, MipMap.SizeY, MipIndex, BulkDataSize);
-				
-				// reads directly into MipData[] , doesn't respect Pitch
-				// we do get a completion callback at AsyncFileCallBack
-				// so in theory could fix Pitch there
-				uint32 DestPitch = MipData[MipIndex].Pitch;
-				FTexture2DResource::WarnRequiresTightPackedMip(MipMap.SizeX, MipMap.SizeY, Context.Resource->GetPixelFormat(), DestPitch);
-
-				EAsyncIOPriorityAndFlags Priority = AIOP_Low;
-				if (bPrioritizedIORequest)
-				{
-					static IConsoleVariable* CVarAsyncLoadingPrecachePriority = IConsoleManager::Get().FindConsoleVariable(TEXT("s.AsyncLoadingPrecachePriority"));
-					const bool bLoadBeforeAsyncPrecache = CVarStreamingLowResHandlingMode.GetValueOnAnyThread() == (int32)FRenderAssetStreamingSettings::LRHM_LoadBeforeAsyncPrecache;
-
-					if (CVarAsyncLoadingPrecachePriority && bLoadBeforeAsyncPrecache)
-					{
-						const int32 AsyncIOPriority = CVarAsyncLoadingPrecachePriority->GetInt();
-						// Higher priority than regular requests but don't go over max
-						Priority = (EAsyncIOPriorityAndFlags)FMath::Clamp<int32>(AsyncIOPriority + 1, AIOP_BelowNormal, AIOP_MAX);
-					}
-					else
-					{
-						Priority = AIOP_BelowNormal;
-					}
-				}
-
-				FIoBuffer Dst(FIoBuffer::Wrap, MipData[MipIndex].Data, BulkDataSize);
-				Batch.Read(MipMap.BulkData, 0, BulkDataSize, Priority | AIOP_FLAG_DONTCACHE | AIOP_FLAG_HW_TARGET_MEMORY, Dst);
-			}
-			else // Bulk data size can only be 0 when not available, in which case, we need to recache the file state.
-			{
-				bFailedOnIOError = true;
-				MarkAsCancelled();
-				break;
-			}
-		}
-
-		if (bFailedOnIOError || Batch.IsEmpty())
-		{
-			return;
-		}
-
-		TaskSynchronization.Increment();
-
-		return Batch.Issue(
-			[this](FBulkDataBatchRequest::EStatus Status)
-			{
-				TaskSynchronization.Decrement();
-
-				if (Status != FBulkDataBatchRequest::EStatus::Ok && !bIsCancelled)
-				{
-					// If IO requests was cancelled but the streaming request wasn't, this is an IO error.
-					bFailedOnIOError = true;
-					MarkAsCancelled();
-				}
-#if !UE_BUILD_SHIPPING
-				// On some platforms the IO is too fast to test cancelation requests timing issues.
-				if (FRenderAssetStreamingSettings::ExtraIOLatency > 0 && TaskSynchronization.GetValue() == 0)
-				{
-					FPlatformProcess::Sleep(FRenderAssetStreamingSettings::ExtraIOLatency * .001f); // Slow down the streaming.
-				}
-#endif
-				// The tick here is intended to schedule the success or cancel callback.
-				// Using TT_None ensure gets which could create a dead lock.
-				Tick(FTexture2DUpdate::TT_None);
-			},
-			BatchRequest);
-	}
-
-	check(BatchRequest.IsNone());
-	SetAsyncFileCallback();
-	
+	const int32 BatchCount = CurrentFirstLODIdx - PendingFirstLODIdx;
+	FBulkDataBatchRequest::FBatchBuilder Batch = FBulkDataBatchRequest::NewBatch(BatchCount);
 	for (int32 MipIndex = PendingFirstLODIdx; MipIndex < CurrentFirstLODIdx && !IsCancelled(); ++MipIndex)
 	{
 		const FTexture2DMipMap& MipMap = *Context.MipsView[MipIndex];
@@ -161,10 +58,6 @@ void FTexture2DStreamIn_IO::SetIORequests(const FContext& Context)
 		int64 BulkDataSize = MipMap.BulkData.GetBulkDataSize();
 		if (BulkDataSize > 0)
 		{
-			// Increment as we push the requests. If a requests complete immediately, then it will call the callback
-			// but that won't do anything because the tick would not try to acquire the lock since it is already locked.
-			TaskSynchronization.Increment();
-
 			// Validate buffer size for the mip, so we don't overrun it on streaming
 			// note: MipData[] should have size
 			// ValidateMipBulkDataSize only does anything on Android
@@ -194,12 +87,8 @@ void FTexture2DStreamIn_IO::SetIORequests(const FContext& Context)
 				}
 			}
 
-			IORequests[MipIndex] = MipMap.BulkData.CreateStreamingRequest(
-				0,
-				BulkDataSize,
-				Priority | AIOP_FLAG_DONTCACHE | AIOP_FLAG_HW_TARGET_MEMORY,
-				&AsyncFileCallBack,
-				(uint8*)MipData[MipIndex].Data);
+			FIoBuffer Dst(FIoBuffer::Wrap, MipData[MipIndex].Data, BulkDataSize);
+			Batch.Read(MipMap.BulkData, 0, BulkDataSize, Priority | AIOP_FLAG_DONTCACHE | AIOP_FLAG_HW_TARGET_MEMORY, Dst);
 		}
 		else // Bulk data size can only be 0 when not available, in which case, we need to recache the file state.
 		{
@@ -208,57 +97,49 @@ void FTexture2DStreamIn_IO::SetIORequests(const FContext& Context)
 			break;
 		}
 	}
-}
 
-void FTexture2DStreamIn_IO::CancelIORequests()
-{
-	if (bBatchIORequest)
+	if (bFailedOnIOError || Batch.IsEmpty())
 	{
-		check(IORequests.IsEmpty());
-		if (BatchRequest.IsPending())
-		{
-			BatchRequest.Cancel();
-		}
 		return;
 	}
 
-	check(BatchRequest.IsNone());
-	for (int32 MipIndex = 0; MipIndex < IORequests.Num(); ++MipIndex)
+	TaskSynchronization.Increment();
+
+	Batch.Issue([this](FBulkDataBatchRequest::EStatus Status)
 	{
-		IBulkDataIORequest* IORequest = IORequests[MipIndex];
-		if (IORequest)
+		TaskSynchronization.Decrement();
+
+		if (Status != FBulkDataBatchRequest::EStatus::Ok && !bIsCancelled)
 		{
-			// Calling cancel will trigger the SetAsyncFileCallback() which will also try a tick but will fail.
-			IORequest->Cancel();
+			// If IO requests was cancelled but the streaming request wasn't, this is an IO error.
+			bFailedOnIOError = true;
+			MarkAsCancelled();
 		}
-	}
+#if !UE_BUILD_SHIPPING
+		// On some platforms the IO is too fast to test cancelation requests timing issues.
+		if (FRenderAssetStreamingSettings::ExtraIOLatency > 0 && TaskSynchronization.GetValue() == 0)
+		{
+			FPlatformProcess::Sleep(FRenderAssetStreamingSettings::ExtraIOLatency * .001f); // Slow down the streaming.
+		}
+#endif
+		// The tick here is intended to schedule the success or cancel callback.
+		// Using TT_None ensure gets which could create a dead lock.
+		Tick(FTexture2DUpdate::TT_None);
+	},
+	BatchRequest);
 }
 
-void FTexture2DStreamIn_IO::ClearIORequests(const FContext& Context)
+void FTexture2DStreamIn_IO::CancelIORequests()
 {
-	if (bBatchIORequest)
+	if (BatchRequest.IsPending())
 	{
-		check(IORequests.IsEmpty());
-		return BatchRequest.Reset();
+		BatchRequest.Cancel();
 	}
+}
 
-	check(BatchRequest.IsNone());
-	for (int32 MipIndex = PendingFirstLODIdx; MipIndex < CurrentFirstLODIdx; ++MipIndex)
-	{
-		IBulkDataIORequest* IORequest = IORequests[MipIndex];
-		IORequests[MipIndex] = nullptr;
-
-		if (IORequest != nullptr)
-		{
-			// If clearing requests not yet completed, cancel and wait.
-			if (!IORequest->PollCompletion())
-			{
-				IORequest->Cancel();
-				IORequest->WaitCompletion();
-			}
-			delete IORequest;
-		}
-	}
+void FTexture2DStreamIn_IO::ClearIORequests(const FContext& Context)
+{
+	BatchRequest.Reset();
 }
 
 void FTexture2DStreamIn_IO::ReportIOError(const FContext& Context)
@@ -276,45 +157,13 @@ void FTexture2DStreamIn_IO::ReportIOError(const FContext& Context)
 	}
 }
 
-void FTexture2DStreamIn_IO::SetAsyncFileCallback()
-{
-	AsyncFileCallBack = [this](bool bWasCancelled, IBulkDataIORequest*)
-	{
-		// At this point task synchronization would hold the number of pending requests.
-		TaskSynchronization.Decrement();
-		
-		if (bWasCancelled)
-		{
-			// If IO requests was cancelled but the streaming request wasn't, this is an IO error.
-			if (!bIsCancelled)
-			{
-				bFailedOnIOError = true;
-			}
-
-			MarkAsCancelled();
-		}
-
-#if !UE_BUILD_SHIPPING
-		// On some platforms the IO is too fast to test cancelation requests timing issues.
-		if (FRenderAssetStreamingSettings::ExtraIOLatency > 0 && TaskSynchronization.GetValue() == 0)
-		{
-			FPlatformProcess::Sleep(FRenderAssetStreamingSettings::ExtraIOLatency * .001f); // Slow down the streaming.
-		}
-#endif
-
-		// The tick here is intended to schedule the success or cancel callback.
-		// Using TT_None ensure gets which could create a dead lock.
-		Tick(FTexture2DUpdate::TT_None);
-	};
-}
-
 void FTexture2DStreamIn_IO::Abort()
 {
 	if (!IsCancelled() && !IsCompleted())
 	{
 		FTexture2DStreamIn::Abort();
 
-		if (HasPendingIORequests())
+		if (BatchRequest.IsPending())
 		{
 			// Prevent the update from being considered done before this is finished.
 			// By checking that it was not already canceled, we make sure this doesn't get called twice.
@@ -323,26 +172,6 @@ void FTexture2DStreamIn_IO::Abort()
 	}
 }
 
-bool FTexture2DStreamIn_IO::HasPendingIORequests()
-{
-	if (bBatchIORequest)
-	{
-		check(IORequests.IsEmpty());
-		return BatchRequest.IsPending();
-	}
-
-	check(BatchRequest.IsNone());
-	for (IBulkDataIORequest* IORequest : IORequests)
-	{
-		if (IORequest != nullptr)
-		{
-			return true;
-		}
-	}
-
-	return false;
-}
-
 void FTexture2DStreamIn_IO::FCancelIORequestsTask::DoWork()
 {
 	check(PendingUpdate);
diff --git a/Engine/Source/Runtime/Engine/Private/Streaming/Texture2DStreamIn_IO.h b/Engine/Source/Runtime/Engine/Private/Streaming/Texture2DStreamIn_IO.h
--- a/Engine/Source/Runtime/Engine/Private/Streaming/Texture2DStreamIn_IO.h
+++ b/Engine/Source/Runtime/Engine/Private/Streaming/Texture2DStreamIn_IO.h
@@ -9,7 +9,6 @@ Texture2DStreamIn.h: Stream in helper for 2D textures using texture streaming fi
 #include "CoreMinimal.h"
 #include "Texture2DStreamIn.h"
 
-class IBulkDataIORequest;
 class FBulkDataRequest;
 
 // Base StreamIn framework exposing MipData
@@ -33,7 +32,7 @@ protected:
 	// Report IO errors if any.
 	void ReportIOError(const FContext& Context);
 	// Set the IO callback used for streaming the mips.
-	void SetAsyncFileCallback();
+	void SetAsyncFileCallback() { };
 	// Cancel all IO requests.
 	void CancelIORequests();
 
@@ -42,9 +41,6 @@ protected:
 
 private:
 
-	// Poll if any of the mips currently have an active IO request
-	bool HasPendingIORequests();
-
 	class FCancelIORequestsTask : public FNonAbandonableTask
 	{
 	public:
@@ -61,20 +57,13 @@ private:
 	typedef FAutoDeleteAsyncTask<FCancelIORequestsTask> FAsyncCancelIORequestsTask;
 	friend class FCancelIORequestsTask;
 
-
-	// Request for loading into each mip.
-	TArray<IBulkDataIORequest*> IORequests;
-
 	// Whether an IO error was detected (when files do not exists).
 	bool bFailedOnIOError = false;
 
 	// Whether the IO request should be created with an higher priority for quicker response time.
 	bool bPrioritizedIORequest = false;
 
-	// Whether to batch IO request(s) 
-	bool bBatchIORequest;
-
-	FBulkDataIORequestCallBack AsyncFileCallBack;
+	// I/O request handle
 	FBulkDataBatchRequest BatchRequest;
 };
 
