[SoundWave]
 - Reduce class size to fit in 1024 MB3 bucket from the 1168 byte bucket.

#rnx

diff --git a/Engine/Source/Runtime/Engine/Classes/Sound/SoundWave.h b/Engine/Source/Runtime/Engine/Classes/Sound/SoundWave.h
--- a/Engine/Source/Runtime/Engine/Classes/Sound/SoundWave.h
+++ b/Engine/Source/Runtime/Engine/Classes/Sound/SoundWave.h
@@ -13,6 +13,8 @@
 #include "UObject/ObjectMacros.h"
 #include "Misc/Guid.h"
 #include "Async/AsyncWork.h"
+#include "Async/Mutex.h"
+#include "Async/UniqueLock.h"
 #include "Sound/SoundBase.h"
 #include "Sound/SoundWaveTimecodeInfo.h"
 #include "Interfaces/Interface_AsyncCompilation.h"
@@ -469,10 +471,6 @@ public:
 
 private:
 
-	/** The compression type to use for the sound wave asset. */
-	UPROPERTY(EditAnywhere, Category = "Format")
-	ESoundAssetCompressionType SoundAssetCompressionType = ESoundAssetCompressionType::PlatformSpecific;
-
 	// Deprecated compression type properties
 	UPROPERTY(meta = (DeprecatedProperty, DeprecationMessage = "5.0 - Property is deprecated. bSeekableStreaming now means ADPCM codec in SoundAssetCompressionType."))
 	uint8 bSeekableStreaming : 1;
@@ -480,11 +478,11 @@ private:
 	UPROPERTY(meta = (DeprecatedProperty, DeprecationMessage = "5.0 - Property is deprecated. bUseBinkAudio now means Bink codec in SoundAssetCompressionType."))
 	uint8 bUseBinkAudio : 1;
 
-public:
-
-	/** the number of sounds currently playing this sound wave. */
-	FThreadSafeCounter NumSourcesPlaying;
+	/** The compression type to use for the sound wave asset. */
+	UPROPERTY(EditAnywhere, Category = "Format")
+	ESoundAssetCompressionType SoundAssetCompressionType = ESoundAssetCompressionType::PlatformSpecific;
 
+public:
 	void AddPlayingSource()
 	{
 		NumSourcesPlaying.Increment();
@@ -622,9 +620,11 @@ public:
 	 */
 	ENGINE_API ESoundWaveLoadingBehavior GetLoadingBehavior(bool bCheckSoundClasses = true) const;
 
+#if WITH_EDITORONLY_DATA
 	/** Please use size of First Chunk in Seconds. */
-	UPROPERTY(AdvancedDisplay, meta=(DeprecatedProperty))
+	UPROPERTY(AdvancedDisplay, meta = (DeprecatedProperty))
 	int32 InitialChunkSize_DEPRECATED;
+#endif
 
 #if WITH_EDITOR
 	ENGINE_API const FWaveTransformUObjectConfiguration& GetTransformationChainConfig() const;
@@ -640,8 +640,11 @@ private:
 	/** What state the precache decompressor is in. */
 	FThreadSafeCounter PrecacheState;
 
+public:
 	/** the number of sounds currently playing this sound wave. */
-	mutable FCriticalSection SourcesPlayingCs;
+	FThreadSafeCounter NumSourcesPlaying;
+
+private:
 
 	TArray<FSoundWaveClientPtr> SourcesPlaying;
 
@@ -737,9 +740,13 @@ public:
 	/** Specifies how and when compressed audio data is loaded for asset if stream caching is enabled. */
 	UPROPERTY(EditAnywhere, Category = "Loading", meta = (DisplayName = "Loading Behavior Override"))
 	mutable ESoundWaveLoadingBehavior LoadingBehavior;
+	
+private:
+	/** the number of sounds currently playing this sound wave. */
+	mutable UE::FMutex SourcesPlayingCs;
 
-#if WITH_EDITORONLY_DATA
 public:
+#if WITH_EDITORONLY_DATA
    	/** How much audio to add to First Audio Chunk (in seconds) */
 	UPROPERTY(EditAnywhere, Category = Loading, meta = (UIMin = 0, UIMax = 10, EditCondition = "LoadingBehavior == ESoundWaveLoadingBehavior::RetainOnLoad || LoadingBehavior == ESoundWaveLoadingBehavior::PrimeOnLoad"), DisplayName="Size of First Audio Chunk (seconds)")
    	FPerPlatformFloat SizeOfFirstAudioChunkInSeconds = 0.0f;
@@ -1164,7 +1171,7 @@ public:
 	bool IsGeneratingAudio() const
 	{
 		bool bIsGeneratingAudio = false;
-		FScopeLock Lock(&SourcesPlayingCs);
+		UE::TUniqueLock Lock(SourcesPlayingCs);
 		bIsGeneratingAudio = SourcesPlaying.Num() > 0;
 
 		return bIsGeneratingAudio;
diff --git a/Engine/Source/Runtime/Engine/Private/SoundWave.cpp b/Engine/Source/Runtime/Engine/Private/SoundWave.cpp
--- a/Engine/Source/Runtime/Engine/Private/SoundWave.cpp
+++ b/Engine/Source/Runtime/Engine/Private/SoundWave.cpp
@@ -2244,7 +2244,7 @@ void USoundWave::BeginDestroy()
 	Super::BeginDestroy();
 
 	{
-		FScopeLock Lock(&SourcesPlayingCs);
+		UE::TUniqueLock Lock(SourcesPlayingCs);
 		int32 CurrNumSourcesPlaying = SourcesPlaying.Num();
 
 		for (int32 i = CurrNumSourcesPlaying - 1; i >= 0; --i)
@@ -3192,7 +3192,7 @@ FWaveInstance& USoundWave::HandleStart(FActiveSound& ActiveSound, const UPTRINT
 bool USoundWave::IsReadyForFinishDestroy()
 {
 	{
-		FScopeLock Lock(&SourcesPlayingCs);
+		UE::TUniqueLock Lock(SourcesPlayingCs);
 
 		for (ISoundWaveClient* SoundWaveClientPtr : SourcesPlaying)
 		{
@@ -3203,7 +3203,7 @@ bool USoundWave::IsReadyForFinishDestroy()
 	}
 
 	{
-		FScopeLock Lock(&SourcesPlayingCs);
+		UE::TUniqueLock Lock(SourcesPlayingCs);
 
 		for (ISoundWaveClient* SoundWaveClientPtr : SourcesPlaying)
 		{
@@ -3248,17 +3248,19 @@ void USoundWave::FinishDestroy()
 {
 	Super::FinishDestroy();
 
-	FScopeLock Lock(&SourcesPlayingCs);
-	int32 CurrNumSourcesPlaying = SourcesPlaying.Num();
-
-	for (int32 i = CurrNumSourcesPlaying - 1; i >= 0; --i)
 	{
-		ISoundWaveClient* SoundWaveClientPtr = SourcesPlaying[i];
+		UE::TUniqueLock Lock(SourcesPlayingCs);
+		int32 CurrNumSourcesPlaying = SourcesPlaying.Num();
 
-		if (SoundWaveClientPtr)
+		for (int32 i = CurrNumSourcesPlaying - 1; i >= 0; --i)
 		{
-			SoundWaveClientPtr->OnFinishDestroy(this);
-			SourcesPlaying.RemoveAtSwap(i, EAllowShrinking::No);
+			ISoundWaveClient* SoundWaveClientPtr = SourcesPlaying[i];
+
+			if (SoundWaveClientPtr)
+			{
+				SoundWaveClientPtr->OnFinishDestroy(this);
+				SourcesPlaying.RemoveAtSwap(i, EAllowShrinking::No);
+			}
 		}
 	}
 
@@ -3849,7 +3851,7 @@ void USoundWave::AddPlayingSource(const FSoundWaveClientPtr& Source)
 	check(IsInAudioThread() || IsInGameThread());   // Don't allow incrementing on other threads as it's not safe (for GCing of this wave).
 	if (Source)
 	{
-		FScopeLock Lock(&SourcesPlayingCs);
+		UE::TUniqueLock Lock(SourcesPlayingCs);
 		check(!SourcesPlaying.Contains(Source));
 		SourcesPlaying.Add(Source);
 	}
@@ -3859,7 +3861,7 @@ void USoundWave::RemovePlayingSource(const FSoundWaveClientPtr& Source)
 {
 	if (Source)
 	{
-		FScopeLock Lock(&SourcesPlayingCs);
+		UE::TUniqueLock Lock(SourcesPlayingCs);
 		check(SourcesPlaying.Contains(Source));
 		SourcesPlaying.RemoveSwap(Source);
 	}
