Added some better Audio Compression statistics in the log (produces output like below):
LogAudioDerivedData: Display: voice_en_male_single_number_10 compressed to 7.74% (from 42KB to 3KB) with RADA at Quality 80 with Quality Modifier 1.00

[RN] Better logging for audio compression statistics

diff --git a/Engine/Source/Runtime/Engine/Private/AudioDerivedData.cpp b/Engine/Source/Runtime/Engine/Private/AudioDerivedData.cpp
--- a/Engine/Source/Runtime/Engine/Private/AudioDerivedData.cpp
+++ b/Engine/Source/Runtime/Engine/Private/AudioDerivedData.cpp
@@ -2151,7 +2151,24 @@ bool FDerivedAudioDataCompressor::Build(TArray<uint8>& OutData)
 		CookSurroundWave(*CookInputs, OutData);
 	}
 
+	const uint64 BeforeSize = CookInputs->BulkData.GetPayloadSize();
+	const uint64 AfterSize = OutData.Num();
+	const float Percent = BeforeSize > 0 ? ((float)AfterSize / BeforeSize) * 100.f : 0.0f;
+
+	// Log message about the completed results
+	FFormatNamedArguments Args2;
+	Args2.Add(TEXT("AudioFormat"), FText::FromName(CookInputs->BaseFormat));
+	Args2.Add(TEXT("SoundNodeName"), FText::FromString(CookInputs->SoundName));	
+	Args2.Add(TEXT("BeforeSize"), BeforeSize >> 10);
+	Args2.Add(TEXT("AfterSize"),  AfterSize >> 10);
+	Args2.Add(TEXT("Percent"), FText::FromString(FString::Printf(TEXT("%2.2f"), Percent)));
+	Args2.Add(TEXT("Quality"), FText::FromString(LexToString(CookInputs->CompressionQuality)));
+	Args2.Add(TEXT("QualityMod"), FText::FromString(FString::Printf(TEXT("%2.2f"), CookInputs->CompressionQualityModifier)));
+
+	FAudioStatusMessageContext CompressedMessage(FText::Format(NSLOCTEXT("Engine", "BuildingCompressedAudioTaskResults", "{SoundNodeName} compressed to {Percent}% (from {BeforeSize}KB to {AfterSize}KB) with {AudioFormat} at Quality {Quality} with Quality Modifier {QualityMod}"), Args2));
+
 #endif
+
 	return OutData.Num() > 0;
 }
 
