Fix doc spelling in World.h

diff --git a/Engine/Source/Runtime/Engine/Classes/Engine/World.h b/Engine/Source/Runtime/Engine/Classes/Engine/World.h
--- a/Engine/Source/Runtime/Engine/Classes/Engine/World.h
+++ b/Engine/Source/Runtime/Engine/Classes/Engine/World.h
@@ -514,7 +514,7 @@ struct ENGINE_API FActorSpawnUtils
 	/**
 	 * Function to generate a locally or globally unique actor name. To generate a globally unique name, we store an epoch number
 	 * in the name number (while maintaining compatibility with fast path name generation, see GFastPathUniqueNameGeneration) and
-	 * also append an unique user id to the name.
+	 * also append a unique user id to the name.
 	 *
 	 * @param	Level			the new actor level
 	 * @param	Class			the new actor class
@@ -3057,7 +3057,7 @@ public:
 	void MarkActorComponentForNeededEndOfFrameUpdate(UActorComponent* Component, bool bForceGameThread);
 
 	/**
-	* Clears the need for a component to have a end of frame update
+	* Clears the need for a component to have an end of frame update
 	* @param Component - Component to update at the end of the frame
 	*/
 	void ClearActorComponentEndOfFrameUpdate(UActorComponent* Component);
