PR #11918: Allow custom ACTOR_HAS_LABELS defines on all platforms

diff --git a/Engine/Source/Runtime/Engine/Classes/GameFramework/Actor.h b/Engine/Source/Runtime/Engine/Classes/GameFramework/Actor.h
--- a/Engine/Source/Runtime/Engine/Classes/GameFramework/Actor.h
+++ b/Engine/Source/Runtime/Engine/Classes/GameFramework/Actor.h
@@ -57,7 +57,9 @@ namespace UE::Net
 
 // By default, debug and development builds (even cooked) will keep actor labels. Manually define this if you want to make a local build
 // that keep actor labels for Test or Shipping builds.
+#ifndef ACTOR_HAS_LABELS
 #define ACTOR_HAS_LABELS (UE_BUILD_DEBUG || UE_BUILD_DEVELOPMENT || WITH_PROFILEGPU)
+#endif
 
 /** Chooses a method for actors to update overlap state (objects it is touching) on initialization, currently only used during level streaming. */
 UENUM(BlueprintType)
