Markup compression type as AssetRegistrySearchable so it appears in the content browser ( requires a resave of the soundwaves to work).

#rnx

diff --git a/Engine/Source/Runtime/Engine/Classes/Sound/SoundWave.h b/Engine/Source/Runtime/Engine/Classes/Sound/SoundWave.h
--- a/Engine/Source/Runtime/Engine/Classes/Sound/SoundWave.h
+++ b/Engine/Source/Runtime/Engine/Classes/Sound/SoundWave.h
@@ -479,7 +479,7 @@ private:
 	uint8 bUseBinkAudio : 1;
 
 	/** The compression type to use for the sound wave asset. */
-	UPROPERTY(EditAnywhere, Category = "Format")
+	UPROPERTY(EditAnywhere, Category = "Format", AssetRegistrySearchable)
 	ESoundAssetCompressionType SoundAssetCompressionType = ESoundAssetCompressionType::PlatformSpecific;
 
 public:
