[Shaders] Detect and handle memory limitation violations for SCW processes.

- Add CVar "r.ShaderCompiler.MemoryLimit" (default 0) to allow specifying a memory limit for all SCW processes. Once this limit is exceeded, 50% of SCW processes will be suspended.
- Check periodically if worker memory consumption falls below half the limit, in which case 50% of active workers will be resumed, effectively going back to 75% of total workers.
- Define Windows specific memory limitation code in separate source files.

Example: For a harsh limit of only 2 GiB, the log entries of Lyra can look like this:
  LogShaderCompilers: Display: Shader compiler memory usage of 2053 MiB exceeded limit of 2000 MiB: Backlogged 152 compile jobs from 32 suspended workers (32/64 active)
  LogShaderCompilers: Display: Shader compiler memory usage of 2055 MiB exceeded limit of 2000 MiB: Backlogged 68 compile jobs from 16 suspended workers (16/64 active)
  LogShaderCompilers: Display: Shader compiler memory usage of 2012 MiB exceeded limit of 2000 MiB: Backlogged 43 compile jobs from 8 suspended workers (8/64 active)
  LogShaderCompilers: Display: Shader compiler memory usage of 2029 MiB exceeded limit of 2000 MiB: Backlogged 40 compile jobs from 4 suspended workers (4/64 active)
  LogShaderCompilers: Display: Resumed 2 suspended workers since memory usage of 940 MiB is below half the limit of 2000 MiB (6/64 active)
  LogShaderCompilers: Display: Resumed 3 suspended workers since memory usage of 986 MiB is below half the limit of 2000 MiB (9/64 active)

#rnx

diff --git a/Engine/Source/Runtime/Engine/Private/ShaderCompiler/ShaderCompilerMemoryLimit.cpp b/Engine/Source/Runtime/Engine/Private/ShaderCompiler/ShaderCompilerMemoryLimit.cpp
new file mode 100644
--- /dev/null
+++ b/Engine/Source/Runtime/Engine/Private/ShaderCompiler/ShaderCompilerMemoryLimit.cpp
@@ -0,0 +1,188 @@
+// Copyright Epic Games, Inc. All Rights Reserved.
+
+/*=============================================================================
+	ShaderCompilerMemoryLimit.cpp: Wrapper for Windows specific Job Object functionality.
+=============================================================================*/
+
+#include "ShaderCompilerMemoryLimit.h"
+#include "HAL/ConsoleManager.h"
+
+
+#if PLATFORM_WINDOWS
+
+FWindowsResourceRestrictedJobObject::FWindowsResourceRestrictedJobObject(const FString& InJobName, int32 InitialJobMemoryLimitMiB) :
+	JobName(InJobName)
+{
+	JobObject = ::CreateJobObject(nullptr, *JobName);
+	if (JobObject == nullptr)
+	{
+		UE_LOG(LogWindows, Warning, TEXT("Failed to create job object \"%s\": %s"), *JobName, *GetErrorMessage());
+		return;
+	}
+	CreateAndLinkCompletionPort();
+	if (InitialJobMemoryLimitMiB > 0)
+	{
+		SetMemoryLimit(InitialJobMemoryLimitMiB);
+	}
+}
+
+FWindowsResourceRestrictedJobObject::~FWindowsResourceRestrictedJobObject()
+{
+	if (JobObject != nullptr)
+	{
+		::CloseHandle(JobObject);
+	}
+	if (CompletionPort != nullptr)
+	{
+		::CloseHandle(CompletionPort);
+	}
+}
+
+void FWindowsResourceRestrictedJobObject::AssignProcess(const FProcHandle& Process)
+{
+	if (IsValid())
+	{
+		// Assign input process to job object
+		if (::AssignProcessToJobObject(JobObject, Process.Get()) == 0)
+		{
+			UE_LOG(LogWindows, Warning, TEXT("Failed to assign process to job object \"%s\": %s"), *JobName, *GetErrorMessage());
+		}
+	}
+}
+
+void FWindowsResourceRestrictedJobObject::SetMemoryLimit(int32 InJobMemoryLimitMiB)
+{
+	checkf(InJobMemoryLimitMiB >= 1024, TEXT("Cannot launch ShaderCompileWorker processes with memory restriction of less than 1024 MiB (%u MiB was specified)"), InJobMemoryLimitMiB);
+	if (MemoryLimit != InJobMemoryLimitMiB)
+	{
+		if (IsValid())
+		{
+			JOBOBJECT_NOTIFICATION_LIMIT_INFORMATION JobLimitInfo = {};
+
+			JobLimitInfo.JobMemoryLimit = 1024 * 1024 * static_cast<SIZE_T>(InJobMemoryLimitMiB);
+			JobLimitInfo.LimitFlags = JOB_OBJECT_LIMIT_JOB_MEMORY;
+
+			if (::SetInformationJobObject(JobObject, JobObjectNotificationLimitInformation, &JobLimitInfo, sizeof(JobLimitInfo)) == 0)
+			{
+				UE_LOG(LogWindows, Warning, TEXT("Failed to set restriction information for job object \"%s\": %s"), *JobName, *GetErrorMessage());
+			}
+		}
+		MemoryLimit = InJobMemoryLimitMiB;
+	}
+}
+
+bool FWindowsResourceRestrictedJobObject::QueryStatus(FJobObjectLimitationInfo& OutInfo)
+{
+	JOBOBJECT_LIMIT_VIOLATION_INFORMATION LimitViolationInfo = {};
+	DWORD ReturnLength = 0;
+	if (::QueryInformationJobObject(JobObject, JobObjectLimitViolationInformation, &LimitViolationInfo, sizeof(LimitViolationInfo), &ReturnLength) != 0 &&
+		ReturnLength == sizeof(LimitViolationInfo) &&
+		LimitViolationInfo.JobMemoryLimit > 0)
+	{
+		OutInfo.MemoryLimit = LimitViolationInfo.JobMemoryLimit;
+		OutInfo.MemoryUsed = LimitViolationInfo.JobMemory;
+		return true;
+	}
+	return false;
+}
+
+bool FWindowsResourceRestrictedJobObject::QueryLimitViolationStatus(FJobObjectLimitationInfo& OutInfo)
+{
+	if (IsValid())
+	{
+		// Check if a packet has been dequeued from the IO completion port
+		constexpr DWORD WaitDurationInMilliseconds = 0;
+		DWORD NumBytesTransferred = 0;
+		ULONG_PTR CompletionKey = 0;
+		LPOVERLAPPED Overlapped = nullptr;
+		while (::GetQueuedCompletionStatus(CompletionPort, &NumBytesTransferred, &CompletionKey, &Overlapped, WaitDurationInMilliseconds) != 0)
+		{
+			// Check if this is the message we are looking for, i.e. notification about limiation violation
+			if (NumBytesTransferred == JOB_OBJECT_MSG_NOTIFICATION_LIMIT)
+			{
+				// Query violation information from job object
+				JOBOBJECT_LIMIT_VIOLATION_INFORMATION LimitViolationInfo = {};
+				DWORD ReturnLength = 0;
+				if (::QueryInformationJobObject(JobObject, JobObjectLimitViolationInformation, &LimitViolationInfo, sizeof(LimitViolationInfo), &ReturnLength) != 0 &&
+					ReturnLength == sizeof(LimitViolationInfo))
+				{
+					// Even though the limit has been exceeded as indicated by the JOB_OBJECT_MSG_NOTIFICATION_LIMIT message,
+					// by the time we query the information it might have already gone down again,
+					// so check if the memory information is actually exceeding our limit here.
+					if (LimitViolationInfo.JobMemory >= LimitViolationInfo.JobMemoryLimit)
+					{
+						OutInfo.MemoryLimit = LimitViolationInfo.JobMemoryLimit;
+						OutInfo.MemoryUsed = LimitViolationInfo.JobMemory;
+						return true;
+					}
+				}
+			}
+		}
+	}
+	return false;
+}
+
+FString FWindowsResourceRestrictedJobObject::GetErrorMessage()
+{
+	TCHAR ErrorMessage[4096];
+	return FString(FPlatformMisc::GetSystemErrorMessage(ErrorMessage, UE_ARRAY_COUNT(ErrorMessage), 0));
+}
+
+bool FWindowsResourceRestrictedJobObject::IsValid() const
+{
+	return JobObject != nullptr && CompletionPort != nullptr;
+}
+
+void FWindowsResourceRestrictedJobObject::CreateAndLinkCompletionPort()
+{
+	// Create completion port
+	CompletionPort = ::CreateIoCompletionPort(INVALID_HANDLE_VALUE, nullptr, 0, 1);
+	if (CompletionPort == nullptr)
+	{
+		UE_LOG(LogWindows, Warning, TEXT("Failed to create I/O completion port for job object \"%s\": %s"), *JobName, *GetErrorMessage());
+		return;
+	}
+
+	// Link completion port to job object
+	JOBOBJECT_ASSOCIATE_COMPLETION_PORT CompletionPortInfo = {};
+
+	CompletionPortInfo.CompletionKey = JobObject;
+	CompletionPortInfo.CompletionPort = CompletionPort;
+
+	if (::SetInformationJobObject(JobObject, JobObjectAssociateCompletionPortInformation, &CompletionPortInfo, sizeof(CompletionPortInfo)) == 0)
+	{
+		UE_LOG(LogWindows, Warning, TEXT("Failed to set completion port for job object \"%s\": %s"), *JobName, *GetErrorMessage());
+	}
+}
+
+#else
+
+FGenericResourceRestrictedJobObject::FGenericResourceRestrictedJobObject(const FString& JobName, int32 InitialJobMemoryLimitMiB)
+{
+	// dummy
+}
+
+void FGenericResourceRestrictedJobObject::AssignProcess(const FProcHandle& Process)
+{
+	// dummy
+}
+
+void FGenericResourceRestrictedJobObject::SetMemoryLimit(int32 InJobMemoryLimitMiB)
+{
+	// dummy
+}
+
+bool FGenericResourceRestrictedJobObject::QueryStatus(FJobObjectLimitationInfo& OutInfo)
+{
+	return false; // dummy
+}
+
+bool FGenericResourceRestrictedJobObject::QueryLimitViolationStatus(FJobObjectLimitationInfo& OutInfo)
+{
+	return false; // dummy
+}
+
+#endif // PLATFORM_WINDOWS
+
+
+
diff --git a/Engine/Source/Runtime/Engine/Private/ShaderCompiler/ShaderCompilerMemoryLimit.h b/Engine/Source/Runtime/Engine/Private/ShaderCompiler/ShaderCompilerMemoryLimit.h
new file mode 100644
--- /dev/null
+++ b/Engine/Source/Runtime/Engine/Private/ShaderCompiler/ShaderCompilerMemoryLimit.h
@@ -0,0 +1,86 @@
+// Copyright Epic Games, Inc. All Rights Reserved.
+
+/*=============================================================================
+	ShaderCompilerMemoryLimit.h: Wrapper for Windows specific Job Object functionality.
+=============================================================================*/
+
+#pragma once
+
+#include "CoreMinimal.h"
+#include "HAL/PlatformProcess.h"
+
+
+/** Output structure for polling job object limitation violation status. See FWindowsResourceRestrictedJobObject::QueryLimitViolationStatus(). */
+struct FJobObjectLimitationInfo
+{
+	/** Size (in bytes) of the job memory limitation. */
+	int64 MemoryLimit = 0;
+
+	/** Size (in bytes) of the job memory usage. When QueryLimitViolationStatus() returns true, this will be greater than JobMemoryLimit, since the job object violated the limitation requirements. */
+	int64 MemoryUsed = 0;
+};
+
+
+#if PLATFORM_WINDOWS
+
+#include "Windows/MinimalWindowsApi.h"
+#include "Windows/WindowsHWrapper.h"
+
+/** Wrapper lass for a resource restricted job object. This is only available on Windows and acts as a placeholder on other platforms. */
+class FWindowsResourceRestrictedJobObject
+{
+public:
+	FWindowsResourceRestrictedJobObject(const FWindowsResourceRestrictedJobObject&) = delete;
+	FWindowsResourceRestrictedJobObject& operator = (const FWindowsResourceRestrictedJobObject&) = delete;
+
+	FWindowsResourceRestrictedJobObject(const FString& InJobName, int32 InitialJobMemoryLimitMiB = 0);
+	~FWindowsResourceRestrictedJobObject();
+
+	/** Assigns the specified process to this job object. */
+	void AssignProcess(const FProcHandle& Process);
+
+	/** Sets the memory limitation for this job object in MiB. This must be greater than or equal to 1024 MiB. */
+	void SetMemoryLimit(int32 InJobMemoryLimitMiB);
+
+	/** Queryies the status of the current memory usage. */
+	bool QueryStatus(FJobObjectLimitationInfo& OutInfo);
+
+	/** Queries the status of limitation violation notifications. */
+	bool QueryLimitViolationStatus(FJobObjectLimitationInfo& OutInfo);
+
+private:
+	static FString GetErrorMessage();
+
+	bool IsValid() const;
+
+	void CreateAndLinkCompletionPort();
+
+private:
+	FString JobName;
+	HANDLE JobObject = nullptr;
+	HANDLE CompletionPort = nullptr;
+	int32 MemoryLimit = 0;
+};
+
+using FResourceRestrictedJobObject = FWindowsResourceRestrictedJobObject;
+
+#else
+
+/** Wrapper lass for a resource restricted job object. This is only available on Windows and acts as a placeholder on other platforms. */
+class FGenericResourceRestrictedJobObject
+{
+public:
+	FGenericResourceRestrictedJobObject(const FGenericResourceRestrictedJobObject&) = delete;
+	FGenericResourceRestrictedJobObject& operator = (const FGenericResourceRestrictedJobObject&) = delete;
+
+	FGenericResourceRestrictedJobObject(const FString& JobName, int32 InitialJobMemoryLimitMiB = 0);
+
+	void AssignProcess(const FProcHandle& Process);
+	void SetMemoryLimit(int32 InJobMemoryLimitMiB);
+	bool QueryStatus(FJobObjectLimitationInfo& OutInfo);
+	bool QueryLimitViolationStatus(FJobObjectLimitationInfo& OutInfo);
+};
+
+using FResourceRestrictedJobObject = FGenericResourceRestrictedJobObject;
+
+#endif // PLATFORM_WINDOWS
diff --git a/Engine/Source/Runtime/Engine/Private/ShaderCompiler/ShaderCompilerThreadRunnable.cpp b/Engine/Source/Runtime/Engine/Private/ShaderCompiler/ShaderCompilerThreadRunnable.cpp
--- a/Engine/Source/Runtime/Engine/Private/ShaderCompiler/ShaderCompilerThreadRunnable.cpp
+++ b/Engine/Source/Runtime/Engine/Private/ShaderCompiler/ShaderCompilerThreadRunnable.cpp
@@ -6,10 +6,12 @@
 =============================================================================*/
 
 #include "ShaderCompilerPrivate.h"
+#include "ShaderCompilerMemoryLimit.h"
 
 #include "Async/ParallelFor.h"
 #include "HAL/PlatformFile.h"
 #include "HAL/PlatformFileManager.h"
+#include "Logging/StructuredLog.h"
 #include "Misc/ScopeTryLock.h"
 #include "ProfilingDebugging/CookStats.h"
 #include "ProfilingDebugging/DiagnosticTable.h"
@@ -31,6 +33,15 @@ static TAutoConsoleVariable<bool> CVarCompileParallelInProcess(
 	TEXT("EXPERIMENTAL- If true, shader compilation will be executed in-process in parallel. Note that this will serialize if the legacy preprocessor is enabled."),
 	ECVF_ReadOnly);
 
+int32 GShaderCompilerMemoryLimit = 0;
+static FAutoConsoleVariableRef CVarShaderCompilerMemoryLimit(
+	TEXT("r.ShaderCompiler.MemoryLimit"),
+	GShaderCompilerMemoryLimit,
+	TEXT("Specifies a memory limit (in MiB) for all ShaderCompileWorker (SCW) processes.") \
+	TEXT("If the total memory consumption of all SCW processes exceeds this limit, the editor will start to suspend workers and reschedule compile jobs.") \
+	TEXT("By default 0, effectively disabling the limitation. If this is non-zero, it must be greater than or equal to 1024 since shader compilation must be granted at least 1024 MiB of memory in total."),
+	ECVF_ReadOnly);
+
 // Configuration to retry shader compile through workers after a worker has been abandoned
 static constexpr int32 GSingleThreadedRunsDisabled = -2;
 static constexpr int32 GSingleThreadedRunsIncreaseFactor = 8;
@@ -40,6 +51,25 @@ static const TCHAR* GWorkerInputFilename = TEXT("WorkerInputOnly.in");
 static const TCHAR* GWorkerOutputFilename = TEXT("WorkerOutputOnly.out");
 
 
+static FResourceRestrictedJobObject GSCWResourceRestrictedJobObject(TEXT("UE.ShaderCompileWorker.JobGroup"));
+
+// Apply memory limits (see CVar r.ShaderCompiler.MemoryLimit) to by assigning input process to resource restricted job object
+// and initialize this job object here, since we can't guarantee execution order of static global object (i.e. global cvar and the job object).
+static void ApplyWorkerProcessMemoryLimits(const FProcHandle& Process)
+{
+	if (GShaderCompilerMemoryLimit > 0)
+	{
+		static bool bIsJobObjectLimitInitialized;
+		if (!bIsJobObjectLimitInitialized)
+		{
+			GSCWResourceRestrictedJobObject.SetMemoryLimit(GShaderCompilerMemoryLimit);
+			bIsJobObjectLimitInitialized = true;
+		}
+		GSCWResourceRestrictedJobObject.AssignProcess(Process);
+	}
+}
+
+
 /** Information tracked for each shader compile worker process instance. */
 struct FShaderCompileWorkerInfo
 {
@@ -246,6 +276,12 @@ int32 FShaderCompileThreadRunnable::PullTasksFromQueue()
 			}
 		};
 
+	// Check if memory limitations has been violated and suspend workers as needed
+	if (GShaderCompilerMemoryLimit > 0)
+	{
+		CheckMemoryLimitViolation();
+	}
+
 	FScopeLock WorkerScopeLock(&WorkerInfosLock); // Must be entered before CompileQueueSection
 
 	int32 NumActiveThreads = 0;
@@ -387,13 +423,13 @@ void FShaderCompileThreadRunnable::PushCompletedJobsToManager()
 				Manager->ProcessFinishedJob(Job.GetReference());
 			}
 
-			const float ElapsedTime = FPlatformTime::Seconds() - CurrentWorkerInfo.StartTime;
+			const double ElapsedTime = FPlatformTime::Seconds() - CurrentWorkerInfo.StartTime;
 
 			Manager->WorkersBusyTime += ElapsedTime;
 			COOK_STAT(ShaderCompilerCookStats::AsyncCompileTimeSec += ElapsedTime);
 
 			// Log if requested or if there was an exceptionally slow batch, to see the offender easily
-			if (Manager->bLogJobCompletionTimes || ElapsedTime > 60.0f)
+			if (Manager->bLogJobCompletionTimes || ElapsedTime > 60.0)
 			{
 				TArray<FShaderCommonCompileJobPtr> SortedJobs = CurrentWorkerInfo.QueuedJobs;
 				SortedJobs.Sort([](const FShaderCommonCompileJobPtr& JobA, const FShaderCommonCompileJobPtr& JobB)
@@ -401,8 +437,8 @@ void FShaderCompileThreadRunnable::PushCompletedJobsToManager()
 						const FShaderCompileJob* SingleJobA = JobA->GetSingleShaderJob();
 						const FShaderCompileJob* SingleJobB = JobB->GetSingleShaderJob();
 
-						const float TimeA = SingleJobA ? SingleJobA->Output.CompileTime : 0.0f;
-						const float TimeB = SingleJobB ? SingleJobB->Output.CompileTime : 0.0f;
+						const double TimeA = SingleJobA ? SingleJobA->Output.CompileTime : 0.0f;
+						const double TimeB = SingleJobB ? SingleJobB->Output.CompileTime : 0.0f;
 
 						return TimeA > TimeB;
 					});
@@ -566,7 +602,7 @@ void FShaderCompileThreadRunnable::WriteNewTasks()
 		double IODuration = FPlatformTime::Seconds() - StartIOWork;
 		if (IODuration > GShaderCompilerTooLongIOThresholdSeconds)
 		{
-			UE_LOG(LogShaderCompilers, Display, TEXT("FShaderCompileThreadRunnable::WriteNewTasks()() took too long (%.3f seconds, threshold is %.3f s), will parallelize next time."), IODuration, GShaderCompilerTooLongIOThresholdSeconds);
+			UE_LOG(LogShaderCompilers, Display, TEXT("FShaderCompileThreadRunnable::WriteNewTasks() took too long (%.3f seconds, threshold is %.3f s), will parallelize next time."), IODuration, GShaderCompilerTooLongIOThresholdSeconds);
 			bParallelizeIO = true;
 		}
 	}
@@ -654,6 +690,9 @@ bool FShaderCompileThreadRunnable::LaunchWorkersIfNeeded()
 				CurrentWorkerInfo.WorkerProcess = Manager->LaunchWorker(WorkingDirectory, Manager->ProcessId, WorkerIndex, GWorkerInputFilename, GWorkerOutputFilename);
 				CurrentWorkerInfo.bLaunchedWorker = true;
 
+				// Assign process to job object to monitor the total memory consumption of all SCW processes
+				ApplyWorkerProcessMemoryLimits(CurrentWorkerInfo.WorkerProcess);
+
 				NumberLaunched++;
 			}
 		}
@@ -842,6 +881,12 @@ void FShaderCompileThreadRunnable::PrintWorkerMemoryUsageWithLockTaken()
 	}
 }
 
+int32 FShaderCompileThreadRunnable::GetNumberOfWorkers() const
+{
+	FScopeLock WorkerScopeLock(&WorkerInfosLock);
+	return WorkerInfos.Num();
+}
+
 int32 FShaderCompileThreadRunnable::GetNumberOfAvailableWorkersUnsafe() const
 {
 	// Don't lock WorkerScopeLock critical section here, since this function might be called inside an already locked scope, hence the "Unsafe" name
@@ -858,12 +903,18 @@ int32 FShaderCompileThreadRunnable::GetNumberOfAvailableWorkersUnsafe() const
 	return NumAvailableWorkers;
 }
 
+int32 FShaderCompileThreadRunnable::GetNumberOfAvailableWorkers() const
+{
+	FScopeLock WorkerScopeLock(&WorkerInfosLock);
+	return GetNumberOfAvailableWorkersUnsafe();
+}
+
 int32 FShaderCompileThreadRunnable::GetNumberOfSuspendedWorkersUnsafe() const
 {
 	return WorkerInfos.Num() - GetNumberOfAvailableWorkersUnsafe();
 }
 
-int32 FShaderCompileThreadRunnable::SuspendWorkersAndBacklogJobs(int32 NumWorkersToSuspend)
+int32 FShaderCompileThreadRunnable::SuspendWorkersAndBacklogJobs(int32 NumWorkersToSuspend, int32* OutNumBackloggedJobs)
 {
 	int32 NumSuspendedWorkers = 0;
 	int32 NumBackloggedJobs = 0;
@@ -908,16 +959,9 @@ int32 FShaderCompileThreadRunnable::SuspendWorkersAndBacklogJobs(int32 NumWorker
 		}
 	}
 
-	// Report about backlogged jobs
-	if (NumBackloggedJobs > 0)
+	if (OutNumBackloggedJobs)
 	{
-		UE_LOG(
-			LogShaderCompilers, Display, TEXT("Backlogged %d compile %s from %d suspended %s"),
-			NumBackloggedJobs,
-			NumBackloggedJobs == 1 ? TEXT("job") : TEXT("jobs"),
-			NumSuspendedWorkers,
-			NumSuspendedWorkers == 1 ? TEXT("worker") : TEXT("workers")
-		);
+		*OutNumBackloggedJobs = NumBackloggedJobs;
 	}
 
 	return NumSuspendedWorkers;
@@ -952,18 +996,6 @@ int32 FShaderCompileThreadRunnable::ResumeSuspendedWorkers(int32 NumWorkersToRes
 				}
 			}
 		}
-
-		if (NumResumedWorkers > 0)
-		{
-			const int32 NumRemainingSuspendedWorkers = NumSuspendedWorkers - NumResumedWorkers;
-			UE_LOG(
-				LogShaderCompilers, Display, TEXT("Resumed %d suspended %s (%d %s suspended)"),
-				NumResumedWorkers,
-				NumResumedWorkers == 1 ? TEXT("worker") : TEXT("workers"),
-				NumRemainingSuspendedWorkers,
-				NumRemainingSuspendedWorkers == 1 ? TEXT("remains") : TEXT("remain")
-			);
-		}
 	}
 
 	return NumResumedWorkers;
@@ -985,6 +1017,99 @@ FString FShaderCompileThreadRunnable::GetWorkingDirectoryForWorker(int32 WorkerI
 	return FPaths::Combine(bRelativePath ? Manager->ShaderBaseWorkingDirectory : Manager->AbsoluteShaderBaseWorkingDirectory, FString::FromInt(WorkerIndex));
 }
 
+void FShaderCompileThreadRunnable::CheckMemoryLimitViolation()
+{
+	constexpr double kMemoryLimitPollInterval = 0.1; // Check every 0.1s if the memory limit has been exceeded
+	constexpr double kResumingWorkersPollInterval = 1.0; // Check every second since the last time workers have been suspending if we can resume some workers again
+
+	const double CurrentTime = FPlatformTime::Seconds();
+
+	// Check memory limit violations periodically
+	if (CurrentTime - MemoryMonitoringState.LastTimeOfMemoryLimitPoll > kMemoryLimitPollInterval)
+	{
+		MemoryMonitoringState.LastTimeOfMemoryLimitPoll = CurrentTime;
+
+		// Check if memory limit has been exceeded
+		FJobObjectLimitationInfo LimitInfo;
+		if (GSCWResourceRestrictedJobObject.QueryLimitViolationStatus(LimitInfo))
+		{
+			MemoryMonitoringState.LastTimeOfSuspeningOrResumingWorkers = CurrentTime;
+
+			// Try to halve the number of workers
+			const int32 NumWorkersToSuspend = GetNumberOfAvailableWorkers() / 2;
+
+			int32 NumBackloggedJobs = 0;
+			const int32 NumSuspendedWorkers = SuspendWorkersAndBacklogJobs(NumWorkersToSuspend, &NumBackloggedJobs);
+			if (NumSuspendedWorkers > 0)
+			{
+				UE_LOGFMT_NSLOC(
+					LogShaderCompilers, Display, "ShaderCompilers", "SuspendingWorkers",
+					"Shader compiler memory usage of {MemoryUsed} MiB exceeded limit of {MemoryLimit} MiB: " \
+					"Backlogged {BackloggedJobs} compile {BackloggedJobsName} from {SuspendedWorkers} suspended {SuspendedWorkersName} ({ActiveWorkerCount}/{TotalWorkerCount} active)",
+					("MemoryUsed", static_cast<int32>(LimitInfo.MemoryUsed / 1024 / 1024)),
+					("MemoryLimit", static_cast<int32>(LimitInfo.MemoryLimit / 1024 / 1024)),
+					("BackloggedJobs", NumBackloggedJobs),
+					("BackloggedJobsName", NumBackloggedJobs == 1 ? TEXT("job") : TEXT("jobs")),
+					("SuspendedWorkers", NumSuspendedWorkers),
+					("SuspendedWorkersName", NumSuspendedWorkers == 1 ? TEXT("worker") : TEXT("workers")),
+					("ActiveWorkerCount", GetNumberOfAvailableWorkers()),
+					("TotalWorkerCount", GetNumberOfWorkers())
+				);
+				MemoryMonitoringState.bHasSuspendedWorkers = true;
+				MemoryMonitoringState.bHasFailedToSuspendWorkers = false;
+			}
+			else if (!MemoryMonitoringState.bHasFailedToSuspendWorkers)
+			{
+				UE_LOGFMT_NSLOC(
+					LogShaderCompilers, Warning, "ShaderCompilers", "SuspendingWorkersFailed",
+					"Shader compiler memory usage of {MemoryUsed} MiB exceeded limit of {MemoryLimit} MiB, but cannot suspend any more workers",
+					("MemoryUsed", static_cast<int32>(LimitInfo.MemoryUsed / 1024 / 1024)),
+					("MemoryLimit", static_cast<int32>(LimitInfo.MemoryLimit / 1024 / 1024))
+				);
+				MemoryMonitoringState.bHasFailedToSuspendWorkers = true; // Don't show this warning again unless we were able to suspend workers again
+			}
+		}
+	}
+
+	// Check if we can resume previously suspended workers periodically
+	if (MemoryMonitoringState.bHasSuspendedWorkers && CurrentTime - MemoryMonitoringState.LastTimeOfSuspeningOrResumingWorkers > kResumingWorkersPollInterval)
+	{
+		MemoryMonitoringState.LastTimeOfSuspeningOrResumingWorkers = CurrentTime;
+
+		FJobObjectLimitationInfo LimitInfo;
+		if (GSCWResourceRestrictedJobObject.QueryStatus(LimitInfo))
+		{
+			// If we are below half of our memory limit, resume 50% of available workers.
+			// This approach suspends workers from 100% to 50% and then resumes them back up to 75%.
+			if (LimitInfo.MemoryUsed < LimitInfo.MemoryLimit / 2)
+			{
+				// Number of workers to resume should be a lower bound (i.e. division by 2 can result in zero),
+				// in which case we accept that we cannot allow the last suspended worker to be resumed.
+				const int32 NumWorkersToResume = GetNumberOfAvailableWorkers() / 2;
+				const int32 NumResumedWorkers = ResumeSuspendedWorkers(NumWorkersToResume);
+				if (NumResumedWorkers > 0)
+				{
+					UE_LOGFMT_NSLOC(
+						LogShaderCompilers, Display, "ShaderCompilers", "ResumingWorkers",
+						"Resumed {ResumedWorkers} suspended {ResumedWorkersName} since memory usage of {MemoryUsed} MiB is below half the limit of {MemoryLimit} MiB ({ActiveWorkerCount}/{TotalWorkerCount} active)",
+						("ResumedWorkers", NumResumedWorkers),
+						("ResumedWorkersName", NumResumedWorkers == 1 ? TEXT("worker") : TEXT("workers")),
+						("MemoryUsed", static_cast<int32>(LimitInfo.MemoryUsed / 1024 / 1024)),
+						("MemoryLimit", static_cast<int32>(LimitInfo.MemoryLimit / 1024 / 1024)),
+						("ActiveWorkerCount", GetNumberOfAvailableWorkers()),
+						("TotalWorkerCount", GetNumberOfWorkers())
+					);
+				}
+				else
+				{
+					// No more workers that could be resumed
+					MemoryMonitoringState.bHasSuspendedWorkers = false;
+				}
+			}
+		}
+	}
+}
+
 bool FShaderCompileThreadRunnable::PrintWorkerMemoryUsage(bool bAllowToWaitForLock)
 {
 	if (bAllowToWaitForLock)
diff --git a/Engine/Source/Runtime/Engine/Public/ShaderCompiler.h b/Engine/Source/Runtime/Engine/Public/ShaderCompiler.h
--- a/Engine/Source/Runtime/Engine/Public/ShaderCompiler.h
+++ b/Engine/Source/Runtime/Engine/Public/ShaderCompiler.h
@@ -201,7 +201,7 @@ private:
 
 	/** Information about the active workers that this thread is tracking. */
 	TArray<TUniquePtr<struct FShaderCompileWorkerInfo>> WorkerInfos;
-	FCriticalSection WorkerInfosLock;
+	mutable FCriticalSection WorkerInfosLock;
 
 	/** Tracks the last time that this thread checked if the workers were still active. */
 	double LastCheckForWorkersTime;
@@ -212,6 +212,15 @@ private:
 	/** List of jobs that have been backlogged when workers had to be closed due to reaching memory limits. These jobs will be picked up first before new jobs are pulled from the manager job queue. */
 	TArray<FShaderCommonCompileJobPtr> BackloggedJobs;
 
+	struct FMemoryMonitoringState
+	{
+		double LastTimeOfMemoryLimitPoll = 0.0;
+		double LastTimeOfSuspeningOrResumingWorkers = 0.0;
+		bool bHasFailedToSuspendWorkers = false;
+		bool bHasSuspendedWorkers = false;
+	}
+	MemoryMonitoringState;
+
 public:
 	/** Initialization constructor. */
 	FShaderCompileThreadRunnable(class FShaderCompilingManager* InManager);
@@ -252,8 +261,12 @@ private:
 
 	void PrintWorkerMemoryUsageWithLockTaken();
 
+	/** Returns the total number of workers this thread is handling. */
+	int32 GetNumberOfWorkers() const;
+
 	/** Returns the number of available workers. Only call inside the critical section WorkerInfosLock. */
 	int32 GetNumberOfAvailableWorkersUnsafe() const;
+	int32 GetNumberOfAvailableWorkers() const;
 
 	/** Returns the number of suspended workers. Only call inside the critical section WorkerInfosLock. */
 	int32 GetNumberOfSuspendedWorkersUnsafe() const;
@@ -262,7 +275,7 @@ private:
 	 * Suspends the specified number of workers and moves all their compile jobs to the backlog queue.
 	 * Returns the number of workers that have been suspended. The last worker cannot be suspended.
 	 */
-	int32 SuspendWorkersAndBacklogJobs(int32 NumWorkers);
+	int32 SuspendWorkersAndBacklogJobs(int32 NumWorkers, int32* OutNumBackloggedJobs = nullptr);
 
 	/**
 	 * Makes the specified number of workers available again after they have been suspended.
@@ -275,6 +288,9 @@ private:
 
 	/** Returns the working directory for the specified shader compile worker. */
 	FString GetWorkingDirectoryForWorker(int32 WorkerIndex, bool bRelativePath = false) const;
+
+	/** Checks it the memory limit for shader compile workers has been exceeded and suspend workers as needed. */
+	void CheckMemoryLimitViolation();
 };
 
 class FShaderCompileUtilities
