Added extra documentation to SetCanEverAffectNavigation

diff --git a/Engine/Source/Runtime/Engine/Classes/Components/ActorComponent.h b/Engine/Source/Runtime/Engine/Classes/Components/ActorComponent.h
--- a/Engine/Source/Runtime/Engine/Classes/Components/ActorComponent.h
+++ b/Engine/Source/Runtime/Engine/Classes/Components/ActorComponent.h
@@ -1166,7 +1166,13 @@ public:
 	/** Can this component potentially influence navigation */
 	ENGINE_API bool CanEverAffectNavigation() const;
 
-	/** set value of bCanEverAffectNavigation flag and update navigation octree if needed */
+	/**
+	 * Sets value of bCanEverAffectNavigation flag and update navigation octree if needed.
+	 * @note This flag influences the required updates when setting other properties associated
+	 * to the navigation data (e.g., transform, primitive, mesh, etc.).
+	 * So when setting it to 'true' it is more efficient to set the other properties first,
+	 * otherwise ('false') it is more efficient to set this one first, then the other properties.
+	 */
 	ENGINE_API void SetCanEverAffectNavigation(bool bRelevant);
 
 	/** Override to specify that a component is relevant to the navigation system */
