[AutoRTFM] Migrate more critical sections to using the transactionally safe variants.

Backing out backout CL34912307, to re-land CL34912307 with the following changes:
This CL keeps the #includes that are transitively used in various places across the codebase, which was the cause of the original backout.
I've annotated the #includes which should be removed with SOL-6812. I'll remove these on main, and add the missing #includes in the places that have the transitive dependency on the old include.

diff --git a/Engine/Source/Programs/AutoRTFMEngineTests/Private/AutoRTFMEngineTests.cpp b/Engine/Source/Programs/AutoRTFMEngineTests/Private/AutoRTFMEngineTests.cpp
--- a/Engine/Source/Programs/AutoRTFMEngineTests/Private/AutoRTFMEngineTests.cpp
+++ b/Engine/Source/Programs/AutoRTFMEngineTests/Private/AutoRTFMEngineTests.cpp
@@ -213,6 +213,11 @@ static void PreInit()
 
 	FDelayedAutoRegisterHelper::RunAndClearDelayedAutoRegisterDelegates(EDelayedRegisterRunPhase::StatSystemReady);
 
+
+	// Ensure GCachedScalabilityCVars.bInitialized is set (required by some mesh components)
+	FTaskTagScope::SwapTag(ETaskTag::EGameThread);
+	extern void ScalabilityCVarsSinkCallback();
+	ScalabilityCVarsSinkCallback();
 }
 
 static void LoadModules()
diff --git a/Engine/Source/Runtime/Core/Private/Features/ModularFeatures.cpp b/Engine/Source/Runtime/Core/Private/Features/ModularFeatures.cpp
--- a/Engine/Source/Runtime/Core/Private/Features/ModularFeatures.cpp
+++ b/Engine/Source/Runtime/Core/Private/Features/ModularFeatures.cpp
@@ -4,7 +4,7 @@
 
 #include "CoreTypes.h"
 #include "Misc/AssertionMacros.h"
-#include "Misc/ScopeLock.h"
+#include "Misc/TransactionallySafeScopeLock.h"
 
 IModularFeatures& IModularFeatures::Get()
 {
@@ -54,7 +54,7 @@ IModularFeature* FModularFeatures::GetModularFeatureImplementation( const FName
 
 void FModularFeatures::RegisterModularFeature( const FName Type, IModularFeature* ModularFeature )
 {
-	FScopeLock ScopeLock(&ModularFeaturesMapCriticalSection);
+	FTransactionallySafeScopeLock ScopeLock(&ModularFeaturesMapCriticalSection);
 
 	ModularFeaturesMap.AddUnique( Type, ModularFeature );
 	ModularFeatureRegisteredEvent.Broadcast( Type, ModularFeature );
@@ -63,7 +63,7 @@ void FModularFeatures::RegisterModularFeature( const FName Type, IModularFeature
 
 void FModularFeatures::UnregisterModularFeature( const FName Type, IModularFeature* ModularFeature )
 {
-	FScopeLock ScopeLock(&ModularFeaturesMapCriticalSection);
+	FTransactionallySafeScopeLock ScopeLock(&ModularFeaturesMapCriticalSection);
 
 	ModularFeaturesMap.RemoveSingle( Type, ModularFeature );
 	ModularFeatureUnregisteredEvent.Broadcast( Type, ModularFeature );
diff --git a/Engine/Source/Runtime/Core/Private/Features/ModularFeatures.h b/Engine/Source/Runtime/Core/Private/Features/ModularFeatures.h
--- a/Engine/Source/Runtime/Core/Private/Features/ModularFeatures.h
+++ b/Engine/Source/Runtime/Core/Private/Features/ModularFeatures.h
@@ -6,7 +6,8 @@
 #include "Delegates/Delegate.h"
 #include "Features/IModularFeature.h"
 #include "Features/IModularFeatures.h"
-#include "HAL/CriticalSection.h"
+#include "HAL/CriticalSection.h"  // jira SOL-6812: Remove this.
+#include "Misc/TransactionallySafeCriticalSection.h"
 #include "UObject/NameTypes.h"
 
 /**
@@ -37,7 +38,7 @@ private:
 	TMultiMap< FName, class IModularFeature* > ModularFeaturesMap;
 
 	/** Lock modular features map so it can be used across threads, this could be a RW lock instead of a crit sec */
-	FCriticalSection ModularFeaturesMapCriticalSection;
+	FTransactionallySafeCriticalSection ModularFeaturesMapCriticalSection;
 
 	/** Event used to inform clients that a modular feature has been registered */
 	FOnModularFeatureRegistered ModularFeatureRegisteredEvent;
diff --git a/Engine/Source/Runtime/Core/Public/Delegates/DelegateAccessHandler.h b/Engine/Source/Runtime/Core/Public/Delegates/DelegateAccessHandler.h
--- a/Engine/Source/Runtime/Core/Public/Delegates/DelegateAccessHandler.h
+++ b/Engine/Source/Runtime/Core/Public/Delegates/DelegateAccessHandler.h
@@ -3,9 +3,11 @@
 #pragma once
 
 #include "CoreTypes.h"
-#include "HAL/CriticalSection.h"
+#include "HAL/CriticalSection.h"  // jira SOL-6812: Remove this.
 #include "Misc/MTAccessDetector.h"
-#include "Misc/ScopeLock.h"
+#include "Misc/TransactionallySafeCriticalSection.h"
+#include "Misc/TransactionallySafeScopeLock.h"
+#include "Misc/ScopeLock.h"  // jira SOL-6812: Remove this.
 #include "AutoRTFM/AutoRTFM.h"
 
 //#define UE_DETECT_DELEGATES_RACE_CONDITIONS 0
@@ -65,8 +67,8 @@ template<>
 class TDelegateAccessHandlerBase<FThreadSafeDelegateMode>
 {
 protected:
-	struct FReadAccessScope { FScopeLock Lock; };
-	struct FWriteAccessScope { FScopeLock Lock; };
+	struct FReadAccessScope { FTransactionallySafeScopeLock Lock; };
+	struct FWriteAccessScope { FTransactionallySafeScopeLock Lock; };
 
 	[[nodiscard]] FReadAccessScope GetReadAccessScope() const
 	{
@@ -79,7 +81,7 @@ protected:
 	}
 
 private:
-	mutable FCriticalSection Mutex;
+	mutable FTransactionallySafeCriticalSection Mutex;
 };
 
 template<>
diff --git a/Engine/Source/Runtime/Engine/Classes/Animation/Skeleton.h b/Engine/Source/Runtime/Engine/Classes/Animation/Skeleton.h
--- a/Engine/Source/Runtime/Engine/Classes/Animation/Skeleton.h
+++ b/Engine/Source/Runtime/Engine/Classes/Animation/Skeleton.h
@@ -12,11 +12,12 @@
 #include "UObject/Object.h"
 #include "HAL/CriticalSection.h"
 #include "Misc/Guid.h"
+#include "Misc/TransactionallySafeRWLock.h"
 #include "ReferenceSkeleton.h"
 #include "Animation/PreviewAssetAttachComponent.h"
 #include "Animation/SmartName.h"
 #include "Engine/AssetUserData.h"
-#include "HAL/CriticalSection.h"
+#include "HAL/CriticalSection.h"  // jira SOL-6812: Remove this (duplicate include).
 #include "Interfaces/Interface_AssetUserData.h"
 #include "Interfaces/Interface_PreviewMeshProvider.h"
 #include "Delegates/DelegateCombinations.h"
@@ -658,7 +659,7 @@ private:
 	TMap<TWeakObjectPtr<USkinnedAsset>, int32> SkinnedAsset2LinkupCache;
 	
 	//Use this Lock everytime you change or access SkinnedAssetLinkupCache member.
-	FRWLock SkinnedAssetLinkupCacheLock;
+	FTransactionallySafeRWLock SkinnedAssetLinkupCacheLock;
 
 	/** Runtime built mapping table between SkinnedAssets and Mesh Linkup Data*/
 	TMap<TObjectKey<USkinnedAsset>, TUniquePtr<FSkeletonToMeshLinkup>> SkinnedAssetLinkupCache;
diff --git a/Engine/Source/Runtime/Engine/Private/Animation/Skeleton.cpp b/Engine/Source/Runtime/Engine/Private/Animation/Skeleton.cpp
--- a/Engine/Source/Runtime/Engine/Private/Animation/Skeleton.cpp
+++ b/Engine/Source/Runtime/Engine/Private/Animation/Skeleton.cpp
@@ -750,7 +750,7 @@ bool USkeleton::IsCompatibleMesh(const USkinnedAsset* InSkinnedAsset, bool bDoPa
 void USkeleton::ClearCacheData()
 {
 	{
-		FRWScopeLock Lock(SkinnedAssetLinkupCacheLock, SLT_Write);
+		FTransactionallySafeWriteScopeLock Lock(SkinnedAssetLinkupCacheLock);
 		SkinnedAssetLinkupCache.Empty();
 	}
 
@@ -767,7 +767,7 @@ const FSkeletonToMeshLinkup& USkeleton::FindOrAddMeshLinkupData(const USkinnedAs
 	const TUniquePtr<FSkeletonToMeshLinkup>* SkeletonToMeshLinkupPtr = nullptr;
 
 	{
-		FRWScopeLock Lock(SkinnedAssetLinkupCacheLock, SLT_ReadOnly);
+		FTransactionallySafeReadScopeLock Lock(SkinnedAssetLinkupCacheLock);
 		SkeletonToMeshLinkupPtr = SkinnedAssetLinkupCache.Find(InSkinnedAsset);
 	}
 
@@ -782,7 +782,7 @@ const FSkeletonToMeshLinkup& USkeleton::AddMeshLinkupData(const USkinnedAsset* I
 	BuildLinkupData(InSkinnedAsset, *TmpLinkup.Get());
 
 	{
-		FRWScopeLock Lock(SkinnedAssetLinkupCacheLock, SLT_Write);
+		FTransactionallySafeWriteScopeLock Lock(SkinnedAssetLinkupCacheLock);
 		return *SkinnedAssetLinkupCache.Add(TObjectKey<USkinnedAsset>(InSkinnedAsset), MoveTemp(TmpLinkup)).Get();
 	}
 }
@@ -823,7 +823,7 @@ int32 USkeleton::GetMeshLinkupIndex(const USkinnedAsset* InSkinnedAsset)
 void USkeleton::RemoveLinkup(const USkinnedAsset* InSkinnedAsset)
 {
 	{
-		FRWScopeLock Lock(SkinnedAssetLinkupCacheLock, SLT_Write);
+		FTransactionallySafeWriteScopeLock Lock(SkinnedAssetLinkupCacheLock);
 		SkinnedAssetLinkupCache.Remove(InSkinnedAsset);
 	}
 
diff --git a/Engine/Source/Runtime/Engine/Private/Tests/AutoRTFM/AutoRTFMActorComponentTests.cpp b/Engine/Source/Runtime/Engine/Private/Tests/AutoRTFM/AutoRTFMActorComponentTests.cpp
--- a/Engine/Source/Runtime/Engine/Private/Tests/AutoRTFM/AutoRTFMActorComponentTests.cpp
+++ b/Engine/Source/Runtime/Engine/Private/Tests/AutoRTFM/AutoRTFMActorComponentTests.cpp
@@ -1,22 +1,25 @@
 // Copyright Epic Games, Inc. All Rights Reserved.
 
 #include "Misc/AutomationTest.h"
-#include "Misc/App.h"
+#include "Animation/AnimInstance.h"
+#include "Animation/Skeleton.h"
 #include "AutoRTFM/AutoRTFM.h"
 #include "AutoRTFMTestActor.h"
 #include "AutoRTFMTestAnotherActor.h"
 #include "AutoRTFMTestBodySetup.h"
+#include "AutoRTFMTestChildActorComponent.h"
 #include "AutoRTFMTestLevel.h"
 #include "AutoRTFMTestObject.h"
-#include "AutoRTFMTestChildActorComponent.h"
 #include "AutoRTFMTestPrimitiveComponent.h"
-#include "PhysicsProxy/SingleParticlePhysicsProxy.h"
-#include "Physics/Experimental/PhysScene_Chaos.h"
-#include "PBDRigidsSolver.h"
 #include "Chaos/Core.h"
+#include "Components/SkeletalMeshComponent.h"
+#include "Engine/SkeletalMesh.h"
 #include "Engine/World.h"
-
-#include <utility>
+#include "PBDRigidsSolver.h"
+#include "Physics/Experimental/PhysScene_Chaos.h"
+#include "PhysicsProxy/SingleParticlePhysicsProxy.h"
+#include "Rendering/SkeletalMeshRenderData.h"
+#include "UObject/UObjectGlobals.h"
 
 #if WITH_DEV_AUTOMATION_TESTS
 
@@ -396,6 +399,32 @@ AUTORTFM_ACTOR_COMPONENT_TEST(ChildActor)
 	TEST_CHECK_TRUE(nullptr == ChildActorComponent->GetChildActor());
 }
 
+// Test aborting a call to USkeletalMeshComponent::RegisterComponent() with an assigned skeletal
+// mesh and empty PostProcessAnimInstance.
+// See: SOL-6779
+AUTORTFM_ACTOR_COMPONENT_TEST(USkeletalMeshComponent)
+{
+	USkeleton* Skeleton = NewObject<USkeleton>();
+	USkeletalMesh* SkeletalMesh = NewObject<USkeletalMesh>();
+	SkeletalMesh->SetSkeleton(Skeleton);
+	SkeletalMesh->AllocateResourceForRendering();
+	FSkeletalMeshRenderData* RenderData = SkeletalMesh->GetResourceForRendering();
+	TRefCountPtr<FSkeletalMeshLODRenderData> LODRenderData = MakeRefCount<FSkeletalMeshLODRenderData>();
+	RenderData->LODRenderData.Add(LODRenderData);
+	USkeletalMeshComponent* SkeletalMeshComponent = NewObject<USkeletalMeshComponent>(Actor);
+	SkeletalMeshComponent->SetSkeletalMeshAsset(SkeletalMesh);
+	SkeletalMeshComponent->PostProcessAnimInstance = NewObject<UAnimInstance>(SkeletalMeshComponent);
+
+	AutoRTFM::ETransactionResult Result;
+	Result = AutoRTFM::Transact([&]
+		{
+			SkeletalMeshComponent->RegisterComponent();
+			AutoRTFM::AbortTransaction();
+		});
+
+	TEST_CHECK_TRUE(AutoRTFM::ETransactionResult::AbortedByRequest == Result);
+}
+
 }  // anonymous namespace
 
 #undef TEST_CHECK_FALSE
