REworking the Legacy converstion path

diff --git a/Engine/Source/Runtime/Engine/Private/Materials/MaterialExpressions.cpp b/Engine/Source/Runtime/Engine/Private/Materials/MaterialExpressions.cpp
--- a/Engine/Source/Runtime/Engine/Private/Materials/MaterialExpressions.cpp
+++ b/Engine/Source/Runtime/Engine/Private/Materials/MaterialExpressions.cpp
@@ -17931,39 +17931,25 @@ bool FMaterialLayersFunctionsRuntimeData::SerializeFromMismatchedTag(const FProp
 }
 
 
-bool FMaterialLayersFunctionsRuntimeData::Serialize(FArchive& Ar)
+void FMaterialLayersFunctionsRuntimeData::PostSerialize(const FArchive& Ar)
 {
-#if ENABLE_MATERIAL_LAYER_PROTOTYPE
-	// TODO: We d like to do rely on the custom version of the branch or something:
-	// Ar.UsingCustomVersion(FRenderingObjectVersion::GUID);
-	// we are using our custom version indicator for now in the tree
-	if (!Ar.IsLoading())
-	{
-		Tree.EnableMaterialLayersFunctionsTree();
-	}
-#endif
-
-	// Use default serialization for most properties for
-	// archive configurations supported by UStruct::SerializeVersionedTaggedProperties 
-	if (Ar.IsLoading() || Ar.IsSaving() || Ar.IsCountingMemory() || Ar.IsObjectReferenceCollector())
-	{
-		StaticStruct()->SerializeTaggedProperties(Ar, (uint8*)this, StaticStruct(), nullptr);
-	}
 
 #if ENABLE_MATERIAL_LAYER_PROTOTYPE
-	// For PRE substrate tree layer version:
-	// TODO: As commented above, we could use:
-	// 	if (Ar.CustomVer(FRenderingObjectVersion::GUID) < FRenderingObjectVersion::AddedMaterialLayersSubstrateSupport)
-	if (!Tree.IsMaterialLayersFunctionsTreeEnabled())
-	{
-		// When loading, The tree need to be recreated from the list
+	if (Ar.IsLoading())
+	{	
+		// When loading, if a legacy PRE substrate tree layer version:
+		// The tree need to be recreated from the list
 		// of layers and blends assuming this is a flat hierarchy of layers:
-		if (Ar.IsLoading())
+		// For PRE substrate tree layer version:
+		// TODO: We could use:
+		// 	if (Ar.CustomVer(FRenderingObjectVersion::GUID) < FRenderingObjectVersion::AddedMaterialLayersSubstrateSupport)
+		if ((Tree.Nodes.Num() != Layers.Num()) || (Layers.Num() > Blends.Num()))
 		{
-			Tree.Empty();
-
 			int NumLegacyLayers = Layers.Num();
 			int NumLegacyBlends = Blends.Num();
+
+			Tree.Empty();
+
 			// First add an extra empty Blend to match the number of layers
 			while (Blends.Num() < Layers.Num())
 			{
@@ -17985,8 +17971,6 @@ bool FMaterialLayersFunctionsRuntimeData::Serialize(FArchive& Ar)
 		}
 	}
 #endif
-
-	return true;
 }
 
 
@@ -18055,6 +18039,11 @@ void FMaterialLayersFunctions::SerializeLegacy(FArchive& Ar)
 
 void FMaterialLayersFunctions::PostSerialize(const FArchive& Ar)
 {
+#if ENABLE_MATERIAL_LAYER_PROTOTYPE
+	FMaterialLayersFunctionsRuntimeData::PostSerialize(Ar);
+#endif
+
+
 #if WITH_EDITORONLY_DATA
 	if (Ar.IsLoading())
 	{
@@ -18074,8 +18063,10 @@ void FMaterialLayersFunctions::PostSerialize(const FArchive& Ar)
 		{
 			EditorOnly.RestrictToBlendRelatives = MoveTemp(RestrictToBlendRelatives_DEPRECATED);
 		}
+		bool bIsDeprecatedLegacyAsset = false;
 		if (LayerGuids_DEPRECATED.Num() > 0)
 		{
+			bIsDeprecatedLegacyAsset = true;
 			EditorOnly.LayerGuids = MoveTemp(LayerGuids_DEPRECATED);
 		}
 		if (LayerLinkStates_DEPRECATED.Num() > 0)
@@ -18087,6 +18078,36 @@ void FMaterialLayersFunctions::PostSerialize(const FArchive& Ar)
 			EditorOnly.DeletedParentLayerGuids = MoveTemp(DeletedParentLayerGuids_DEPRECATED);
 		}
 
+#if ENABLE_MATERIAL_LAYER_PROTOTYPE
+
+		if (bIsDeprecatedLegacyAsset)
+		{
+			const int32 NumLayers = EditorOnly.LayerStates.Num();
+			if (EditorOnly.LayerGuids.Num() != NumLayers ||
+				EditorOnly.LayerLinkStates.Num() != NumLayers)
+			{
+				EditorOnly.LayerGuids.Empty(NumLayers);
+				EditorOnly.LayerLinkStates.Empty(NumLayers);
+
+				if (NumLayers > 0)
+				{
+					EditorOnly.LayerGuids.Add(BackgroundGuid);
+					EditorOnly.LayerLinkStates.Add(EMaterialLayerLinkState::Uninitialized);
+
+					for (int32 i = 1; i < NumLayers; ++i)
+					{
+						// Need to allocate deterministic guids for layers loaded from old data
+						// These guids will be saved into any child material layers that have this material as their parent,
+						// But it's possible *this* material may not actually be saved in that case
+						// If that happens, need to ensure that the guids remain consistent if this material is loaded again;
+						// otherwise they will no longer match the guids that were saved into the child material
+						EditorOnly.LayerGuids.Add(FGuid(3u, 0u, 0u, i));
+						EditorOnly.LayerLinkStates.Add(EMaterialLayerLinkState::Uninitialized);
+					}
+				}
+			}
+		}
+#else
 		if (EditorOnly.LayerGuids.Num() != Layers.Num() ||
 			EditorOnly.LayerLinkStates.Num() != Layers.Num())
 		{
@@ -18111,9 +18132,34 @@ void FMaterialLayersFunctions::PostSerialize(const FArchive& Ar)
 				}
 			}
 		}
+#endif
+
+#if ENABLE_MATERIAL_LAYER_PROTOTYPE
+		// After the deprecated content as moved to the EditorOnly block, let's update it to be in sync with the LayersTree
+		{
+			// This method is called right after unserialization
+			// Use this fact to fix editor only data in case of a legacy version of the LayersFunctions
+			// Catch case when a legacy layer has been loaded
+			// EditorOnly data need to be updated
+			if ((Layers.Num() == Blends.Num()) && ((EditorOnly.LayerNames.Num() + 1) == Layers.Num()))
+			{
+				int32 NumLegacyLayers = EditorOnly.LayerNames.Num();
+				// Second add an extra Layer and Blend for each legacy layer.
+				{
+					EditorOnly.LayerStates.Add(true);
+					FText LayerName = FText::FromString(TEXT("Legacy Layers"));
+					EditorOnly.LayerNames.Add(LayerName);
+					EditorOnly.RestrictToLayerRelatives.Add(false);
+					EditorOnly.RestrictToBlendRelatives.Add(false);
+					EditorOnly.LayerGuids.Add(FGuid::NewGuid());
+					EditorOnly.LayerLinkStates.Add(EMaterialLayerLinkState::Uninitialized);
+				}
+			}
 
+		}
+#endif // ENABLE_MATERIAL_LAYER_PROTOTYPE
+	} // Ar.IsLoading()
 
-	}
 #endif // WITH_EDITORONLY_DATA	
 }
 
@@ -18437,6 +18483,7 @@ bool FMaterialLayersFunctions::ResolveParent(const FMaterialLayersFunctionsRunti
 	FMaterialLayersFunctionsEditorOnlyData& EditorOnly,
 	TArray<int32>& OutRemapLayerIndices)
 {
+/*
 #if ENABLE_MATERIAL_LAYER_PROTOTYPE
 	// This method is called right after unserialization
 	// Use this fact to fix editor only data in case of a legacy version of the LayersFunctions
@@ -18458,7 +18505,7 @@ bool FMaterialLayersFunctions::ResolveParent(const FMaterialLayersFunctionsRunti
 		}
 	}
 #endif
-
+*/
  	check(EditorOnly.LayerGuids.Num() == Runtime.Layers.Num());
 	check(EditorOnly.LayerLinkStates.Num() == Runtime.Layers.Num());
 
diff --git a/Engine/Source/Runtime/Engine/Public/Materials/MaterialLayersFunctions.h b/Engine/Source/Runtime/Engine/Public/Materials/MaterialLayersFunctions.h
--- a/Engine/Source/Runtime/Engine/Public/Materials/MaterialLayersFunctions.h
+++ b/Engine/Source/Runtime/Engine/Public/Materials/MaterialLayersFunctions.h
@@ -179,14 +179,8 @@ public:
 
 	TArray<FNode>		Nodes;
 	TArray<FPayload>	Payloads;
-	FNodeId				Root = -2; // the index of the head of the top level siblings
-							   // and temporarly, capture the information if the LayersFunctionsTree is enabled in this asset
-							   // -2 means is is disabled
-
-	// Custom version system until we can rely on the archive custom version system
-	void			EnableMaterialLayersFunctionsTree() { Root = (Root == -2 ? -1 : Root); }
-	bool			IsMaterialLayersFunctionsTreeEnabled() const { return Root != -2; }
-
+	FNodeId				Root = -1; // the index of the head of the top level siblings
+							   	
 	FMaterialLayersFunctionsTree() = default;
 	FMaterialLayersFunctionsTree(const FMaterialLayersFunctionsTree& Rhs)
 		:	Nodes(Rhs.Nodes),
@@ -414,7 +408,8 @@ struct FMaterialLayersFunctionsRuntimeData
 	}
 
 	bool SerializeFromMismatchedTag(const FPropertyTag& Tag, FStructuredArchive::FSlot Slot);
-	bool Serialize(FArchive& Ar);
+	
+	ENGINE_API void PostSerialize(const FArchive& Ar);
 
 #if WITH_EDITOR
 	const FMaterialLayersFunctionsID GetID(const FMaterialLayersFunctionsEditorOnlyData& EditorOnly) const;
@@ -435,7 +430,7 @@ template<>
 struct TStructOpsTypeTraits<FMaterialLayersFunctionsRuntimeData> : TStructOpsTypeTraitsBase2<FMaterialLayersFunctionsRuntimeData>
 {
 	enum { WithStructuredSerializeFromMismatchedTag = true };
-	enum { WithSerializer = true };
+	enum { WithPostSerialize = true };
 };
 
 USTRUCT()
