[World Partition Validations] Make invalid HLODs layers warnings instead of errors.

diff --git a/Engine/Plugins/Editor/DataValidation/Source/DataValidation/Private/WorldPartitionChangelistValidator.cpp b/Engine/Plugins/Editor/DataValidation/Source/DataValidation/Private/WorldPartitionChangelistValidator.cpp
--- a/Engine/Plugins/Editor/DataValidation/Source/DataValidation/Private/WorldPartitionChangelistValidator.cpp
+++ b/Engine/Plugins/Editor/DataValidation/Source/DataValidation/Private/WorldPartitionChangelistValidator.cpp
@@ -547,11 +547,11 @@ void UWorldPartitionChangelistValidator::OnInvalidHLODLayer(const IWorldPartitio
 {
 	if (Filter(ActorDescView))
 	{
-		FText CurrentError = FText::Format(LOCTEXT("DataValidation.Changelist.WorldPartition.InvalidHLODLayer", "Actor {0} has an invalid HLOD Layer {1}"),
+		FText CurrentWarning = FText::Format(LOCTEXT("DataValidation.Changelist.WorldPartition.InvalidHLODLayer", "Actor {0} has an invalid HLOD Layer {1}"),
 											FText::FromString(GetFullActorName(ActorDescView)), 
 											FText::FromString(ActorDescView.GetHLODLayer().ToString()));
 
-		AssetFails(CurrentAsset, CurrentError);
+		AssetWarning(CurrentAsset, CurrentWarning);
 	}
 }
 
diff --git a/Engine/Source/Runtime/Engine/Private/WorldPartition/ErrorHandling/WorldPartitionStreamingGenerationTokenizedMessageErrorHandler.cpp b/Engine/Source/Runtime/Engine/Private/WorldPartition/ErrorHandling/WorldPartitionStreamingGenerationTokenizedMessageErrorHandler.cpp
--- a/Engine/Source/Runtime/Engine/Private/WorldPartition/ErrorHandling/WorldPartitionStreamingGenerationTokenizedMessageErrorHandler.cpp
+++ b/Engine/Source/Runtime/Engine/Private/WorldPartition/ErrorHandling/WorldPartitionStreamingGenerationTokenizedMessageErrorHandler.cpp
@@ -273,7 +273,7 @@ void ITokenizedMessageErrorHandler::OnInvalidActorFilterReference(const IWorldPa
 
 void ITokenizedMessageErrorHandler::OnInvalidHLODLayer(const IWorldPartitionActorDescInstanceView& ActorDescView)
 {
-	TSharedRef<FTokenizedMessage> Message = FTokenizedMessage::Create(EMessageSeverity::Error);
+	TSharedRef<FTokenizedMessage> Message = FTokenizedMessage::Create(EMessageSeverity::Warning);
 	Message->AddToken(FTextToken::Create(LOCTEXT("TokenMessage_WorldPartition_Actor", "Actor")))
 		->AddToken(FActorToken::Create(ActorDescView.GetActorSoftPath().ToString(), ActorDescView.GetGuid(), FText::FromString(GetActorName(ActorDescView))))
 		->AddToken(FTextToken::Create(LOCTEXT("TokenMessage_WorldPartition_HaveInvalidHLODLayer", "has an invalid HLOD layer")))
