- Handle case where we have a valid skin cache entry but the position buffers have not been filled in yet.

#tests Sampling SKMesh with skin cache, and separately with deformer running

diff --git a/Engine/Source/Runtime/Engine/Private/SkeletalRenderGPUSkin.cpp b/Engine/Source/Runtime/Engine/Private/SkeletalRenderGPUSkin.cpp
--- a/Engine/Source/Runtime/Engine/Private/SkeletalRenderGPUSkin.cpp
+++ b/Engine/Source/Runtime/Engine/Private/SkeletalRenderGPUSkin.cpp
@@ -2233,8 +2233,15 @@ bool FSkeletalMeshObjectGPUSkin::GetCachedGeometry(FCachedGeometry& OutCachedGeo
 		if (SkinCacheEntry != nullptr)
 		{
 			// Get the cached geometry SRVs from the skin cache.
-			CachedSection.PositionBuffer = FGPUSkinCache::GetPositionBuffer(SkinCacheEntry, SectionIndex)->SRV;
-			CachedSection.PreviousPositionBuffer = FGPUSkinCache::GetPreviousPositionBuffer(SkinCacheEntry, SectionIndex)->SRV;
+			FRWBuffer* PositionBuffer = FGPUSkinCache::GetPositionBuffer(SkinCacheEntry, SectionIndex);
+			if (!PositionBuffer)
+			{
+				return false;
+			}
+			FRWBuffer* PreviousPositionBuffer = FGPUSkinCache::GetPreviousPositionBuffer(SkinCacheEntry, SectionIndex);
+
+			CachedSection.PositionBuffer = PositionBuffer->SRV;
+			CachedSection.PreviousPositionBuffer = PreviousPositionBuffer ? PreviousPositionBuffer->SRV : PositionBuffer->SRV;
 
 			FRWBuffer* TangentBuffer = FGPUSkinCache::GetTangentBuffer(SkinCacheEntry, SectionIndex);
 			CachedSection.TangentBuffer = TangentBuffer ? TangentBuffer->SRV : nullptr;
