Some mesh paint setting changes on StaticMesh and StaticMeshComponent.
Added MeshPaintTextureCoordinateIndex setting on UStaticMesh to allow it to define which channel to default to when mesh painting.
Added local override for this on the UStaticMeshComponent: OverridenMeshPaintTextureCoordinateIndex
Removed temporary default creation function of the StaticMeshComponent MeshPaintTexture and added a setter interface for use by the Mesh Paint tool.
Added r.MeshPaintVirtualTexture.DefaultTexelsPerVertex to control the logic for setting the default mesh paint texture size to use on a mesh. By default we try to match the number of texels to some multiple of the number of vertices to give a similar quality/memory usage between texture paint and vertex paint. The ratio can be modified with the cvar.
Set the default fallback color for MeshPaintVirtualTextures to white instead of black. This matches behavior for missing vertex painting.

diff --git a/Engine/Source/Runtime/Engine/Classes/Components/MeshComponent.h b/Engine/Source/Runtime/Engine/Classes/Components/MeshComponent.h
--- a/Engine/Source/Runtime/Engine/Classes/Components/MeshComponent.h
+++ b/Engine/Source/Runtime/Engine/Classes/Components/MeshComponent.h
@@ -157,8 +157,12 @@ public:
 
 	/** Get the mesh paint texture set on this component. This does not take into account any transient override. */
 	ENGINE_API virtual UTexture* GetMeshPaintTexture() const { return nullptr; }
+	/** Set the mesh paint texture on this component. */
+	ENGINE_API virtual void SetMeshPaintTexture(UTexture* Texture) {}
 	/** Set a transient override mesh paint texture on this component. */
 	ENGINE_API virtual void SetMeshPaintTextureOverride(UTexture* OverrideTexture) {}
+	/** Get the default coordinate index for painting to the mesh paint texture on this component. */
+	ENGINE_API virtual int32 GetMeshPaintTextureCoordinateIndex() const { return 0; }
 
 public:
 	/** Material parameter setting and caching */
diff --git a/Engine/Source/Runtime/Engine/Classes/Components/StaticMeshComponent.h b/Engine/Source/Runtime/Engine/Classes/Components/StaticMeshComponent.h
--- a/Engine/Source/Runtime/Engine/Classes/Components/StaticMeshComponent.h
+++ b/Engine/Source/Runtime/Engine/Classes/Components/StaticMeshComponent.h
@@ -236,14 +236,10 @@ public:
 	UPROPERTY(transient)
 	uint8 bForceNavigationObstacle : 1;
 
-	/** If true, mesh painting is disallowed on this instance. Set if vertex colors are overridden in a construction script. */
-	UPROPERTY(EditAnywhere, AdvancedDisplay, BlueprintReadWrite, Category = "Mesh Painting")
+	/** If true, vertex mesh painting is disallowed on this instance. Set if vertex colors are overridden in a construction script. */
+	UPROPERTY(EditAnywhere, BlueprintReadWrite, Category = "Mesh Painting")
 	uint8 bDisallowMeshPaintPerInstance : 1;
 
-	/** true if we store a mesh painting texture on this mesh component. */
-	UPROPERTY(EditAnywhere, BlueprintReadOnly, Category = "Mesh Painting", meta = (InlineEditConditionToggle))
-	uint8 bMeshPaintTexture : 1;
-
 #if STATICMESH_ENABLE_DEBUG_RENDERING
 	/** Draw mesh collision if used for complex collision */
 	uint8 bDrawMeshCollisionIfComplex : 1;
@@ -308,15 +304,23 @@ public:
 	 */
 	UPROPERTY(EditAnywhere, AdvancedDisplay, BlueprintReadOnly, Category=Lighting)
 	uint8 bReverseCulling : 1;
+	
+	/** Whether to override the MeshPaintTextureCoordinateIndex set on the static mesh. */
+	UPROPERTY(EditAnywhere, Category = "Mesh Painting", meta=(InlineEditConditionToggle))
+	uint8 bOverrideMeshPaintTextureCoordinateIndex : 1;
 
 	/** Texture containing mesh painting for this mesh component. */
-	UPROPERTY(VisibleAnywhere, Category = "Mesh Painting", meta = (editcondition = "bMeshPaintTexture"))
+	UPROPERTY(VisibleAnywhere, Category = "Mesh Painting")
 	TObjectPtr<UTexture> MeshPaintTexture;
 
 	/** Set this to override the locally stored mesh paint texture. */
 	UPROPERTY(Transient)
 	TObjectPtr<UTexture> MeshPaintTextureOverride;
 
+	/** The overriden coordinate index to use when painting the MeshPaintTexture on this mesh. */
+	UPROPERTY(EditAnywhere, Category = "Mesh Painting", meta=(UIMin = "0", UIMax = "3", editcondition = "bOverrideMeshPaintTextureCoordinateIndex"))
+	int32 OverridenMeshPaintTextureCoordinateIndex;
+
 	/** Light map resolution to use on this component, used if bOverrideLightMapRes is true and there is a valid StaticMesh. */
 	UPROPERTY(EditAnywhere, BlueprintReadOnly, Category=Lighting, meta=(ClampMax = 4096, editcondition="bOverrideLightMapRes") )
 	int32 OverriddenLightMapRes;
@@ -594,7 +598,9 @@ public:
 	ENGINE_API virtual void RegisterLODStreamingCallback(FLODStreamingCallback&& CallbackStreamingStart, FLODStreamingCallback&& CallbackStreamingDone, float TimeoutStartSecs, float TimeoutDoneSecs) override;
 	ENGINE_API virtual bool PrestreamMeshLODs(float Seconds) override;
 	ENGINE_API virtual UTexture* GetMeshPaintTexture() const override;
+	ENGINE_API virtual void SetMeshPaintTexture(UTexture* Texture) override;
 	ENGINE_API virtual void SetMeshPaintTextureOverride(UTexture* OverrideTexture) override;
+	ENGINE_API virtual int32 GetMeshPaintTextureCoordinateIndex() const override;
 	//~ End UMeshComponent Interface
 
 	//~ Begin INavRelevantInterface Interface.
diff --git a/Engine/Source/Runtime/Engine/Classes/Engine/StaticMesh.h b/Engine/Source/Runtime/Engine/Classes/Engine/StaticMesh.h
--- a/Engine/Source/Runtime/Engine/Classes/Engine/StaticMesh.h
+++ b/Engine/Source/Runtime/Engine/Classes/Engine/StaticMesh.h
@@ -990,6 +990,10 @@ public:
 		PRAGMA_ENABLE_DEPRECATION_WARNINGS
 	}
 
+	/** The default coordinate index to use when painting the MeshPaintTexture on this mesh. */
+	UPROPERTY(EditAnywhere, AdvancedDisplay, Category = StaticMesh, meta=(UIMin = "0", UIMax = "3"))
+	int32 MeshPaintTextureCoordinateIndex;
+
 	/** Useful for reducing self shadowing from distance field methods when using world position offset to animate the mesh's vertices. */
 	UPROPERTY(EditAnywhere, AdvancedDisplay, Category = StaticMesh)
 	float DistanceFieldSelfShadowBias;
diff --git a/Engine/Source/Runtime/Engine/Classes/VT/MeshPaintVirtualTexture.h b/Engine/Source/Runtime/Engine/Classes/VT/MeshPaintVirtualTexture.h
--- a/Engine/Source/Runtime/Engine/Classes/VT/MeshPaintVirtualTexture.h
+++ b/Engine/Source/Runtime/Engine/Classes/VT/MeshPaintVirtualTexture.h
@@ -39,6 +39,12 @@ namespace MeshPaintVirtualTexture
 	/** Returns per platform support of mesh paint virtual textures. */
 	ENGINE_API bool IsSupported(EShaderPlatform InShaderPlatform);
 
+	/** Returns the fallback color to use for unmapped virtual textures. Use white (same as default vertex color). */
+	inline uint32 GetDefaultFallbackColor() { return 0xFFFFFFFF; }
+
+	/** Returns the default texture size to use for a mesh based on the number of vertices. */
+	ENGINE_API uint32 GetDefaultTextureSize(int32 InNumVertices);
+
 	/** 
 	 * Get the 2 dword texture descriptor from texture resource. 
 	 * Will return a null descriptor if the texture resource is not from a UMeshPaintVirtualTexture.
@@ -50,7 +56,7 @@ namespace MeshPaintVirtualTexture
 	{
 		FTextureRHIRef PageTableTexture;
 		FTextureRHIRef PhysicalTexture;
-		FUintVector4 PackedUniform = FUintVector4(0, 0, 0, 0);
+		FUintVector4 PackedUniform = FUintVector4(GetDefaultFallbackColor(), 0, 0, 0);
 	};
 	
 	/** Get the global scene view parameters shared by all UMeshPaintVirtualTexture objects. */
diff --git a/Engine/Source/Runtime/Engine/Private/Components/StaticMeshComponent.cpp b/Engine/Source/Runtime/Engine/Private/Components/StaticMeshComponent.cpp
--- a/Engine/Source/Runtime/Engine/Private/Components/StaticMeshComponent.cpp
+++ b/Engine/Source/Runtime/Engine/Private/Components/StaticMeshComponent.cpp
@@ -49,7 +49,6 @@
 #include "NaniteVertexFactory.h"
 #include "StaticMeshSceneProxyDesc.h"
 #include "WorldPartition/ActorInstanceGuids.h"
-#include "VT/MeshPaintVirtualTexture.h"
 
 #include UE_INLINE_GENERATED_CPP_BY_NAME(StaticMeshComponent)
 
@@ -1995,23 +1994,6 @@ void UStaticMeshComponent::PostEditChangeProperty(FPropertyChangedEvent& Propert
 			// If the materials changed, then the component needs a texture streaming rebuild.
 			StreamingTextureData.Empty();
 		}
-
-		if (PropertyThatChanged->GetFName() == GET_MEMBER_NAME_CHECKED(UStaticMeshComponent, MeshPaintTexture))
-		{
-			//todo: Move creation and management of MeshPaintTexture to the mesh paint tools.
-			if (!bMeshPaintTexture)
-			{
-				MeshPaintTexture = nullptr;
-			}
-			else if (!MeshPaintTexture)
-			{
-				UMeshPaintVirtualTexture* NewTexture = NewObject<UMeshPaintVirtualTexture>(GetOutermost());
-				NewTexture->Source.Init(128, 128, 1, 8, TSF_BGRA8);
-				NewTexture->OwningComponent = MakeWeakObjectPtr(this);
-				NewTexture->UpdateResource();
-				MeshPaintTexture = NewTexture;
-			}
-		}
 	}
 
 	FBodyInstanceEditorHelpers::EnsureConsistentMobilitySimulationSettingsOnPostEditChange(this, PropertyChangedEvent);
@@ -3118,6 +3100,15 @@ UTexture* UStaticMeshComponent::GetMeshPaintTexture() const
 	return MeshPaintTexture; 
 }
 
+void UStaticMeshComponent::SetMeshPaintTexture(UTexture* InTexture)
+{
+	if (MeshPaintTexture != InTexture)
+	{
+		MeshPaintTexture = InTexture;
+		MarkRenderStateDirty();
+	}
+}
+
 void UStaticMeshComponent::SetMeshPaintTextureOverride(UTexture* OverrideTexture)
 {
 	if (MeshPaintTextureOverride != OverrideTexture)
@@ -3127,6 +3118,19 @@ void UStaticMeshComponent::SetMeshPaintTextureOverride(UTexture* OverrideTexture
 	}
 }
 
+int32 UStaticMeshComponent::GetMeshPaintTextureCoordinateIndex() const
+{
+	if (bOverrideMeshPaintTextureCoordinateIndex)
+	{
+		return OverridenMeshPaintTextureCoordinateIndex;
+	}
+	if (StaticMesh != nullptr)
+	{
+		return StaticMesh->MeshPaintTextureCoordinateIndex;
+	}
+	return 0;
+}
+
 bool UStaticMeshComponent::IsNavigationRelevant() const
 {
 	if (const UStaticMesh* Mesh = GetStaticMesh())
diff --git a/Engine/Source/Runtime/Engine/Private/VT/MeshPaintVirtualTexture.cpp b/Engine/Source/Runtime/Engine/Private/VT/MeshPaintVirtualTexture.cpp
--- a/Engine/Source/Runtime/Engine/Private/VT/MeshPaintVirtualTexture.cpp
+++ b/Engine/Source/Runtime/Engine/Private/VT/MeshPaintVirtualTexture.cpp
@@ -2,6 +2,7 @@
 
 #include "VT/MeshPaintVirtualTexture.h"
 
+#include "ComponentRecreateRenderStateContext.h"
 #include "Components/PrimitiveComponent.h"
 #include "EngineModule.h"
 #include "RendererInterface.h"
@@ -21,6 +22,10 @@ static TAutoConsoleVariable<bool> CVarMeshPaintVirtualTextureEnable(
 	TEXT("r.MeshPaintVirtualTexture.Enable"),
 	true,
 	TEXT("Run time enable mesh painting with virtual textures"),
+	FConsoleVariableDelegate::CreateLambda([](IConsoleVariable* InVariable)
+	{
+		FGlobalComponentRecreateRenderStateContext Context;
+	}),
 	ECVF_RenderThreadSafe);
 
 static TAutoConsoleVariable<int32> CVarMeshPaintVirtualTextureTileSize(
@@ -35,6 +40,11 @@ static TAutoConsoleVariable<int32> CVarMeshPaintVirtualTextureTileBorderSize(
 	TEXT("Virtual texture tile border size for mesh paint textures"),
 	ECVF_RenderThreadSafe | ECVF_ReadOnly);
 
+static TAutoConsoleVariable<int32> CVarMeshPaintVirtualTextureTexelsPerVertex(
+	TEXT("r.MeshPaintVirtualTexture.DefaultTexelsPerVertex"),
+	4,
+	TEXT("Default ratio of texels to vertices when creating a texture for a mesh"),
+	ECVF_Default);
 
 namespace MeshPaintVirtualTexture
 {
@@ -59,6 +69,13 @@ namespace MeshPaintVirtualTexture
 		return FVirtualTextureBuildSettings::ClampAndAlignTileBorderSize(CVarMeshPaintVirtualTextureTileBorderSize.GetValueOnAnyThread());
 	}
 
+	uint32 GetDefaultTextureSize(int32 InNumVertices)
+	{
+		const int32 NumTexels = InNumVertices * CVarMeshPaintVirtualTextureTexelsPerVertex.GetValueOnGameThread();
+		const uint32 TextureSize = (uint32)FMath::Sqrt((float)NumTexels);
+		return FMath::DivideAndRoundUp(TextureSize, GetTileSize()) * GetTileSize();
+	}
+
 	/** 
 	 * Fill out the scene uniforms from an allocated VT. 
 	 * We expect the result to be constant for all allocated VTs (so that they can share one uniform buffer).
@@ -78,7 +95,7 @@ namespace MeshPaintVirtualTexture
 		const float RcpPhysicalTextureSize = 1.0f / float(InAllocatedVT->GetPhysicalTextureSize(0));
 		const uint32 pPageSize = vPageSize + PageBorderSize * 2u;
 
-		OutParams.PackedUniform.X = ~0u; // todo fallback color
+		OutParams.PackedUniform.X = GetDefaultFallbackColor();
 		OutParams.PackedUniform.Y = FMath::AsUInt((float)vPageSize * RcpPhysicalTextureSize);
 		OutParams.PackedUniform.Z = FMath::AsUInt((float)PageBorderSize * RcpPhysicalTextureSize);
 		
@@ -203,6 +220,11 @@ UMeshPaintVirtualTexture::UMeshPaintVirtualTexture(const FObjectInitializer& Obj
 	: UTexture2D(ObjectInitializer)
 {
 	VirtualTextureStreaming = true;
+	
+#if WITH_EDITORONLY_DATA
+	// Force alpha channel so that we the platform format is consistent for all content.
+	CompressionForceAlpha = true;
+#endif
 }
 
 void UMeshPaintVirtualTexture::GetVirtualTextureBuildSettings(FVirtualTextureBuildSettings& OutSettings) const
