Restore default acceleration mode to "true" as it was internally prior to being exposed

diff --git a/Engine/Source/Runtime/Engine/Classes/PhysicsEngine/ConstraintDrives.h b/Engine/Source/Runtime/Engine/Classes/PhysicsEngine/ConstraintDrives.h
--- a/Engine/Source/Runtime/Engine/Classes/PhysicsEngine/ConstraintDrives.h
+++ b/Engine/Source/Runtime/Engine/Classes/PhysicsEngine/ConstraintDrives.h
@@ -84,7 +84,7 @@ struct FLinearDriveConstraint
 	FConstraintDrive ZDrive;
 	
 	UPROPERTY(EditAnywhere, Category = AngularMotor)
-	bool bAccelerationMode;
+	bool bAccelerationMode = true;
 
 #if WITH_EDITORONLY_DATA
 	PRAGMA_DISABLE_DEPRECATION_WARNINGS
@@ -151,7 +151,7 @@ struct FAngularDriveConstraint
 	TEnumAsByte<enum EAngularDriveMode::Type> AngularDriveMode;
 
 	UPROPERTY(EditAnywhere, Category = AngularMotor)
-	bool bAccelerationMode;
+	bool bAccelerationMode = true;
 
 	ENGINE_API FAngularDriveConstraint();
 
