[World Partition] Fixed missing actors replacement events.

diff --git a/Engine/Source/Runtime/Engine/Private/WorldPartition/ActorDescContainerInstance.cpp b/Engine/Source/Runtime/Engine/Private/WorldPartition/ActorDescContainerInstance.cpp
--- a/Engine/Source/Runtime/Engine/Private/WorldPartition/ActorDescContainerInstance.cpp
+++ b/Engine/Source/Runtime/Engine/Private/WorldPartition/ActorDescContainerInstance.cpp
@@ -245,6 +245,7 @@ void UActorDescContainerInstance::RegisterDelegates()
 	if (GetOuterWorldPartition())
 	{
 		FCoreUObjectDelegates::OnObjectsReplaced.AddUObject(this, &UActorDescContainerInstance::OnObjectsReplaced);
+		FEditorDelegates::OnEditorActorReplaced.AddUObject(this, &UActorDescContainerInstance::OnEditorActorReplaced);
 	}
 
 	// Only listen to this event if we have registered Child Container Instances
@@ -270,6 +271,7 @@ void UActorDescContainerInstance::RegisterDelegates()
 void UActorDescContainerInstance::UnregisterDelegates()
 {
 	FCoreUObjectDelegates::OnObjectsReplaced.RemoveAll(this);
+	FEditorDelegates::OnEditorActorReplaced.RemoveAll(this);
 
 	check(Container);
 	Container->OnActorDescAddedEvent.RemoveAll(this);
@@ -513,17 +515,22 @@ void UActorDescContainerInstance::OnObjectsReplaced(const TMap<UObject*, UObject
 	{
 		if (AActor* OldActor = Cast<AActor>(OldObject))
 		{
-			if (Container->ShouldHandleActorEvent(OldActor, InstancingContext.GetPtrOrNull() != nullptr))
-			{
-				// Only replace if OldActor matches ActorDescInstance ActorPtr.
-				// No need to replace the ActorDescInstance ActorPtr if it isn't set, it will get refreshed on its next GetActor call. This avoids replacing pointers on the wrong ActorDescInstance for Instanced worlds
-				if (FWorldPartitionActorDescInstance* ActorDescInstance = GetActorDescInstance(OldActor->GetActorGuid()); ActorDescInstance && ActorDescInstance->ActorPtr == OldActor)
-				{
-					AActor* NewActor = Cast<AActor>(NewObject);
-					FWorldPartitionActorDescUtils::ReplaceActorDescriptorPointerFromActor(OldActor, NewActor, ActorDescInstance);
-					OnActorReplacedEvent.Broadcast(ActorDescInstance);
-				}
-			}
+			OnEditorActorReplaced(OldActor, Cast<AActor>(NewObject));
+		}
+	}
+}
+
+void UActorDescContainerInstance::OnEditorActorReplaced(AActor* InOldActor, AActor* InNewActor)
+{
+	// Patch up Actor pointers in ActorDescInstances
+	if (Container->ShouldHandleActorEvent(InOldActor, InstancingContext.GetPtrOrNull() != nullptr))
+	{
+		// Only replace if OldActor matches ActorDescInstance ActorPtr.
+		// No need to replace the ActorDescInstance ActorPtr if it isn't set, it will get refreshed on its next GetActor call. This avoids replacing pointers on the wrong ActorDescInstance for Instanced worlds
+		if (FWorldPartitionActorDescInstance* ActorDescInstance = GetActorDescInstance(InOldActor->GetActorGuid()); ActorDescInstance && ActorDescInstance->ActorPtr == InOldActor)
+		{
+			FWorldPartitionActorDescUtils::ReplaceActorDescriptorPointerFromActor(InOldActor, InNewActor, ActorDescInstance);
+			OnActorReplacedEvent.Broadcast(ActorDescInstance);
 		}
 	}
 }
diff --git a/Engine/Source/Runtime/Engine/Public/ExternalDirtyActorsTracker.h b/Engine/Source/Runtime/Engine/Public/ExternalDirtyActorsTracker.h
--- a/Engine/Source/Runtime/Engine/Public/ExternalDirtyActorsTracker.h
+++ b/Engine/Source/Runtime/Engine/Public/ExternalDirtyActorsTracker.h
@@ -7,6 +7,7 @@
 #include "TickableEditorObject.h"
 #include "UObject/Package.h"
 #include "UObject/WeakObjectPtr.h"
+#include "Editor.h"
 
 class ULevel;
 
@@ -26,12 +27,14 @@ public:
 	{
 		UPackage::PackageDirtyStateChangedEvent.AddRaw(this, &TExternalDirtyActorsTracker::OnPackageDirtyStateChanged);
 		FCoreUObjectDelegates::OnObjectsReplaced.AddRaw(this, &TExternalDirtyActorsTracker::OnObjectsReplaced);
+		FEditorDelegates::OnEditorActorReplaced.AddRaw(this, &TExternalDirtyActorsTracker::OnEditorActorReplaced);
 	}
 
 	~TExternalDirtyActorsTracker()
 	{
 		UPackage::PackageDirtyStateChangedEvent.RemoveAll(this);
 		FCoreUObjectDelegates::OnObjectsReplaced.RemoveAll(this);
+		FEditorDelegates::OnEditorActorReplaced.RemoveAll(this);
 	}
 
 	const MapType& GetDirtyActors() const { return DirtyActors; }
@@ -85,18 +88,23 @@ protected:
 			{
 				if (AActor* NewActor = Cast<AActor>(NewObject))
 				{
-					for (auto& [WeakActor, Value] : DirtyActors)
-					{
-						if (WeakActor.IsValid() && WeakActor.Get() == OldActor)
-						{
-							WeakActor = NewActor;
-						}
-					}
+					OnEditorActorReplaced(OldActor, NewActor);
 				}
 			}
 		}
 	}
 
+	void OnEditorActorReplaced(AActor* InOldActor, AActor* InNewActor)
+	{
+		for (auto& [WeakActor, Value] : DirtyActors)
+		{
+			if (WeakActor.IsValid() && WeakActor.Get() == InOldActor)
+			{
+				WeakActor = InNewActor;
+			}
+		}
+	}
+
 	//~ Begin FTickableEditorObject interface
 	virtual TStatId GetStatId() const override
 	{
diff --git a/Engine/Source/Runtime/Engine/Public/WorldPartition/ActorDescContainerInstance.h b/Engine/Source/Runtime/Engine/Public/WorldPartition/ActorDescContainerInstance.h
--- a/Engine/Source/Runtime/Engine/Public/WorldPartition/ActorDescContainerInstance.h
+++ b/Engine/Source/Runtime/Engine/Public/WorldPartition/ActorDescContainerInstance.h
@@ -177,6 +177,7 @@ private:
 	void OnActorDescUpdated(FWorldPartitionActorDesc* InActorDesc);
 
 	void OnObjectsReplaced(const TMap<UObject*, UObject*>& InOldToNewObjectMap);
+	void OnEditorActorReplaced(AActor* InOldActor, AActor* InNewActor);
 
 	void OnRegisterChildContainerInstance(const FGuid& InActorGuid, UActorDescContainerInstance* InChildContainerInstance);
 	void OnUnregisterChildContainerInstance(const FGuid& InActorGuid);
