Added Root Motion Attribute handling to Inertialization and DeadBlending nodes.

diff --git a/Engine/Source/Runtime/Engine/Classes/Animation/AnimNode_DeadBlending.h b/Engine/Source/Runtime/Engine/Classes/Animation/AnimNode_DeadBlending.h
--- a/Engine/Source/Runtime/Engine/Classes/Animation/AnimNode_DeadBlending.h
+++ b/Engine/Source/Runtime/Engine/Classes/Animation/AnimNode_DeadBlending.h
@@ -182,22 +182,25 @@ private:
 	 * 
 	 * @param InPose		The current pose for the animation being transitioned to.
 	 * @param InCurves		The current curves for the animation being transitioned to.
+	 * @param InAttributes	The current attributes for the animation being transitioned to.
 	 * @param SrcPosePrev	The pose recorded as output of the inertializer two frames ago.
 	 * @param SrcPoseCurr	The pose recorded as output of the inertializer on the previous frame.
 	 */
 	void InitFrom(
 		const FCompactPose& InPose, 
 		const FBlendedCurve& InCurves, 
+		const UE::Anim::FStackAttributeContainer& Attributes,
 		const FInertializationSparsePose& SrcPosePrev,
 		const FInertializationSparsePose& SrcPoseCurr);
 
 	/**
 	 * Computes the extrapolated pose and blends it with the input pose.
 	 * 
-	 * @param InOutPose		The current pose to blend with the extrapolated pose.
-	 * @param InOutCurves	The current curves to blend with the extrapolated curves.
+	 * @param InOutPose			The current pose to blend with the extrapolated pose.
+	 * @param InOutCurves		The current curves to blend with the extrapolated curves.
+	 * @param InOutAttributes	The current attributes to blend with the extrapolated attributes.
 	 */
-	void ApplyTo(FCompactPose& InOutPose, FBlendedCurve& InOutCurves);
+	void ApplyTo(FCompactPose& InOutPose, FBlendedCurve& InOutCurves, UE::Anim::FStackAttributeContainer& InOutAttributes);
 
 public: // IBoneReferenceSkeletonProvider
 	ENGINE_API class USkeleton* GetSkeleton(bool& bInvalidSkeletonIsError, const IPropertyHandle* PropertyHandle) override;
@@ -243,6 +246,10 @@ private:
 	TArray<FVector3f> BoneRotationDecayHalfLives;
 	TArray<FVector3f> BoneScaleDecayHalfLives;
 
+	FVector3f RootTranslationVelocity;
+	FVector3f RootRotationVelocity;
+	FVector3f RootScaleVelocity;
+
 	// Recorded curve state at the point of transition.
 
 	struct FDeadBlendingCurveElement : public UE::Anim::FCurveElement
diff --git a/Engine/Source/Runtime/Engine/Classes/Animation/AnimNode_Inertialization.h b/Engine/Source/Runtime/Engine/Classes/Animation/AnimNode_Inertialization.h
--- a/Engine/Source/Runtime/Engine/Classes/Animation/AnimNode_Inertialization.h
+++ b/Engine/Source/Runtime/Engine/Classes/Animation/AnimNode_Inertialization.h
@@ -158,6 +158,12 @@ private:
 	// Transform of the component at the point of the snapshot
 	FTransform ComponentTransform;
 	
+	// Has Root Motion
+	bool bHasRootMotion = false;
+
+	// Root Motion Delta at the point of the snapshot
+	FTransform RootMotionDelta;
+
 	// For each SkeletonPoseBoneIndex this array stores the index into the BoneTranslations, BoneRotations, and 
 	// BoneScales arrays which contains that bone's data. Or INDEX_NONE if this bone's data is not in the snapshot.
 	TArray<int32> BoneIndices;
@@ -180,7 +186,14 @@ private:
 	// Delta Time since last snapshot
 	float DeltaTime = 0.0f;
 
-	void InitFrom(const FCompactPose& Pose, const FBlendedCurve& InCurves, const FTransform& InComponentTransform, const FName InAttachParentName, const float InDeltaTime);
+	void InitFrom(
+		const FCompactPose& Pose, 
+		const FBlendedCurve& InCurves, 
+		const UE::Anim::FStackAttributeContainer& Attributes, 
+		const FTransform& InComponentTransform, 
+		const FName InAttachParentName, 
+		const float InDeltaTime);
+
 	bool IsEmpty() const;
 	void Empty();
 };
@@ -386,6 +399,7 @@ private:
 	 *
 	 * @param InPose				The current pose for the animation being transitioned to.
 	 * @param InCurves				The current curves for the animation being transitioned to.
+	 * @param InAttributes			The current attributes for the animation being transitioned to.
 	 * @param ComponentTransform	The component transform of the current pose
 	 * @param AttachParentName		The name of the attached parent object
 	 * @param PreviousPose1			The pose recorded as output of the inertializer on the previous frame.
@@ -394,6 +408,7 @@ private:
 	void InitFrom(
 		const FCompactPose& InPose, 
 		const FBlendedCurve& InCurves, 
+		const UE::Anim::FStackAttributeContainer& InAttributes,
 		const FTransform& ComponentTransform, 
 		const FName AttachParentName, 
 		const FInertializationSparsePose& PreviousPose1, 
@@ -402,10 +417,11 @@ private:
 	/**
 	 * Applies the inertialization difference to the given pose (decaying to zero as ElapsedTime approaches Duration)
 	 *
-	 * @param InOutPose		The current pose to blend with the extrapolated pose.
-	 * @param InOutCurves	The current curves to blend with the extrapolated curves.
+	 * @param InOutPose		    The current pose to blend with the extrapolated pose.
+	 * @param InOutCurves	    The current curves to blend with the extrapolated curves.
+	 * @param InOutAttributes	The current attributes to blend with the extrapolated attributes.
 	 */
-	void ApplyTo(FCompactPose& InOutPose, FBlendedCurve& InOutCurves);
+	void ApplyTo(FCompactPose& InOutPose, FBlendedCurve& InOutCurves, UE::Anim::FStackAttributeContainer& InOutAttributes);
 
 	// Snapshots of the actor pose generated as output.
 	FInertializationSparsePose PrevPoseSnapshot;
@@ -450,6 +466,12 @@ private:
 	TArray<FVector3f> BoneScaleDiffAxis;
 	TArray<float> BoneScaleDiffMagnitude;
 	TArray<float> BoneScaleDiffSpeed;
+	FVector3f RootTranslationVelocityDiffDirection;
+	float RootTranslationVelocityDiffMagnitude;
+	FVector3f RootRotationVelocityDiffDirection;
+	float RootRotationVelocityDiffMagnitude;
+	FVector3f RootScaleVelocityDiffDirection;
+	float RootScaleVelocityDiffMagnitude;
 
 	// Curve differences
 	TBaseBlendedCurve<FDefaultAllocator, FInertializationCurveDiffElement> CurveDiffs;
diff --git a/Engine/Source/Runtime/Engine/Private/Animation/AnimNode_DeadBlending.cpp b/Engine/Source/Runtime/Engine/Private/Animation/AnimNode_DeadBlending.cpp
--- a/Engine/Source/Runtime/Engine/Private/Animation/AnimNode_DeadBlending.cpp
+++ b/Engine/Source/Runtime/Engine/Private/Animation/AnimNode_DeadBlending.cpp
@@ -2,6 +2,7 @@
 
 #include "Animation/AnimNode_DeadBlending.h"
 #include "Animation/AnimInstanceProxy.h"
+#include "Animation/AnimRootMotionProvider.h"
 #include "Animation/AnimNode_SaveCachedPose.h"
 #include "Animation/BlendProfile.h"
 #include "Algo/MaxElement.h"
@@ -277,7 +278,12 @@ void FAnimNode_DeadBlending::Deactivate()
 	InertializationDurationPerBone.Empty();
 }
 
-void FAnimNode_DeadBlending::InitFrom(const FCompactPose& InPose, const FBlendedCurve& InCurves, const FInertializationSparsePose& SrcPosePrev, const FInertializationSparsePose& SrcPoseCurr)
+void FAnimNode_DeadBlending::InitFrom(
+	const FCompactPose& InPose, 
+	const FBlendedCurve& InCurves, 
+	const UE::Anim::FStackAttributeContainer& InAttributes, 
+	const FInertializationSparsePose& SrcPosePrev, 
+	const FInertializationSparsePose& SrcPoseCurr)
 {
 	TRACE_CPUPROFILER_EVENT_SCOPE(FAnimNode_DeadBlending::InitFrom);
 
@@ -473,9 +479,35 @@ void FAnimNode_DeadBlending::InitFrom(const FCompactPose& InPose, const FBlended
 	{
 		UE::Anim::FCurveUtils::Filter(CurveData, ExtrapolatedCurveFilter);
 	}
+
+	// Record Root Motion Delta
+
+	// We don't compute the root acceleration difference which can be quite noisy and unreliable
+	// and so can cause the extrapolated root velocity to be quite bad and not very useful when
+	// being blended out.
+
+	RootTranslationVelocity = FVector3f::ZeroVector;
+	RootRotationVelocity = FVector3f::ZeroVector;
+	RootScaleVelocity = FVector3f::ZeroVector;
+
+	if (const UE::Anim::IAnimRootMotionProvider* RootMotionProvider = UE::Anim::IAnimRootMotionProvider::Get())
+	{
+		FTransform CurrRootMotionDelta = FTransform::Identity;
+
+		if (RootMotionProvider->ExtractRootMotion(InAttributes, CurrRootMotionDelta) &&
+			SrcPoseCurr.bHasRootMotion &&
+			SrcPoseCurr.DeltaTime > UE_KINDA_SMALL_NUMBER)
+		{
+			RootTranslationVelocity = (FVector3f)SrcPoseCurr.RootMotionDelta.GetTranslation() / SrcPoseCurr.DeltaTime;
+			RootRotationVelocity = (FVector3f)SrcPoseCurr.RootMotionDelta.GetRotation().ToRotationVector() / SrcPoseCurr.DeltaTime;
+			RootScaleVelocity = bLinearlyInterpolateScales ?
+				((FVector3f)SrcPoseCurr.RootMotionDelta.GetScale3D() / SrcPoseCurr.DeltaTime) :
+				((FVector3f)UE::Anim::DeadBlending::Private::VectorLogSafe(SrcPoseCurr.RootMotionDelta.GetScale3D()) / SrcPoseCurr.DeltaTime);
+		}
+	}
 }
 
-void FAnimNode_DeadBlending::ApplyTo(FCompactPose& InOutPose, FBlendedCurve& InOutCurves)
+void FAnimNode_DeadBlending::ApplyTo(FCompactPose& InOutPose, FBlendedCurve& InOutCurves, UE::Anim::FStackAttributeContainer& InOutAttributes)
 {
 	TRACE_CPUPROFILER_EVENT_SCOPE(FAnimNode_DeadBlending::ApplyTo);
 
@@ -623,6 +655,38 @@ void FAnimNode_DeadBlending::ApplyTo(FCompactPose& InOutPose, FBlendedCurve& InO
 			OutResultElement.Value = FMath::Lerp(InElement0.Value, ExtrapolatedCurve, CurveAlpha);
 			OutResultElement.Flags = InElement0.Flags | InElement1.Flags;
 		});
+
+	// Blend Root Motion Delta
+
+	if (const UE::Anim::IAnimRootMotionProvider* RootMotionProvider = UE::Anim::IAnimRootMotionProvider::Get())
+	{
+		// Compute Blend Alpha
+
+		const float Alpha = 1.0f - FAlphaBlend::AlphaToBlendOption(
+			FMath::Clamp(InertializationTime / FMath::Max(InertializationDurationPerBone[0], UE_SMALL_NUMBER), 0.0f, 1.0f),
+			InertializationBlendMode, InertializationCustomBlendCurve);
+
+		if (Alpha != 0.0f)
+		{
+			FTransform CurrRootMotionDelta = FTransform::Identity;
+			if (RootMotionProvider->ExtractRootMotion(InOutAttributes, CurrRootMotionDelta))
+			{
+				CurrRootMotionDelta.SetTranslation(FMath::Lerp(CurrRootMotionDelta.GetTranslation(), DeltaTime * (FVector)RootTranslationVelocity, Alpha));
+				CurrRootMotionDelta.SetRotation(FQuat::MakeFromRotationVector(DeltaTime * (FVector)RootRotationVelocity * Alpha) * CurrRootMotionDelta.GetRotation());
+
+				if (bLinearlyInterpolateScales)
+				{
+					CurrRootMotionDelta.SetScale3D(FMath::Lerp(CurrRootMotionDelta.GetScale3D(), DeltaTime * (FVector)RootScaleVelocity, Alpha));
+				}
+				else
+				{
+					CurrRootMotionDelta.SetScale3D(UE::Anim::DeadBlending::Private::VectorEerp(CurrRootMotionDelta.GetScale3D(), DeltaTime * (FVector)RootScaleVelocity, Alpha));
+				}
+
+				RootMotionProvider->OverrideRootMotion(CurrRootMotionDelta, InOutAttributes);
+			}
+		}
+	}
 }
 
 class USkeleton* FAnimNode_DeadBlending::GetSkeleton(bool& bInvalidSkeletonIsError, const IPropertyHandle* PropertyHandle)
@@ -899,6 +963,7 @@ void FAnimNode_DeadBlending::Evaluate_AnyThread(FPoseContext& Output)
 			InitFrom(
 				Output.Pose,
 				Output.Curve,
+				Output.CustomAttributes,
 				PrevPoseSnapshot,
 				CurrPoseSnapshot);
 		}
@@ -909,6 +974,7 @@ void FAnimNode_DeadBlending::Evaluate_AnyThread(FPoseContext& Output)
 			InitFrom(
 				Output.Pose,
 				Output.Curve,
+				Output.CustomAttributes,
 				CurrPoseSnapshot,
 				CurrPoseSnapshot);
 		}
@@ -940,7 +1006,7 @@ void FAnimNode_DeadBlending::Evaluate_AnyThread(FPoseContext& Output)
 
 	if (InertializationState == EInertializationState::Active)
 	{
-		ApplyTo(Output.Pose, Output.Curve);
+		ApplyTo(Output.Pose, Output.Curve, Output.CustomAttributes);
 	}
 
 	// Find AttachParentName
@@ -963,7 +1029,7 @@ void FAnimNode_DeadBlending::Evaluate_AnyThread(FPoseContext& Output)
 	}
 	
 	// Initialize the current pose
-	CurrPoseSnapshot.InitFrom(Output.Pose, Output.Curve, ComponentTransform, AttachParentName, DeltaTime);
+	CurrPoseSnapshot.InitFrom(Output.Pose, Output.Curve, Output.CustomAttributes, ComponentTransform, AttachParentName, DeltaTime);
 
 	// Reset Delta Time
 
diff --git a/Engine/Source/Runtime/Engine/Private/Animation/AnimNode_Inertialization.cpp b/Engine/Source/Runtime/Engine/Private/Animation/AnimNode_Inertialization.cpp
--- a/Engine/Source/Runtime/Engine/Private/Animation/AnimNode_Inertialization.cpp
+++ b/Engine/Source/Runtime/Engine/Private/Animation/AnimNode_Inertialization.cpp
@@ -2,6 +2,7 @@
 
 #include "Animation/AnimNode_Inertialization.h"
 #include "Animation/AnimInstanceProxy.h"
+#include "Animation/AnimRootMotionProvider.h"
 #include "Animation/AnimNode_SaveCachedPose.h"
 #include "Animation/AnimStats.h"
 #include "Animation/BlendProfile.h"
@@ -632,6 +633,7 @@ void FAnimNode_Inertialization::Evaluate_AnyThread(FPoseContext& Output)
 			InitFrom(
 				Output.Pose,
 				Output.Curve,
+				Output.CustomAttributes,
 				ComponentTransform,
 				AttachParentName,
 				CurrPoseSnapshot,
@@ -644,6 +646,7 @@ void FAnimNode_Inertialization::Evaluate_AnyThread(FPoseContext& Output)
 			InitFrom(
 				Output.Pose,
 				Output.Curve,
+				Output.CustomAttributes,
 				ComponentTransform,
 				AttachParentName,
 				CurrPoseSnapshot,
@@ -663,7 +666,7 @@ void FAnimNode_Inertialization::Evaluate_AnyThread(FPoseContext& Output)
 
 	if (InertializationState == EInertializationState::Active)
 	{
-		ApplyTo(Output.Pose, Output.Curve);
+		ApplyTo(Output.Pose, Output.Curve, Output.CustomAttributes);
 	}
 
 	// Record Pose Snapshot
@@ -675,7 +678,7 @@ void FAnimNode_Inertialization::Evaluate_AnyThread(FPoseContext& Output)
 	}
 	
 	// Initialize the current pose
-	CurrPoseSnapshot.InitFrom(Output.Pose, Output.Curve, ComponentTransform, AttachParentName, DeltaTime);
+	CurrPoseSnapshot.InitFrom(Output.Pose, Output.Curve, Output.CustomAttributes, ComponentTransform, AttachParentName, DeltaTime);
 	
 
 	// Reset the time accumulator and teleport state
@@ -771,7 +774,14 @@ void FAnimNode_Inertialization::ApplyInertialization(FPoseContext& Context, cons
 
 PRAGMA_ENABLE_DEPRECATION_WARNINGS
 
-void FAnimNode_Inertialization::InitFrom(const FCompactPose& InPose, const FBlendedCurve& InCurves, const FTransform& ComponentTransform, const FName AttachParentName, const FInertializationSparsePose& Prev1, const FInertializationSparsePose& Prev2)
+void FAnimNode_Inertialization::InitFrom(
+	const FCompactPose& InPose, 
+	const FBlendedCurve& InCurves, 
+	const UE::Anim::FStackAttributeContainer& InAttributes,
+	const FTransform& ComponentTransform, 
+	const FName AttachParentName, 
+	const FInertializationSparsePose& Prev1, 
+	const FInertializationSparsePose& Prev2)
 {	
 	check(!Prev1.IsEmpty() && !Prev2.IsEmpty());
 
@@ -986,9 +996,54 @@ void FAnimNode_Inertialization::InitFrom(const FCompactPose& InPose, const FBlen
 	{
 		UE::Anim::FCurveUtils::Filter(CurveDiffs, CurveFilter);
 	}
+
+	// Compute Root Motion Delta Difference
+
+	// We don't compute the speed difference since this essentially represents
+	// the acceleration of the root motion, which can be quite noisy and unreliable
+	// and so can cause the computed offsets to be bad when they are blended out
+
+	RootTranslationVelocityDiffDirection = FVector3f::ZeroVector;
+	RootTranslationVelocityDiffMagnitude = 0.0f;
+	RootRotationVelocityDiffDirection = FVector3f::ZeroVector;
+	RootRotationVelocityDiffMagnitude = 0.0f;
+	RootScaleVelocityDiffDirection = FVector3f::ZeroVector;
+	RootScaleVelocityDiffMagnitude = 0.0f;
+
+	if (const UE::Anim::IAnimRootMotionProvider* RootMotionProvider = UE::Anim::IAnimRootMotionProvider::Get())
+	{
+		FTransform CurrRootMotionDelta = FTransform::Identity;
+
+		if (RootMotionProvider->ExtractRootMotion(InAttributes, CurrRootMotionDelta) && 
+			Prev1.bHasRootMotion &&
+			Prev1.DeltaTime > UE_KINDA_SMALL_NUMBER && 
+			DeltaTime > UE_KINDA_SMALL_NUMBER)
+		{
+			const FVector RootTranslationVelocityDiff = (Prev1.RootMotionDelta.GetTranslation() / Prev1.DeltaTime) - (CurrRootMotionDelta.GetTranslation() / DeltaTime);
+			RootTranslationVelocityDiffMagnitude = RootTranslationVelocityDiff.Size();
+			if (RootTranslationVelocityDiffMagnitude > UE_KINDA_SMALL_NUMBER)
+			{
+				RootTranslationVelocityDiffDirection = (FVector3f)RootTranslationVelocityDiff / RootTranslationVelocityDiffMagnitude;
+			}
+
+			const FVector RootRotationVelocityDiff = (Prev1.RootMotionDelta.GetRotation().ToRotationVector() / Prev1.DeltaTime) - (CurrRootMotionDelta.GetRotation().ToRotationVector() / DeltaTime);
+			RootRotationVelocityDiffMagnitude = RootRotationVelocityDiff.Size();
+			if (RootRotationVelocityDiffMagnitude > UE_KINDA_SMALL_NUMBER)
+			{
+				RootRotationVelocityDiffDirection = (FVector3f)RootRotationVelocityDiff / RootRotationVelocityDiffMagnitude;
+			}
+
+			const FVector RootScaleVelocityDiff = (Prev1.RootMotionDelta.GetScale3D() / Prev1.DeltaTime) - (CurrRootMotionDelta.GetScale3D() / DeltaTime);
+			RootScaleVelocityDiffMagnitude = RootScaleVelocityDiff.Size();
+			if (RootScaleVelocityDiffMagnitude > UE_KINDA_SMALL_NUMBER)
+			{
+				RootScaleVelocityDiffDirection = (FVector3f)RootScaleVelocityDiff / RootScaleVelocityDiffMagnitude;
+			}
+		}
+	}
 }
 
-void FAnimNode_Inertialization::ApplyTo(FCompactPose& InOutPose, FBlendedCurve& InOutCurves)
+void FAnimNode_Inertialization::ApplyTo(FCompactPose& InOutPose, FBlendedCurve& InOutCurves, UE::Anim::FStackAttributeContainer& InOutAttributes)
 {
 	const FBoneContainer& BoneContainer = InOutPose.GetBoneContainer();
 
@@ -1049,6 +1104,35 @@ void FAnimNode_Inertialization::ApplyTo(FCompactPose& InOutPose, FBlendedCurve&
 			OutResultElement.Value = InElement0.Value + UE::Anim::Inertialization::Private::CalcInertialFloat(InElement1.Delta, InElement1.Derivative, InertializationElapsedTime, InertializationDuration);
 			OutResultElement.Flags = InElement0.Flags | InElement1.Flags;
 		});
+
+	// Apply Root Motion Delta Difference
+
+	if (const UE::Anim::IAnimRootMotionProvider* RootMotionProvider = UE::Anim::IAnimRootMotionProvider::Get())
+	{
+		FTransform CurrRootMotionDelta = FTransform::Identity;
+		if (RootMotionProvider->ExtractRootMotion(InOutAttributes, CurrRootMotionDelta))
+		{
+			// Use Blend Duration from Root Bone
+			const float Duration = InertializationDurationPerBone[0];
+
+			// Apply the root translation velocity difference
+			const FVector T = DeltaTime * (FVector)RootTranslationVelocityDiffDirection *
+				UE::Anim::Inertialization::Private::CalcInertialFloat(RootTranslationVelocityDiffMagnitude, 0.0f, InertializationElapsedTime, Duration);
+			CurrRootMotionDelta.AddToTranslation(T);
+
+			// Apply the root rotation velocity difference
+			const FQuat Q = FQuat::MakeFromRotationVector(DeltaTime * (FVector)RootRotationVelocityDiffDirection *
+				UE::Anim::Inertialization::Private::CalcInertialFloat(RootRotationVelocityDiffMagnitude, 0.0f, InertializationElapsedTime, Duration));
+			CurrRootMotionDelta.SetRotation(Q * CurrRootMotionDelta.GetRotation());
+
+			// Apply the root scale velocity difference
+			const FVector S = DeltaTime * (FVector)RootScaleVelocityDiffDirection *
+				UE::Anim::Inertialization::Private::CalcInertialFloat(RootScaleVelocityDiffMagnitude, 0.0f, InertializationElapsedTime, Duration);
+			CurrRootMotionDelta.SetScale3D(S + CurrRootMotionDelta.GetScale3D());
+
+			RootMotionProvider->OverrideRootMotion(CurrRootMotionDelta, InOutAttributes);
+		}
+	}
 }
 
 void FAnimNode_Inertialization::Deactivate()
@@ -1112,6 +1196,7 @@ PRAGMA_ENABLE_DEPRECATION_WARNINGS
 void FInertializationSparsePose::InitFrom(
 	const FCompactPose& Pose,
 	const FBlendedCurve& InCurves,
+	const UE::Anim::FStackAttributeContainer& Attributes,
 	const FTransform& InComponentTransform,
 	const FName InAttachParentName,
 	const float InDeltaTime)
@@ -1169,6 +1254,16 @@ void FInertializationSparsePose::InitFrom(
 		BoneScales[InertializationBoneIndex] = BoneTransform.GetScale3D();
 	}
 
+	// Init the Root Motion Delta
+
+	bHasRootMotion = false;
+	RootMotionDelta = FTransform::Identity;
+
+	if (const UE::Anim::IAnimRootMotionProvider* RootMotionProvider = UE::Anim::IAnimRootMotionProvider::Get())
+	{
+		bHasRootMotion = RootMotionProvider->ExtractRootMotion(Attributes, RootMotionDelta);
+	}
+
 	// Init the rest of the snapshot data
 
 	Curves.InitFrom(InCurves);
