#UE Move static mesh LODOverride message to a shared function and call it during data validation

diff --git a/Engine/Source/Runtime/Engine/Classes/Engine/StaticMesh.h b/Engine/Source/Runtime/Engine/Classes/Engine/StaticMesh.h
--- a/Engine/Source/Runtime/Engine/Classes/Engine/StaticMesh.h
+++ b/Engine/Source/Runtime/Engine/Classes/Engine/StaticMesh.h
@@ -2091,6 +2091,7 @@ private:
 	// Notification about missing Nanite required shader models.
 	TWeakPtr<class SNotificationItem> ShaderModelNotificationPtr;
 
+	void CheckForValidMinLODs(FPerQualityLevelInt& QualityLocalMinLOD, FPerPlatformInt& LocalMinLOD, int32& OutMinAvailableLOD, TArray<TPair<int32, FName>>& OutInvalidMinLODs) const;
 	void CheckForMissingShaderModels();
 
 	FOnPreMeshBuild PreMeshBuild;
diff --git a/Engine/Source/Runtime/Engine/Private/StaticMesh.cpp b/Engine/Source/Runtime/Engine/Private/StaticMesh.cpp
--- a/Engine/Source/Runtime/Engine/Private/StaticMesh.cpp
+++ b/Engine/Source/Runtime/Engine/Private/StaticMesh.cpp
@@ -4634,6 +4634,39 @@ EDataValidationResult UStaticMesh::IsDataValid(FDataValidationContext& Context)
 		}
 	}
 
+	{
+		// check the MinLOD values are all within range
+		FPerQualityLevelInt QualityLocalMinLOD;
+		FPerPlatformInt LocalMinLOD;
+		int32 MinAvailableLOD = INDEX_NONE;
+		TArray<TPair<int32, FName>> InvalidMinLODs;
+		CheckForValidMinLODs(QualityLocalMinLOD, LocalMinLOD, MinAvailableLOD, InvalidMinLODs);
+		if (InvalidMinLODs.Num() > 0)
+		{
+			for (const TPair<int32, FName>& InvalidMinLOD : InvalidMinLODs)
+			{
+				const int32 LODIdx = InvalidMinLOD.Key;
+				const FName OverrideName = InvalidMinLOD.Value;
+
+				FFormatNamedArguments Arguments;
+				Arguments.Add(TEXT("MinLOD"), FText::AsNumber(LODIdx));
+				Arguments.Add(TEXT("MinAvailLOD"), FText::AsNumber(MinAvailableLOD));
+				Arguments.Add(TEXT("OverrideName"), FText::FromName(OverrideName));
+
+				if (OverrideName.IsNone())
+				{
+					Context.AddWarning(FText::Format(LOCTEXT("LoadError_BadMinLOD", "Min LOD value of {MinLOD} is out of range 0..{MinAvailLOD}."), Arguments));
+				}
+				else
+				{
+					Context.AddWarning(FText::Format(LOCTEXT("LoadError_BadMinLODWithOverride", "Min LOD override of {MinLOD} for {OverrideName} is out of range 0..{MinAvailLOD}."), Arguments));
+				}
+			}
+
+			ValidationResult = EDataValidationResult::Invalid;
+		}
+	}
+
 	return ValidationResult;
 }
 #endif
@@ -6949,124 +6982,71 @@ void UStaticMesh::ExecutePostLoadInternal(FStaticMeshPostLoadContext& Context)
 			}
 		}
 #endif
+	}
 
-		// check the MinLOD values are all within range
-
-		int32 MinAvailableLOD = FMath::Max<int32>(GetRenderData()->LODResources.Num() - 1, 0);
-		
+	// check the MinLOD values are all within range
+	FPerQualityLevelInt QualityLocalMinLOD;
+	FPerPlatformInt LocalMinLOD;
+	int32 MinAvailableLOD = INDEX_NONE;
+	TArray<TPair<int32, FName>> InvalidMinLODs;
+	CheckForValidMinLODs(QualityLocalMinLOD, LocalMinLOD, MinAvailableLOD, InvalidMinLODs);
+	if (InvalidMinLODs.Num())
+	{
 		if (IsMinLodQualityLevelEnable())
 		{
-			bool bFixedQualityMinLOD = false;
-			FPerQualityLevelInt QualityLocalMinLOD = GetQualityLevelMinLOD();
+			SetQualityLevelMinLOD(QualityLocalMinLOD);
+		}
+		else
+		{
+			SetMinLOD(LocalMinLOD);
+		}
 
-			if (!GetRenderData()->LODResources.IsValidIndex(QualityLocalMinLOD.Default))
+		TArray<FText> MinLODErrors;
+		for (const TPair<int32, FName>& InvalidMinLOD : InvalidMinLODs)
+		{
+			const int32 LODIdx = InvalidMinLOD.Key;
+			const FName OverrideName = InvalidMinLOD.Value;
+
+			FFormatNamedArguments Arguments;
+			Arguments.Add(TEXT("MinLOD"), FText::AsNumber(LODIdx));
+			Arguments.Add(TEXT("MinAvailLOD"), FText::AsNumber(MinAvailableLOD));
+			Arguments.Add(TEXT("OverrideName"), FText::FromName(OverrideName));
+			if (OverrideName.IsNone())
 			{
-				FFormatNamedArguments Arguments;
-				Arguments.Add(TEXT("DefaultMinLOD"), FText::AsNumber(QualityLocalMinLOD.Default));
-				Arguments.Add(TEXT("MinAvailLOD"), FText::AsNumber(MinAvailableLOD));
-				TSharedRef<FUObjectToken> TokenRef = FUObjectToken::Create(this);
-				Async(
-					EAsyncExecution::TaskGraphMainThread,
-					// No choice to MoveTemp here, the SharedRef is not thread safe so it cannot
-					// be copied to another thread, only moved.
-					[Token = MoveTemp(TokenRef), Arguments]()
-					{
-						FMessageLog("LoadErrors").Warning()
-							->AddToken(Token)
-							->AddToken(FTextToken::Create(FText::Format(LOCTEXT("LoadError_DefaultMinLODOutOfRange", "Min LOD value of {DefaultMinLOD} is out of range 0..{MinAvailLOD} and has been adjusted to {MinAvailLOD}. Please verify and resave the asset."), Arguments)));
-					}	
-				);
-				QualityLocalMinLOD.Default = MinAvailableLOD;
-				bFixedQualityMinLOD = true;
+				MinLODErrors.Add(FText::Format(LOCTEXT("LoadError_BadMinLOD_Fixed", "Min LOD value of {MinLOD} is out of range 0..{MinAvailLOD} and has been adjusted to {MinAvailLOD}. Please verify and resave the asset."), Arguments));
 			}
-			for (TMap<int32, int32>::TIterator It(QualityLocalMinLOD.PerQuality); It; ++It)
+			else
 			{
-				if (!GetRenderData()->LODResources.IsValidIndex(It.Value()))
-				{
-					FFormatNamedArguments Arguments;
-					Arguments.Add(TEXT("QualityLevel"), FText::FromString(QualityLevelProperty::QualityLevelToFName(It.Key()).ToString()));
-					Arguments.Add(TEXT("QualityLevelMinLOD"), FText::AsNumber(It.Value()));
-					Arguments.Add(TEXT("MinAvailLOD"), FText::AsNumber(MinAvailableLOD));
-					TSharedRef<FUObjectToken> TokenRef = FUObjectToken::Create(this);
-					Async(
-						EAsyncExecution::TaskGraphMainThread,
-						// No choice to MoveTemp here, the SharedRef is not thread safe so it cannot
-						// be copied to another thread, only moved.
-						[Token = MoveTemp(TokenRef), Arguments]()
-						{
-							FMessageLog("LoadErrors").Warning()
-								->AddToken(Token)
-								->AddToken(FTextToken::Create(FText::Format(LOCTEXT("LoadError_MinLODOverrideForQualityLevel", "Min LOD override of {QualityLevelMinLOD} for {QualityLevel} is out of range 0..{MinAvailLOD} and has been adjusted to {MinAvailLOD}. Please verify and resave the asset."), Arguments)));
-						}
-					);
-					It.Value() = MinAvailableLOD;
-					bFixedQualityMinLOD = true;
-				}
+				MinLODErrors.Add(FText::Format(LOCTEXT("LoadError_BadMinLODWithOverride_Fixed", "Min LOD override of {MinLOD} for {OverrideName} is out of range 0..{MinAvailLOD} and has been adjusted to {MinAvailLOD}. Please verify and resave the asset."), Arguments));
 			}
+		}
 
-			if (bFixedQualityMinLOD)
+		if (IsRunningCommandlet())
+		{
+			for (const FText& MinLODError : MinLODErrors)
 			{
-				SetQualityLevelMinLOD(MoveTemp(QualityLocalMinLOD));
-				// Make sure Slate gets called from the game thread
-				Async(EAsyncExecution::TaskGraphMainThread, []() { FMessageLog("LoadErrors").Open(); });
+				UE_LOG(LogStaticMesh, Warning, TEXT("%s %s"), *GetNameSafe(this), *MinLODError.ToString());
 			}
 		}
-		else 
+		else
 		{
-			bool bFixedMinLOD = false;
-			FPerPlatformInt LocalMinLOD = GetMinLOD();
-
-			if (!GetRenderData()->LODResources.IsValidIndex(LocalMinLOD.Default))
-			{
-				FFormatNamedArguments Arguments;
-				Arguments.Add(TEXT("MinLOD"), FText::AsNumber(LocalMinLOD.Default));
-				Arguments.Add(TEXT("MinAvailLOD"), FText::AsNumber(MinAvailableLOD));
-				TSharedRef<FUObjectToken> TokenRef = FUObjectToken::Create(this);
-				Async(
-					EAsyncExecution::TaskGraphMainThread,
-					// No choice to MoveTemp here, the SharedRef is not thread safe so it cannot
-					// be copied to another thread, only moved.
-					[Token = MoveTemp(TokenRef), Arguments]()
+			TSharedRef<FUObjectToken> TokenRef = FUObjectToken::Create(this);
+			Async(
+				EAsyncExecution::TaskGraphMainThread,
+				// No choice to MoveTemp here, the SharedRef is not thread safe so it cannot
+				// be copied to another thread, only moved.
+				[Token = MoveTemp(TokenRef), MinAvailableLOD, MinLODErrors]()
+				{
+					for (const FText& MinLODError : MinLODErrors)
 					{
 						FMessageLog("LoadErrors").Warning()
 							->AddToken(Token)
-							->AddToken(FTextToken::Create(FText::Format(LOCTEXT("LoadError_BadMinLOD", "Min LOD value of {MinLOD} is out of range 0..{MinAvailLOD} and has been adjusted to {MinAvailLOD}. Please verify and resave the asset."), Arguments)));
+							->AddToken(FTextToken::Create(MinLODError));
 					}
-				);
-				LocalMinLOD.Default = MinAvailableLOD;
-				bFixedMinLOD = true;
-			}
-			for (TMap<FName, int32>::TIterator It(LocalMinLOD.PerPlatform); It; ++It)
-			{
-				if (!GetRenderData()->LODResources.IsValidIndex(It.Value()))
-				{
-					FFormatNamedArguments Arguments;
-					Arguments.Add(TEXT("MinLOD"), FText::AsNumber(It.Value()));
-					Arguments.Add(TEXT("MinAvailLOD"), FText::AsNumber(MinAvailableLOD));
-					Arguments.Add(TEXT("Platform"), FText::FromString(It.Key().ToString()));
-					TSharedRef<FUObjectToken> TokenRef = FUObjectToken::Create(this);
-					Async(
-						EAsyncExecution::TaskGraphMainThread,
-						// No choice to MoveTemp here, the SharedRef is not thread safe so it cannot
-						// be copied to another thread, only moved.
-						[Token = MoveTemp(TokenRef), Arguments]()
-						{
-							FMessageLog("LoadErrors").Warning()
-								->AddToken(Token)
-								->AddToken(FTextToken::Create(FText::Format(LOCTEXT("LoadError_BadMinLODOverride", "Min LOD override of {MinLOD} for {Platform} is out of range 0..{MinAvailLOD} and has been adjusted to {MinAvailLOD}. Please verify and resave the asset."), Arguments)));
-						}
-					);
-					It.Value() = MinAvailableLOD;
-					bFixedMinLOD = true;
-				}
-			}
 
-			if (bFixedMinLOD)
-			{
-				SetMinLOD(MoveTemp(LocalMinLOD));
-				// Make sure Slate gets called from the game thread
-				Async(EAsyncExecution::TaskGraphMainThread, []() { FMessageLog("LoadErrors").Open(); });
-			}
+					FMessageLog("LoadErrors").Open();
+				}
+			);
 		}
 	}
 
@@ -7180,6 +7160,47 @@ void UStaticMesh::ExecutePostLoadInternal(FStaticMeshPostLoadContext& Context)
 }
 
 #if WITH_EDITOR
+void UStaticMesh::CheckForValidMinLODs(FPerQualityLevelInt& QualityLocalMinLOD, FPerPlatformInt& LocalMinLOD, int32& OutMinAvailableLOD, TArray<TPair<int32, FName>>& OutInvalidMinLODs) const
+{
+	const FStaticMeshRenderData* LocalRenderData = GetRenderData();
+	if (!LocalRenderData)
+	{
+		return;
+	}
+
+	OutMinAvailableLOD = FMath::Max<int32>(LocalRenderData->LODResources.Num() - 1, 0);
+
+	auto CheckValidMinLOD = [LocalRenderData, OutMinAvailableLOD, &OutInvalidMinLODs](int32& LODIdx, FName OverrideName)
+	{
+		if (!LocalRenderData->LODResources.IsValidIndex(LODIdx))
+		{
+			OutInvalidMinLODs.Emplace(LODIdx, OverrideName);
+			LODIdx = OutMinAvailableLOD;
+		}
+	};
+
+	if (IsMinLodQualityLevelEnable())
+	{
+		QualityLocalMinLOD = GetQualityLevelMinLOD();
+		CheckValidMinLOD(QualityLocalMinLOD.Default, NAME_None);
+
+		for (TMap<int32, int32>::TIterator It(QualityLocalMinLOD.PerQuality); It; ++It)
+		{
+			CheckValidMinLOD(It.Value(), QualityLevelProperty::QualityLevelToFName(It.Key()));
+		}
+	}
+	else
+	{
+		LocalMinLOD = GetMinLOD();
+		CheckValidMinLOD(LocalMinLOD.Default, NAME_None);
+
+		for (TMap<FName, int32>::TIterator It(LocalMinLOD.PerPlatform); It; ++It)
+		{
+			CheckValidMinLOD(It.Value(), It.Key());
+		}
+	}
+}
+
 void UStaticMesh::CheckForMissingShaderModels()
 {
 #if PLATFORM_WINDOWS || PLATFORM_LINUX
