ODSC bugfixes
In FMaterial::TryGetShaders, if a shadermap is from ODSC, bShouldForceRecompile is false, and then RequestShaderIds is empty. Because of that, AddThreadedShaderPipelineRequest will not request additional permutations on an incomplete shadermap received from the ODSC server. We also need to fill RequestShaderIds if we detect that a shader is missing from the shadermap.

In FMaterialShaderMap::LoadForRemoteRecompile, if we received additional permutations on an existing ODSC shadermap, we were discarding it because it already exists in the global shadermap. Since a batch of requests may contain multiple shadermaps with the same id, make a first pass to deduplicate shadermaps with the same ID, register them, and associate with the materials afterwards

#rnx

diff --git a/Engine/Source/Runtime/Engine/Private/Materials/MaterialShader.cpp b/Engine/Source/Runtime/Engine/Private/Materials/MaterialShader.cpp
--- a/Engine/Source/Runtime/Engine/Private/Materials/MaterialShader.cpp
+++ b/Engine/Source/Runtime/Engine/Private/Materials/MaterialShader.cpp
@@ -2076,21 +2076,29 @@ void FMaterialShaderMap::SaveForRemoteRecompile(FArchive& Ar, const TMap<FString
 
 void FMaterialShaderMap::LoadForRemoteRecompile(FArchive& Ar, EShaderPlatform ShaderPlatform, TArray<UMaterialInterface*>& OutLoadedMaterials)
 {
-	FString MaterialName;
-	TArray<TRefCountPtr<FMaterialShaderMap>> LoadedShaderMaps;
+	TMap<FMaterialShaderMapId, TRefCountPtr<FMaterialShaderMap>> LoadedShaderMapsDictionary;
+
+	struct FMaterialShaderMapData
+	{
+		FString MaterialName;
+		TArray<FMaterialShaderMapId> LoadedShaderMapsIds;
+	};
+
+	TArray<FMaterialShaderMapData> MaterialShaderMapDataArray;
 
 	int32 MapSize;
 	Ar << MapSize;
 
+	MaterialShaderMapDataArray.Reserve(MapSize);
 	for (int32 MaterialIndex = 0; MaterialIndex < MapSize; MaterialIndex++)
 	{
-		MaterialName.Reset();
-		Ar << MaterialName;
+		FMaterialShaderMapData& MaterialShaderMapData = MaterialShaderMapDataArray.Emplace_GetRef();
+		Ar << MaterialShaderMapData.MaterialName;
 
 		int32 NumShaderMaps = 0;
 		Ar << NumShaderMaps;
 
-		LoadedShaderMaps.Reset();
+		MaterialShaderMapData.LoadedShaderMapsIds.Reserve(NumShaderMaps);
 		for (int32 ShaderMapIndex = 0; ShaderMapIndex < NumShaderMaps; ShaderMapIndex++)
 		{
 			uint8 bIsValid = 0;
@@ -2103,23 +2111,31 @@ void FMaterialShaderMap::LoadForRemoteRecompile(FArchive& Ar, EShaderPlatform Sh
 				// serialize the id and the material shader map
 				ShaderMap->Serialize(Ar, false);
 
-				// If we already registered a material shadermap with the same ID, just re-use it. The material will do the same in FMaterial::CacheShaders anyway
-				TRefCountPtr<FMaterialShaderMap> ExistingShaderMap = FMaterialShaderMap::FindId(ShaderMap->GetShaderMapId(), ShaderPlatform);
-#if WITH_ODSC
-				if (ExistingShaderMap && ExistingShaderMap->IsFromODSC())
-				{
-					ShaderMap = ExistingShaderMap;
-				}
-				else
-#endif
-				{
-					// Register in the global map
-					ShaderMap->RegisterForODSC(ShaderPlatform);
-				}
-
-				LoadedShaderMaps.Add(ShaderMap);
+				LoadedShaderMapsDictionary.Add(ShaderMap->GetShaderMapId(), ShaderMap);
+				MaterialShaderMapData.LoadedShaderMapsIds.Add(ShaderMap->GetShaderMapId());
 			}
 		}
+	}
+
+	for (auto IterLoadedShaderMaps = LoadedShaderMapsDictionary.CreateIterator(); IterLoadedShaderMaps; ++IterLoadedShaderMaps)
+	{
+		// Register in the global map
+		IterLoadedShaderMaps.Value()->RegisterForODSC(ShaderPlatform);
+	}
+
+	check(MaterialShaderMapDataArray.Num() == MapSize);
+	for (int32 MaterialIndex = 0; MaterialIndex < MapSize; MaterialIndex++)
+	{
+		TArray<TRefCountPtr<FMaterialShaderMap>> LoadedShaderMaps;
+		const FMaterialShaderMapData& MaterialShaderMapData = MaterialShaderMapDataArray[MaterialIndex];
+		const FString MaterialName = MaterialShaderMapData.MaterialName;
+
+
+		LoadedShaderMaps.Reserve(MaterialShaderMapData.LoadedShaderMapsIds.Num());
+		for (const FMaterialShaderMapId& MaterialShaderMapId : MaterialShaderMapData.LoadedShaderMapsIds)
+		{
+			LoadedShaderMaps.Add(LoadedShaderMapsDictionary.FindChecked(MaterialShaderMapId));
+		}
 
 #if WITH_ODSC
 		FODSCManager::RegisterMaterialShaderMaps(MaterialName, LoadedShaderMaps);
diff --git a/Engine/Source/Runtime/Engine/Private/Materials/MaterialShared.cpp b/Engine/Source/Runtime/Engine/Private/Materials/MaterialShared.cpp
--- a/Engine/Source/Runtime/Engine/Private/Materials/MaterialShared.cpp
+++ b/Engine/Source/Runtime/Engine/Private/Materials/MaterialShared.cpp
@@ -3731,8 +3731,10 @@ bool FMaterial::TryGetShaders(const FMaterialShaderTypes& InTypes, const FVertex
 	}
 
 #if WITH_ODSC
-	bool bShouldForceRecompile = FODSCManager::ShouldForceRecompile(ShaderMap, this) && !IsDefaultMaterial();
+	const bool bIsODSCActive = FODSCManager::IsODSCActive();
+	bool bShouldForceRecompile = bIsODSCActive && FODSCManager::ShouldForceRecompile(ShaderMap, this) && !IsDefaultMaterial();
 #else
+	constexpr bool bIsODSCActive = false;
 	constexpr bool bShouldForceRecompile = false;
 #endif
 
@@ -3796,7 +3798,7 @@ bool FMaterial::TryGetShaders(const FMaterialShaderTypes& InTypes, const FVertex
 		TArray<FShaderId> RequestShaderIds;
 		TArray<FString> ShaderStageNamesToCompile;
         bool bODSCRequestAlreadySent = false;
-		if (bShouldForceRecompile)
+		if (bIsODSCActive && (Pipeline == nullptr || bShouldForceRecompile))
 		{
 			for (auto* ShaderType : InTypes.PipelineType->GetStages())
 			{
@@ -3837,7 +3839,7 @@ bool FMaterial::TryGetShaders(const FMaterialShaderTypes& InTypes, const FVertex
 #if WITH_ODSC
 				if (FPlatformProperties::RequiresCookedData() && !bODSCRequestAlreadySent)
 				{
-					if (GODSCManager->IsHandlingRequests())
+					if (bIsODSCActive)
 					{
 						const FString VFTypeName(InVertexFactoryType ? InVertexFactoryType->GetName() : TEXT(""));
 						const FString PipelineName(InTypes.PipelineType->GetName());
@@ -3915,7 +3917,7 @@ bool FMaterial::TryGetShaders(const FMaterialShaderTypes& InTypes, const FVertex
 
 				TArray<FShaderId> RequestShaderIds;
                 bool bODSCRequestAlreadySent = false;
-				if (bShouldForceRecompile)
+				if (bIsODSCActive && (Shader == nullptr || bShouldForceRecompile))
 				{
 					RequestShaderIds.Add(FShaderId(ShaderType, ShaderMap->GetShaderMapId().CookedShaderMapIdHash, FHashedName(), InVertexFactoryType, PermutationId, ShaderPlatform));
 
@@ -3941,7 +3943,7 @@ bool FMaterial::TryGetShaders(const FMaterialShaderTypes& InTypes, const FVertex
 #if WITH_ODSC
 					if (FPlatformProperties::RequiresCookedData() && !bODSCRequestAlreadySent)
 					{
-						if (GODSCManager->IsHandlingRequests())
+						if (bIsODSCActive)
 						{
 							const FString VFTypeName(InVertexFactoryType ? InVertexFactoryType->GetName() : TEXT(""));
 							const FString PipelineName;
