Added new CVars that allow CSV Tracking for the following Movement related stats:
- Correction Distance: whether the client received a correction from the server and how big it was.
- Num movement RPCs the client sends to the server per frame
- Num movement RPCs the client receives from the server per frame
- Size of the sent Movement RPCs
- Size of the received Movement RPCs
#tests Tested in a local build and generated local PRS reports

diff --git a/Engine/Source/Runtime/Engine/Private/Components/CharacterMovementComponent.cpp b/Engine/Source/Runtime/Engine/Private/Components/CharacterMovementComponent.cpp
--- a/Engine/Source/Runtime/Engine/Private/Components/CharacterMovementComponent.cpp
+++ b/Engine/Source/Runtime/Engine/Private/Components/CharacterMovementComponent.cpp
@@ -36,6 +36,7 @@
 #include UE_INLINE_GENERATED_CPP_BY_NAME(CharacterMovementComponent)
 
 CSV_DEFINE_CATEGORY(CharacterMovement, true);
+CSV_DEFINE_CATEGORY(CharacterMovementDetailed, (!UE_BUILD_SHIPPING));
 
 DEFINE_LOG_CATEGORY_STATIC(LogCharacterMovement, Log, All);
 DEFINE_LOG_CATEGORY_STATIC(LogNavMeshMovement, Log, All);
@@ -360,6 +361,24 @@ namespace CharacterMovementCVars
 		TEXT("Apply the ledge movement vector directly, rather than the old method that reapplied acceleration."),
 		ECVF_Default);
 
+#if CSV_PROFILER_STATS
+	bool bClientRecordMovePackedRpcStatsToCsv = false;
+	FAutoConsoleVariableRef CVarClientRecordMovePackedRpcStatsToCsv(
+		TEXT("p.ClientRecordMovePackedRpcStatsToCsv"),
+		bClientRecordMovePackedRpcStatsToCsv,
+		TEXT("Whether to record MovePacked RPC stats to Csv on the client (RPC count and RPC size)\n")
+		TEXT("0: Disable, 1: Record MovePacked RPC stats"),
+		ECVF_Default);
+
+	bool bClientRecordNetCorrectionDistanceToCsv = false;
+	FAutoConsoleVariableRef CVarRecordNetCorrectionDistanceToCsv(
+		TEXT("p.ClientRecordNetCorrectionDistanceToCsv"),
+		bClientRecordNetCorrectionDistanceToCsv,
+		TEXT("Whether to record net correction distances to Csv on the client. This can be useful to understand how often corrections happen and how big are they\n")
+		TEXT("0: Disable, 1: Record the correction distance"),
+		ECVF_Default);
+#endif //CSV_PROFILER_STATS
+
 #if !UE_BUILD_SHIPPING
 
 	int32 NetShowCorrections = 0;
@@ -9555,6 +9574,18 @@ bool FCharacterNetworkMoveData::Serialize(UCharacterMovementComponent& Character
 
 void UCharacterMovementComponent::ServerMovePacked_ClientSend(const FCharacterServerMovePackedBits& PackedBits)
 {
+#if CSV_PROFILER_STATS
+	if (CharacterMovementCVars::bClientRecordMovePackedRpcStatsToCsv)
+	{
+		// track the amount of ServerMovePacked RPCs the client sends to the server
+		CSV_CUSTOM_STAT(CharacterMovementDetailed, Client_SendServerMovePacked_RPC, 1, ECsvCustomStatOp::Accumulate);
+
+		// track the size of the RPC sent to the server
+		CSV_CUSTOM_STAT(CharacterMovementDetailed, Client_SendServerMovePacked_RPC_NumBits, PackedBits.DataBits.Num(), ECsvCustomStatOp::Accumulate);
+	}
+#endif //CSV_PROFILER_STATS
+
+
 	// Pass through RPC call to character on server, there is less RPC bandwidth overhead when used on an Actor rather than a Component.
 	CharacterOwner->ServerMovePacked(PackedBits);
 }
@@ -10413,6 +10444,17 @@ void UCharacterMovementComponent::MoveResponsePacked_ServerSend(const FCharacter
 
 void UCharacterMovementComponent::MoveResponsePacked_ClientReceive(const FCharacterMoveResponsePackedBits& PackedBits)
 {
+#if CSV_PROFILER_STATS
+	if (CharacterMovementCVars::bClientRecordMovePackedRpcStatsToCsv)
+	{
+		// track the amount of MoveResponsePacked RPCs the client receives from the server
+		CSV_CUSTOM_STAT(CharacterMovementDetailed, Client_ReceiveMovePackedResponse_RPC, 1, ECsvCustomStatOp::Accumulate);
+
+		// track the size of the RPC data received from the server
+		CSV_CUSTOM_STAT(CharacterMovementDetailed, Client_ReceiveMovePackedResponse_RPC_NumBits, PackedBits.DataBits.Num(), ECsvCustomStatOp::Accumulate);
+	}
+#endif //CSV_PROFILER_STATS
+
 	if (!HasValidData() || !IsActive())
 	{
 		return;
@@ -10986,6 +11028,20 @@ void UCharacterMovementComponent::ClientAdjustRootMotionPosition(float TimeStamp
 
 void UCharacterMovementComponent::OnClientCorrectionReceived(FNetworkPredictionData_Client_Character& ClientData, float TimeStamp, FVector NewLocation, FVector NewVelocity, UPrimitiveComponent* NewBase, FName NewBaseBoneName, bool bHasBase, bool bBaseRelativePosition, uint8 ServerMovementMode, FVector ServerGravityDirection)
 {
+#if CSV_PROFILER_STATS
+	const bool bIsRecordingDetailedCsvStats = FCsvProfiler::Get()->IsCategoryEnabled(CSV_CATEGORY_INDEX(CharacterMovementDetailed));
+	if (bIsRecordingDetailedCsvStats)
+	{
+		if (CharacterMovementCVars::bClientRecordNetCorrectionDistanceToCsv)
+		{
+			const FVector ClientLocAtCorrectedMove = ClientData.LastAckedMove.IsValid() ? ClientData.LastAckedMove->SavedLocation : UpdatedComponent->GetComponentLocation();
+			const FVector LocDiff = ClientLocAtCorrectedMove - NewLocation;
+			const double LocDiffSize = LocDiff.Size();
+			CSV_CUSTOM_STAT(CharacterMovementDetailed, Correction_Distance, LocDiffSize, ECsvCustomStatOp::Set);
+		}
+	}
+#endif //CSV_PROFILER_STATS
+
 #if !UE_BUILD_SHIPPING
 	if (CharacterMovementCVars::NetShowCorrections != 0)
 	{
@@ -11019,7 +11075,7 @@ void UCharacterMovementComponent::OnClientCorrectionReceived(FNetworkPredictionD
 													  *NewVelocity.ToCompactString(), *Velocity.ToCompactString(), *VelocityCorrection.ToCompactString(), TimeStamp);
 		RootMotionSourceDebug::PrintOnScreen(*CharacterOwner, AdjustedDebugString);
 	}
-#endif
+#endif //ROOT_MOTION_DEBUG
 }
 
 
