Networking
Trivial code and comment fixes

#rnx

diff --git a/Engine/Source/Runtime/Engine/Private/DataReplication.cpp b/Engine/Source/Runtime/Engine/Private/DataReplication.cpp
--- a/Engine/Source/Runtime/Engine/Private/DataReplication.cpp
+++ b/Engine/Source/Runtime/Engine/Private/DataReplication.cpp
@@ -2242,10 +2242,7 @@ bool FObjectReplicator::ReadyForDormancy(bool bSuppressLogs)
 	// Can't go dormant until last update produced no new property updates
 	if (!bLastUpdateEmpty)
 	{
-		if (!bSuppressLogs)
-		{
-			UE_LOG(LogRepTraffic, Verbose, TEXT("    [%d] Not ready for dormancy. bLastUpdateEmpty = false"), OwningChannel->ChIndex);
-		}
+		UE_CLOG(!bSuppressLogs, LogRepTraffic, Verbose, TEXT("    [%d] Not ready for dormancy. bLastUpdateEmpty = false"), OwningChannel->ChIndex);
 
 		return false;
 	}
@@ -2278,10 +2275,7 @@ bool FObjectReplicator::ReadyForDormancy(bool bSuppressLogs)
 		{
 			if (Retirement.Next != nullptr)
 			{
-				if (!bSuppressLogs)
-				{
-					UE_LOG(LogRepTraffic, Verbose, TEXT("    [%d] OutAckPacketId: %d First: %d Last: %d "), OwningChannel->ChIndex, OwningChannel->Connection->OutAckPacketId, Retirement.OutPacketIdRange.First, Retirement.OutPacketIdRange.Last);
-				}
+				UE_CLOG(!bSuppressLogs, LogRepTraffic, Verbose, TEXT("    [%d] Not ready for dormancy. OutAckPacketId: %d First: %d Last: %d "), OwningChannel->ChIndex, OwningChannel->Connection->OutAckPacketId, Retirement.OutPacketIdRange.First, Retirement.OutPacketIdRange.Last);
 				return false;
 			}
 		}
diff --git a/Engine/Source/Runtime/Engine/Private/NetConnection.cpp b/Engine/Source/Runtime/Engine/Private/NetConnection.cpp
--- a/Engine/Source/Runtime/Engine/Private/NetConnection.cpp
+++ b/Engine/Source/Runtime/Engine/Private/NetConnection.cpp
@@ -2094,7 +2094,7 @@ void UNetConnection::PreTickDispatch()
 
 	if (!IsReplay())
 	{
-		double LastTickDispatchRealtime = Driver->LastTickDispatchRealtime;
+		const double LastTickDispatchRealtime = Driver->LastTickDispatchRealtime;
 
 		if (bIsServer && RPCDoS.IsValid())
 		{
diff --git a/Engine/Source/Runtime/Engine/Private/NetDriver.cpp b/Engine/Source/Runtime/Engine/Private/NetDriver.cpp
--- a/Engine/Source/Runtime/Engine/Private/NetDriver.cpp
+++ b/Engine/Source/Runtime/Engine/Private/NetDriver.cpp
@@ -2365,21 +2365,22 @@ void UNetDriver::TickDispatch( float DeltaTime )
 {
 	SendCycles=0;
 
-	const double CurrentRealtime = FPlatformTime::Seconds();
-
-	const float DeltaRealtime = CurrentRealtime - LastTickDispatchRealtime;
-
-	LastTickDispatchRealtime = CurrentRealtime;
+	// Manage realtime values
+	{
+		const double CurrentRealtime = FPlatformTime::Seconds();
+		LastTickDispatchRealtime = CurrentRealtime;
 
-	// Check to see if too much time is passing between ticks
-	// Setting this to somewhat large value for now, but small enough to catch blocking calls that are causing timeouts
-	const float TickLogThreshold = 5.0f;
+		// Check to see if too much time is passing between ticks
+		// Setting this to somewhat large value for now, but small enough to catch blocking calls that are causing timeouts
+		constexpr float TickLogThreshold = 5.0f;
 
-	bDidHitchLastFrame = (DeltaTime > TickLogThreshold || DeltaRealtime > TickLogThreshold);
+		const float DeltaRealtime = CurrentRealtime - LastTickDispatchRealtime;
+		bDidHitchLastFrame = (DeltaTime > TickLogThreshold || DeltaRealtime > TickLogThreshold);
 
-	if (bDidHitchLastFrame)
-	{
-		UE_LOG( LogNet, Log, TEXT( "UNetDriver::TickDispatch: Very long time between ticks. DeltaTime: %2.2f, Realtime: %2.2f. %s" ), DeltaTime, DeltaRealtime, *GetName() );
+		if (bDidHitchLastFrame)
+		{
+			UE_LOG( LogNet, Log, TEXT( "UNetDriver::TickDispatch: Very long time between ticks. DeltaTime: %2.2f, Realtime: %2.2f. %s" ), DeltaTime, DeltaRealtime, *GetName() );
+		}
 	}
 
 	// Get new time.
@@ -2423,7 +2424,7 @@ void UNetDriver::TickDispatch( float DeltaTime )
 
 			for (const FDisconnectedClient& CurElement : RecentlyDisconnectedClients)
 			{
-				if ((CurrentRealtime - CurElement.DisconnectTime) >= RecentlyDisconnectedTrackingTime)
+				if ((LastTickDispatchRealtime - CurElement.DisconnectTime) >= RecentlyDisconnectedTrackingTime)
 				{
 					verify(MappedClientConnections.Remove(CurElement.Address) == 1);
 
@@ -5017,7 +5018,7 @@ void UNetDriver::ServerReplicateActors_BuildConsiderList( TArray<FNetworkObjectI
 {
 	SCOPE_CYCLE_COUNTER( STAT_NetConsiderActorsTime );
 
-	UE_LOG( LogNetTraffic, Log, TEXT( "ServerReplicateActors_BuildConsiderList, Building ConsiderList %4.2f" ), World->GetTimeSeconds() );
+	UE_LOG( LogNetTraffic, Log, TEXT( "ServerReplicateActors_BuildConsiderList, Building ConsiderList at WorldTime: %f ServerTickTime: %f" ), World->GetTimeSeconds(), ServerTickTime );
 
 	int32 NumInitiallyDormant = 0;
 
@@ -5031,6 +5032,7 @@ void UNetDriver::ServerReplicateActors_BuildConsiderList( TArray<FNetworkObjectI
 
 		if ( !ActorInfo->bPendingNetUpdate && World->TimeSeconds <= ActorInfo->NextUpdateTime )
 		{
+			UE_LOG(LogNetTraffic, VeryVerbose, TEXT("Skipping actor: %s. bPendingNetUpdate: %s | NextUpdateTime: %f (currently %f)"), *GetNameSafe(ActorInfo->Actor), ActorInfo->bPendingNetUpdate?TEXT("true"):TEXT("false"),ActorInfo->NextUpdateTime, World->TimeSeconds);
 			continue;		// It's not time for this actor to perform an update, skip it
 		}
 
diff --git a/Engine/Source/Runtime/Net/Core/Public/Net/Core/NetBitArray.h b/Engine/Source/Runtime/Net/Core/Public/Net/Core/NetBitArray.h
--- a/Engine/Source/Runtime/Net/Core/Public/Net/Core/NetBitArray.h
+++ b/Engine/Source/Runtime/Net/Core/Public/Net/Core/NetBitArray.h
@@ -328,7 +328,7 @@ public:
 	/** Return true if equal including BitCount and padding bits */
 	bool operator==(const FNetBitArrayView& Other) const;
 
-	/** Return if a specified bit is set or not */
+	/** Return if a specific bit is set or not */
 	inline bool IsBitSet(uint32 Index) const { return GetBit(Index); }
 
 	/** Returns true if any bit is set in the bitset Note: Padding bits in storage are expected to be zero. */
