Fixed blend parameter overrides not coming through

diff --git a/Engine/Source/Runtime/Engine/Private/Materials/MaterialExpressions.cpp b/Engine/Source/Runtime/Engine/Private/Materials/MaterialExpressions.cpp
--- a/Engine/Source/Runtime/Engine/Private/Materials/MaterialExpressions.cpp
+++ b/Engine/Source/Runtime/Engine/Private/Materials/MaterialExpressions.cpp
@@ -17953,7 +17953,7 @@ void FMaterialLayersFunctionsRuntimeData::PostSerialize(const FArchive& Ar)
 			// First add an extra empty Blend to match the number of layers
 			while (Blends.Num() < Layers.Num())
 			{
-				Blends.InsertDefaulted(0); // normally we should go through this only once
+				Blends.AddDefaulted(); // normally we should go through this only once
 			}
 
 			Layers.AddDefaulted();
@@ -17966,7 +17966,7 @@ void FMaterialLayersFunctionsRuntimeData::PostSerialize(const FArchive& Ar)
 			for (int32 l = 0; l < NumLegacyLayers; ++l)
 			{
 				// And then ONE MORE node to be the attributes first child
-				Tree.AddNode({ l,  l}, LayerNodeId);
+				Tree.AddNode({ l,  l == 0 ? NumLegacyLayers - 1 : l-1}, LayerNodeId);
 			}
 		}
 	}
