Added Raytracing metadata to CSV profiling data

diff --git a/Engine/Source/Runtime/Engine/Private/Rendering/RayTracingGeometryManager.cpp b/Engine/Source/Runtime/Engine/Private/Rendering/RayTracingGeometryManager.cpp
--- a/Engine/Source/Runtime/Engine/Private/Rendering/RayTracingGeometryManager.cpp
+++ b/Engine/Source/Runtime/Engine/Private/Rendering/RayTracingGeometryManager.cpp
@@ -16,6 +16,8 @@
 
 #include "Math/UnitConversion.h"
 
+#include "ProfilingDebugging/CsvProfiler.h"
+
 #if RHI_RAYTRACING
 
 static bool bHasRayTracingEnableChanged = false;
@@ -127,6 +129,16 @@ CSV_DEFINE_CATEGORY(RayTracingGeometry, true);
 FRayTracingGeometryManager::FRayTracingGeometryManager()
 {
 	StreamingRequests.SetNum(GRayTracingStreamingMaxPendingRequests);
+
+#if CSV_PROFILER_STATS
+	if (FCsvProfiler* CSVProfiler = FCsvProfiler::Get())
+	{
+		CSVProfiler->OnCSVProfileStart().AddLambda([]()
+		{
+			CSV_METADATA(TEXT("RayTracing"), IsRayTracingEnabled() ? TEXT("1") : TEXT("0"));
+		});
+	}
+#endif
 }
 
 FRayTracingGeometryManager::~FRayTracingGeometryManager()
