Make sure the bIsInstanceDataApplyCompleted flag is reset in all control flow paths

diff --git a/Engine/Source/Runtime/Engine/Private/InstancedStaticMesh.cpp b/Engine/Source/Runtime/Engine/Private/InstancedStaticMesh.cpp
--- a/Engine/Source/Runtime/Engine/Private/InstancedStaticMesh.cpp
+++ b/Engine/Source/Runtime/Engine/Private/InstancedStaticMesh.cpp
@@ -2126,6 +2126,11 @@ void UInstancedStaticMeshComponent::ApplyComponentInstanceData(FInstancedStaticM
 #if WITH_EDITOR
 	check(InstancedMeshData);
 
+	ON_SCOPE_EXIT
+	{		
+		bIsInstanceDataApplyCompleted = true;	
+	};
+
 	if (GetStaticMesh() != InstancedMeshData->StaticMesh)
 	{
 		return;
@@ -2176,8 +2181,6 @@ void UInstancedStaticMeshComponent::ApplyComponentInstanceData(FInstancedStaticM
 
 	bHasPerInstanceHitProxies = InstancedMeshData->bHasPerInstanceHitProxies;
 
-	bIsInstanceDataApplyCompleted = true;
-
 	// TODO: restore ID mapping either from the serialized stuff, or the InstancedMeshData
 	PrimitiveInstanceDataManager.Invalidate(PerInstanceSMData.Num());
 #endif
