Fixed assert in SkyLight capture on using static feature level before RHI was initialized

diff --git a/Engine/Source/Runtime/Engine/Private/Components/SkyLightComponent.cpp b/Engine/Source/Runtime/Engine/Private/Components/SkyLightComponent.cpp
--- a/Engine/Source/Runtime/Engine/Private/Components/SkyLightComponent.cpp
+++ b/Engine/Source/Runtime/Engine/Private/Components/SkyLightComponent.cpp
@@ -92,7 +92,7 @@ static bool FeatureLevelSupportsRealTimeReflectionCapture(FStaticFeatureLevel Fe
 
 void OnChangeSkylightRealTimeReflectionCapture(IConsoleVariable* Var)
 {
-	if (!FeatureLevelSupportsRealTimeReflectionCapture(GMaxRHIFeatureLevel))
+	if (!GIsRHIInitialized || !FeatureLevelSupportsRealTimeReflectionCapture(GMaxRHIFeatureLevel))
 	{
 		return;
 	}
