[Shaders] Add alternative code path to detect memory limit violation of SCW processes to support suspending/resuming of workers on POSIX/Wine.

Wine only implements a subset of the Job Object API, so we run an alternative code path in that environment in which we estimate the total amount of memory committed for the SCW processes.
This uses the existing function GetExternalWorkerMemoryUsage() which accumulates the virtual memory of all worker processes and takes advantage of the already opened process handles.

#rnx

diff --git a/Engine/Source/Runtime/Engine/Private/ShaderCompiler/ShaderCompilerThreadRunnable.cpp b/Engine/Source/Runtime/Engine/Private/ShaderCompiler/ShaderCompilerThreadRunnable.cpp
--- a/Engine/Source/Runtime/Engine/Private/ShaderCompiler/ShaderCompilerThreadRunnable.cpp
+++ b/Engine/Source/Runtime/Engine/Private/ShaderCompiler/ShaderCompilerThreadRunnable.cpp
@@ -154,7 +154,11 @@ void FShaderCompileThreadRunnableBase::StartThread()
 
 FShaderCompileThreadRunnable::FShaderCompileThreadRunnable(FShaderCompilingManager* InManager)
 	: FShaderCompileThreadRunnableBase(InManager)
-	, LastCheckForWorkersTime(0)
+#if PLATFORM_WINDOWS
+	, bEstimateCommittedMemory(FPlatformMisc::IsWine()) // Use alternative code path to estimate memory when we're running on POSIX/Wine instead of a real Windows host system
+#else
+	, bEstimateCommittedMemory(true) // Use alternative code path to estimate memory when we're running on POSIX
+#endif
 {
 	for (uint32 WorkerIndex = 0; WorkerIndex < Manager->NumShaderCompilingThreads; WorkerIndex++)
 	{
@@ -1031,7 +1035,7 @@ void FShaderCompileThreadRunnable::CheckMemoryLimitViolation()
 
 		// Check if memory limit has been exceeded
 		FJobObjectLimitationInfo LimitInfo;
-		if (GSCWResourceRestrictedJobObject.QueryLimitViolationStatus(LimitInfo))
+		if (QueryMemoryLimitViolationStatus(LimitInfo))
 		{
 			MemoryMonitoringState.LastTimeOfSuspeningOrResumingWorkers = CurrentTime;
 
@@ -1077,7 +1081,7 @@ void FShaderCompileThreadRunnable::CheckMemoryLimitViolation()
 		MemoryMonitoringState.LastTimeOfSuspeningOrResumingWorkers = CurrentTime;
 
 		FJobObjectLimitationInfo LimitInfo;
-		if (GSCWResourceRestrictedJobObject.QueryStatus(LimitInfo))
+		if (QueryMemoryStatus(LimitInfo))
 		{
 			// If we are below half of our memory limit, resume 50% of available workers.
 			// This approach suspends workers from 100% to 50% and then resumes them back up to 75%.
@@ -1110,6 +1114,45 @@ void FShaderCompileThreadRunnable::CheckMemoryLimitViolation()
 	}
 }
 
+bool FShaderCompileThreadRunnable::QueryMemoryStatus(FJobObjectLimitationInfo& OutInfo)
+{
+	if (bEstimateCommittedMemory)
+	{
+		const FShaderCompileMemoryUsage MemoryUsage = GetExternalWorkerMemoryUsage();
+		if (MemoryUsage.VirtualMemory > 0)
+		{
+			OutInfo.MemoryLimit = static_cast<int64>(GShaderCompilerMemoryLimit) * 1024 * 1024;
+			OutInfo.MemoryUsed = MemoryUsage.VirtualMemory;
+			return true;
+		}
+		return false;
+	}
+	else
+	{
+		return GSCWResourceRestrictedJobObject.QueryStatus(OutInfo);
+	}
+}
+
+bool FShaderCompileThreadRunnable::QueryMemoryLimitViolationStatus(FJobObjectLimitationInfo& OutInfo)
+{
+	if (bEstimateCommittedMemory)
+	{
+		const FShaderCompileMemoryUsage MemoryUsage = GetExternalWorkerMemoryUsage();
+		const int64 MemoryLimitInBytes = static_cast<int64>(GShaderCompilerMemoryLimit) * 1024 * 1024;
+		if (MemoryUsage.VirtualMemory >= static_cast<uint64>(MemoryLimitInBytes))
+		{
+			OutInfo.MemoryLimit = MemoryLimitInBytes;
+			OutInfo.MemoryUsed = MemoryUsage.VirtualMemory;
+			return true;
+		}
+		return false;
+	}
+	else
+	{
+		return GSCWResourceRestrictedJobObject.QueryLimitViolationStatus(OutInfo);
+	}
+}
+
 bool FShaderCompileThreadRunnable::PrintWorkerMemoryUsage(bool bAllowToWaitForLock)
 {
 	if (bAllowToWaitForLock)
diff --git a/Engine/Source/Runtime/Engine/Public/ShaderCompiler.h b/Engine/Source/Runtime/Engine/Public/ShaderCompiler.h
--- a/Engine/Source/Runtime/Engine/Public/ShaderCompiler.h
+++ b/Engine/Source/Runtime/Engine/Public/ShaderCompiler.h
@@ -51,6 +51,8 @@ bool AreShaderErrorsFatal();
 extern ENGINE_API bool IsShaderJobCacheDDCEnabled();
 extern ENGINE_API bool IsMaterialMapDDCEnabled();
 
+struct FJobObjectLimitationInfo;
+
 struct FShaderJobCacheStoredOutput;
 class FShaderJobCache;
 
@@ -199,12 +201,14 @@ class FShaderCompileThreadRunnable : public FShaderCompileThreadRunnableBase
 	friend class FShaderCompilingManager;
 private:
 
+	const bool bEstimateCommittedMemory = false; // Must be true on POSIX/Wine where only a small subset of the Job Object functionality is implemented
+
 	/** Information about the active workers that this thread is tracking. */
 	TArray<TUniquePtr<struct FShaderCompileWorkerInfo>> WorkerInfos;
 	mutable FCriticalSection WorkerInfosLock;
 
 	/** Tracks the last time that this thread checked if the workers were still active. */
-	double LastCheckForWorkersTime;
+	double LastCheckForWorkersTime = 0.0;
 
 	/** Whether to read/write files for SCW in parallel (can help situations when this takes too long for a number of reasons) */
 	bool bParallelizeIO = false;
@@ -291,6 +295,12 @@ private:
 
 	/** Checks it the memory limit for shader compile workers has been exceeded and suspend workers as needed. */
 	void CheckMemoryLimitViolation();
+
+	/** Queries the memory status of all worker processes. Either uses FResourceRestrictedJobObject or QueryEstimatedCommittedMemory() when running on POSIX/Wine. */
+	bool QueryMemoryStatus(FJobObjectLimitationInfo& OutInfo);
+
+	/** Queries the status if the job object for all worker processes has violated the memory limitation. Either uses FResourceRestrictedJobObject or QueryEstimatedCommittedMemory() when running on POSIX/Wine. */
+	bool QueryMemoryLimitViolationStatus(FJobObjectLimitationInfo& OutInfo);
 };
 
 class FShaderCompileUtilities
