Level and WorldParition objects being renamed to avoid name conflicts are now outered to the transient package instead of generating unique FNames polluting the FName table. Saves memory when streaming is used heavily as unique names are generated permanently bloating FName storage if UE_FNAME_OUTLINE_NUMBER is enabled

diff --git a/Engine/Source/Runtime/Engine/Private/Level.cpp b/Engine/Source/Runtime/Engine/Private/Level.cpp
--- a/Engine/Source/Runtime/Engine/Private/Level.cpp
+++ b/Engine/Source/Runtime/Engine/Private/Level.cpp
@@ -548,7 +548,7 @@ void ULevel::CleanupLevel(bool bCleanupResources, bool bUnloadFromEditor)
 			if (bTrashPackage && (InPackage != GetTransientPackage()))
 			{
 				// Rename package to make sure it won't be reused
-				FName NewPackageName = MakeUniqueObjectName(nullptr, UPackage::StaticClass(), FName(*FString::Printf(TEXT("%s_Trashed"), *InPackage->GetName())));
+				FName NewPackageName = MakeUniqueObjectName(GetTransientPackage(), UPackage::StaticClass(), InPackage->GetFName());
 				InPackage->Rename(*NewPackageName.ToString(), nullptr, REN_DontCreateRedirectors | REN_NonTransactional | REN_DoNotDirty);
 			}
 		}
diff --git a/Engine/Source/Runtime/Engine/Private/LevelStreamingGCHelper.cpp b/Engine/Source/Runtime/Engine/Private/LevelStreamingGCHelper.cpp
--- a/Engine/Source/Runtime/Engine/Private/LevelStreamingGCHelper.cpp
+++ b/Engine/Source/Runtime/Engine/Private/LevelStreamingGCHelper.cpp
@@ -189,7 +189,7 @@ void FLevelStreamingGCHelper::PrepareStreamedOutLevelForGC(ULevel* InLevel)
 			for (UPackage* Package : Packages)
 			{
 				FCoreUObjectInternalDelegates::GetOnLeakedPackageRenameDelegate().Broadcast(Package);
-				const FName NewName = MakeUniqueObjectName(nullptr, UPackage::StaticClass(), FName(FString::Printf(TEXT("%s_PreparedForGC"), *Package->GetFName().GetPlainNameString())));
+				const FName NewName = MakeUniqueObjectName(GetTransientPackage(), UPackage::StaticClass(), Package->GetFName());
 				Package->Rename(*NewName.ToString(), nullptr, REN_DontCreateRedirectors | REN_NonTransactional);
 			}
 
diff --git a/Engine/Source/Runtime/Engine/Private/WorldPartition/WorldPartitionLevelHelper.cpp b/Engine/Source/Runtime/Engine/Private/WorldPartition/WorldPartitionLevelHelper.cpp
--- a/Engine/Source/Runtime/Engine/Private/WorldPartition/WorldPartitionLevelHelper.cpp
+++ b/Engine/Source/Runtime/Engine/Private/WorldPartition/WorldPartitionLevelHelper.cpp
@@ -307,7 +307,7 @@ void FWorldPartitionLevelHelper::MoveExternalActorsToLevel(const TArray<FWorldPa
 
 				// Trash this package to guarantee that any potential future load of this actor won't find the old empty package
 				// @todo_ow: Decide if we want to support actor reloads during cook. If not, remove this code, detect the reload and report an error.
-				FName NewPackageName = MakeUniqueObjectName(nullptr, UPackage::StaticClass(), FName(*FString::Printf(TEXT("%s_Trashed"), *ActorPackage->GetName())));
+				FName NewPackageName = MakeUniqueObjectName(GetTransientPackage(), UPackage::StaticClass(), ActorPackage->GetFName());
 				ActorPackage->Rename(*NewPackageName.ToString(), nullptr, REN_DontCreateRedirectors | REN_NonTransactional | REN_DoNotDirty);
 			}
 
diff --git a/Engine/Source/Runtime/Engine/Private/WorldPartition/WorldPartitionPackageHelper.cpp b/Engine/Source/Runtime/Engine/Private/WorldPartition/WorldPartitionPackageHelper.cpp
--- a/Engine/Source/Runtime/Engine/Private/WorldPartition/WorldPartitionPackageHelper.cpp
+++ b/Engine/Source/Runtime/Engine/Private/WorldPartition/WorldPartitionPackageHelper.cpp
@@ -20,7 +20,7 @@ void FWorldPartitionPackageHelper::UnloadPackage(UPackage* InPackage)
 		}, false);
 
 		// Rename so it isn't found again
-		FName NewUniqueTrashName = MakeUniqueObjectName(nullptr, UPackage::StaticClass(), FName(*FString::Printf(TEXT("%s_Trashed"), *InPackage->GetName())));
+		FName NewUniqueTrashName = MakeUniqueObjectName(GetTransientPackage(), UPackage::StaticClass(), InPackage->GetFName());
 		InPackage->Rename(*NewUniqueTrashName.ToString(), nullptr, REN_DontCreateRedirectors | REN_NonTransactional | REN_DoNotDirty);
 	};
 
diff --git a/Engine/Source/Runtime/Engine/Private/WorldPartition/WorldPartitionRuntimeSpatialHash.cpp b/Engine/Source/Runtime/Engine/Private/WorldPartition/WorldPartitionRuntimeSpatialHash.cpp
--- a/Engine/Source/Runtime/Engine/Private/WorldPartition/WorldPartitionRuntimeSpatialHash.cpp
+++ b/Engine/Source/Runtime/Engine/Private/WorldPartition/WorldPartitionRuntimeSpatialHash.cpp
@@ -2123,7 +2123,7 @@ bool UWorldPartitionRuntimeSpatialHash::UnregisterWorldAssetStreaming(const FGui
 		{
 			auto TrashObject = [](UObject* InObject)
 			{
-				FName NewUniqueTrashName = MakeUniqueObjectName(InObject->GetOuter(), InObject->GetClass(), FName(*FString::Printf(TEXT("%s_Trashed"), *InObject->GetName())));
+				FName NewUniqueTrashName = MakeUniqueObjectName(GetTransientPackage(), InObject->GetClass(), InObject->GetFName());
 				InObject->Rename(*NewUniqueTrashName.ToString(), nullptr, REN_DontCreateRedirectors | REN_NonTransactional | REN_DoNotDirty);
 			};
 
