[World Partition] Make sure to unregister OnLevelShown/OnLevelHidden events when destroying WorldPartitionRuntimeLevelStreamingCell objects.

diff --git a/Engine/Source/Runtime/Engine/Private/WorldPartition/WorldPartitionRuntimeLevelStreamingCell.cpp b/Engine/Source/Runtime/Engine/Private/WorldPartition/WorldPartitionRuntimeLevelStreamingCell.cpp
--- a/Engine/Source/Runtime/Engine/Private/WorldPartition/WorldPartitionRuntimeLevelStreamingCell.cpp
+++ b/Engine/Source/Runtime/Engine/Private/WorldPartition/WorldPartitionRuntimeLevelStreamingCell.cpp
@@ -26,6 +26,17 @@ UWorldPartitionRuntimeLevelStreamingCell::UWorldPartitionRuntimeLevelStreamingCe
 	, LevelStreaming(nullptr)
 {}
 
+void UWorldPartitionRuntimeLevelStreamingCell::BeginDestroy()
+{
+	if (LevelStreaming)
+	{
+		LevelStreaming->OnLevelShown.RemoveAll(this);
+		LevelStreaming->OnLevelHidden.RemoveAll(this);
+	}
+
+	Super::BeginDestroy();
+}
+
 EWorldPartitionRuntimeCellState UWorldPartitionRuntimeLevelStreamingCell::GetCurrentState() const
 {
 	if (LevelStreaming)
diff --git a/Engine/Source/Runtime/Engine/Public/WorldPartition/WorldPartitionRuntimeLevelStreamingCell.h b/Engine/Source/Runtime/Engine/Public/WorldPartition/WorldPartitionRuntimeLevelStreamingCell.h
--- a/Engine/Source/Runtime/Engine/Public/WorldPartition/WorldPartitionRuntimeLevelStreamingCell.h
+++ b/Engine/Source/Runtime/Engine/Public/WorldPartition/WorldPartitionRuntimeLevelStreamingCell.h
@@ -17,6 +17,10 @@ class UWorldPartitionRuntimeLevelStreamingCell : public UWorldPartitionRuntimeCe
 {
 	GENERATED_UCLASS_BODY()
 
+	//~Begin UObject interface
+	virtual void BeginDestroy() override;
+	//~End UObject interface
+
 	//~Begin UWorldPartitionRuntimeCell interface
 	ENGINE_API virtual void Load() const override;
 	ENGINE_API virtual void Unload() const override;
