Alter behaviour of MaterialMayModifyMeshPosition to capture all instances of when a material may output a vertex modification to avoid skipping Depth and Velocity shader compilations for materials (which meant the depth buffer did not apply WPO whilst the base pass did, leading to a Depth/Color mismatch)

diff --git a/Engine/Source/Runtime/Engine/Private/Materials/MaterialShared.cpp b/Engine/Source/Runtime/Engine/Private/Materials/MaterialShared.cpp
--- a/Engine/Source/Runtime/Engine/Private/Materials/MaterialShared.cpp
+++ b/Engine/Source/Runtime/Engine/Private/Materials/MaterialShared.cpp
@@ -1193,7 +1193,17 @@ bool FMaterial::MaterialMayModifyMeshPosition() const
 {
 	// Conservative estimate when called before material translation has occurred. 
 	// This function is only intended for use in deciding whether or not shader permutations are required.
-	return HasVertexPositionOffsetConnected() || HasPixelDepthOffsetConnected() || HasDisplacementConnected();
+	bool bMayModifyMeshPosition = false;
+	if (IsInRenderingThread())
+	{
+		bMayModifyMeshPosition = MaterialUsesWorldPositionOffset_RenderThread() || MaterialUsesDisplacement_RenderThread() || MaterialUsesPixelDepthOffset_RenderThread();
+	}
+	else if (IsInGameThread())
+	{
+		bMayModifyMeshPosition = MaterialUsesWorldPositionOffset_GameThread() || MaterialUsesDisplacement_GameThread() || MaterialUsesPixelDepthOffset_GameThread();
+	}
+
+	return bMayModifyMeshPosition || HasVertexPositionOffsetConnected() || HasPixelDepthOffsetConnected() || HasDisplacementConnected();
 }
 
 bool FMaterial::MaterialUsesPixelDepthOffset_GameThread() const
