Morph Target removal: invalidate the DDC cache on removal if MorphTargets array and SourceModels's MeshDescription are not synchronized.

#rnx

diff --git a/Engine/Source/Runtime/Engine/Private/SkeletalMesh.cpp b/Engine/Source/Runtime/Engine/Private/SkeletalMesh.cpp
--- a/Engine/Source/Runtime/Engine/Private/SkeletalMesh.cpp
+++ b/Engine/Source/Runtime/Engine/Private/SkeletalMesh.cpp
@@ -4227,6 +4227,12 @@ bool USkeletalMesh::RemoveMorphTargets(TConstArrayView<FName> InMorphTargetNames
 					MeshAttributes.UnregisterMorphTargetAttribute(MorphTargetName);
 					CommitMeshDescription(LODIndex);
 				}
+				else
+				{
+					// this means that MorphTargets and MeshDescription are not synchronized (which should not happen)
+					// if the DDC is not invalidated, the MorphTargets array will be reset to its previous value in the next build. 
+					InvalidateDeriveDataCacheGUID();
+				}
 			}
 
 			UnregisterMorphTarget(MorphTarget);
