[Level Streaming] Added helper function FLevelStreamingGCHelper::TrashPackage to trash a package so it doesn't get reloaded when needed again.
This change makes sure to avoid creating new names for trashed packages by routing them all into the shared "TrashedPackage" name with a different number suffix.
This change will probably make UE_FNAME_OUTLINE_NUMBER use more memory in some situations.

diff --git a/Engine/Source/Runtime/Core/Public/UObject/UnrealNames.inl b/Engine/Source/Runtime/Core/Public/UObject/UnrealNames.inl
--- a/Engine/Source/Runtime/Core/Public/UObject/UnrealNames.inl
+++ b/Engine/Source/Runtime/Core/Public/UObject/UnrealNames.inl
@@ -204,6 +204,7 @@ REGISTER_NAME(258, Gzip)
 REGISTER_NAME(259, LZ4)
 REGISTER_NAME(260, Mobile)
 REGISTER_NAME(261, Oodle)
+REGISTER_NAME(262, TrashedPackage)
 
 // Online
 REGISTER_NAME(280,DGram)
diff --git a/Engine/Source/Runtime/Engine/Classes/Engine/LevelStreamingGCHelper.h b/Engine/Source/Runtime/Engine/Classes/Engine/LevelStreamingGCHelper.h
--- a/Engine/Source/Runtime/Engine/Classes/Engine/LevelStreamingGCHelper.h
+++ b/Engine/Source/Runtime/Engine/Classes/Engine/LevelStreamingGCHelper.h
@@ -40,6 +40,11 @@ struct FLevelStreamingGCHelper
 	 */
 	static ENGINE_API int32 GetNumLevelsPendingPurge();
 
+	/**
+	 * Trash a package, making sure it will be reloaded if needed again by level streaming
+	 */
+	static ENGINE_API void TrashPackage(UPackage* InPackage);
+
 	/**
 	 * Allows FLevelStreamingGCHelper to be used in a commandlet.
 	 */
diff --git a/Engine/Source/Runtime/Engine/Private/Level.cpp b/Engine/Source/Runtime/Engine/Private/Level.cpp
--- a/Engine/Source/Runtime/Engine/Private/Level.cpp
+++ b/Engine/Source/Runtime/Engine/Private/Level.cpp
@@ -5,7 +5,7 @@ Level.cpp: Level-related functions
 =============================================================================*/
 
 #include "Engine/Level.h"
-
+#include "Engine/LevelStreamingGCHelper.h"
 #include "EngineLogs.h"
 #include "Algo/Copy.h"
 #include "Algo/ForEach.h"
@@ -548,8 +548,7 @@ void ULevel::CleanupLevel(bool bCleanupResources, bool bUnloadFromEditor)
 			if (bTrashPackage && (InPackage != GetTransientPackage()))
 			{
 				// Rename package to make sure it won't be reused
-				FName NewPackageName = MakeUniqueObjectName(nullptr, UPackage::StaticClass(), FName(*FString::Printf(TEXT("%s_Trashed"), *InPackage->GetName())));
-				InPackage->Rename(*NewPackageName.ToString(), nullptr, REN_DontCreateRedirectors | REN_NonTransactional | REN_DoNotDirty);
+				FLevelStreamingGCHelper::TrashPackage(InPackage);
 			}
 		}
 	};
diff --git a/Engine/Source/Runtime/Engine/Private/LevelStreamingGCHelper.cpp b/Engine/Source/Runtime/Engine/Private/LevelStreamingGCHelper.cpp
--- a/Engine/Source/Runtime/Engine/Private/LevelStreamingGCHelper.cpp
+++ b/Engine/Source/Runtime/Engine/Private/LevelStreamingGCHelper.cpp
@@ -189,8 +189,7 @@ void FLevelStreamingGCHelper::PrepareStreamedOutLevelForGC(ULevel* InLevel)
 			for (UPackage* Package : Packages)
 			{
 				FCoreUObjectInternalDelegates::GetOnLeakedPackageRenameDelegate().Broadcast(Package);
-				const FName NewName = MakeUniqueObjectName(nullptr, UPackage::StaticClass(), FName(FString::Printf(TEXT("%s_PreparedForGC"), *Package->GetFName().GetPlainNameString())));
-				Package->Rename(*NewName.ToString(), nullptr, REN_DontCreateRedirectors | REN_NonTransactional);
+				TrashPackage(Package);
 			}
 
 #if !WITH_EDITOR
@@ -260,3 +259,9 @@ int32 FLevelStreamingGCHelper::GetNumLevelsPendingPurge()
 {
 	return LevelsPendingUnload.Num() + NumberOfPreparedStreamedOutLevelsForGC;
 }
+
+void FLevelStreamingGCHelper::TrashPackage(UPackage* InPackage)
+{
+	const FName NewName = MakeUniqueObjectName(nullptr, UPackage::StaticClass(), NAME_TrashedPackage);
+	InPackage->Rename(*NewName.ToString(), nullptr, REN_DontCreateRedirectors | REN_NonTransactional | REN_DoNotDirty);
+}
diff --git a/Engine/Source/Runtime/Engine/Private/WorldPartition/RuntimeHashSet/WorldPartitionRuntimeHashSet.cpp b/Engine/Source/Runtime/Engine/Private/WorldPartition/RuntimeHashSet/WorldPartitionRuntimeHashSet.cpp
--- a/Engine/Source/Runtime/Engine/Private/WorldPartition/RuntimeHashSet/WorldPartitionRuntimeHashSet.cpp
+++ b/Engine/Source/Runtime/Engine/Private/WorldPartition/RuntimeHashSet/WorldPartitionRuntimeHashSet.cpp
@@ -779,7 +779,7 @@ bool UWorldPartitionRuntimeHashSet::UnregisterWorldAssetStreaming(const FGuid& I
 		{
 			auto TrashObject = [](UObject* InObject)
 			{
-				FName NewUniqueTrashName = MakeUniqueObjectName(InObject->GetOuter(), InObject->GetClass(), FName(*FString::Printf(TEXT("%s_Trashed"), *InObject->GetName())));
+				FName NewUniqueTrashName = MakeUniqueObjectName(InObject->GetOuter(), InObject->GetClass(), NAME_TrashedPackage);
 				InObject->Rename(*NewUniqueTrashName.ToString(), nullptr, REN_DontCreateRedirectors | REN_NonTransactional | REN_DoNotDirty);
 			};
 
diff --git a/Engine/Source/Runtime/Engine/Private/WorldPartition/WorldPartitionLevelHelper.cpp b/Engine/Source/Runtime/Engine/Private/WorldPartition/WorldPartitionLevelHelper.cpp
--- a/Engine/Source/Runtime/Engine/Private/WorldPartition/WorldPartitionLevelHelper.cpp
+++ b/Engine/Source/Runtime/Engine/Private/WorldPartition/WorldPartitionLevelHelper.cpp
@@ -1,9 +1,5 @@
 // Copyright Epic Games, Inc. All Rights Reserved.
 
-/*
- * WorldPartitionLevelHelper implementation
- */
-
 #include "WorldPartition/WorldPartitionLevelHelper.h"
 #include "Misc/PackageName.h"
 #include "WorldPartition/WorldPartitionPackageHelper.h"
@@ -13,8 +9,8 @@
 #include "WorldPartition/IWorldPartitionObjectResolver.h"
 
 #if WITH_EDITOR
-
 #include "Engine/Level.h"
+#include "Engine/LevelStreamingGCHelper.h"
 #include "Misc/Paths.h"
 #include "Model.h"
 #include "UnrealEngine.h"
@@ -29,7 +25,6 @@
 #include "ActorFolder.h"
 
 FUObjectAnnotationSparse<FWorldPartitionLevelHelper::FActorPropertyOverridesAnnotation, true> FWorldPartitionLevelHelper::ActorPropertyOverridesAnnotation;
-
 #endif
 
 bool FWorldPartitionResolveData::ResolveObject(UWorld* InWorld, const FSoftObjectPath& InObjectPath, UObject*& OutObject) const
@@ -307,8 +302,7 @@ void FWorldPartitionLevelHelper::MoveExternalActorsToLevel(const TArray<FWorldPa
 
 				// Trash this package to guarantee that any potential future load of this actor won't find the old empty package
 				// @todo_ow: Decide if we want to support actor reloads during cook. If not, remove this code, detect the reload and report an error.
-				FName NewPackageName = MakeUniqueObjectName(nullptr, UPackage::StaticClass(), FName(*FString::Printf(TEXT("%s_Trashed"), *ActorPackage->GetName())));
-				ActorPackage->Rename(*NewPackageName.ToString(), nullptr, REN_DontCreateRedirectors | REN_NonTransactional | REN_DoNotDirty);
+				FLevelStreamingGCHelper::TrashPackage(ActorPackage);
 			}
 
 			OutModifiedPackages.Add(ActorPackage);
diff --git a/Engine/Source/Runtime/Engine/Private/WorldPartition/WorldPartitionPackageHelper.cpp b/Engine/Source/Runtime/Engine/Private/WorldPartition/WorldPartitionPackageHelper.cpp
--- a/Engine/Source/Runtime/Engine/Private/WorldPartition/WorldPartitionPackageHelper.cpp
+++ b/Engine/Source/Runtime/Engine/Private/WorldPartition/WorldPartitionPackageHelper.cpp
@@ -4,8 +4,8 @@
 #include "UObject/Package.h"
 
 #if WITH_EDITOR
-
 #include "Engine/Level.h"
+#include "Engine/LevelStreamingGCHelper.h"
 #include "Engine/World.h"
 
 void FWorldPartitionPackageHelper::UnloadPackage(UPackage* InPackage)
@@ -20,8 +20,7 @@ void FWorldPartitionPackageHelper::UnloadPackage(UPackage* InPackage)
 		}, false);
 
 		// Rename so it isn't found again
-		FName NewUniqueTrashName = MakeUniqueObjectName(nullptr, UPackage::StaticClass(), FName(*FString::Printf(TEXT("%s_Trashed"), *InPackage->GetName())));
-		InPackage->Rename(*NewUniqueTrashName.ToString(), nullptr, REN_DontCreateRedirectors | REN_NonTransactional | REN_DoNotDirty);
+		FLevelStreamingGCHelper::TrashPackage(InPackage);
 	};
 
 	TrashPackage(InPackage);
diff --git a/Engine/Source/Runtime/Engine/Private/WorldPartition/WorldPartitionRuntimeSpatialHash.cpp b/Engine/Source/Runtime/Engine/Private/WorldPartition/WorldPartitionRuntimeSpatialHash.cpp
--- a/Engine/Source/Runtime/Engine/Private/WorldPartition/WorldPartitionRuntimeSpatialHash.cpp
+++ b/Engine/Source/Runtime/Engine/Private/WorldPartition/WorldPartitionRuntimeSpatialHash.cpp
@@ -2123,7 +2123,7 @@ bool UWorldPartitionRuntimeSpatialHash::UnregisterWorldAssetStreaming(const FGui
 		{
 			auto TrashObject = [](UObject* InObject)
 			{
-				FName NewUniqueTrashName = MakeUniqueObjectName(InObject->GetOuter(), InObject->GetClass(), FName(*FString::Printf(TEXT("%s_Trashed"), *InObject->GetName())));
+				FName NewUniqueTrashName = MakeUniqueObjectName(InObject->GetOuter(), InObject->GetClass(), NAME_TrashedPackage);
 				InObject->Rename(*NewUniqueTrashName.ToString(), nullptr, REN_DontCreateRedirectors | REN_NonTransactional | REN_DoNotDirty);
 			};
 
