[World Partition] Fix FWorldPartitionLevelHelper::AddActorContainerID to append the container name before the name number suffix to reduce the name table bloat.

diff --git a/Engine/Source/Runtime/Engine/Private/WorldPartition/WorldPartitionLevelHelper.cpp b/Engine/Source/Runtime/Engine/Private/WorldPartition/WorldPartitionLevelHelper.cpp
--- a/Engine/Source/Runtime/Engine/Private/WorldPartition/WorldPartitionLevelHelper.cpp
+++ b/Engine/Source/Runtime/Engine/Private/WorldPartition/WorldPartitionLevelHelper.cpp
@@ -46,8 +46,9 @@ bool FWorldPartitionResolveData::ResolveObject(UWorld* InWorld, const FSoftObjec
 
 FString FWorldPartitionLevelHelper::AddActorContainerID(const FActorContainerID& InContainerID, const FString& InActorName)
 {
-	return InActorName + TEXT("_") + InContainerID.ToShortString();
-}
+	const FName ActorName(*InActorName);
+	const FString ActorPlainName(ActorName.GetPlainNameString() + TEXT("_") + InContainerID.ToShortString());
+	return FName(*ActorPlainName, ActorName.GetNumber()).ToString();}
 
 FString FWorldPartitionLevelHelper::AddActorContainerIDToSubPathString(const FActorContainerID& InContainerID, const FString& InSubPathString)
 {
@@ -745,14 +746,14 @@ bool FWorldPartitionLevelHelper::LoadActorsInternal(FLoadActorsParams&& InParams
 					OuterWorld->GetSoftObjectPathMapping(SourceOuterWorldPath, DummyUnusedPath);
 
 					// Rename through UObject to avoid changing Actor's external packaging and folder properties
-					Actor->UObject::Rename(*FString::Printf(TEXT("%s_%s"), *Actor->GetName(), *PackageObjectMapping->ContainerID.ToShortString()), DestLevel, REN_NonTransactional | REN_DoNotDirty | REN_DontCreateRedirectors);
+					Actor->UObject::Rename(*AddActorContainerID(PackageObjectMapping->ContainerID, Actor->GetName()), DestLevel, REN_NonTransactional | REN_DoNotDirty | REN_DontCreateRedirectors);
 
 					// Handle child actors
 					Actor->ForEachComponent<UChildActorComponent>(true, [DestLevel = DestLevel, PackageObjectMapping](UChildActorComponent* ChildActorComponent)
 					{
 						if (AActor* ChildActor = ChildActorComponent->GetChildActor())
 						{
-							ChildActor->UObject::Rename(*FString::Printf(TEXT("%s_%s"), *ChildActor->GetName(), *PackageObjectMapping->ContainerID.ToShortString()), DestLevel, REN_NonTransactional | REN_DoNotDirty | REN_DontCreateRedirectors);
+							ChildActor->UObject::Rename(*AddActorContainerID(PackageObjectMapping->ContainerID, ChildActor->GetName()), DestLevel, REN_NonTransactional | REN_DoNotDirty | REN_DontCreateRedirectors);
 						}
 					});
 
