Change the Role validation check in actor spawn to an ensure, to stop it from crashing in the content browser thumbnail preview.
Role is not a transient variable so this could be caused by corrupting a blueprint asset and then saving.

diff --git a/Engine/Source/Runtime/Engine/Private/Actor.cpp b/Engine/Source/Runtime/Engine/Private/Actor.cpp
--- a/Engine/Source/Runtime/Engine/Private/Actor.cpp
+++ b/Engine/Source/Runtime/Engine/Private/Actor.cpp
@@ -3836,7 +3836,7 @@ void AActor::PostSpawnInitialize(FTransform const& UserSpawnTransform, AActor* I
 	CreationTime = (World ? World->GetTimeSeconds() : 0.f);
 
 	// Set network role.
-	check(GetLocalRole() == ROLE_Authority);
+	ensureMsgf(GetLocalRole() == ROLE_Authority, TEXT("Actor %s has an invalid Role and may be a corrupt asset!"), *GetFullName());
 	ExchangeNetRoles(bRemoteOwned);
 
 	// Set owner.
