[World Partition] Properly initialize actor descriptor bounds to zero when they are considered invalid to avoid having random values (no consequences, just non-deterministic behavior).

diff --git a/Engine/Source/Runtime/Engine/Private/WorldPartition/WorldPartitionActorDesc.cpp b/Engine/Source/Runtime/Engine/Private/WorldPartition/WorldPartitionActorDesc.cpp
--- a/Engine/Source/Runtime/Engine/Private/WorldPartition/WorldPartitionActorDesc.cpp
+++ b/Engine/Source/Runtime/Engine/Private/WorldPartition/WorldPartitionActorDesc.cpp
@@ -632,6 +632,11 @@ void FWorldPartitionActorDesc::Serialize(FArchive& Ar)
 				RuntimeBounds = FBox(BoundsLocation - BoundsExtent, BoundsLocation + BoundsExtent);
 				EditorBounds = RuntimeBounds;
 			}
+			else
+			{
+				RuntimeBounds.Init();
+				EditorBounds.Init();
+			}
 		}
 		else
 		{
