Avoid unnecessary TArray copy in USceneComponent::OnComponentDestroyed.

diff --git a/Engine/Source/Runtime/Engine/Private/Components/SceneComponent.cpp b/Engine/Source/Runtime/Engine/Private/Components/SceneComponent.cpp
--- a/Engine/Source/Runtime/Engine/Private/Components/SceneComponent.cpp
+++ b/Engine/Source/Runtime/Engine/Private/Components/SceneComponent.cpp
@@ -1219,7 +1219,7 @@ void USceneComponent::OnComponentDestroyed(bool bDestroyingHierarchy)
 			int32 ChildCount = AttachChildren.Num();
 
 			// We cache the actual children to put back after the detach process
-			TArray<USceneComponent*> CachedChildren;
+			TArray<TObjectPtr<USceneComponent>> CachedChildren;
 			CachedChildren.Reserve(ChildCount);
 
 			while (ChildCount > 0)
