Reduce memstack pressure during anim compression
Fix whitespace

diff --git a/Engine/Source/Runtime/Engine/Private/Animation/AnimCompressionTypes.cpp b/Engine/Source/Runtime/Engine/Private/Animation/AnimCompressionTypes.cpp
--- a/Engine/Source/Runtime/Engine/Private/Animation/AnimCompressionTypes.cpp
+++ b/Engine/Source/Runtime/Engine/Private/Animation/AnimCompressionTypes.cpp
@@ -528,17 +528,21 @@ void FCompressibleAnimData::ResampleAnimationTrackData(const FFrameRate& SampleR
 					TrackData.InternalTrackData.ScaleKeys.SetNumUninitialized(SampledKeys);
 				}
 			}
-			
+
+			FCompactPose Pose;
+			Pose.SetBoneContainer(&EvalContext.RequiredBones);
+
 			FBlendedCurve Curve;
 			Curve.InitFrom(EvalContext.RequiredBones);
+
 			UE::Anim::FStackAttributeContainer AttributeContainer;
-						
+
+			const FName RetargetTransformsSourceName = AnimSequence->GetRetargetTransformsSourceName();
+			const TArray<FTransform>& RetargetTransforms = AnimSequence->GetRetargetTransforms();
+
 			for (int32 FrameIndex = 0; FrameIndex < SampledKeys; ++FrameIndex)
 			{
-				UE::Anim::DataModel::FEvaluationContext EvaluationContext(FFrameTime(FrameIndex), SampleRate, AnimSequence->GetRetargetTransformsSourceName(), AnimSequence->GetRetargetTransforms());
-				
-				FCompactPose Pose;
-				Pose.SetBoneContainer(&EvalContext.RequiredBones);
+				UE::Anim::DataModel::FEvaluationContext EvaluationContext(FFrameTime(FrameIndex), SampleRate, RetargetTransformsSourceName, RetargetTransforms);
 
 				FAnimationPoseData PoseData(Pose, Curve, AttributeContainer);
 				DataModelInterface->Evaluate(PoseData, EvaluationContext);
@@ -893,7 +897,7 @@ void FCompressibleAnimData::FetchData(const ITargetPlatform* InPlatform)
 		TArray<FTrackToSkeletonMap> TempTrackToSkeletonMapTable;
 		TempTrackToSkeletonMapTable.Reserve(OriginalTrackNames.Num());
 		TempRawAnimationData.Reserve(OriginalTrackNames.Num());
-	FinalTrackNames.Reserve(ResampledTrackData.Num());	
+		FinalTrackNames.Reserve(ResampledTrackData.Num());	
 		TempAdditiveBaseAnimationData.Reserve(AdditiveBaseAnimationData.Num() ? AdditiveBaseAnimationData.Num() : 0);
 
 		// Include root bone track
@@ -907,17 +911,17 @@ void FCompressibleAnimData::FetchData(const ITargetPlatform* InPlatform)
 
 		const int32 NumTracks = RawAnimationData.Num();
 		for (int32 TrackIndex = 1; TrackIndex < NumTracks; ++TrackIndex)
-	{
-		const FRawAnimSequenceTrack& Track = RawAnimationData[TrackIndex];
-		// Try find correct bone index
-		const int32 BoneIndex = RefSkeleton.FindBoneIndex(OriginalTrackNames[TrackIndex]);
+		{
+			const FRawAnimSequenceTrack& Track = RawAnimationData[TrackIndex];
+			// Try find correct bone index
+			const int32 BoneIndex = RefSkeleton.FindBoneIndex(OriginalTrackNames[TrackIndex]);
 
 			const bool bValidBoneIndex = BoneIndex != INDEX_NONE;
 			const bool bValidAdditiveTrack = !IsRawTrackZeroAdditive(Track);
 
 			// Only include track if it contains valid (additive) data and its name corresponds to a bone on the skeleton
 			if ((!bIsAdditiveAnimation || bValidAdditiveTrack) && bValidBoneIndex)
-		{
+			{
 				FinalTrackNames.Add(OriginalTrackNames[TrackIndex]);
 				TempTrackToSkeletonMapTable.Add(TrackToSkeletonMapTable[TrackIndex]);
 				TempRawAnimationData.Add(RawAnimationData[TrackIndex]);
@@ -926,7 +930,7 @@ void FCompressibleAnimData::FetchData(const ITargetPlatform* InPlatform)
 				{
 					TempAdditiveBaseAnimationData.Add(AdditiveBaseAnimationData[TrackIndex]);
 				}
-		}
+			}
 		}
 
 		// Swap out maintained track data
@@ -941,7 +945,6 @@ void FCompressibleAnimData::FetchData(const ITargetPlatform* InPlatform)
 
 	if (bShouldPerformStripping)
 	{
-
 		const FName TargetPlatformName = InPlatform->GetPlatformInfo().IniPlatformName;
 		const TObjectPtr<class UVariableFrameStrippingSettings> VarFrameStrippingSettings = AnimSequence->VariableFrameStrippingSettings;
 		const FPerPlatformBool PlatformBool = VarFrameStrippingSettings->UseVariableFrameStripping;
