Add option to flush montages to prevent infinite accumulation of montage instances
-review-35048047 [REDACTED_HANDLE], [REDACTED_HANDLE], [REDACTED_HANDLE]

diff --git a/Engine/Source/Runtime/Engine/Classes/Animation/AnimInstance.h b/Engine/Source/Runtime/Engine/Classes/Animation/AnimInstance.h
--- a/Engine/Source/Runtime/Engine/Classes/Animation/AnimInstance.h
+++ b/Engine/Source/Runtime/Engine/Classes/Animation/AnimInstance.h
@@ -911,6 +911,10 @@ protected:
 
 public:
 
+	/** Flush completed montages when animation tick is paused. 
+		This can be used to prevent montages from infinitely stacking up when not ticking. */
+	ENGINE_API void ConditionalFlushCompletedMontages();
+
 	/**  Builds an inertialization request from the montage's group, provided duration and optional blend profile*/
 	ENGINE_API void RequestMontageInertialization(const UAnimMontage* Montage, float Duration, const UBlendProfile* BlendProfile = nullptr);
 
diff --git a/Engine/Source/Runtime/Engine/Classes/Components/SkinnedMeshComponent.h b/Engine/Source/Runtime/Engine/Classes/Components/SkinnedMeshComponent.h
--- a/Engine/Source/Runtime/Engine/Classes/Components/SkinnedMeshComponent.h
+++ b/Engine/Source/Runtime/Engine/Classes/Components/SkinnedMeshComponent.h
@@ -1583,12 +1583,14 @@ protected:
 	/** Flip the editable space base buffer */
 	ENGINE_API void FlipEditableSpaceBases();
 
+public:
 	/** 
 	 * Should tick  pose (by calling TickPose) in Tick
 	 * 
 	 * @return : return true if should Tick. false otherwise.
 	 */
 	ENGINE_API virtual bool ShouldTickPose() const;
+protected:
 
 	/**
 	 * Allocate Transform Data array including SpaceBases, BoneVisibilityStates 
diff --git a/Engine/Source/Runtime/Engine/Private/Animation/AnimInstance.cpp b/Engine/Source/Runtime/Engine/Private/Animation/AnimInstance.cpp
--- a/Engine/Source/Runtime/Engine/Private/Animation/AnimInstance.cpp
+++ b/Engine/Source/Runtime/Engine/Private/Animation/AnimInstance.cpp
@@ -97,6 +97,15 @@ static FAutoConsoleVariableRef CVarRK4SpringInterpolatorUpdateRate(TEXT("p.RK4Sp
 ENGINE_API int32 RK4_SPRING_INTERPOLATOR_MAX_ITER = 4;
 static FAutoConsoleVariableRef CVarRK4SpringInterpolatorMaxIter(TEXT("p.RK4SpringInterpolator.MaxIter"), RK4_SPRING_INTERPOLATOR_MAX_ITER, TEXT("RK4 Spring Interpolator's max number of iterations"), ECVF_Default);
 
+namespace MontageCVars
+{
+	static bool bFlushCompletedMontagesOnPlay = false;
+	static FAutoConsoleVariableRef CVarFlushCompletedMontagesOnPlay(
+		TEXT("a.Montage.FlushCompletedMontagesOnPlay"),
+		bFlushCompletedMontagesOnPlay,
+		TEXT("Whether we should flush all completed montages IMMEDIATELY when a new montage stomps the group. Use this to prevent accumulating montages when animation tick is paused, and firing all at once when we unpause."));
+}
+
 /////////////////////////////////////////////////////
 // UAnimInstance
 /////////////////////////////////////////////////////
@@ -1901,6 +1910,43 @@ void UAnimInstance::RequestSlotGroupInertialization(FName InSlotGroupName, float
 	GetProxyOnAnyThread<FAnimInstanceProxy>().GetSlotGroupInertializationRequestDataMap().FindOrAdd(InSlotGroupName) = Request;
 }
 
+void UAnimInstance::ConditionalFlushCompletedMontages()
+{
+	if (MontageInstances.IsEmpty())
+	{
+		// No montages, nothing to do.
+		return;
+	}
+
+	const USkeletalMeshComponent* MeshComp = GetSkelMeshComponent();
+	check(MeshComp);
+
+	// If we ticked this frame, then montages are not paused
+	// If we don't tick animation, then montages are paused
+	// If we don't tick the pose, then montages are paused
+	const bool bTickedThisFrame = MeshComp->PoseTickedThisFrame();
+	const bool bShouldTickAnimation = MeshComp->ShouldTickAnimation();
+	const bool bShouldTickPose = MeshComp->ShouldTickPose();
+	const bool bShouldFlush = !bTickedThisFrame && (!bShouldTickAnimation || !bShouldTickPose); 
+	if (bShouldFlush)
+	{
+		for (int32 InstanceIndex = MontageInstances.Num() - 1; InstanceIndex >= 0; InstanceIndex--)
+		{
+			FAnimMontageInstance* MontageInstance = MontageInstances[InstanceIndex];
+			if (MontageInstance && MontageInstance->IsValid() && MontageInstance->IsStopped() && MontageInstance->GetBlend().IsComplete())
+			{
+				// Need this to trigger Montage ended events.
+				MontageInstance->Terminate();
+
+				// Make sure we've cleared our references before deleting memory. Terminate might miss this call.
+				ClearMontageInstanceReferences(*MontageInstance);
+				delete MontageInstance;
+				MontageInstances.RemoveAt(InstanceIndex);
+			}
+		}
+	}
+}
+
 void UAnimInstance::RequestMontageInertialization(const UAnimMontage* Montage, const FInertializationRequest& Request)
 {
 	if (Montage)
@@ -2197,6 +2243,11 @@ float UAnimInstance::Montage_PlayInternal(UAnimMontage* MontageToPlay, const FMo
 			{
 				// Enforce 'a single montage at once per group' rule
 				StopAllMontagesByGroupName(NewMontageGroupName, BlendInSettings);
+
+				if (MontageCVars::bFlushCompletedMontagesOnPlay)
+				{
+					ConditionalFlushCompletedMontages();
+				}
 			}
 
 			// Enforce 'a single root motion montage at once' rule.
@@ -2999,7 +3050,7 @@ void UAnimInstance::StopAllMontagesByGroupName(FName InGroupName, const FMontage
 		FAnimMontageInstance* MontageInstance = MontageInstances[InstanceIndex];
 		if (MontageInstance && MontageInstance->Montage && (MontageInstance->Montage->GetGroupName() == InGroupName))
 		{
-			MontageInstances[InstanceIndex]->Stop(BlendOutSettings, true);
+			MontageInstance->Stop(BlendOutSettings, true);
 		}
 	}
 }
