Fix case where RVT streaming mips wouldn't initialize on the first load in the editor.
The IsCompatibleFormat() could return false if the streaming mip texture hadn't finished compilation (in that case it would assume pixel format from the default texture).
Changed the logic to look directly at the PlatformData which is the structure that we read soon after in CreateStreamingTextureProducer(). This also allows us to do a full validation check of each layer.

diff --git a/Engine/Source/Runtime/Engine/Private/Components/RuntimeVirtualTextureComponent.cpp b/Engine/Source/Runtime/Engine/Private/Components/RuntimeVirtualTextureComponent.cpp
--- a/Engine/Source/Runtime/Engine/Private/Components/RuntimeVirtualTextureComponent.cpp
+++ b/Engine/Source/Runtime/Engine/Private/Components/RuntimeVirtualTextureComponent.cpp
@@ -21,6 +21,7 @@
 #include "VT/RuntimeVirtualTexture.h"
 #include "VT/VirtualTexture.h"
 #include "VT/VirtualTextureBuilder.h"
+#include "VT/VirtualTextureBuiltData.h"
 
 #include UE_INLINE_GENERATED_CPP_BY_NAME(RuntimeVirtualTextureComponent)
 
@@ -309,11 +310,27 @@ bool URuntimeVirtualTextureComponent::IsStreamingLowMipsOnly()
 	return bUseStreamingMipsOnly || CVarVTStreamingMipsUseAlways.GetValueOnAnyThread();
 }
 
-bool IsCompatibleFormat(URuntimeVirtualTexture const& RuntimeVirtualTexture, UVirtualTexture2D const& StreamingVirtualTexture)
+/** 
+ * This test should be covered by the BuildHash check, but there was an bug where the texture compilation built the streaming virtual texture with an unexpected pixel format. 
+ * The bug was fixed but keeping this extra check to catch any similar regression in future.
+ */
+static bool IsCompatibleFormat(URuntimeVirtualTexture const& RuntimeVirtualTexture, UVirtualTexture2D const& StreamingVirtualTexture)
 {
-	// During texture compilation we can't validate anything other than first layer, so restrict validation to that.
-	// This should catch any 99% of issues anyway. 
-	return (RuntimeVirtualTexture.GetLayerFormat(0) == StreamingVirtualTexture.GetPixelFormat(0));
+	if (FTexturePlatformData const* StreamingTextureData = StreamingVirtualTexture.GetPlatformData())
+	{
+		if (FVirtualTextureBuiltData const* VTData = StreamingTextureData->VTData)
+		{
+			for (int32 LayerIndex = 0; LayerIndex < RuntimeVirtualTexture.GetLayerCount(); ++LayerIndex)
+			{
+				if (RuntimeVirtualTexture.GetLayerFormat(LayerIndex) != VTData->LayerTypes[LayerIndex])
+				{
+					return false;
+				}
+			}
+			return true;
+		}
+	}
+	return false;
 }
 
 bool URuntimeVirtualTextureComponent::IsStreamingTextureInvalid(EShadingPath ShadingPath) const
