[shaders] replace slow job batch logging with slow individual job logging and move into ProcessFinishedJob so it executes regardless of whether workers are used or not. replace the hardcoded 60s batch threshold with a cvar driven individual job threshold (defaulting to 30s). often the full batch contains a mix of fast and slow compilations so printing details for every job in the batch along with the total batch time isn't terribly useful. if debug info is enabled for the job, the job debug out path will be printed instead of the default shadertype/permutation/format string - the debug path contains all of this information anyways, and this allows quick access to the debug artifacts for slow jobs.

diff --git a/Engine/Source/Runtime/Engine/Private/ShaderCompiler/ShaderCompiler.cpp b/Engine/Source/Runtime/Engine/Private/ShaderCompiler/ShaderCompiler.cpp
--- a/Engine/Source/Runtime/Engine/Private/ShaderCompiler/ShaderCompiler.cpp
+++ b/Engine/Source/Runtime/Engine/Private/ShaderCompiler/ShaderCompiler.cpp
@@ -162,6 +162,13 @@ static TAutoConsoleVariable<bool> CVarShaderCompilerDumpDDCKeys(
 	ECVF_Default
 );
 
+static TAutoConsoleVariable<int> CVarShaderCompilerLogSlowJobThreshold(
+	TEXT("r.ShaderCompiler.LogSlowJobThreshold"),
+	30,
+	TEXT("If a single compilation job's compile time exceeds the specified value (in seconds), info about the job will be automatically logged for investigation."),
+	ECVF_Default
+);
+
 
 bool AreShaderErrorsFatal()
 {
@@ -1371,6 +1378,26 @@ FShaderPipelineCompileJob* FShaderCompilingManager::PreparePipelineCompileJob(ui
 
 void FShaderCompilingManager::ProcessFinishedJob(FShaderCommonCompileJob* FinishedJob)
 {
+	bool bIsPipelineJob = FinishedJob->Type == EShaderCompileJobType::Pipeline;
+	FinishedJob->ForEachSingleShaderJob([this](FShaderCompileJob& SingleJob)
+		{
+			// Log if requested or if there was an exceptionally slow batch, to see the offender easily
+			if (bLogJobCompletionTimes || SingleJob.Output.CompileTime > (double)CVarShaderCompilerLogSlowJobThreshold.GetValueOnAnyThread())
+			{
+				TStringBuilder<256> JobName;
+				if (SingleJob.Input.DumpDebugInfoEnabled())
+				{
+					JobName << SingleJob.Input.DumpDebugInfoPath;
+				}
+				else
+				{
+					JobName << SingleJob.Key.ShaderType->GetName();
+					JobName.Appendf(TEXT("(permutation %d, format %s)"), SingleJob.Key.PermutationId, *SingleJob.Input.ShaderFormat.ToString());
+				}
+				UE_LOG(LogShaderCompilers, Display, TEXT("Job %s compile time exceeded threshold (%.3fs)"), JobName.ToString(), SingleJob.Output.CompileTime);
+			}
+		});
+
 	AllJobs.ProcessFinishedJob(FinishedJob);
 }
 
diff --git a/Engine/Source/Runtime/Engine/Private/ShaderCompiler/ShaderCompilerThreadRunnable.cpp b/Engine/Source/Runtime/Engine/Private/ShaderCompiler/ShaderCompilerThreadRunnable.cpp
--- a/Engine/Source/Runtime/Engine/Private/ShaderCompiler/ShaderCompilerThreadRunnable.cpp
+++ b/Engine/Source/Runtime/Engine/Private/ShaderCompiler/ShaderCompilerThreadRunnable.cpp
@@ -432,44 +432,6 @@ void FShaderCompileThreadRunnable::PushCompletedJobsToManager()
 			Manager->WorkersBusyTime += ElapsedTime;
 			COOK_STAT(ShaderCompilerCookStats::AsyncCompileTimeSec += ElapsedTime);
 
-			// Log if requested or if there was an exceptionally slow batch, to see the offender easily
-			if (Manager->bLogJobCompletionTimes || ElapsedTime > 60.0)
-			{
-				TArray<FShaderCommonCompileJobPtr> SortedJobs = CurrentWorkerInfo.QueuedJobs;
-				SortedJobs.Sort([](const FShaderCommonCompileJobPtr& JobA, const FShaderCommonCompileJobPtr& JobB)
-					{
-						const FShaderCompileJob* SingleJobA = JobA->GetSingleShaderJob();
-						const FShaderCompileJob* SingleJobB = JobB->GetSingleShaderJob();
-
-						const double TimeA = SingleJobA ? SingleJobA->Output.CompileTime : 0.0f;
-						const double TimeB = SingleJobB ? SingleJobB->Output.CompileTime : 0.0f;
-
-						return TimeA > TimeB;
-					});
-
-				FString JobNames;
-
-				for (int32 JobIndex = 0; JobIndex < SortedJobs.Num(); JobIndex++)
-				{
-					const FShaderCommonCompileJob& Job = *SortedJobs[JobIndex];
-					if (const FShaderCompileJob* SingleJob = Job.GetSingleShaderJob())
-					{
-						const TCHAR* JobName = Manager->bLogJobCompletionTimes ? *SingleJob->Input.DebugGroupName : SingleJob->Key.ShaderType->GetName();
-						JobNames += FString::Printf(TEXT("%s:%d(%s) [WorkerTime=%.3fs]"), JobName, SingleJob->Key.PermutationId, *SingleJob->Input.ShaderFormat.ToString(), SingleJob->Output.CompileTime);
-					}
-					else
-					{
-						const FShaderPipelineCompileJob* PipelineJob = Job.GetShaderPipelineJob();
-						JobNames += FString(PipelineJob->Key.ShaderPipeline->GetName());
-					}
-					if (JobIndex < SortedJobs.Num() - 1)
-					{
-						JobNames += TEXT(", ");
-					}
-				}
-
-				UE_LOG(LogShaderCompilers, Display, TEXT("Worker (%d/%d) finished batch of %u jobs in %.3fs, %s"), WorkerIndex + 1, WorkerInfos.Num(), SortedJobs.Num(), ElapsedTime, *JobNames);
-			}
 
 			CurrentWorkerInfo.FinishTime = FPlatformTime::Seconds();
 			CurrentWorkerInfo.bComplete = false;
