[Shaders] Improve SCW crash report message.

When SCW crashes with an unspecified type of error (HandleGeneralCrash) the error was only logged and then the editor crashes later due to invalid data in the output file.
Show the modal or log entry with the crash details instead of showing an unrelated error.

#rnx

diff --git a/Engine/Source/Runtime/Engine/Private/ShaderCompiler/ShaderCompilerUtil.cpp b/Engine/Source/Runtime/Engine/Private/ShaderCompiler/ShaderCompilerUtil.cpp
--- a/Engine/Source/Runtime/Engine/Private/ShaderCompiler/ShaderCompilerUtil.cpp
+++ b/Engine/Source/Runtime/Engine/Private/ShaderCompiler/ShaderCompilerUtil.cpp
@@ -43,25 +43,33 @@ static void ModalErrorOrLog(const FString& Title, const FString& Text, int64 Cur
 		BadFile = FString::Printf(TEXT(" (Truncated or corrupt output file! Current file pos %lld, file size %lld)"), CurrentFilePos, ExpectedFileSize);
 	}
 
-	if (FPlatformProperties::SupportsWindowedMode() && !FApp::IsUnattended())
+	if (bIsErrorFatal)
 	{
-		UE_LOG(LogShaderCompilers, Error, TEXT("%s\n%s"), *Text, *BadFile);
-		if (!bModalReported.AtomicSet(true))
+		// Ensure errors are logged before exiting
+		GLog->Panic();
+
+		if (FPlatformProperties::SupportsWindowedMode() && !FApp::IsUnattended())
 		{
-			// Show dialog box with error message and request exit
-			FMessageDialog::Open(EAppMsgType::Ok, FText::FromString(Text), FText::FromString(Title));
-			FPlatformMisc::RequestExit(false, TEXT("ShaderCompiler.ModalErrorOrLog"));
+			if (!bModalReported.AtomicSet(true))
+			{
+				UE_LOG(LogShaderCompilers, Error, TEXT("%s\n%s"), *Text, *BadFile);
+
+				// Show dialog box with error message and request exit
+				FMessageDialog::Open(EAppMsgType::Ok, FText::FromString(Text), FText::FromString(Title));
+				constexpr bool bForceExit = true;
+				FPlatformMisc::RequestExit(bForceExit, TEXT("ShaderCompiler.ModalErrorOrLog"));
+			}
+			else
+			{
+				// Another thread already opened a dialog box and requests exit
+				FPlatformProcess::SleepInfinite();
+			}
 		}
 		else
 		{
-			// Another thread already opened a dialog box and requests exit
-			FPlatformProcess::SleepInfinite();
+			UE_LOG(LogShaderCompilers, Fatal, TEXT("%s\n%s\n%s"), *Title, *Text, *BadFile);
 		}
 	}
-	else if (bIsErrorFatal)
-	{
-		UE_LOG(LogShaderCompilers, Fatal, TEXT("%s\n%s\n%s"), *Title, *Text, *BadFile);
-	}
 	else
 	{
 		UE_LOG(LogShaderCompilers, Error, TEXT("%s\n%s\n%s"), *Title, *Text, *BadFile);
@@ -110,8 +118,7 @@ namespace ShaderCompileWorkerError
 {
 	void HandleGeneralCrash(const TCHAR* ExceptionInfo, const TCHAR* Callstack)
 	{
-		GLog->Panic();
-		UE_LOG(LogShaderCompilers, Error, TEXT("ShaderCompileWorker crashed!\n%s\n%s"), ExceptionInfo, Callstack);
+		ModalErrorOrLog(TEXT("ShaderCompileWorker crashed"), *FString::Printf(TEXT("Exception:\n%s\n\nCallstack:\n%s"), ExceptionInfo, Callstack));
 	}
 
 	void HandleBadShaderFormatVersion(const TCHAR* Data)
