Fix for flickering translucents on skeletal meshes due to missing MeshIdInPrimitive.

#rnx

diff --git a/Engine/Source/Runtime/Engine/Private/SkeletalMesh.cpp b/Engine/Source/Runtime/Engine/Private/SkeletalMesh.cpp
--- a/Engine/Source/Runtime/Engine/Private/SkeletalMesh.cpp
+++ b/Engine/Source/Runtime/Engine/Private/SkeletalMesh.cpp
@@ -6718,6 +6718,7 @@ void FSkeletalMeshSceneProxy::DrawStaticElements(FStaticPrimitiveDrawInterface*
 					MeshElement.Type = PT_TriangleList;
 					MeshElement.LODIndex = LODIndex;
 					MeshElement.SegmentIndex = SectionIndex;
+					MeshElement.MeshIdInPrimitive = SectionIndex;
 
 					BatchElement.PrimitiveUniformBuffer = GetUniformBuffer();
 					BatchElement.FirstIndex = Section.BaseIndex;
@@ -6929,6 +6930,7 @@ void FSkeletalMeshSceneProxy::GetDynamicElementsSection(const TArray<const FScen
 			FMeshBatch& Mesh = Collector.AllocateMesh();
 
 			CreateBaseMeshBatch(View, LODData, LODIndex, SectionIndex, SectionElementInfo, Mesh);
+			//For dynamic mesh elements, Mesh.MeshIdInPrimitive is setup in Collector.AddMesh.
 			
 			if(!Mesh.VertexFactory)
 			{
