Add shader map friendly name and vertex factory/shader type name to assertions

diff --git a/Engine/Plugins/Runtime/ComputeFramework/Source/ComputeFramework/Private/ComputeKernelShaderMap.cpp b/Engine/Plugins/Runtime/ComputeFramework/Source/ComputeFramework/Private/ComputeKernelShaderMap.cpp
--- a/Engine/Plugins/Runtime/ComputeFramework/Source/ComputeFramework/Private/ComputeKernelShaderMap.cpp
+++ b/Engine/Plugins/Runtime/ComputeFramework/Source/ComputeFramework/Private/ComputeKernelShaderMap.cpp
@@ -356,7 +356,7 @@ void FComputeKernelShaderMap::Compile(
 				if (ShaderType && ShouldCacheComputeKernelShader(ShaderType, InPlatform, InKernel))
 				{
 					// Verify that the shader map Id contains inputs for any shaders that will be put into this shader map
-					check(InShaderMapId.ContainsShaderType(ShaderType));
+					checkf(InShaderMapId.ContainsShaderType(ShaderType), TEXT("Compute shader map %s missing expected shader type %s"), *GetFriendlyName(), ShaderType->GetName());
 					
 					// Compile this ComputeKernel shader.
 					TArray<FString> ShaderErrors;
diff --git a/Engine/Source/Runtime/Engine/Private/Materials/MaterialShader.cpp b/Engine/Source/Runtime/Engine/Private/Materials/MaterialShader.cpp
--- a/Engine/Source/Runtime/Engine/Private/Materials/MaterialShader.cpp
+++ b/Engine/Source/Runtime/Engine/Private/Materials/MaterialShader.cpp
@@ -2326,8 +2326,8 @@ int32 FMaterialShaderMap::SubmitCompileJobs(uint32 CompilingShaderMapId,
 			}
 
 			// Verify that the shader map Id contains inputs for any shaders that will be put into this shader map
-			check(ShaderMapId.ContainsVertexFactoryType(MeshLayout.VertexFactoryType));
-			check(ShaderMapId.ContainsShaderType(ShaderType, kUniqueShaderPermutationId));
+			checkf(ShaderMapId.ContainsVertexFactoryType(MeshLayout.VertexFactoryType), TEXT("Material shader map %s is missing expected vertex factory type %s"), GetFriendlyName(), MeshLayout.VertexFactoryType->GetName());
+			checkf(ShaderMapId.ContainsShaderType(ShaderType, kUniqueShaderPermutationId), TEXT("Material shader map %s is missing expected shader type %s"), GetFriendlyName(), ShaderType->GetName());
 
 			NumShadersPerVF++;
 			// only compile the shader if we don't already have it and it is not a pipelined one
@@ -2372,7 +2372,7 @@ int32 FMaterialShaderMap::SubmitCompileJobs(uint32 CompilingShaderMapId,
 				for (auto* ShaderType : StageTypes)
 				{
 					// Verify that the shader map Id contains inputs for any shaders that will be put into this shader map
-					check(ShaderMapId.ContainsShaderType(ShaderType, kUniqueShaderPermutationId));
+					checkf(ShaderMapId.ContainsShaderType(ShaderType, kUniqueShaderPermutationId), TEXT("Material shader map %s missing expected shader type %s"), GetFriendlyName(), ShaderType->GetName());
 				}
 
 				// Make a pipeline job with all the stages
@@ -2430,7 +2430,7 @@ int32 FMaterialShaderMap::SubmitCompileJobs(uint32 CompilingShaderMapId,
 		}
 
 		// Verify that the shader map Id contains inputs for any shaders that will be put into this shader map
-		check(ShaderMapId.ContainsShaderType(ShaderType, kUniqueShaderPermutationId));
+		checkf(ShaderMapId.ContainsShaderType(ShaderType, kUniqueShaderPermutationId), TEXT("Material shader map %s missing expected shader type %s"), GetFriendlyName(), ShaderType->GetName());
 
 		// Compile this material shader for this material.
 		TArray<FString> ShaderErrors;
