Fix virtual texture adapter behavior for case where source data is linear and adapted texture is sRGB.
Also ensure that sRGB setting is stored to the virtual texture resource since without that the material sampling can chose the incorrect view when the adapter is present.

diff --git a/Engine/Source/Runtime/Engine/Private/VT/VirtualTextureAdapter.cpp b/Engine/Source/Runtime/Engine/Private/VT/VirtualTextureAdapter.cpp
--- a/Engine/Source/Runtime/Engine/Private/VT/VirtualTextureAdapter.cpp
+++ b/Engine/Source/Runtime/Engine/Private/VT/VirtualTextureAdapter.cpp
@@ -125,8 +125,10 @@ namespace VirtualTextureAdapter
 		const uint32 FinalTexelCountX = InDestRect.Max.X - InDestRect.Min.X;
 		const uint32 FinalTexelCountY = InDestRect.Max.Y - InDestRect.Min.Y;
 
-		const bool bSrgb = InSourceSRV->GetDesc().Texture.SRV.GetViewInfo(InSourceSRV->GetTexture()).bSRGB;
-		const ETextureCreateFlags TextureCreateFlagsSrgb = bSrgb ? ETextureCreateFlags::SRGB : ETextureCreateFlags::None;
+		const bool bSourceSrgb = InSourceSRV->GetDesc().Texture.SRV.GetViewInfo(InSourceSRV->GetTexture()).bSRGB;
+		const bool bDestSrgb = EnumHasAnyFlags(InDestTexture->GetDesc().Flags, ETextureCreateFlags::SRGB);
+		const ETextureCreateFlags TextureCreateFlagsSrgb = bSourceSrgb ? ETextureCreateFlags::SRGB : ETextureCreateFlags::None;
+		const bool bFinalPassConvertToSrgb = !bSourceSrgb && bDestSrgb;
 
 		FRDGTextureRef CurrentOutput = nullptr;
 		FBox2f CurrentUVRange(InUVRange);
@@ -147,7 +149,7 @@ namespace VirtualTextureAdapter
 
 			FCopyCompressCS::FPermutationDomain PermutationVector;
 			PermutationVector.Set<FCopyCompressCS::FSourceTextureSelector>(true);
-			PermutationVector.Set<FCopyCompressCS::FDestSrgb>(bSrgb);
+			PermutationVector.Set<FCopyCompressCS::FDestSrgb>(bIsFinalPass ? bFinalPassConvertToSrgb : false);
 			PermutationVector.Set<FCopyCompressCS::FCompressionFormatDim>(0);
 			TShaderMapRef<FCopyCompressCS> Shader(GlobalShaderMap, PermutationVector);
 
@@ -196,7 +198,7 @@ namespace VirtualTextureAdapter
 
 				FCopyCompressCS::FPermutationDomain PermutationVector;
 				PermutationVector.Set<FCopyCompressCS::FSourceTextureSelector>(bUseSourceTextureA);
-				PermutationVector.Set<FCopyCompressCS::FDestSrgb>(bSrgb);
+				PermutationVector.Set<FCopyCompressCS::FDestSrgb>(bIsFinalPass ? bFinalPassConvertToSrgb : false);
 				PermutationVector.Set<FCopyCompressCS::FCompressionFormatDim>(0);
 				TShaderMapRef<FCopyCompressCS> Shader(GlobalShaderMap, PermutationVector);
 
@@ -232,7 +234,7 @@ namespace VirtualTextureAdapter
 
 			FCopyCompressCS::FPermutationDomain PermutationVector;
 			PermutationVector.Set<FCopyCompressCS::FSourceTextureSelector>(bUseSourceTextureA);
-			PermutationVector.Set<FCopyCompressCS::FDestSrgb>(bSrgb);
+			PermutationVector.Set<FCopyCompressCS::FDestSrgb>(bFinalPassConvertToSrgb);
 			PermutationVector.Set<FCopyCompressCS::FCompressionFormatDim>(CompressionPermutation);
 			TShaderMapRef<FCopyCompressCS> Shader(GlobalShaderMap, PermutationVector);
 
@@ -487,10 +489,12 @@ public:
 		{
 			SourceFormat = Texture2D->GetPixelFormat(0);
 			NumSourceMips = Texture2D->GetNumMips();
+			bSRGB = Texture2D->SRGB;
 		}
 		else if (UTextureRenderTarget2D* RenderTarget2D = Cast<UTextureRenderTarget2D>(InTexture))
 		{
 			SourceFormat = RenderTarget2D->GetFormat();
+			bSRGB = RenderTarget2D->SRGB;
 		}
 		Format = VirtualTextureAdapter::GetFinalFormat(SourceFormat, bUseCompressedFormat);
 
