Fixed regression with UMaterial::RecursiveGetExpressionChain introduced in 33967097

#tests Editor, mobile preview

diff --git a/Engine/Source/Runtime/Engine/Private/Materials/Material.cpp b/Engine/Source/Runtime/Engine/Private/Materials/Material.cpp
--- a/Engine/Source/Runtime/Engine/Private/Materials/Material.cpp
+++ b/Engine/Source/Runtime/Engine/Private/Materials/Material.cpp
@@ -6441,13 +6441,12 @@ bool UMaterial::RecursiveGetExpressionChain(
 		}
 
 		// here we assume ALL inputs to the MaterialFunctionCall are active
-		Inputs.Empty(InExpression->CountInputs());
+		Inputs.Reserve(Inputs.Num() + InExpression->CountInputs());
 		for (FExpressionInputIterator It{ InExpression }; It; ++It)
 		{
 			Inputs.Push(It.Input);
+			InputsFrequency.Push(InShaderFrequency);
 		}
-
-		InputsFrequency.Init(InShaderFrequency, Inputs.Num());
 	}
 	else
 	{
