Expose r.MeshPaintVirtualTexture.DefaultTexelsPerVertex in the renderer project settings.
This option controls the size of mesh paint textures when they are created.

diff --git a/Engine/Source/Runtime/Engine/Classes/Engine/RendererSettings.h b/Engine/Source/Runtime/Engine/Classes/Engine/RendererSettings.h
--- a/Engine/Source/Runtime/Engine/Classes/Engine/RendererSettings.h
+++ b/Engine/Source/Runtime/Engine/Classes/Engine/RendererSettings.h
@@ -457,6 +457,12 @@ class URendererSettings : public UDeveloperSettings
 		ConfigRestartRequired = true))
 	uint32 MeshPaintVirtualTextureTileBorderSize;
 
+	UPROPERTY(config, EditAnywhere, Category = VirtualTextures, meta = (
+		EditCondition = "bVirtualTextures",
+		ConsoleVariable = "r.MeshPaintVirtualTexture.DefaultTexelsPerVertex", DisplayName = "Mesh paint texels per vertex",
+		ToolTip = "Default ratio of texels to vertices when creating a mesh paint texture for a mesh."))
+	uint32 MeshPaintDefaultTexelsPerVertex;
+
 	UPROPERTY(config, EditAnywhere, Category = "Runtime Virtual Textures|Enable Material Types", meta = (
 		EditCondition = "bVirtualTextures",
 		ConsoleVariable = "r.vt.rvt.EnableBaseColor", DisplayName = "Enable 'Base Color'",
