CMC now has option to only apply JumpZVelocity at start of jump instead of all through Jump Max Hold Time, which was default.
The option defaults to existing behavior, so the new behavior is opt-in.

As well, we are now able to pass deltaTime through to DoJump, so it can in turn pass this time slice value down through whatever needs it
Being an engine API change we're going to deprecate the old signature while introducing the new one
- Old DoJump marked UE_DEPRECATED_FORGAME and calls the new DoJump (preserve client compilation)
- New Character CVar that defaults for now to telling CheckJumpInput to use the old DoJump (preserve client behavior)
- On game projects, setting this CVar to false will cause CheckJumpInput to use the new DoJump

diff --git a/Engine/Source/Runtime/Engine/Classes/GameFramework/CharacterMovementComponent.h b/Engine/Source/Runtime/Engine/Classes/GameFramework/CharacterMovementComponent.h
--- a/Engine/Source/Runtime/Engine/Classes/GameFramework/CharacterMovementComponent.h
+++ b/Engine/Source/Runtime/Engine/Classes/GameFramework/CharacterMovementComponent.h
@@ -407,6 +407,13 @@ public:
 	UPROPERTY(Category="Character Movement (General Settings)", EditDefaultsOnly, BlueprintReadWrite)
 	uint8 bUseSeparateBrakingFriction:1;
 
+	/** 
+	 * True means while the jump key is held, we will not allow the vertical speed to fall below the JumpZVelocity tuning value 
+	 * even if a stronger force, such as gravity, is opposing the jump. 
+	 */
+	UPROPERTY(Category="Character Movement: Jumping / Falling", EditAnywhere, BlueprintReadWrite, AdvancedDisplay)
+	uint8 bDontFallBelowJumpZVelocityDuringJump:1;
+
 	/**
 	 *	Apply gravity while the character is actively jumping (e.g. holding the jump key).
 	 *	Helps remove frame-rate dependent jump height, but may alter base jump height.
@@ -1396,16 +1403,24 @@ public:
 	UPROPERTY(Category = "Character Movement: Jumping / Falling", EditAnywhere, BlueprintReadWrite, meta = (editcondition = "bStayBasedInAir"))
 	float StayBasedInAirHeight = 1000.0f;
 
+
 	/** changes physics based on MovementMode */
 	ENGINE_API virtual void StartNewPhysics(float deltaTime, int32 Iterations);
 	
+	/**
+	 * NOTE: THIS FUNCTION IS DEPRECATED, PLEASE CALL DoJump(bool bReplayingMoves, float DeltaTime) ... 
+	 */
+	UE_DEPRECATED_FORGAME(5.5, "This function has been deprecated. Please call DoJump(bool bReplayingMoves, float DeltaTime)")
+	ENGINE_API virtual bool DoJump(bool bReplayingMoves);
+
 	/**
 	 * Perform jump. Called by Character when a jump has been detected because Character->bPressedJump was true. Checks Character->CanJump().
 	 * Note that you should usually trigger a jump through Character::Jump() instead.
 	 * @param	bReplayingMoves: true if this is being done as part of replaying moves on a locally controlled client after a server correction.
+	 * @param	DeltaTime: time slice for this jump move
 	 * @return	True if the jump was triggered successfully.
 	 */
-	ENGINE_API virtual bool DoJump(bool bReplayingMoves);
+	ENGINE_API virtual bool DoJump(bool bReplayingMoves, float DeltaTime);
 
 	/**
 	 * Returns true if current movement state allows an attempt at jumping. Used by Character::CanJump().
diff --git a/Engine/Source/Runtime/Engine/Private/Character.cpp b/Engine/Source/Runtime/Engine/Private/Character.cpp
--- a/Engine/Source/Runtime/Engine/Private/Character.cpp
+++ b/Engine/Source/Runtime/Engine/Private/Character.cpp
@@ -42,6 +42,13 @@ namespace CharacterCVars
 		TEXT("p.EnableCharacterAccelerationReplication"),
 		EnableCharacterAccelerationReplication,
 		TEXT("Whether to author acceleration data with character movement replication to sim proxies."));
+
+	static int32 UseLegacyDoJump = 1;
+	static FAutoConsoleVariableRef CVarUseLegacyDoJump(
+		TEXT("p.UseLegacyDoJump"),
+		UseLegacyDoJump,
+		TEXT("Should CheckJumpInput call the legacy DoJump or the new DoJump. Default is true while we are still deprecating the old DoJump"),
+		ECVF_Default);
 }
 
 ACharacter::ACharacter(const FObjectInitializer& ObjectInitializer)
@@ -1131,7 +1138,7 @@ void ACharacter::CheckJumpInput(float DeltaTime)
 				JumpCurrentCount++;
 			}
 
-			const bool bDidJump = CanJump() && CharacterMovement->DoJump(bClientUpdating);
+			const bool bDidJump = CanJump() && (CharacterCVars::UseLegacyDoJump)? CharacterMovement->DoJump(bClientUpdating) : CharacterMovement->DoJump(bClientUpdating, DeltaTime);
 			if (bDidJump)
 			{
 				// Transition from not (actively) jumping to jumping.
diff --git a/Engine/Source/Runtime/Engine/Private/Components/CharacterMovementComponent.cpp b/Engine/Source/Runtime/Engine/Private/Components/CharacterMovementComponent.cpp
--- a/Engine/Source/Runtime/Engine/Private/Components/CharacterMovementComponent.cpp
+++ b/Engine/Source/Runtime/Engine/Private/Components/CharacterMovementComponent.cpp
@@ -633,7 +633,7 @@ UCharacterMovementComponent::UCharacterMovementComponent(const FObjectInitialize
 		PrePhysicsTickFunction.SetTickFunctionEnable(true);
 		PrePhysicsTickFunction.TickGroup = TG_PrePhysics;
 	}
-
+	bDontFallBelowJumpZVelocityDuringJump = true;
 	bApplyGravityWhileJumping = true;
 
 	GravityScale = 1.f;
@@ -1084,21 +1084,41 @@ FVector UCharacterMovementComponent::GetPawnCapsuleExtent(const EShrinkCapsuleEx
 	return CapsuleExtent;
 }
 
-
 bool UCharacterMovementComponent::DoJump(bool bReplayingMoves)
+{
+	return DoJump(bReplayingMoves, 0.f);
+}
+
+bool UCharacterMovementComponent::DoJump(bool bReplayingMoves, float DeltaTime)
 {
 	if ( CharacterOwner && CharacterOwner->CanJump() )
 	{
 		// Don't jump if we can't move up/down.
 		if (!bConstrainToPlane || !FMath::IsNearlyEqual(FMath::Abs(GetGravitySpaceZ(PlaneConstraintNormal)), 1.f))
 		{
-			if (HasCustomGravity())
-			{
-				SetGravitySpaceZ(Velocity, FMath::Max<FVector::FReal>(GetGravitySpaceZ(Velocity), JumpZVelocity));
-			}
-			else
+			// If first frame of DoJump, we want to always inject the initial jump velocity.
+			// For subsequent frames, during the time Jump is held, it depends... 
+			// bDontFallXXXX == true means we want to ensure the character's Z velocity is never less than JumpZVelocity in this period
+			// bDontFallXXXX == false means we just want to leave Z velocity alone and "let the chips fall where they may" (e.g. fall properly in physics)
+
+			// NOTE: 
+			// Checking JumpCurrentCountPreJump instead of JumpCurrentCount because Character::CheckJumpInput might have
+			// incremented JumpCurrentCount just before entering this function... in order to compensate for the case when
+			// on the first frame of the jump, we're already in falling stage. So we want the original value before any 
+			// modification here.
+			// 
+			const bool bFirstJump = (CharacterOwner->JumpCurrentCountPreJump == 0);
+
+			if (bFirstJump || bDontFallBelowJumpZVelocityDuringJump)
 			{
-				Velocity.Z = FMath::Max<FVector::FReal>(Velocity.Z, JumpZVelocity);
+				if (HasCustomGravity())
+				{
+					SetGravitySpaceZ(Velocity, FMath::Max<FVector::FReal>(GetGravitySpaceZ(Velocity), JumpZVelocity));
+				}
+				else
+				{
+					Velocity.Z = FMath::Max<FVector::FReal>(Velocity.Z, JumpZVelocity);
+				}
 			}
 			
 			SetMovementMode(MOVE_Falling);
@@ -13217,6 +13237,7 @@ void UCharacterMovementComponent::FillAsyncInput(const FVector& InputVector, FCh
 	AsyncInput.AirControl = AirControl;
 	AsyncInput.AirControlBoostMultiplier = AirControlBoostMultiplier;
 	AsyncInput.AirControlBoostVelocityThreshold = AirControlBoostVelocityThreshold;
+	AsyncInput.bDontFallBelowJumpZVelocityDuringJump = bDontFallBelowJumpZVelocityDuringJump;
 	AsyncInput.bApplyGravityWhileJumping = bApplyGravityWhileJumping;
 	AsyncInput.PhysicsVolumeTerminalVelocity = GetPhysicsVolume()->TerminalVelocity;
 	AsyncInput.MaxJumpApexAttemptsPerSimulation = MaxJumpApexAttemptsPerSimulation;
diff --git a/Engine/Source/Runtime/Engine/Public/CharacterMovementComponentAsync.h b/Engine/Source/Runtime/Engine/Public/CharacterMovementComponentAsync.h
--- a/Engine/Source/Runtime/Engine/Public/CharacterMovementComponentAsync.h
+++ b/Engine/Source/Runtime/Engine/Public/CharacterMovementComponentAsync.h
@@ -446,6 +446,7 @@ struct FCharacterMovementComponentAsyncInput : public Chaos::FSimCallbackInput
 	float AirControl;
 	float AirControlBoostMultiplier;
 	float AirControlBoostVelocityThreshold;
+	bool bDontFallBelowJumpZVelocityDuringJump;
 	bool bApplyGravityWhileJumping;
 	float PhysicsVolumeTerminalVelocity;
 	int32 MaxJumpApexAttemptsPerSimulation;
