Removing deprecated variables from FSinglePrimitiveStructured.

diff --git a/Engine/Source/Runtime/Engine/Private/PrimitiveUniformShaderParameters.cpp b/Engine/Source/Runtime/Engine/Private/PrimitiveUniformShaderParameters.cpp
--- a/Engine/Source/Runtime/Engine/Private/PrimitiveUniformShaderParameters.cpp
+++ b/Engine/Source/Runtime/Engine/Private/PrimitiveUniformShaderParameters.cpp
@@ -40,15 +40,6 @@ void FSinglePrimitiveStructured::InitRHI(FRHICommandListBase& RHICmdList)
 {
 	SCOPED_LOADTIMER(FSinglePrimitiveStructuredBuffer_InitRHI);
 
-	FRHIResourceCreateInfo CreateInfo(TEXT("PrimitiveSceneDataBuffer"));
-
-	PRAGMA_DISABLE_DEPRECATION_WARNINGS
-
-	{	
-		PrimitiveSceneDataBufferRHI = RHICmdList.CreateStructuredBuffer(sizeof(FVector4f), FPrimitiveSceneShaderData::DataStrideInFloat4s * sizeof(FVector4f), BUF_Static | BUF_ShaderResource, CreateInfo);
-		PrimitiveSceneDataBufferSRV = RHICmdList.CreateShaderResourceView(PrimitiveSceneDataBufferRHI);
-	}
-
 	{
 		const static FLazyName ClassName(TEXT("FSinglePrimitiveStructured"));
 		const FRHITextureCreateDesc Desc =
@@ -60,21 +51,7 @@ void FSinglePrimitiveStructured::InitRHI(FRHICommandListBase& RHICmdList)
 		PrimitiveSceneDataTextureSRV = RHICmdList.CreateShaderResourceView(PrimitiveSceneDataTextureRHI, 0);
 	}
 
-	CreateInfo.DebugName = TEXT("LightmapSceneDataBuffer");
-	LightmapSceneDataBufferRHI = RHICmdList.CreateStructuredBuffer(sizeof(FVector4f), FLightmapSceneShaderData::DataStrideInFloat4s * sizeof(FVector4f), BUF_Static | BUF_ShaderResource, CreateInfo);
-	LightmapSceneDataBufferSRV = RHICmdList.CreateShaderResourceView(LightmapSceneDataBufferRHI);
-
-	CreateInfo.DebugName = TEXT("InstanceSceneDataBuffer");
-	InstanceSceneDataBufferRHI = RHICmdList.CreateStructuredBuffer(sizeof(FVector4f), FInstanceSceneShaderData::GetDataStrideInFloat4s() * sizeof(FVector4f), BUF_Static | BUF_ShaderResource, CreateInfo);
-	InstanceSceneDataBufferSRV = RHICmdList.CreateShaderResourceView(InstanceSceneDataBufferRHI);
-
-	CreateInfo.DebugName = TEXT("InstancePayloadDataBuffer");
-	InstancePayloadDataBufferRHI = RHICmdList.CreateStructuredBuffer(sizeof(FVector4f), 1 /* unused dummy */ * sizeof(FVector4f), BUF_Static | BUF_ShaderResource, CreateInfo);
-	InstancePayloadDataBufferSRV = RHICmdList.CreateShaderResourceView(InstancePayloadDataBufferRHI);
-
-	PRAGMA_ENABLE_DEPRECATION_WARNINGS
-
-	CreateInfo.DebugName = TEXT("SkyIrradianceEnvironmentMap");
+	FRHIResourceCreateInfo CreateInfo(TEXT("SkyIrradianceEnvironmentMap"));
 	SkyIrradianceEnvironmentMapRHI = RHICmdList.CreateStructuredBuffer(sizeof(FVector4f), sizeof(FVector4f) * 8, BUF_Static | BUF_ShaderResource, CreateInfo);
 	SkyIrradianceEnvironmentMapSRV = RHICmdList.CreateShaderResourceView(SkyIrradianceEnvironmentMapRHI);
 
@@ -83,62 +60,9 @@ void FSinglePrimitiveStructured::InitRHI(FRHICommandListBase& RHICmdList)
 
 void FSinglePrimitiveStructured::UploadToGPU(FRHICommandListBase& RHICmdList)
 {
-	void* LockedData = nullptr;
-
-	PRAGMA_DISABLE_DEPRECATION_WARNINGS
-
-	LockedData = RHICmdList.LockBuffer(PrimitiveSceneDataBufferRHI, 0, FPrimitiveSceneShaderData::DataStrideInFloat4s * sizeof(FVector4f), RLM_WriteOnly);
-	FPlatformMemory::Memcpy(LockedData, PrimitiveSceneData.Data.GetData(), FPrimitiveSceneShaderData::DataStrideInFloat4s * sizeof(FVector4f));
-	RHICmdList.UnlockBuffer(PrimitiveSceneDataBufferRHI);
-
-	LockedData = RHICmdList.LockBuffer(LightmapSceneDataBufferRHI, 0, FLightmapSceneShaderData::DataStrideInFloat4s * sizeof(FVector4f), RLM_WriteOnly);
-	FPlatformMemory::Memcpy(LockedData, LightmapSceneData.Data.GetData(), FLightmapSceneShaderData::DataStrideInFloat4s * sizeof(FVector4f));
-	RHICmdList.UnlockBuffer(LightmapSceneDataBufferRHI);
-
-	LockedData = RHICmdList.LockBuffer(InstanceSceneDataBufferRHI, 0, FInstanceSceneShaderData::GetDataStrideInFloat4s() * sizeof(FVector4f), RLM_WriteOnly);
-	FPlatformMemory::Memcpy(LockedData, InstanceSceneData.Data.GetData(), FInstanceSceneShaderData::GetDataStrideInFloat4s() * sizeof(FVector4f));
-	RHICmdList.UnlockBuffer(InstanceSceneDataBufferRHI);
-
-	LockedData = RHICmdList.LockBuffer(InstancePayloadDataBufferRHI, 0, 1 /* unused dummy */ * sizeof(FVector4f), RLM_WriteOnly);
-	FPlatformMemory::Memset(LockedData, 0x00, sizeof(FVector4f));
-	RHICmdList.UnlockBuffer(InstancePayloadDataBufferRHI);
-
-//#if WITH_EDITOR
-	if (IsFeatureLevelSupported(GMaxRHIShaderPlatform, ERHIFeatureLevel::SM5))
-	{
-		// Create level instance SRV
-		FRHIResourceCreateInfo LevelInstanceBufferCreateInfo(TEXT("EditorVisualizeLevelInstanceDataBuffer"));
-		EditorVisualizeLevelInstanceDataBufferRHI = RHICmdList.CreateVertexBuffer(sizeof(uint32), BUF_Static | BUF_ShaderResource, LevelInstanceBufferCreateInfo);
-
-		LockedData = RHICmdList.LockBuffer(EditorVisualizeLevelInstanceDataBufferRHI, 0, sizeof(uint32), RLM_WriteOnly);
-
-		*reinterpret_cast<uint32*>(LockedData) = 0;
-
-		RHICmdList.UnlockBuffer(EditorVisualizeLevelInstanceDataBufferRHI);
-
-		EditorVisualizeLevelInstanceDataBufferSRV = RHICmdList.CreateShaderResourceView(EditorVisualizeLevelInstanceDataBufferRHI, sizeof(uint32), PF_R32_UINT);
-
-		// Create selection outline SRV
-		FRHIResourceCreateInfo SelectionBufferCreateInfo(TEXT("EditorSelectedDataBuffer"));
-		EditorSelectedDataBufferRHI = RHICmdList.CreateVertexBuffer(sizeof(uint32), BUF_Static | BUF_ShaderResource, SelectionBufferCreateInfo);
-
-		LockedData = RHICmdList.LockBuffer(EditorSelectedDataBufferRHI, 0, sizeof(uint32), RLM_WriteOnly);
-
-		*reinterpret_cast<uint32*>(LockedData) = 0;
-
-		RHICmdList.UnlockBuffer(EditorSelectedDataBufferRHI);
-
-		EditorSelectedDataBufferSRV = RHICmdList.CreateShaderResourceView(EditorSelectedDataBufferRHI, sizeof(uint32), PF_R32_UINT);
-	}
-//#endif
-
-	PRAGMA_ENABLE_DEPRECATION_WARNINGS
 }
 
 TGlobalResource<FSinglePrimitiveStructured> GIdentityPrimitiveBuffer;
-PRAGMA_DISABLE_DEPRECATION_WARNINGS
-TGlobalResource<FSinglePrimitiveStructured> GTilePrimitiveBuffer;
-PRAGMA_ENABLE_DEPRECATION_WARNINGS
 
 FPrimitiveSceneShaderData::FPrimitiveSceneShaderData(const FPrimitiveSceneProxy* RESTRICT Proxy)
 	: Data(InPlace, NoInit)
diff --git a/Engine/Source/Runtime/Engine/Public/PrimitiveSceneShaderData.h b/Engine/Source/Runtime/Engine/Public/PrimitiveSceneShaderData.h
--- a/Engine/Source/Runtime/Engine/Public/PrimitiveSceneShaderData.h
+++ b/Engine/Source/Runtime/Engine/Public/PrimitiveSceneShaderData.h
@@ -49,81 +49,28 @@ struct FPrimitiveSceneShaderData
 class FSinglePrimitiveStructured : public FRenderResource
 {
 public:
-
 	FSinglePrimitiveStructured()
-		: ShaderPlatform(SP_NumPlatforms)
 	{}
 
 	ENGINE_API virtual void InitRHI(FRHICommandListBase& RHICmdList) override;
 
 	virtual void ReleaseRHI() override
 	{
-		PRAGMA_DISABLE_DEPRECATION_WARNINGS
-		PrimitiveSceneDataBufferRHI.SafeRelease();
-		PrimitiveSceneDataBufferSRV.SafeRelease();
 		SkyIrradianceEnvironmentMapRHI.SafeRelease();
 		SkyIrradianceEnvironmentMapSRV.SafeRelease();
-		InstanceSceneDataBufferRHI.SafeRelease();
-		InstanceSceneDataBufferSRV.SafeRelease();
-		InstancePayloadDataBufferRHI.SafeRelease();
-		InstancePayloadDataBufferSRV.SafeRelease();
 		PrimitiveSceneDataTextureRHI.SafeRelease();
 		PrimitiveSceneDataTextureSRV.SafeRelease();
-		LightmapSceneDataBufferRHI.SafeRelease();
-		LightmapSceneDataBufferSRV.SafeRelease();
-		EditorVisualizeLevelInstanceDataBufferRHI.SafeRelease();
-		EditorVisualizeLevelInstanceDataBufferSRV.SafeRelease();
-		EditorSelectedDataBufferRHI.SafeRelease();
-		EditorSelectedDataBufferSRV.SafeRelease();
-		PRAGMA_ENABLE_DEPRECATION_WARNINGS
 	}
 
 	ENGINE_API void UploadToGPU(FRHICommandListBase& RHICmdList);
 
 	EShaderPlatform ShaderPlatform=SP_NumPlatforms;
 
-	UE_DEPRECATED(5.3, "Use IRendererModule::CreateSinglePrimitiveSceneUniformBuffer instead")
-	FPrimitiveSceneShaderData PrimitiveSceneData;
-	UE_DEPRECATED(5.3, "Use IRendererModule::CreateSinglePrimitiveSceneUniformBuffer instead")
-	FInstanceSceneShaderData InstanceSceneData;
-	UE_DEPRECATED(5.3, "Use IRendererModule::CreateSinglePrimitiveSceneUniformBuffer instead")
-	FLightmapSceneShaderData LightmapSceneData;
-
-	UE_DEPRECATED(5.3, "Use IRendererModule::CreateSinglePrimitiveSceneUniformBuffer instead")
-	FBufferRHIRef PrimitiveSceneDataBufferRHI;
-	UE_DEPRECATED(5.3, "Use IRendererModule::CreateSinglePrimitiveSceneUniformBuffer instead")
-	FShaderResourceViewRHIRef PrimitiveSceneDataBufferSRV;
-
 	FBufferRHIRef SkyIrradianceEnvironmentMapRHI;
 	FShaderResourceViewRHIRef SkyIrradianceEnvironmentMapSRV;
 
-	UE_DEPRECATED(5.3, "Use IRendererModule::CreateSinglePrimitiveSceneUniformBuffer instead")
-	FBufferRHIRef InstanceSceneDataBufferRHI;
-	UE_DEPRECATED(5.3, "Use IRendererModule::CreateSinglePrimitiveSceneUniformBuffer instead")
-	FShaderResourceViewRHIRef InstanceSceneDataBufferSRV;
-
-	UE_DEPRECATED(5.3, "Use IRendererModule::CreateSinglePrimitiveSceneUniformBuffer instead")
-	FBufferRHIRef InstancePayloadDataBufferRHI;
-	UE_DEPRECATED(5.3, "Use IRendererModule::CreateSinglePrimitiveSceneUniformBuffer instead")
-	FShaderResourceViewRHIRef InstancePayloadDataBufferSRV;
-
 	FTextureRHIRef PrimitiveSceneDataTextureRHI;
 	FShaderResourceViewRHIRef PrimitiveSceneDataTextureSRV;
-
-	UE_DEPRECATED(5.3, "Use IRendererModule::CreateSinglePrimitiveSceneUniformBuffer instead")
-	FBufferRHIRef LightmapSceneDataBufferRHI;
-	UE_DEPRECATED(5.3, "Use IRendererModule::CreateSinglePrimitiveSceneUniformBuffer instead")
-	FShaderResourceViewRHIRef LightmapSceneDataBufferSRV;
-
-	UE_DEPRECATED(5.3, "Use IRendererModule::CreateSinglePrimitiveSceneUniformBuffer instead")
-	FBufferRHIRef EditorVisualizeLevelInstanceDataBufferRHI;
-	UE_DEPRECATED(5.3, "Use IRendererModule::CreateSinglePrimitiveSceneUniformBuffer instead")
-	FShaderResourceViewRHIRef EditorVisualizeLevelInstanceDataBufferSRV;
-
-	UE_DEPRECATED(5.3, "Use IRendererModule::CreateSinglePrimitiveSceneUniformBuffer instead")
-	FBufferRHIRef EditorSelectedDataBufferRHI;
-	UE_DEPRECATED(5.3, "Use IRendererModule::CreateSinglePrimitiveSceneUniformBuffer instead")
-	FShaderResourceViewRHIRef EditorSelectedDataBufferSRV;
 };
 
 /**
@@ -131,5 +78,3 @@ public:
 * This is used when the VF is used for rendering outside normal mesh passes, where there is no valid scene.
 */
 extern ENGINE_API TGlobalResource<FSinglePrimitiveStructured> GIdentityPrimitiveBuffer;
-UE_DEPRECATED(5.3, "Use IRendererModule::CreateSinglePrimitiveSceneUniformBuffer instead")
-extern ENGINE_API TGlobalResource<FSinglePrimitiveStructured> GTilePrimitiveBuffer;
