Fix bug caused by mesh paint virtual textures being allowed to be non power of 2 size.
This would disable virtual texturing on the textures.

diff --git a/Engine/Source/Runtime/Engine/Private/VT/MeshPaintVirtualTexture.cpp b/Engine/Source/Runtime/Engine/Private/VT/MeshPaintVirtualTexture.cpp
--- a/Engine/Source/Runtime/Engine/Private/VT/MeshPaintVirtualTexture.cpp
+++ b/Engine/Source/Runtime/Engine/Private/VT/MeshPaintVirtualTexture.cpp
@@ -73,7 +73,8 @@ namespace MeshPaintVirtualTexture
 	{
 		const int32 NumTexels = InNumVertices * CVarMeshPaintVirtualTextureTexelsPerVertex.GetValueOnGameThread();
 		const uint32 TextureSize = (uint32)FMath::Sqrt((float)NumTexels);
-		return FMath::DivideAndRoundUp(TextureSize, GetTileSize()) * GetTileSize();
+		const uint32 TextureSizePow2Aligned = FMath::RoundUpToPowerOfTwo(TextureSize);
+		return FMath::Max(GetTileSize(), TextureSizePow2Aligned);
 	}
 
 	/** 
