UE-135976 - Max Distance Parameter on Sound Wave and Source Bus is always 0.0

diff --git a/Engine/Source/Runtime/Engine/Classes/Sound/SoundBase.h b/Engine/Source/Runtime/Engine/Classes/Sound/SoundBase.h
--- a/Engine/Source/Runtime/Engine/Classes/Sound/SoundBase.h
+++ b/Engine/Source/Runtime/Engine/Classes/Sound/SoundBase.h
@@ -162,7 +162,8 @@ public:
 	UPROPERTY(Category = Developer, AssetRegistrySearchable, VisibleAnywhere, BlueprintReadOnly)
 	float Duration;
 
-	/** The max distance of the asset, as determined by attenuation settings. */
+	/** The MaxDistance property is calculated statically on load or at asset edit time, but is not reliable at runtime.
+	  * the GetMaxDistance function should be used to determine the applied max distance based on runtime behavior. */
 	UPROPERTY(Category = Developer, AssetRegistrySearchable, VisibleAnywhere, BlueprintReadOnly)
 	float MaxDistance;
 
@@ -355,4 +356,7 @@ public:
 
 	void InjectPropertySheet();
 #endif //WITH_EDITORONLY_DATA
+
+protected:
+	ENGINE_API virtual float ComputeMaxDistance() const;
 };
diff --git a/Engine/Source/Runtime/Engine/Classes/Sound/SoundCue.h b/Engine/Source/Runtime/Engine/Classes/Sound/SoundCue.h
--- a/Engine/Source/Runtime/Engine/Classes/Sound/SoundCue.h
+++ b/Engine/Source/Runtime/Engine/Classes/Sound/SoundCue.h
@@ -279,14 +279,13 @@ public:
 
 protected:
 	ENGINE_API bool RecursiveFindPathToNode(USoundNode* CurrentNode, const UPTRINT CurrentHash, const UPTRINT NodeHashToFind, TArray<USoundNode*>& OutPath) const;
+	ENGINE_API float ComputeMaxDistance() const override;
 
 private:
 	ENGINE_API void AudioQualityChanged();
 	ENGINE_API void OnPostEngineInit();
 	ENGINE_API void EvaluateNodes(bool bAddToRoot);
 
-	ENGINE_API float FindMaxDistanceInternal() const;
-
 	FDelegateHandle OnPostEngineInitHandle;
 	static ENGINE_API int32 CachedQualityLevel;
 
diff --git a/Engine/Source/Runtime/Engine/Private/SoundBase.cpp b/Engine/Source/Runtime/Engine/Private/SoundBase.cpp
--- a/Engine/Source/Runtime/Engine/Private/SoundBase.cpp
+++ b/Engine/Source/Runtime/Engine/Private/SoundBase.cpp
@@ -359,3 +359,22 @@ void USoundBase::InjectPropertySheet()
 }
 
 #endif //WITH_EDITORONLY_DATA
+
+float USoundBase::ComputeMaxDistance() const
+{
+	if (const FSoundAttenuationSettings* Settings = GetAttenuationSettingsToApply())
+	{
+		if (!Settings->bAttenuate)
+		{
+			return FAudioDevice::GetMaxWorldDistance();
+		}
+
+		const float MaxDimension = Settings->GetMaxDimension();
+		if (MaxDimension > UE_KINDA_SMALL_NUMBER)
+		{
+			return MaxDimension;
+		}
+	}
+
+	return GetMaxDistance();
+}
\ No newline at end of file
diff --git a/Engine/Source/Runtime/Engine/Private/SoundCue.cpp b/Engine/Source/Runtime/Engine/Private/SoundCue.cpp
--- a/Engine/Source/Runtime/Engine/Private/SoundCue.cpp
+++ b/Engine/Source/Runtime/Engine/Private/SoundCue.cpp
@@ -83,7 +83,7 @@ void USoundCue::CacheAggregateValues()
 			Duration = NewDuration;
 		}
 
-		MaxDistance = FindMaxDistanceInternal();
+		MaxDistance = ComputeMaxDistance();
 		bHasDelayNode = FirstNode->HasDelayNode();
 		bHasConcatenatorNode = FirstNode->HasConcatenatorNode();
 		bHasPlayWhenSilent = FirstNode->IsPlayWhenSilent();
@@ -360,7 +360,7 @@ void USoundCue::CacheQualityLevel()
 	}
 }
 
-float USoundCue::FindMaxDistanceInternal() const
+float USoundCue::ComputeMaxDistance() const
 {
 	float OutMaxDistance = 0.0f;
 	if (const FSoundAttenuationSettings* Settings = GetAttenuationSettingsToApply())
@@ -599,7 +599,7 @@ float USoundCue::GetMaxDistance() const
 	// Always recalc the max distance when in the editor as it could change
 	// from a referenced attenuation asset being updated without this cue
 	// asset re-caching the aggregate 'MaxDistance' value
-	return GIsEditor ? FindMaxDistanceInternal() : MaxDistance;
+	return GIsEditor ? ComputeMaxDistance() : MaxDistance;
 }
 
 float USoundCue::GetDuration() const
diff --git a/Engine/Source/Runtime/Engine/Private/SoundSourceBus.cpp b/Engine/Source/Runtime/Engine/Private/SoundSourceBus.cpp
--- a/Engine/Source/Runtime/Engine/Private/SoundSourceBus.cpp
+++ b/Engine/Source/Runtime/Engine/Private/SoundSourceBus.cpp
@@ -46,6 +46,7 @@ void USoundSourceBus::Init()
 void USoundSourceBus::PostEditChangeProperty(FPropertyChangedEvent& PropertyChangedEvent)
 {
 	Init();
+	MaxDistance = ComputeMaxDistance();
 }
 #endif
 
@@ -58,5 +59,3 @@ float USoundSourceBus::GetDuration() const
 {
 	return (SourceBusDuration > 0.0f) ? SourceBusDuration : INDEFINITELY_LOOPING_DURATION;
 }
-
-
diff --git a/Engine/Source/Runtime/Engine/Private/SoundWave.cpp b/Engine/Source/Runtime/Engine/Private/SoundWave.cpp
--- a/Engine/Source/Runtime/Engine/Private/SoundWave.cpp
+++ b/Engine/Source/Runtime/Engine/Private/SoundWave.cpp
@@ -1474,6 +1474,8 @@ void USoundWave::PostInitProperties()
 		InvalidateCompressedData(false, false);
 	}
 
+	MaxDistance = ComputeMaxDistance();
+
 #if WITH_EDITORONLY_DATA
 	if (!HasAnyFlags(RF_ClassDefaultObject))
 	{
@@ -2031,6 +2033,7 @@ void USoundWave::PostLoad()
 
 
 	CacheInheritedLoadingBehavior();
+	MaxDistance = ComputeMaxDistance();
 	
 	if (FApp::CanEverRenderAudio())
 	{
@@ -2367,8 +2370,6 @@ void USoundWave::RemoveAudioResource()
 	SoundWaveDataPtr->ResourceData.Empty();
 }
 
-
-
 #if WITH_EDITOR
 
 void USoundWave::InvalidateSoundWaveIfNeccessary()
@@ -3038,6 +3039,8 @@ void USoundWave::PostEditChangeProperty(FPropertyChangedEvent& PropertyChangedEv
 			}
 		}
 	}
+
+	MaxDistance = ComputeMaxDistance();
 }
 
 bool USoundWave::CanEditChange(const FProperty* InProperty) const
