[Shaders] Ignore Job Object entirely in SCW when we're running in Wine.

When the editor is running in POSIX/Wine, we only estimate the total memory of all SCW processes.
This change ensures we don't use the Job Object API at all in this environment to avoid warnings due to partially implemented Wine functions.

#rnx

diff --git a/Engine/Source/Runtime/Engine/Private/ShaderCompiler/ShaderCompilerThreadRunnable.cpp b/Engine/Source/Runtime/Engine/Private/ShaderCompiler/ShaderCompilerThreadRunnable.cpp
--- a/Engine/Source/Runtime/Engine/Private/ShaderCompiler/ShaderCompilerThreadRunnable.cpp
+++ b/Engine/Source/Runtime/Engine/Private/ShaderCompiler/ShaderCompilerThreadRunnable.cpp
@@ -656,8 +656,11 @@ bool FShaderCompileThreadRunnable::LaunchWorkersIfNeeded()
 				CurrentWorkerInfo.WorkerProcess = Manager->LaunchWorker(WorkingDirectory, Manager->ProcessId, WorkerIndex, GWorkerInputFilename, GWorkerOutputFilename);
 				CurrentWorkerInfo.bLaunchedWorker = true;
 
-				// Assign process to job object to monitor the total memory consumption of all SCW processes
-				ApplyWorkerProcessMemoryLimits(CurrentWorkerInfo.WorkerProcess);
+				// Assign process to job object to monitor the total memory consumption of all SCW processes. Ignore if we only estimate the memory due to partial support in POSIX/Wine.
+				if (!bEstimateCommittedMemory)
+				{
+					ApplyWorkerProcessMemoryLimits(CurrentWorkerInfo.WorkerProcess);
+				}
 
 				NumberLaunched++;
 			}
