Custom HLSL:  support for automatic fixup of SceneTexture and UserSceneTexture input fetches.  The suffix ".ID" can be added to the input name to generate the ID (index) token needed for calls to SceneTextureLookup or SceneTextureFetch.  And shorthand function call style syntax can be used for offset fetches, for example:  Input.Fetch(1,-1) or Input.Fetch(float2(1, -1)), generating a call to SceneTextureFetch.  Without this fixup, any time the input node changes, the custom HLSL must be simultaneously changed, possibly in multiple places, which is tedious and error prone.

diff --git a/Engine/Shaders/Private/MaterialTemplate.ush b/Engine/Shaders/Private/MaterialTemplate.ush
--- a/Engine/Shaders/Private/MaterialTemplate.ush
+++ b/Engine/Shaders/Private/MaterialTemplate.ush
@@ -3174,6 +3174,11 @@ float4 SceneTextureFetchFunc(FMaterialPixelParameters Parameters, int SceneTextu
 #endif
 }
 
+float4 SceneTextureFetchFunc(FMaterialPixelParameters Parameters, int SceneTextureIndex, float PixelOffsetX, float PixelOffsetY)
+{
+	return SceneTextureFetchFunc(Parameters, SceneTextureIndex, float2(PixelOffsetX, PixelOffsetY));
+}
+
 #define SceneTextureFetch(SceneTextureIndex, PixelOffset) SceneTextureFetchFunc(Parameters, SceneTextureIndex, PixelOffset)
 
 #endif // SHADING_PATH_DEFERRED
diff --git a/Engine/Source/Runtime/Engine/Private/Materials/HLSLMaterialTranslator.cpp b/Engine/Source/Runtime/Engine/Private/Materials/HLSLMaterialTranslator.cpp
--- a/Engine/Source/Runtime/Engine/Private/Materials/HLSLMaterialTranslator.cpp
+++ b/Engine/Source/Runtime/Engine/Private/Materials/HLSLMaterialTranslator.cpp
@@ -14263,6 +14263,336 @@ int32 FHLSLMaterialTranslator::AccessMaterialAttribute(int32 CodeIndex, const FG
 		*AttributeName);
 }
 
+enum class AsciiFlags
+{
+	TerminatorOrSlash = (1 << 0),	// Null terminator OR slash
+	Whitespace = (1 << 1),			// Includes other special characters below 32 (in addition to tab / newline)
+	Other = (1 << 2),				// Anything else not one of the other types
+	SymbolStart = (1 << 3),			// Letters plus underscore (anything that can start a symbol)
+	Digit = (1 << 4),
+	Dot = (1 << 5),
+	Quote = (1 << 6),
+	Hash = (1 << 7),
+};
+
+static uint8 AsciiFlagTable[256] =
+{
+	1,2,2,2,2,2,2,2, 2,2,2,2,2,2,2,2, 2,2,2,2,2,2,2,2, 2,2,2,2,2,2,2,2,		// Treat all special characters as whitespace
+
+	2,4,64,128,4,4,4,4,			// 34 == Quote  35 == Hash
+	4,4,4,4,4,4,32,1,			// 46 == Dot    47 == Slash
+	16,16,16,16,16,16,16,16,	// Digits 0-7
+	16,16,4,4,4,4,4,4,			// Digits 8-9
+
+	4,8,8,8,8,8,8,8, 8,8,8,8,8,8,8,8, 8,8,8,8,8,8,8,8, 8,8,8,4,4,4,4,8,		// Upper case letters,  95 == Underscore
+	4,8,8,8,8,8,8,8, 8,8,8,8,8,8,8,8, 8,8,8,8,8,8,8,8, 8,8,8,4,4,4,4,4,		// Lower case letters
+
+	4,4,4,4,4,4,4,4, 4,4,4,4,4,4,4,4, 4,4,4,4,4,4,4,4, 4,4,4,4,4,4,4,4,		// Treat all non-ASCII characters as Other
+	4,4,4,4,4,4,4,4, 4,4,4,4,4,4,4,4, 4,4,4,4,4,4,4,4, 4,4,4,4,4,4,4,4,
+	4,4,4,4,4,4,4,4, 4,4,4,4,4,4,4,4, 4,4,4,4,4,4,4,4, 4,4,4,4,4,4,4,4,
+	4,4,4,4,4,4,4,4, 4,4,4,4,4,4,4,4, 4,4,4,4,4,4,4,4, 4,4,4,4,4,4,4,4,
+};
+
+// Utility function to search for the given HLSL Identifier in a block of Code.  Takes into account comments and quoted text
+// that should be ignored when searching for identifiers.  Returns the start of the identifier, and the input pointer "Code"
+// will be advanced to the first non-whitespace, non-comment character after the end of the identifier.  Or "null" will be
+// returned if the identifier isn't found.  OutCompoundTokens contain additional tokens if it's part of a compound identifier
+// separated by periods (for "Identifier.x", OutCompoundTokens would include "x").
+static const TCHAR* FindHlslIdentifierInCode(const TCHAR*& Code, const TCHAR* Identifier, TArray<FStringView>& OutCompoundTokens)
+{
+	FStringView RootIdentifier;
+	OutCompoundTokens.Empty();
+
+	#define FETCH_ASCII_FLAG_TABLE(CHARACTER) AsciiFlagTable[(uint8)FMath::Min(CHARACTER, (TCHAR)UINT8_MAX)]
+
+	const TCHAR* SearchChar = Code;
+	uint8 SearchCharFlag = FETCH_ASCII_FLAG_TABLE(*SearchChar);
+	bool bInCompoundIdentifier = false;
+
+	// Scanning loop
+	while (1)
+	{
+		static constexpr uint8 AsciiFlagsSymbol = (uint8)AsciiFlags::SymbolStart | (uint8)AsciiFlags::Digit;
+		static constexpr uint8 AsciiFlagsStartNumberOrDirective = (uint8)AsciiFlags::Digit | (uint8)AsciiFlags::Dot | (uint8)AsciiFlags::Hash;
+		static constexpr uint8 AsciiFlagsEndNumberOrDirective = (uint8)AsciiFlags::Whitespace | (uint8)AsciiFlags::Other | (uint8)AsciiFlags::Quote | (uint8)AsciiFlags::TerminatorOrSlash;
+
+		if (SearchCharFlag & (uint8)AsciiFlags::Whitespace)
+		{
+			// Always skip over whitespace
+			SearchChar++;
+			SearchCharFlag = FETCH_ASCII_FLAG_TABLE(*SearchChar);
+		}
+		else if (SearchChar[0] == '/' && (SearchChar[1] == '/' || SearchChar[1] == '*'))
+		{
+			// Always skip over comments
+			if (SearchChar[1] == '/')
+			{
+				SearchChar += 2;
+				while (*SearchChar && *SearchChar != '\r' && *SearchChar != '\n')
+				{
+					SearchChar++;
+				}
+
+				// Could be end of string, if not, skip over newline
+				if (*SearchChar)
+				{
+					SearchChar++;
+				}
+			}
+			else
+			{
+				SearchChar += 2;
+				while (*SearchChar && (SearchChar[0] != '*' || SearchChar[1] != '/'))
+				{
+					SearchChar++;
+				}
+
+				// Could be end of string, if not, skip over comment terminator
+				if (*SearchChar)
+				{
+					SearchChar += 2;
+				}
+			}
+			SearchCharFlag = FETCH_ASCII_FLAG_TABLE(*SearchChar);
+		}
+		else if (!RootIdentifier.IsEmpty() && (SearchCharFlag & ((uint8)AsciiFlags::Dot)))
+		{
+			// Dot after an identifier starts a compound identifier
+			bInCompoundIdentifier = true;
+			SearchChar++;
+			SearchCharFlag = FETCH_ASCII_FLAG_TABLE(*SearchChar);
+		}
+		else if (!RootIdentifier.IsEmpty() && !bInCompoundIdentifier)
+		{
+			// Any other character besides whitespace, comment, or dot terminates an identifier.  Check if it matches what we are searching for.
+			if (RootIdentifier == Identifier)
+			{
+				Code = SearchChar;
+				return RootIdentifier.GetData();
+			}
+
+			// Clear the state
+			RootIdentifier = FStringView();
+			OutCompoundTokens.Empty();
+
+			// Don't advance character -- restart loop, no longer tracking an identifier
+		}
+		else if (SearchCharFlag & (uint8)AsciiFlags::SymbolStart)
+		{
+			// Start of an identifier, determine the bounds of the identifier
+			const TCHAR* IdentifierStart = SearchChar;
+
+			SearchChar++;
+			while ((SearchCharFlag = FETCH_ASCII_FLAG_TABLE(*SearchChar)) & AsciiFlagsSymbol)
+			{
+				SearchChar++;
+			}
+
+			FStringView IdentifierView(IdentifierStart, SearchChar - IdentifierStart);
+
+			// Add identifier as a compound identifier or the root identifier
+			if (bInCompoundIdentifier)
+			{
+				OutCompoundTokens.Add(IdentifierView);
+				bInCompoundIdentifier = false;
+			}
+			else
+			{
+				check(RootIdentifier.IsEmpty());
+				RootIdentifier = IdentifierView;
+			}
+		}
+		else if (!RootIdentifier.IsEmpty())
+		{
+			// If we reach this point, we are in an invalid syntax compound identifier, where an identifier token and dot are not followed
+			// by another identifier.  Clear the state and restart parsing.
+			bInCompoundIdentifier = false;
+			RootIdentifier = FStringView();
+			OutCompoundTokens.Empty();
+
+			// Don't advance character -- restart loop, no longer tracking an identifier.
+		}
+		else if ((SearchCharFlag & (uint8)AsciiFlags::Other) || (*SearchChar == '/'))
+		{
+			// Not in an identifier or comment -- skip over various characters with no special handling
+			SearchChar++;
+			SearchCharFlag = FETCH_ASCII_FLAG_TABLE(*SearchChar);
+		}
+		else if (SearchCharFlag & AsciiFlagsStartNumberOrDirective)
+		{
+			// Number or directive, skip to Whitespace, Other, or Quote (numbers may contain letters or #, i.e. "1.#INF" for infinity, or "e" for an exponent)
+			SearchChar++;
+			while (!((SearchCharFlag = FETCH_ASCII_FLAG_TABLE(*SearchChar)) & AsciiFlagsEndNumberOrDirective))
+			{
+				SearchChar++;
+			}
+		}
+		else if (SearchCharFlag & (uint8)AsciiFlags::Quote)
+		{
+			// Quote, skip to next Quote (or maybe end of string if text is malformed), ignoring the quote if it's escaped
+			SearchChar++;
+			while (*SearchChar && (*SearchChar != '\"' || *(SearchChar - 1) == '\\'))
+			{
+				SearchChar++;
+			}
+
+			// Could be end of string or the quote -- skip over the quote if not the null terminator
+			if (*SearchChar)
+			{
+				SearchChar++;
+			}
+			SearchCharFlag = FETCH_ASCII_FLAG_TABLE(*SearchChar);
+		}
+		// Must be null terminator -- we've tested all other possibilities
+		else
+		{
+			// End of string
+			Code = SearchChar;
+			OutCompoundTokens.Empty();
+			return nullptr;
+		}
+	}
+
+	#undef FETCH_ASCII_FLAG_TABLE
+}
+
+// Fixup for SceneTexture and UserSceneTexture inputs to custom HLSL.  Returns a new string, rather than modifying a Code FString
+// in place, to allow the code to be shared with the HLSL tree code path (MaterialExpressionHLSL.cpp), which uses TStringBuilder.
+// An empty string is returned if no fixup was required (the common case), avoiding reallocation.
+//
+// The goal is to allow the scene texture ID to be automatically generated in the code, without the user needing to manually
+// insert the correct identifier based on settings in the connected node.  Without this fixup, the user needs to change the
+// custom HLSL any time the input node changes, which is tedious and error prone.  For SceneTexture, if they (for example) change
+// SceneColor to Depth, they need to change PPI_SceneColor to PPI_Depth.  For UserSceneTexture, if the name is changed from Input
+// to NewInput, they need to change PPIUser_Input to PPIUser_NewInput.  Similar manual changes are also required if a completely
+// different SceneTexture node is linked.  Besides substituting the ID, minimalist shorthand ".Fetch" function call style syntax
+// is provided for offset fetches (accepts either float2 or two floats).  The original syntax of referencing "Input" without the
+// ".ID" or ".Fetch" suffixes still works for fetching a non-offset sample.
+//
+//		Input.ID -> Replace text with ID enum of connected input (PPI_* or PPIUser_*)
+//			Example usage:	SceneTextureFetch(Input.ID, float2(0,0))
+//							SceneTextureFetch(/*Input.ID=*/ PPI_SceneColor, float2(0,0))
+//
+//		Input.Fetch(*) -> Replace with SceneTextureFetchFunc(Parameters, ID, *)
+//			Example usage:	Input.Fetch(1,-1)
+//							SceneTextureFetchFunc(Parameters, /*Input=*/ PPI_SceneColor, 1,-1)
+//
+// Note that this text replacement happens before preprocessor macro substitution, and so it's not possible to use the shorthand
+// syntax with the user specified "InputName" as a macro argument, but you can pass "InputName.ID", for example:
+//
+//		// Shorthand syntax will not work, because substitution of "InputName" as the "Input" macro argument happens later in preprocess
+//		#define WEIGHTED_FETCH_SCENE_TEXTURE(Input, SampleIndex) \
+//			(Input.Fetch(Offsets[SampleIndex], 0) * Weights[SampleIndex])
+//
+//		Color += WEIGHTED_FETCH_SCENE_TEXTURE(InputName, SampleIndex);
+//
+//		// But the syntax that accepts a scene texture ID works fine
+//		#define WEIGHTED_FETCH_SCENE_TEXTURE(InputID, SampleIndex) \
+//			(SceneTextureFetch(InputID, float2(Offsets[SampleIndex], 0)) * Weights[SampleIndex])
+//
+//		Color += WEIGHTED_FETCH_SCENE_TEXTURE(InputName.ID, SampleIndex);
+//
+//		// This also works fine, because InputName is inside the macro, and not a macro argument
+//		#define WEIGHTED_FETCH_SCENE_TEXTURE(SampleIndex) \
+//			(InputName.Fetch(Offsets[SampleIndex], 0) * Weights[SampleIndex])
+//
+FString CustomExpressionSceneTextureInputFixup(const UMaterialExpressionCustom* Custom, const TCHAR* Code)
+{
+	// Allocated when required fixup is first encountered
+	FString ModifiedCode;
+
+	for (const FCustomInput& Input : Custom->Inputs)
+	{
+		if (Input.InputName.IsNone())
+		{
+			continue;
+		}
+
+		FExpressionInput ExpressionInput = Input.Input.GetTracedInput();
+		const UMaterialExpression* Expression = ExpressionInput.Expression;
+		if (!Expression || ExpressionInput.OutputIndex != 0 || !(Expression->IsA<UMaterialExpressionSceneTexture>() || Expression->IsA<UMaterialExpressionUserSceneTexture>()))
+		{
+			continue;
+		}
+
+		FString InputNameString = Input.InputName.ToString();
+		FString InputIDString;
+
+		if (const UMaterialExpressionSceneTexture* SceneTextureExpression = Cast<UMaterialExpressionSceneTexture>(Expression))
+		{
+			InputIDString = StaticEnum<ESceneTextureId>()->GetNameStringByValue(static_cast<int>(SceneTextureExpression->SceneTextureId));
+		}
+		else
+		{
+			const UMaterialExpressionUserSceneTexture* UserSceneTextureExpression = Cast<UMaterialExpressionUserSceneTexture>(Expression);
+			InputIDString = TEXT("PPIUser_") + UserSceneTextureExpression->UserSceneTexture.ToString();
+		}
+
+		int32 CodeOffset = 0;
+		while (1)
+		{
+			TArray<FStringView> CompoundTokens;
+			const TCHAR* CodeStart = ModifiedCode.IsEmpty() ? Code : *ModifiedCode;
+			const TCHAR* CodeString = &CodeStart[CodeOffset];
+			const TCHAR* FoundInputName = FindHlslIdentifierInCode(CodeString, *InputNameString, CompoundTokens);
+
+			if (!FoundInputName)
+			{
+				break;
+			}
+
+			CodeOffset = FoundInputName - CodeStart;
+
+			int32 OriginalCodeOffset = CodeOffset;
+			if (CompoundTokens.Num() == 1)
+			{
+				// Check if this is using the "Fetch" function call shorthand syntax.  If so, replace it with the corresponding fetch function.
+				if (CompoundTokens[0] == TEXT("Fetch") && *CodeString == '(')
+				{
+					FString ReplacementText = FString::Printf(TEXT("SceneTextureFetchFunc(Parameters, /*%s=*/ %s, "), *InputNameString, *InputIDString);
+
+					if (ModifiedCode.IsEmpty())
+					{
+						ModifiedCode.Reserve(FCString::Strlen(Code) + ReplacementText.Len());
+						ModifiedCode = Code;
+					}
+
+					// Add one character to amount removed, as we are overwriting the open parentheses as well
+					ModifiedCode.RemoveAt(CodeOffset, CodeString - FoundInputName + 1, EAllowShrinking::No);
+					ModifiedCode.InsertAt(CodeOffset, ReplacementText);
+
+					CodeOffset += ReplacementText.Len();
+				}
+				else if (CompoundTokens[0] == TEXT("ID"))
+				{
+					// Replace "Input.ID" identifier with ID (including comment for readability if there are compile errors)
+					FString ReplacementText = FString::Printf(TEXT("/*%s.ID=*/ %s"), *InputNameString, *InputIDString);
+
+					if (ModifiedCode.IsEmpty())
+					{
+						ModifiedCode.Reserve(FCString::Strlen(Code) + ReplacementText.Len());
+						ModifiedCode = Code;
+					}
+
+					ModifiedCode.RemoveAt(CodeOffset, CompoundTokens[0].GetData() + CompoundTokens[0].Len() - FoundInputName, EAllowShrinking::No);
+					ModifiedCode.InsertAt(CodeOffset, ReplacementText);
+
+					CodeOffset += ReplacementText.Len();
+				}
+			}
+
+			// If we didn't insert and skip over replacement code, skip over the token that was found (CodeString will have
+			// been set past the end of the token by the function "FindHlslIdentifierInCode").
+			if (CodeOffset == OriginalCodeOffset)
+			{
+				CodeOffset = CodeString - CodeStart;
+			}
+		}
+	}
+
+	return ModifiedCode;
+}
+
 int32 FHLSLMaterialTranslator::CustomExpression( class UMaterialExpressionCustom* Custom, int32 OutputIndex, TArray<int32>& CompiledInputs )
 {
 	const FMaterialCustomExpressionEntry* CustomEntry = nullptr;
@@ -14466,6 +14796,12 @@ int32 FHLSLMaterialTranslator::CustomExpression( class UMaterialExpressionCustom
 		}
 		Code.ReplaceInline(TEXT("\r\n"), TEXT("\n"), ESearchCase::CaseSensitive);
 
+		FString ModifiedCode = CustomExpressionSceneTextureInputFixup(Custom, *Code);
+		if (!ModifiedCode.IsEmpty())
+		{
+			Code = ModifiedCode;
+		}
+
 		FString ParametersType = ShaderFrequency == SF_Vertex ? TEXT("Vertex") : TEXT("Pixel");
 
 		FMaterialCustomExpressionEntry& Entry = CustomExpressions.AddDefaulted_GetRef();
diff --git a/Engine/Source/Runtime/Engine/Private/Materials/MaterialExpressionHLSL.cpp b/Engine/Source/Runtime/Engine/Private/Materials/MaterialExpressionHLSL.cpp
--- a/Engine/Source/Runtime/Engine/Private/Materials/MaterialExpressionHLSL.cpp
+++ b/Engine/Source/Runtime/Engine/Private/Materials/MaterialExpressionHLSL.cpp
@@ -4569,6 +4569,8 @@ bool UMaterialExpressionObjectRadius::GenerateHLSLExpression(FMaterialHLSLGenera
 	return true;
 }
 
+extern FString CustomExpressionSceneTextureInputFixup(const UMaterialExpressionCustom* Custom, const TCHAR* Code);
+
 bool UMaterialExpressionCustom::GenerateHLSLExpression(FMaterialHLSLGenerator& Generator, UE::HLSLTree::FScope& Scope, int32 OutputIndex, UE::HLSLTree::FExpression const*& OutExpression) const
 {
 	using namespace UE::HLSLTree;
@@ -4637,18 +4639,28 @@ bool UMaterialExpressionCustom::GenerateHLSLExpression(FMaterialHLSLGenerator& G
 		}
 	}
 
+	TStringBuilder<8 * 1024> FormattedCode;
+	if (!Code.Contains(TEXT("return")))
+	{
+		FormattedCode.Appendf(TEXT("return %s;"), *Code);
+	}
+
+	FString SceneTextureFixupCode = CustomExpressionSceneTextureInputFixup(this, FormattedCode.Len() ? FormattedCode.ToString() : *Code);
+
 	FStringView FunctionCode;
-	if (Code.Contains(TEXT("return")))
+	if (SceneTextureFixupCode.Len())
 	{
-		// Can just reference to 'Code' field directly, the UMaterialExpressionCustom lifetime will be longer than the resulting HLSLTree
-		FunctionCode = Code;
+		FunctionCode = UE::MemStack::AllocateStringView(Allocator, FStringView(SceneTextureFixupCode));
 	}
-	else
+	else if (FormattedCode.Len())
 	{
-		TStringBuilder<8 * 1024> FormattedCode;
-		FormattedCode.Appendf(TEXT("return %s;"), *Code);
 		FunctionCode = UE::MemStack::AllocateStringView(Allocator, FormattedCode.ToView());
 	}
+	else
+	{
+		// Can just reference to 'Code' field directly, the UMaterialExpressionCustom lifetime will be longer than the resulting HLSLTree
+		FunctionCode = Code;
+	}
 
 	const FExpression* ExpressionCustom = Generator.GetTree().NewExpression<FExpressionCustomHLSL>(
 		UE::MemStack::AllocateStringView(Allocator, DeclarationCode.ToView()),
