Adjust the WorldPartitionEditorSoftObjectPathTest since 35057571.

diff --git a/Engine/Source/Runtime/Engine/Private/Tests/WorldPartition/WorldPartitionEditorSoftObjectPathTests.cpp b/Engine/Source/Runtime/Engine/Private/Tests/WorldPartition/WorldPartitionEditorSoftObjectPathTests.cpp
--- a/Engine/Source/Runtime/Engine/Private/Tests/WorldPartition/WorldPartitionEditorSoftObjectPathTests.cpp
+++ b/Engine/Source/Runtime/Engine/Private/Tests/WorldPartition/WorldPartitionEditorSoftObjectPathTests.cpp
@@ -114,7 +114,7 @@ namespace WorldPartitionTests
 			FSoftObjectPath RuntimePath;
 			FSoftObjectPath EditorPath(TEXT("/Engine/WorldPartition/UnitTests/EditorPaths/EditorPaths.EditorPaths:PersistentLevel.LevelInstance_UAID_E04F43E60CCA14A401_2106307871.StaticMeshActor_UAID_E04F43E60CCA14A401_2106052870"));
 			FWorldPartitionHelpers::ConvertEditorPathToRuntimePath(EditorPath, RuntimePath);
-			if (!TestEqual(TEXT("Invalid runtime path"), RuntimePath, FSoftObjectPath("/Engine/WorldPartition/UnitTests/EditorPaths/EditorPaths/_Generated_/2WBX502V75SLTW3UKBT9S2GG6.EditorPaths:PersistentLevel.StaticMeshActor_UAID_E04F43E60CCA14A401_2106052870_b3e12442c23432cb")))
+			if (!TestEqual(TEXT("Invalid runtime path"), RuntimePath, FSoftObjectPath("/Engine/WorldPartition/UnitTests/EditorPaths/EditorPaths/_Generated_/2WBX502V75SLTW3UKBT9S2GG6.EditorPaths:PersistentLevel.StaticMeshActor_UAID_E04F43E60CCA14A401_b3e12442c23432cb_2106052870")))
 			{
 				return false;
 			}
@@ -122,7 +122,7 @@ namespace WorldPartitionTests
 			// Test conversion of Editor Path to Runtime Path (1 Level deep StaticMeshComponent)	
 			EditorPath = TEXT("/Engine/WorldPartition/UnitTests/EditorPaths/EditorPaths.EditorPaths:PersistentLevel.LevelInstance_UAID_E04F43E60CCA14A401_2106307871.StaticMeshActor_UAID_E04F43E60CCA14A401_2106052870.StaticMeshComponent0");
 			FWorldPartitionHelpers::ConvertEditorPathToRuntimePath(EditorPath, RuntimePath);
-			if (!TestEqual(TEXT("Invalid runtime path"), RuntimePath, FSoftObjectPath("/Engine/WorldPartition/UnitTests/EditorPaths/EditorPaths/_Generated_/2WBX502V75SLTW3UKBT9S2GG6.EditorPaths:PersistentLevel.StaticMeshActor_UAID_E04F43E60CCA14A401_2106052870_b3e12442c23432cb.StaticMeshComponent0")))
+			if (!TestEqual(TEXT("Invalid runtime path"), RuntimePath, FSoftObjectPath("/Engine/WorldPartition/UnitTests/EditorPaths/EditorPaths/_Generated_/2WBX502V75SLTW3UKBT9S2GG6.EditorPaths:PersistentLevel.StaticMeshActor_UAID_E04F43E60CCA14A401_b3e12442c23432cb_2106052870.StaticMeshComponent0")))
 			{
 				return false;
 			}
@@ -130,7 +130,7 @@ namespace WorldPartitionTests
 			// Test conversion of Editor Path to Runtime Path (1 Level deep StaticMeshActor, different top level instance)
 			EditorPath =(TEXT("/Engine/WorldPartition/UnitTests/EditorPaths/EditorPaths.EditorPaths:PersistentLevel.LevelInstance_UAID_E04F43E60CCA16A401_1435304224.StaticMeshActor_UAID_E04F43E60CCA14A401_2106052870"));
 			FWorldPartitionHelpers::ConvertEditorPathToRuntimePath(EditorPath, RuntimePath);
-			if (!TestEqual(TEXT("Invalid runtime path"), RuntimePath, FSoftObjectPath("/Engine/WorldPartition/UnitTests/EditorPaths/EditorPaths/_Generated_/2WBX502V75SLTW3UKBT9S2GG6.EditorPaths:PersistentLevel.StaticMeshActor_UAID_E04F43E60CCA14A401_2106052870_855651f5ca95d6b9")))
+			if (!TestEqual(TEXT("Invalid runtime path"), RuntimePath, FSoftObjectPath("/Engine/WorldPartition/UnitTests/EditorPaths/EditorPaths/_Generated_/2WBX502V75SLTW3UKBT9S2GG6.EditorPaths:PersistentLevel.StaticMeshActor_UAID_E04F43E60CCA14A401_855651f5ca95d6b9_2106052870")))
 			{
 				return false;
 			}
@@ -138,7 +138,7 @@ namespace WorldPartitionTests
 			// Test conversion of Editor Path to Runtime Path (1 Level deep StaticMeshComponent, different top level instance)	
 			EditorPath = TEXT("/Engine/WorldPartition/UnitTests/EditorPaths/EditorPaths.EditorPaths:PersistentLevel.LevelInstance_UAID_E04F43E60CCA16A401_1435304224.StaticMeshActor_UAID_E04F43E60CCA14A401_2106052870.StaticMeshComponent0");
 			FWorldPartitionHelpers::ConvertEditorPathToRuntimePath(EditorPath, RuntimePath);
-			if (!TestEqual(TEXT("Invalid runtime path"), RuntimePath, FSoftObjectPath("/Engine/WorldPartition/UnitTests/EditorPaths/EditorPaths/_Generated_/2WBX502V75SLTW3UKBT9S2GG6.EditorPaths:PersistentLevel.StaticMeshActor_UAID_E04F43E60CCA14A401_2106052870_855651f5ca95d6b9.StaticMeshComponent0")))
+			if (!TestEqual(TEXT("Invalid runtime path"), RuntimePath, FSoftObjectPath("/Engine/WorldPartition/UnitTests/EditorPaths/EditorPaths/_Generated_/2WBX502V75SLTW3UKBT9S2GG6.EditorPaths:PersistentLevel.StaticMeshActor_UAID_E04F43E60CCA14A401_855651f5ca95d6b9_2106052870.StaticMeshComponent0")))
 			{
 				return false;
 			}
@@ -146,7 +146,7 @@ namespace WorldPartitionTests
 			// Test conversion of Editor Path to Runtime Path (2 Level deep StaticMeshActor)	
 			EditorPath = TEXT("/Engine/WorldPartition/UnitTests/EditorPaths/EditorPaths.EditorPaths:PersistentLevel.LevelInstance_UAID_E04F43E60CCA14A401_2106307871.LevelInstance_UAID_E04F43E60CCA14A401_2145683877.StaticMeshActor_UAID_E04F43E60CCA14A401_2145611875");
 			FWorldPartitionHelpers::ConvertEditorPathToRuntimePath(EditorPath, RuntimePath);
-			if (!TestEqual(TEXT("Invalid runtime path"), RuntimePath, FSoftObjectPath("/Engine/WorldPartition/UnitTests/EditorPaths/EditorPaths/_Generated_/2WBX502V75SLTW3UKBT9S2GG6.EditorPaths:PersistentLevel.StaticMeshActor_UAID_E04F43E60CCA14A401_2145611875_023ce1270d5e3394")))
+			if (!TestEqual(TEXT("Invalid runtime path"), RuntimePath, FSoftObjectPath("/Engine/WorldPartition/UnitTests/EditorPaths/EditorPaths/_Generated_/2WBX502V75SLTW3UKBT9S2GG6.EditorPaths:PersistentLevel.StaticMeshActor_UAID_E04F43E60CCA14A401_023ce1270d5e3394_2145611875")))
 			{
 				return false;
 			}
@@ -154,7 +154,7 @@ namespace WorldPartitionTests
 			// Test conversion of Editor Path to Runtime Path (2 Level deep StaticMeshComponent)	
 			EditorPath = TEXT("/Engine/WorldPartition/UnitTests/EditorPaths/EditorPaths.EditorPaths:PersistentLevel.LevelInstance_UAID_E04F43E60CCA14A401_2106307871.LevelInstance_UAID_E04F43E60CCA14A401_2145683877.StaticMeshActor_UAID_E04F43E60CCA14A401_2145611875.StaticMeshComponent0");
 			FWorldPartitionHelpers::ConvertEditorPathToRuntimePath(EditorPath, RuntimePath);
-			if (!TestEqual(TEXT("Invalid runtime path"), RuntimePath, FSoftObjectPath("/Engine/WorldPartition/UnitTests/EditorPaths/EditorPaths/_Generated_/2WBX502V75SLTW3UKBT9S2GG6.EditorPaths:PersistentLevel.StaticMeshActor_UAID_E04F43E60CCA14A401_2145611875_023ce1270d5e3394.StaticMeshComponent0")))
+			if (!TestEqual(TEXT("Invalid runtime path"), RuntimePath, FSoftObjectPath("/Engine/WorldPartition/UnitTests/EditorPaths/EditorPaths/_Generated_/2WBX502V75SLTW3UKBT9S2GG6.EditorPaths:PersistentLevel.StaticMeshActor_UAID_E04F43E60CCA14A401_023ce1270d5e3394_2145611875.StaticMeshComponent0")))
 			{
 				return false;
 			}
@@ -162,7 +162,7 @@ namespace WorldPartitionTests
 			// Test conversion of Editor Path to Runtime Path (2 Level deep StaticMeshActor, different top level instance)
 			EditorPath = TEXT("/Engine/WorldPartition/UnitTests/EditorPaths/EditorPaths.EditorPaths:PersistentLevel.LevelInstance_UAID_E04F43E60CCA16A401_1435304224.LevelInstance_UAID_E04F43E60CCA14A401_2145683877.StaticMeshActor_UAID_E04F43E60CCA14A401_2145611875");
 			FWorldPartitionHelpers::ConvertEditorPathToRuntimePath(EditorPath, RuntimePath);
-			if (!TestEqual(TEXT("Invalid runtime path"), RuntimePath, FSoftObjectPath("/Engine/WorldPartition/UnitTests/EditorPaths/EditorPaths/_Generated_/2WBX502V75SLTW3UKBT9S2GG6.EditorPaths:PersistentLevel.StaticMeshActor_UAID_E04F43E60CCA14A401_2145611875_b48c4d72c8d10646")))
+			if (!TestEqual(TEXT("Invalid runtime path"), RuntimePath, FSoftObjectPath("/Engine/WorldPartition/UnitTests/EditorPaths/EditorPaths/_Generated_/2WBX502V75SLTW3UKBT9S2GG6.EditorPaths:PersistentLevel.StaticMeshActor_UAID_E04F43E60CCA14A401_b48c4d72c8d10646_2145611875")))
 			{
 				return false;
 			}
@@ -170,7 +170,7 @@ namespace WorldPartitionTests
 			// Test conversion of Editor Path to Runtime Path (2 Level deep StaticMeshComponent, different top level instance)	
 			EditorPath = TEXT("/Engine/WorldPartition/UnitTests/EditorPaths/EditorPaths.EditorPaths:PersistentLevel.LevelInstance_UAID_E04F43E60CCA16A401_1435304224.LevelInstance_UAID_E04F43E60CCA14A401_2145683877.StaticMeshActor_UAID_E04F43E60CCA14A401_2145611875.StaticMeshComponent0");
 			FWorldPartitionHelpers::ConvertEditorPathToRuntimePath(EditorPath, RuntimePath);
-			if (!TestEqual(TEXT("Invalid runtime path"), RuntimePath, FSoftObjectPath("/Engine/WorldPartition/UnitTests/EditorPaths/EditorPaths/_Generated_/2WBX502V75SLTW3UKBT9S2GG6.EditorPaths:PersistentLevel.StaticMeshActor_UAID_E04F43E60CCA14A401_2145611875_b48c4d72c8d10646.StaticMeshComponent0")))
+			if (!TestEqual(TEXT("Invalid runtime path"), RuntimePath, FSoftObjectPath("/Engine/WorldPartition/UnitTests/EditorPaths/EditorPaths/_Generated_/2WBX502V75SLTW3UKBT9S2GG6.EditorPaths:PersistentLevel.StaticMeshActor_UAID_E04F43E60CCA14A401_b48c4d72c8d10646_2145611875.StaticMeshComponent0")))
 			{
 				return false;
 			}
@@ -178,7 +178,7 @@ namespace WorldPartitionTests
 			// Test conversion of Editor Path to Runtime Path (3 Levels deep StaticMeshActor)	
 			EditorPath = TEXT("/Engine/WorldPartition/UnitTests/EditorPaths/EditorPaths.EditorPaths:PersistentLevel.LevelInstance_UAID_E04F43E60CCA14A401_2106307871.LevelInstance_UAID_E04F43E60CCA14A401_2145683877.LevelInstance_UAID_E04F43E60CCA16A401_1407795222.StaticMeshActor_UAID_E04F43E60CCA16A401_1407522221");
 			FWorldPartitionHelpers::ConvertEditorPathToRuntimePath(EditorPath, RuntimePath);
-			if (!TestEqual(TEXT("Invalid runtime path"), RuntimePath, FSoftObjectPath("/Engine/WorldPartition/UnitTests/EditorPaths/EditorPaths/_Generated_/2WBX502V75SLTW3UKBT9S2GG6.EditorPaths:PersistentLevel.StaticMeshActor_UAID_E04F43E60CCA16A401_1407522221_c72e7f8fcd5f70c5")))
+			if (!TestEqual(TEXT("Invalid runtime path"), RuntimePath, FSoftObjectPath("/Engine/WorldPartition/UnitTests/EditorPaths/EditorPaths/_Generated_/2WBX502V75SLTW3UKBT9S2GG6.EditorPaths:PersistentLevel.StaticMeshActor_UAID_E04F43E60CCA16A401_c72e7f8fcd5f70c5_1407522221")))
 			{
 				return false;
 			}
@@ -186,7 +186,7 @@ namespace WorldPartitionTests
 			// Test conversion of Editor Path to Runtime Path (3 Levels deep StaticMeshComponent)	
 			EditorPath = TEXT("/Engine/WorldPartition/UnitTests/EditorPaths/EditorPaths.EditorPaths:PersistentLevel.LevelInstance_UAID_E04F43E60CCA14A401_2106307871.LevelInstance_UAID_E04F43E60CCA14A401_2145683877.LevelInstance_UAID_E04F43E60CCA16A401_1407795222.StaticMeshActor_UAID_E04F43E60CCA16A401_1407522221.StaticMeshComponent0");
 			FWorldPartitionHelpers::ConvertEditorPathToRuntimePath(EditorPath, RuntimePath);
-			if (!TestEqual(TEXT("Invalid runtime path"), RuntimePath, FSoftObjectPath("/Engine/WorldPartition/UnitTests/EditorPaths/EditorPaths/_Generated_/2WBX502V75SLTW3UKBT9S2GG6.EditorPaths:PersistentLevel.StaticMeshActor_UAID_E04F43E60CCA16A401_1407522221_c72e7f8fcd5f70c5.StaticMeshComponent0")))
+			if (!TestEqual(TEXT("Invalid runtime path"), RuntimePath, FSoftObjectPath("/Engine/WorldPartition/UnitTests/EditorPaths/EditorPaths/_Generated_/2WBX502V75SLTW3UKBT9S2GG6.EditorPaths:PersistentLevel.StaticMeshActor_UAID_E04F43E60CCA16A401_c72e7f8fcd5f70c5_1407522221.StaticMeshComponent0")))
 			{
 				return false;
 			}
@@ -194,7 +194,7 @@ namespace WorldPartitionTests
 			// Test conversion of Editor Path to Runtime Path (3 Levels deep StaticMeshActor, different top level instance)	
 			EditorPath = TEXT("/Engine/WorldPartition/UnitTests/EditorPaths/EditorPaths.EditorPaths:PersistentLevel.LevelInstance_UAID_E04F43E60CCA16A401_1435304224.LevelInstance_UAID_E04F43E60CCA14A401_2145683877.LevelInstance_UAID_E04F43E60CCA16A401_1407795222.StaticMeshActor_UAID_E04F43E60CCA16A401_1407522221");
 			FWorldPartitionHelpers::ConvertEditorPathToRuntimePath(EditorPath, RuntimePath);
-			if (!TestEqual(TEXT("Invalid runtime path"), RuntimePath, FSoftObjectPath("/Engine/WorldPartition/UnitTests/EditorPaths/EditorPaths/_Generated_/2WBX502V75SLTW3UKBT9S2GG6.EditorPaths:PersistentLevel.StaticMeshActor_UAID_E04F43E60CCA16A401_1407522221_0f31c5c4ff7639d5")))
+			if (!TestEqual(TEXT("Invalid runtime path"), RuntimePath, FSoftObjectPath("/Engine/WorldPartition/UnitTests/EditorPaths/EditorPaths/_Generated_/2WBX502V75SLTW3UKBT9S2GG6.EditorPaths:PersistentLevel.StaticMeshActor_UAID_E04F43E60CCA16A401_0f31c5c4ff7639d5_1407522221")))
 			{
 				return false;
 			}
@@ -202,7 +202,7 @@ namespace WorldPartitionTests
 			// Test conversion of Editor Path to Runtime Path (3 Levels deep StaticMeshComponent, different top level instance)	
 			EditorPath = TEXT("/Engine/WorldPartition/UnitTests/EditorPaths/EditorPaths.EditorPaths:PersistentLevel.LevelInstance_UAID_E04F43E60CCA16A401_1435304224.LevelInstance_UAID_E04F43E60CCA14A401_2145683877.LevelInstance_UAID_E04F43E60CCA16A401_1407795222.StaticMeshActor_UAID_E04F43E60CCA16A401_1407522221.StaticMeshComponent0");
 			FWorldPartitionHelpers::ConvertEditorPathToRuntimePath(EditorPath, RuntimePath);
-			if (!TestEqual(TEXT("Invalid runtime path"), RuntimePath, FSoftObjectPath("/Engine/WorldPartition/UnitTests/EditorPaths/EditorPaths/_Generated_/2WBX502V75SLTW3UKBT9S2GG6.EditorPaths:PersistentLevel.StaticMeshActor_UAID_E04F43E60CCA16A401_1407522221_0f31c5c4ff7639d5.StaticMeshComponent0")))
+			if (!TestEqual(TEXT("Invalid runtime path"), RuntimePath, FSoftObjectPath("/Engine/WorldPartition/UnitTests/EditorPaths/EditorPaths/_Generated_/2WBX502V75SLTW3UKBT9S2GG6.EditorPaths:PersistentLevel.StaticMeshActor_UAID_E04F43E60CCA16A401_0f31c5c4ff7639d5_1407522221.StaticMeshComponent0")))
 			{
 				return false;
 			}
