Fixing up a regression in the Visual Logger where Actors failed to have unique names with bForceUniqueLogNames set.

diff --git a/Engine/Source/Runtime/Engine/Private/VisualLogger/VisualLogger.cpp b/Engine/Source/Runtime/Engine/Private/VisualLogger/VisualLogger.cpp
--- a/Engine/Source/Runtime/Engine/Private/VisualLogger/VisualLogger.cpp
+++ b/Engine/Source/Runtime/Engine/Private/VisualLogger/VisualLogger.cpp
@@ -513,10 +513,11 @@ FVisualLogEntry* FVisualLogger::GetEntryToWriteInternal(const UObject* Object, c
 
 		const UWorld* World = GetWorldForVisualLogger(LogOwner);
 		const bool bIsStandalone = (World == nullptr || World->GetNetMode() == NM_Standalone);
-		const FName LogName(*FString::Printf(TEXT("%s%s%s"),
+		const FName LogName(*FString::Printf(TEXT("%s%s%s%s"),
 			*UniqueLogPrefix,
 			bIsStandalone ? TEXT("") : *FString::Printf(TEXT("(%s) "), *GetDebugStringForWorld(World)),
-			*LogOwner->GetName()));
+			*LogOwner->GetName(),
+			bForceUniqueLogNames ? *FString::Printf(TEXT(" [%d]"), LogOwner->GetUniqueID()) : TEXT("")));
 
 		ObjectToNameMap.Add(LogOwner, LogName);
 		ObjectToClassNameMap.Add(LogOwner, *(LogOwner->GetClass()->GetName()));
