Use internal tracking to track when OnBeginPlay is called on the world partition subsytem.  There can be a significant amount of time between the dispatch and the world's boolean getting set.  This resulted in the streaming hash skipping initialization.

diff --git a/Engine/Source/Runtime/Engine/Private/WorldPartition/WorldPartitionSubsystem.cpp b/Engine/Source/Runtime/Engine/Private/WorldPartition/WorldPartitionSubsystem.cpp
--- a/Engine/Source/Runtime/Engine/Private/WorldPartition/WorldPartitionSubsystem.cpp
+++ b/Engine/Source/Runtime/Engine/Private/WorldPartition/WorldPartitionSubsystem.cpp
@@ -565,7 +565,7 @@ void UWorldPartitionSubsystem::OnWorldPartitionInitialized(UWorldPartition* InWo
 		}
 	}
 
-	if (const UWorld* OwningWorld = GetWorld(); OwningWorld && OwningWorld->HasBegunPlay())
+	if (bHasBegunPlay)
 	{
 		InWorldPartition->OnBeginPlay();
 	}
@@ -1095,6 +1095,7 @@ void UWorldPartitionSubsystem::OnWorldBeginPlay(UWorld& InWorld)
 {
 	Super::OnWorldBeginPlay(InWorld);
 
+	bHasBegunPlay = true;
 	for (UWorldPartition* RegisteredWorldPartition : RegisteredWorldPartitions)
 	{
 		RegisteredWorldPartition->OnBeginPlay();
@@ -1710,4 +1711,4 @@ FVector FStreamingSourceVelocity::GetAverageVelocity(const FVector& NewPosition,
 	}
 
 	return AvgVelocity;
-}
\ No newline at end of file
+}
diff --git a/Engine/Source/Runtime/Engine/Public/WorldPartition/WorldPartitionSubsystem.h b/Engine/Source/Runtime/Engine/Public/WorldPartition/WorldPartitionSubsystem.h
--- a/Engine/Source/Runtime/Engine/Public/WorldPartition/WorldPartitionSubsystem.h
+++ b/Engine/Source/Runtime/Engine/Public/WorldPartition/WorldPartitionSubsystem.h
@@ -210,6 +210,8 @@ private:
 	// Tracks world partition loading and pending loads
 	TSet<TWeakObjectPtr<const ULevelStreaming>> WorldPartitionLoadingAndPendingLoadStreamingLevels;
 
+	bool bHasBegunPlay = false;
+
 #if WITH_EDITOR
 	bool bIsRunningConvertWorldPartitionCommandlet;
 #endif
