Fix the UnrealEditor for Linux build failures caused by the returns_nonnull change.

diff --git a/Engine/Plugins/Animation/ControlRig/Source/ControlRigEditor/Private/Editor/SControlRigValidationWidget.cpp b/Engine/Plugins/Animation/ControlRig/Source/ControlRigEditor/Private/Editor/SControlRigValidationWidget.cpp
--- a/Engine/Plugins/Animation/ControlRig/Source/ControlRigEditor/Private/Editor/SControlRigValidationWidget.cpp
+++ b/Engine/Plugins/Animation/ControlRig/Source/ControlRigEditor/Private/Editor/SControlRigValidationWidget.cpp
@@ -118,7 +118,7 @@ void SControlRigValidationWidget::Construct(const FArguments& InArgs, UControlRi
 	ClassItems.Reset();
 	for (TObjectIterator<UClass> ClassIterator; ClassIterator; ++ClassIterator)
 	{
-		const bool bIsValidationPassChild = (*ClassIterator) && ClassIterator->IsChildOf(UControlRigValidationPass::StaticClass());
+		const bool bIsValidationPassChild = ClassIterator->IsChildOf(UControlRigValidationPass::StaticClass());
 		if (bIsValidationPassChild && !ClassIterator->HasAnyClassFlags(CLASS_Abstract | CLASS_Deprecated | CLASS_NewerVersionExists))
 		{
 			ClassItems.Add(MakeShared<FControlRigValidationPassItem>(*ClassIterator));
diff --git a/Engine/Plugins/Runtime/RigVM/Source/RigVMDeveloper/Private/RigVMBlueprintUtils.cpp b/Engine/Plugins/Runtime/RigVM/Source/RigVMDeveloper/Private/RigVMBlueprintUtils.cpp
--- a/Engine/Plugins/Runtime/RigVM/Source/RigVMDeveloper/Private/RigVMBlueprintUtils.cpp
+++ b/Engine/Plugins/Runtime/RigVM/Source/RigVMDeveloper/Private/RigVMBlueprintUtils.cpp
@@ -81,14 +81,11 @@ void FRigVMBlueprintUtils::ForAllRigVMStructs(TFunction<void(UScriptStruct*)> In
 	// Run over all unit types
 	for(TObjectIterator<UStruct> StructIt; StructIt; ++StructIt)
 	{
-		if (*StructIt)
+		if(StructIt->IsChildOf(FRigVMStruct::StaticStruct()) && !StructIt->HasMetaData(FRigVMStruct::AbstractMetaName))
 		{
-			if(StructIt->IsChildOf(FRigVMStruct::StaticStruct()) && !StructIt->HasMetaData(FRigVMStruct::AbstractMetaName))
+			if (UScriptStruct* ScriptStruct = Cast<UScriptStruct>(*StructIt))
 			{
-				if (UScriptStruct* ScriptStruct = Cast<UScriptStruct>(*StructIt))
-				{
-					InFunction(ScriptStruct);
-				}
+				InFunction(ScriptStruct);
 			}
 		}
 	}
diff --git a/Engine/Source/Runtime/AudioMixer/Private/AudioMixerDevice.cpp b/Engine/Source/Runtime/AudioMixer/Private/AudioMixerDevice.cpp
--- a/Engine/Source/Runtime/AudioMixer/Private/AudioMixerDevice.cpp
+++ b/Engine/Source/Runtime/AudioMixer/Private/AudioMixerDevice.cpp
@@ -812,7 +812,7 @@ namespace Audio
 
 				for (TObjectIterator<USoundSubmix> It; It; ++It)
 				{
-					if (*It && It->bMuteWhenBackgrounded)
+					if (It->bMuteWhenBackgrounded)
 					{
 						FMixerSubmixPtr SubmixInstance = GetSubmixInstance(*It).Pin();
 						if (SubmixInstance.IsValid())
diff --git a/Engine/Source/Runtime/Engine/Private/Animation/SkinWeightProfile.cpp b/Engine/Source/Runtime/Engine/Private/Animation/SkinWeightProfile.cpp
--- a/Engine/Source/Runtime/Engine/Private/Animation/SkinWeightProfile.cpp
+++ b/Engine/Source/Runtime/Engine/Private/Animation/SkinWeightProfile.cpp
@@ -82,23 +82,20 @@ static void OnDefaultProfileCVarsChanged(IConsoleVariable* Variable)
 
 			for (TObjectIterator<USkeletalMesh> It; It; ++It)
 			{
-				if (*It)
+				if (FSkeletalMeshRenderData* RenderData = It->GetResourceForRendering())
 				{
-					if (FSkeletalMeshRenderData* RenderData = It->GetResourceForRendering())
+					FSkinnedMeshComponentRecreateRenderStateContext RecreateState(*It);
+					for (int32 LODIndex = 0; LODIndex < RenderData->LODRenderData.Num(); ++LODIndex)
 					{
-						FSkinnedMeshComponentRecreateRenderStateContext RecreateState(*It);
-						for (int32 LODIndex = 0; LODIndex < RenderData->LODRenderData.Num(); ++LODIndex)
+						FSkeletalMeshLODRenderData& LOD = RenderData->LODRenderData[LODIndex];
+						if (bClearBuffer)
 						{
-							FSkeletalMeshLODRenderData& LOD = RenderData->LODRenderData[LODIndex];
-							if (bClearBuffer)
-							{
-								LOD.SkinWeightProfilesData.ClearDynamicDefaultSkinWeightProfile(*It, LODIndex);
-							}
-							else if (bSetBuffer)
-							{
-								LOD.SkinWeightProfilesData.ClearDynamicDefaultSkinWeightProfile(*It, LODIndex);
-								LOD.SkinWeightProfilesData.SetDynamicDefaultSkinWeightProfile(*It, LODIndex);
-							}
+							LOD.SkinWeightProfilesData.ClearDynamicDefaultSkinWeightProfile(*It, LODIndex);
+						}
+						else if (bSetBuffer)
+						{
+							LOD.SkinWeightProfilesData.ClearDynamicDefaultSkinWeightProfile(*It, LODIndex);
+							LOD.SkinWeightProfilesData.SetDynamicDefaultSkinWeightProfile(*It, LODIndex);
 						}
 					}
 				}
