Merging of Get/SetSubstrateAttributes with Get/SetMaterialAttributes. Substrate specific nodes was introduced to simplify the creation of the upgraded material layering system but was only intended to act as a stop gap during the upgrade process. This change ensures we upgrade the existing material attribute nodes to be compatible with Substrate types.

Also updated the existing legacy->Substrate layers/blends to be compatible with this swapped out behaviour, and updated the existing Substrate (non-legacy) test layers/blends manually in the team test project.

#virtualized

diff --git a/Engine/Source/Editor/DetailCustomizations/Private/MaterialAttributePropertyDetails.cpp b/Engine/Source/Editor/DetailCustomizations/Private/MaterialAttributePropertyDetails.cpp
--- a/Engine/Source/Editor/DetailCustomizations/Private/MaterialAttributePropertyDetails.cpp
+++ b/Engine/Source/Editor/DetailCustomizations/Private/MaterialAttributePropertyDetails.cpp
@@ -20,6 +20,7 @@
 #include "PropertyCustomizationHelpers.h"
 #include "PropertyHandle.h"
 #include "Templates/Tuple.h"
+#include "RenderUtils.h"
 
 class SToolTip;
 
@@ -42,7 +43,11 @@ void FMaterialAttributePropertyDetails::CustomizeDetails(IDetailLayoutBuilder& D
 		// We do not want to blend Substrate BSDF which could increase the complexity of the material (BSDF / Slab count) for each blend operation.
 		// Only parameters are allowed to finally be transform into Substrate at the end of the chain.
 		// This filtering is ran only when constructing the UI element.
+#if ENABLE_MATERIAL_LAYER_PROTOTYPE
+		if(Substrate::IsSubstrateEnabled() || NameGUIDPair.Key != ExcludedFrontMaterialName)
+#else
 		if (NameGUIDPair.Key != ExcludedFrontMaterialName)
+#endif
 		{
 			AttributeDisplayNameList.Add(MakeShareable(new FString(NameGUIDPair.Key)));
 		}
diff --git a/Engine/Source/Editor/MaterialEditor/Private/MaterialEditor.cpp b/Engine/Source/Editor/MaterialEditor/Private/MaterialEditor.cpp
--- a/Engine/Source/Editor/MaterialEditor/Private/MaterialEditor.cpp
+++ b/Engine/Source/Editor/MaterialEditor/Private/MaterialEditor.cpp
@@ -147,6 +147,7 @@
 #include "MaterialLayersFunctionsCustomization.h"
 #include "MaterialEditor/MaterialEditorPreviewParameters.h"
 #include "SMaterialLayersFunctionsTree.h"
+#include "Materials/MaterialExpressionGetMaterialAttributes.h"
 #include "Materials/MaterialExpressionSetMaterialAttributes.h"
 #include "Settings/EditorExperimentalSettings.h"
 #include "Materials/MaterialExpressionBlendMaterialAttributes.h"
@@ -774,16 +775,16 @@ void FMaterialEditor::InitMaterialEditor( const EToolkitMode::Type Mode, const T
 							BaseAttributesInput->PreviewValue = FLinearColor::MakeRandomColor();
 						}
 
-						UMaterialExpression* SetSubstrateAttributes = CreateNewMaterialExpression(UMaterialExpressionSubstrateSetAttributes::StaticClass(), FVector2D(-50, 300), false, true);
-						UMaterialExpression* GetSubstrateAttributes = CreateNewMaterialExpression(UMaterialExpressionSubstrateGetAttributes::StaticClass(), FVector2D(-450, 300), false, true);
-						if (Input && SetSubstrateAttributes && GetSubstrateAttributes)
+						UMaterialExpressionSetMaterialAttributes* SetMaterialAttributes = Cast<UMaterialExpressionSetMaterialAttributes>(CreateNewMaterialExpression(UMaterialExpressionSetMaterialAttributes::StaticClass(), FVector2D(0, 300), false, true));
+						UMaterialExpressionGetMaterialAttributes* GetMaterialAttributes = Cast<UMaterialExpressionGetMaterialAttributes>(CreateNewMaterialExpression(UMaterialExpressionGetMaterialAttributes::StaticClass(), FVector2D(-400, 300), false, true));
+						if (Input && SetMaterialAttributes && GetMaterialAttributes)
 						{
-							UMaterialEditingLibrary::ConnectMaterialExpressions(Input, FString(), GetSubstrateAttributes, FString());
+							UMaterialEditingLibrary::ConnectMaterialExpressions(Input, FString(), GetMaterialAttributes, FString());
 						
-							UMaterialEditingLibrary::ConnectMaterialExpressions(GetSubstrateAttributes, FRONT_MATERIAL_ATTRIBUTES_TEXT, SetSubstrateAttributes, FRONT_MATERIAL_ATTRIBUTES_TEXT);
-							UMaterialEditingLibrary::ConnectMaterialExpressions(GetSubstrateAttributes, NON_SUBSTRATE_ATTRIBUTES_TEXT, SetSubstrateAttributes, NON_SUBSTRATE_ATTRIBUTES_TEXT);
+							SetMaterialAttributes->ConnectInputAttribute(MP_FrontMaterial, GetMaterialAttributes, GetMaterialAttributes->CreateOrGetOutputAttribute(MP_FrontMaterial));
+							SetMaterialAttributes->ConnectInputAttribute(MP_MaterialAttributes, GetMaterialAttributes, GetMaterialAttributes->CreateOrGetOutputAttribute(MP_MaterialAttributes));
 
-							UMaterialEditingLibrary::ConnectMaterialExpressions(SetSubstrateAttributes, FString(), Expression, FString());
+							UMaterialEditingLibrary::ConnectMaterialExpressions(SetMaterialAttributes, FString(), Expression, FString());
 							bMaterialDirty = true;
 						}
 
@@ -868,27 +869,36 @@ void FMaterialEditor::InitMaterialEditor( const EToolkitMode::Type Mode, const T
 
 								if (!bMaterialDirty)
 								{
-									UMaterialExpression* SetSubstrateAttributes = CreateNewMaterialExpression(UMaterialExpressionSubstrateSetAttributes::StaticClass(), FVector2D(-75, 300), false, true);
+									UMaterialExpressionSetMaterialAttributes* SetMaterialAttributes = Cast<UMaterialExpressionSetMaterialAttributes>(CreateNewMaterialExpression(UMaterialExpressionSetMaterialAttributes::StaticClass(), FVector2D(-75, 300), false, true));
 									UMaterialExpression* HorizontalMixingNode = CreateNewMaterialExpression(UMaterialExpressionSubstrateHorizontalMixing::StaticClass(), FVector2D(-350, 200), false, true);
-									UMaterialExpression* GetBottomSubstrateAttributes = CreateNewMaterialExpression(UMaterialExpressionSubstrateGetAttributes::StaticClass(), FVector2D(-750, 200), false, true);
+									UMaterialExpressionGetMaterialAttributes* GetBottomMaterialAttributes = Cast<UMaterialExpressionGetMaterialAttributes>(CreateNewMaterialExpression(UMaterialExpressionGetMaterialAttributes::StaticClass(), FVector2D(-750, 200), false, true));
+									int32 BottomFrontMaterialIndex = GetBottomMaterialAttributes->CreateOrGetOutputAttribute(MP_FrontMaterial);
 
 									UMaterialExpression* LegacyBlendNode = CreateNewMaterialExpression(UMaterialExpressionBlendMaterialAttributes::StaticClass(), FVector2D(-350, 400), false, true);
-									UMaterialExpression* GetTopSubstrateAttributes = CreateNewMaterialExpression(UMaterialExpressionSubstrateGetAttributes::StaticClass(), FVector2D(-750, 400), false, true);
-
-									if (SetSubstrateAttributes && HorizontalMixingNode && GetBottomSubstrateAttributes && LegacyBlendNode && GetTopSubstrateAttributes)
+									UMaterialExpressionGetMaterialAttributes* GetTopMaterialAttributes = Cast< UMaterialExpressionGetMaterialAttributes>(CreateNewMaterialExpression(UMaterialExpressionGetMaterialAttributes::StaticClass(), FVector2D(-750, 400), false, true));
+									int32 TopFrontMaterialIndex = GetTopMaterialAttributes->CreateOrGetOutputAttribute(MP_FrontMaterial);
+
+									if (SetMaterialAttributes 
+										&& HorizontalMixingNode 
+										&& GetBottomMaterialAttributes && BottomFrontMaterialIndex > 0 
+										&& LegacyBlendNode 
+										&& GetTopMaterialAttributes && TopFrontMaterialIndex > 0)
 									{
-										UMaterialEditingLibrary::ConnectMaterialExpressions(InputBottom, FString(), GetBottomSubstrateAttributes, FString());
-										UMaterialEditingLibrary::ConnectMaterialExpressions(GetBottomSubstrateAttributes, FRONT_MATERIAL_ATTRIBUTES_TEXT, HorizontalMixingNode, FString("Background"));
-										UMaterialEditingLibrary::ConnectMaterialExpressions(GetBottomSubstrateAttributes, NON_SUBSTRATE_ATTRIBUTES_TEXT, LegacyBlendNode, FString("A"));
+										const FString FrontMaterialName = FMaterialAttributeDefinitionMap::GetAttributeName(MP_FrontMaterial);
+										const FString MaterialAttributesName = FMaterialAttributeDefinitionMap::GetAttributeName(MP_MaterialAttributes);
+
+										UMaterialEditingLibrary::ConnectMaterialExpressions(InputBottom, FString(), GetBottomMaterialAttributes, FString());
+										UMaterialEditingLibrary::ConnectMaterialExpressions(GetBottomMaterialAttributes, FrontMaterialName, HorizontalMixingNode, FString("Background"));
+										UMaterialEditingLibrary::ConnectMaterialExpressions(GetBottomMaterialAttributes, MaterialAttributesName, LegacyBlendNode, FString("A"));
 
-										UMaterialEditingLibrary::ConnectMaterialExpressions(InputTop, FString(), GetTopSubstrateAttributes, FString());
-										UMaterialEditingLibrary::ConnectMaterialExpressions(GetTopSubstrateAttributes, FRONT_MATERIAL_ATTRIBUTES_TEXT, HorizontalMixingNode, FString("Foreground"));
-										UMaterialEditingLibrary::ConnectMaterialExpressions(GetTopSubstrateAttributes, NON_SUBSTRATE_ATTRIBUTES_TEXT, LegacyBlendNode, FString("B"));
+										UMaterialEditingLibrary::ConnectMaterialExpressions(InputTop, FString(), GetTopMaterialAttributes, FString());
+										UMaterialEditingLibrary::ConnectMaterialExpressions(GetTopMaterialAttributes, FrontMaterialName, HorizontalMixingNode, FString("Foreground"));
+										UMaterialEditingLibrary::ConnectMaterialExpressions(GetTopMaterialAttributes, MaterialAttributesName, LegacyBlendNode, FString("B"));
 
-										UMaterialEditingLibrary::ConnectMaterialExpressions(HorizontalMixingNode, FString(), SetSubstrateAttributes, FRONT_MATERIAL_ATTRIBUTES_TEXT);
-										UMaterialEditingLibrary::ConnectMaterialExpressions(LegacyBlendNode, FString(), SetSubstrateAttributes, NON_SUBSTRATE_ATTRIBUTES_TEXT);
+										SetMaterialAttributes->ConnectInputAttribute(MP_FrontMaterial, HorizontalMixingNode);
+										SetMaterialAttributes->ConnectInputAttribute(MP_MaterialAttributes, LegacyBlendNode);
 
-										UMaterialEditingLibrary::ConnectMaterialExpressions(SetSubstrateAttributes, FString(), Expression, FString());
+										UMaterialEditingLibrary::ConnectMaterialExpressions(SetMaterialAttributes, FString(), Expression, FString());
 										bMaterialDirty = true;
 									}
 								}
diff --git a/Engine/Source/Runtime/Engine/Private/Materials/Material.cpp b/Engine/Source/Runtime/Engine/Private/Materials/Material.cpp
--- a/Engine/Source/Runtime/Engine/Private/Materials/Material.cpp
+++ b/Engine/Source/Runtime/Engine/Private/Materials/Material.cpp
@@ -30,6 +30,7 @@
 #include "Materials/MaterialExpressionMakeMaterialAttributes.h"
 #include "Materials/MaterialExpressionMaterialAttributeLayers.h"
 #include "Materials/MaterialExpressionMaterialFunctionCall.h"
+#include "Materials/MaterialExpressionGetMaterialAttributes.h"
 #include "Materials/MaterialExpressionSetMaterialAttributes.h"
 #include "Materials/MaterialExpressionRuntimeVirtualTextureOutput.h"
 #include "Materials/MaterialExpressionStaticSwitchParameter.h"
@@ -3409,11 +3410,11 @@ void UMaterial::ConvertMaterialToSubstrateMaterial()
 		if(FinalNode && LayersNode)
 		{
 			//Separate the layers/blends MAs to Substrate/Non-Substrate respectively and connect to their relevant material output
-			UMaterialExpressionSubstrateGetAttributes* GetAttributesNode = NewObject<UMaterialExpressionSubstrateGetAttributes>(this);
+			UMaterialExpressionGetMaterialAttributes* GetAttributesNode = NewObject<UMaterialExpressionGetMaterialAttributes>(this);
 			GetAttributesNode->Material = this;
 			GetAttributesNode->MaterialAttributes.Connect(0, FinalNode);
-			EditorOnly->FrontMaterial.Connect(ESubstrateAttributeIndex::MSA_FrontMaterial, GetAttributesNode);
-			EditorOnly->MaterialAttributes.Connect(ESubstrateAttributeIndex::MSA_NonSubstrateAttributes, GetAttributesNode);
+			EditorOnly->FrontMaterial.Connect(GetAttributesNode->CreateOrGetOutputAttribute(MP_FrontMaterial), GetAttributesNode);
+			EditorOnly->MaterialAttributes.Connect(GetAttributesNode->CreateOrGetOutputAttribute(MP_MaterialAttributes), GetAttributesNode);
 			SetPosXAndMoveReferenceToTheRight(GetAttributesNode);
 
 			//If a non-Substrate default input is connected to the MLA node, convert it.
@@ -3424,10 +3425,10 @@ void UMaterial::ConvertMaterialToSubstrateMaterial()
 				ConvertAttributeNode->MaterialAttributes.Connect(LayersNode->Input.OutputIndex, LayersNode->Input.Expression);
 				ConvertAttributeNode->ShadingModelOverride = MSM_DefaultLit;
 
-				UMaterialExpressionSubstrateSetAttributes* SetAttributesNode = NewObject<UMaterialExpressionSubstrateSetAttributes>(this);
+				UMaterialExpressionSetMaterialAttributes* SetAttributesNode = NewObject<UMaterialExpressionSetMaterialAttributes>(this);
 				SetAttributesNode->Material = this;
-				SetAttributesNode->FrontMaterial.Connect(0, ConvertAttributeNode);
-				SetAttributesNode->NonSubstrateAttributes.Connect(LayersNode->Input.OutputIndex, LayersNode->Input.Expression);
+				SetAttributesNode->ConnectInputAttribute(MP_FrontMaterial, ConvertAttributeNode);
+				SetAttributesNode->ConnectInputAttribute(MP_MaterialAttributes, LayersNode->Input.Expression, LayersNode->Input.OutputIndex);
 
 				LayersNode->Input.Connect(0, SetAttributesNode);
 
diff --git a/Engine/Source/Runtime/Engine/Private/Materials/MaterialExpressions.cpp b/Engine/Source/Runtime/Engine/Private/Materials/MaterialExpressions.cpp
--- a/Engine/Source/Runtime/Engine/Private/Materials/MaterialExpressions.cpp
+++ b/Engine/Source/Runtime/Engine/Private/Materials/MaterialExpressions.cpp
@@ -7310,6 +7310,7 @@ uint32 UMaterialExpressionBreakMaterialAttributes::GetOutputType(int32 OutputInd
 
 // -----
 
+#define GET_SET_MA_MATERIALATTRIBUTESINDEX 0
 UMaterialExpressionGetMaterialAttributes::UMaterialExpressionGetMaterialAttributes(const FObjectInitializer& ObjectInitializer)
 	: Super(ObjectInitializer)
 {
@@ -7340,6 +7341,40 @@ UMaterialExpressionGetMaterialAttributes::UMaterialExpressionGetMaterialAttribut
 }
 
 #if WITH_EDITOR
+#if ENABLE_MATERIAL_LAYER_PROTOTYPE
+int32 UMaterialExpressionGetMaterialAttributes::CreateOrGetOutputAttribute(EMaterialProperty Attribute)
+{
+	int32 OutputIndex = INDEX_NONE;
+	if (Attribute == MP_MaterialAttributes)
+	{
+		OutputIndex = GET_SET_MA_MATERIALATTRIBUTESINDEX;
+	}
+	else
+	{
+		FGuid AttributeId = FMaterialAttributeDefinitionMap::GetID(Attribute);
+		if (AttributeGetTypes.Find(AttributeId, OutputIndex))
+		{
+			/**
+			* Add one to compensate for the AttributeGetTypes list not containing MP_MaterialAttributes
+			* It's none trivial to iterate the Outputs list for the matching attribute so this is a simpler solution.
+			*/
+			OutputIndex++;
+		}
+		else
+		{
+			int32 GetTypesIndex = AttributeGetTypes.Add(AttributeId);
+			if(GetTypesIndex != INDEX_NONE)
+			{
+				PreEditChange(nullptr);
+				FString AttributeName = FMaterialAttributeDefinitionMap::GetDisplayNameForMaterial(AttributeGetTypes[GetTypesIndex], Material).ToString();
+				OutputIndex = Outputs.Add(FExpressionOutput(*AttributeName, 0, 0, 0, 0, 0));
+			}
+		}
+	}
+	return OutputIndex;
+}
+#endif //ENABLE_MATERIAL_LAYER_PROTOTYPE
+
 int32 UMaterialExpressionGetMaterialAttributes::Compile(class FMaterialCompiler* Compiler, int32 OutputIndex)
 {
 	// Verify setup
@@ -7363,12 +7398,12 @@ int32 UMaterialExpressionGetMaterialAttributes::Compile(class FMaterialCompiler*
 	// Compile attribute
 	int32 Result = INDEX_NONE;
 
-	if (OutputIndex == 0)
+	if (OutputIndex == GET_SET_MA_MATERIALATTRIBUTESINDEX)
 	{
 		const FGuid AttributeID = Compiler->GetMaterialAttribute();
 		Result = MaterialAttributes.CompileWithDefault(Compiler, AttributeID);
 	}
-	else if (OutputIndex > 0)
+	else if (OutputIndex > GET_SET_MA_MATERIALATTRIBUTESINDEX)
 	{
 		checkf(OutputIndex <= AttributeGetTypes.Num(), TEXT("Requested non-existent pin."));
 		Result = MaterialAttributes.CompileWithDefault(Compiler, AttributeGetTypes[OutputIndex-1]);
@@ -7393,7 +7428,7 @@ uint32 UMaterialExpressionGetMaterialAttributes::GetOutputType(int32 OutputIndex
 	uint32 OutputType = Super::GetOutputType(OutputIndex);
 
 	// Override the type if it's a ShadingModel type
-	if (OutputIndex > 0) // "0th" place is the mandatory MaterialAttribute itself, skip it
+	if (OutputIndex > GET_SET_MA_MATERIALATTRIBUTESINDEX) // "0th" place is the mandatory MaterialAttribute itself, skip it
 	{
 		ensure(OutputIndex < AttributeGetTypes.Num() + 1);
 		EMaterialValueType PinType = FMaterialAttributeDefinitionMap::GetValueType(AttributeGetTypes[OutputIndex - 1]);
@@ -7413,9 +7448,9 @@ uint32 UMaterialExpressionGetMaterialAttributes::GetOutputType(int32 OutputIndex
 bool UMaterialExpressionGetMaterialAttributes::IsResultSubstrateMaterial(int32 OutputIndex)
 {	
 #if ENABLE_MATERIAL_LAYER_PROTOTYPE
-	if (OutputIndex >= 1 && OutputIndex <= AttributeGetTypes.Num())
+	if (--OutputIndex >= 0 && OutputIndex < AttributeGetTypes.Num())
 	{
-		return FMaterialAttributeDefinitionMap::GetValueType(AttributeGetTypes[OutputIndex - 1]) == MCT_Substrate;
+		return FMaterialAttributeDefinitionMap::GetValueType(AttributeGetTypes[OutputIndex]) == MCT_Substrate;
 	}
 #endif
 
@@ -7640,9 +7675,9 @@ int32 UMaterialExpressionSetMaterialAttributes::Compile(class FMaterialCompiler*
 	if (CompilingAttributeID == FMaterialAttributeDefinitionMap::GetID(MP_MaterialAttributes))
 	{
 		int32 Result = INDEX_NONE;
-		if (Inputs[0].GetTracedInput().Expression)
+		if (Inputs[GET_SET_MA_MATERIALATTRIBUTESINDEX].GetTracedInput().Expression)
 		{
-			Result = Inputs[0].GetTracedInput().Compile(Compiler);
+			Result = Inputs[GET_SET_MA_MATERIALATTRIBUTESINDEX].GetTracedInput().Compile(Compiler);
 		}
 		else
 		{
@@ -7685,9 +7720,9 @@ int32 UMaterialExpressionSetMaterialAttributes::Compile(class FMaterialCompiler*
 			EMaterialValueType ValueType = FMaterialAttributeDefinitionMap::GetValueType(CompilingAttributeID);
 			return Compiler->ValidCast(AttributeInput->GetTracedInput().Compile(Compiler), ValueType);
 		}
-		else if (Inputs[0].GetTracedInput().Expression)
+		else if (Inputs[GET_SET_MA_MATERIALATTRIBUTESINDEX].GetTracedInput().Expression)
 		{
-			return Inputs[0].GetTracedInput().Compile(Compiler);
+			return Inputs[GET_SET_MA_MATERIALATTRIBUTESINDEX].GetTracedInput().Compile(Compiler);
 		}
 
 		return FMaterialAttributeDefinitionMap::CompileDefaultExpression(Compiler, CompilingAttributeID);
@@ -7719,11 +7754,11 @@ FName UMaterialExpressionSetMaterialAttributes::GetInputName(int32 InputIndex) c
 {
 	FName Name;
 
-	if (InputIndex == 0)
+	if (InputIndex == GET_SET_MA_MATERIALATTRIBUTESINDEX)
 	{
 		Name = *NSLOCTEXT("SetMaterialAttributes", "InputName", "MaterialAttributes").ToString();
 	}
-	else if (InputIndex > 0)
+	else if (InputIndex > GET_SET_MA_MATERIALATTRIBUTESINDEX)
 	{
 		Name = *FMaterialAttributeDefinitionMap::GetDisplayNameForMaterial(AttributeSetTypes[InputIndex-1], Material).ToString();
 	}
@@ -7735,13 +7770,13 @@ uint32 UMaterialExpressionSetMaterialAttributes::GetInputType(int32 InputIndex)
 {
 	uint32 InputType = MCT_Unknown;
 
-	if (InputIndex == 0)
+	if (InputIndex == GET_SET_MA_MATERIALATTRIBUTESINDEX)
 	{
 		InputType = MCT_MaterialAttributes;
 	}
 	else
 	{
-		ensure(InputIndex > 0 && InputIndex < AttributeSetTypes.Num() + 1);
+		ensure(InputIndex > GET_SET_MA_MATERIALATTRIBUTESINDEX && InputIndex < AttributeSetTypes.Num() + 1);
 		InputType = FMaterialAttributeDefinitionMap::GetValueType(AttributeSetTypes[InputIndex - 1]);
 		if (InputType == MCT_ShadingModel)
 		{
@@ -7760,24 +7795,99 @@ uint32 UMaterialExpressionSetMaterialAttributes::GetInputType(int32 InputIndex)
 	return InputType;
 }
 
+#if ENABLE_MATERIAL_LAYER_PROTOTYPE
+int32 UMaterialExpressionSetMaterialAttributes::CreateOrGetInputAttribute(EMaterialProperty Attribute)
+{
+	int32 InputsIndex = INDEX_NONE;
+	if (Attribute == MP_MaterialAttributes)
+	{
+		InputsIndex = GET_SET_MA_MATERIALATTRIBUTESINDEX;
+	}
+	else
+	{
+		FGuid AttributeId = FMaterialAttributeDefinitionMap::GetID(Attribute);
+		if (AttributeSetTypes.Find(AttributeId, InputsIndex))
+		{
+			/**
+			* Add one to compensate for the AttributeGetTypes list not containing MP_MaterialAttributes
+			* It's none trivial to iterate the Inputs list for the matching attribute so this is a simpler solution.
+			*/
+			InputsIndex++;
+		}
+		else
+		{
+			int32 SetTypesIndex = AttributeSetTypes.Add(AttributeId);
+			if(SetTypesIndex != INDEX_NONE)
+			{
+				PreEditChange(nullptr);
+				InputsIndex = Inputs.Add(FExpressionInput());
+				if (Inputs.IsValidIndex(InputsIndex))
+				{
+					Inputs[InputsIndex].InputName = FName(*FMaterialAttributeDefinitionMap::GetDisplayNameForMaterial(AttributeSetTypes[SetTypesIndex], Material).ToString());
+				}
+			}
+		}
+	}
+	return InputsIndex;
+}
+
+bool UMaterialExpressionSetMaterialAttributes::ConnectInputAttribute(EMaterialProperty Attribute, UMaterialExpression* Expression, int32 OutputIndex)
+{
+	int32 Index = CreateOrGetInputAttribute(Attribute);
+	if(Expression && OutputIndex != INDEX_NONE && Inputs.IsValidIndex(Index))
+	{
+		Inputs[Index].Connect(OutputIndex, Expression);
+		return Inputs[Index].IsConnected();
+	}
+	return false;
+}
+
+bool UMaterialExpressionSetMaterialAttributes::GetSubstrateMaterialInputIndex(int32 OutputIndex, int32& InputIndex)
+{
+	for (InputIndex = Inputs.Num() - 1; InputIndex >= GET_SET_MA_MATERIALATTRIBUTESINDEX; InputIndex--)
+	{
+		if (GetInputType(InputIndex) == MCT_Substrate)
+		{			
+			return Inputs[InputIndex].IsConnected();
+		}
+		else if (GetInputType(InputIndex) == MCT_MaterialAttributes && Inputs[InputIndex].IsConnected())
+		{
+			return Inputs[InputIndex].Expression->IsResultSubstrateMaterial(Inputs[InputIndex].OutputIndex);
+		}
+	}
+
+	return false;
+}
+#endif // ENABLE_MATERIAL_LAYER_PROTOTYPE
+
 bool UMaterialExpressionSetMaterialAttributes::IsResultSubstrateMaterial(int32 OutputIndex)
 {
 #if ENABLE_MATERIAL_LAYER_PROTOTYPE
-	if (Inputs[0].IsConnected())
+	int32 InputIndex = INDEX_NONE;
+	return GetSubstrateMaterialInputIndex(OutputIndex, InputIndex);
+#else
+	return false;
+#endif
+}
+
+void UMaterialExpressionSetMaterialAttributes::GatherSubstrateMaterialInfo(FSubstrateMaterialInfo& SubstrateMaterialInfo, int32 OutputIndex)
+{
+#if ENABLE_MATERIAL_LAYER_PROTOTYPE
+	int32 InputIndex = INDEX_NONE;
+	if(GetSubstrateMaterialInputIndex(OutputIndex, InputIndex))
 	{
-		Inputs[0].Expression->IsResultSubstrateMaterial(Inputs[0].OutputIndex);
+		Inputs[InputIndex].Expression->GatherSubstrateMaterialInfo(SubstrateMaterialInfo, Inputs[InputIndex].OutputIndex);
 	}
 #endif
-
-	return false;
 }
 
 FSubstrateOperator* UMaterialExpressionSetMaterialAttributes::SubstrateGenerateMaterialTopologyTree(FMaterialCompiler* Compiler, UMaterialExpression* Parent, int32 OutputIndex)
 {
 #if ENABLE_MATERIAL_LAYER_PROTOTYPE
-	if(IsResultSubstrateMaterial(OutputIndex))
+	int32 InputIndex = INDEX_NONE;
+	if (GetSubstrateMaterialInputIndex(OutputIndex, InputIndex))
 	{
-		Inputs[0].Expression->SubstrateGenerateMaterialTopologyTree(Compiler, Parent, Inputs[0].OutputIndex);
+		return Inputs[InputIndex].Expression->SubstrateGenerateMaterialTopologyTree(Compiler, Parent, Inputs[InputIndex].OutputIndex);
 	}
 #endif
 	return nullptr;
@@ -15980,12 +16090,38 @@ void UMaterialFunction::ConvertExpressionsBetweenLegacyAndSubstrate()
 	{
 		for (UMaterialExpression* Expression : Function->GetExpressions())
 		{
-			if (Expression->IsA<UMaterialExpressionSubstrateBSDF>()
+			if (!Expression || Expression->IsA<UMaterialExpressionSubstrateBSDF>()
 				|| Expression->IsA<UMaterialExpressionSubstrateUtilityBase>())
 			{
 				//If Substrate nodes are present, do not convert this function, assume previously converted or Substrate MF
 				return;
 			}
+
+			//Check if we have already applied FrontMaterial logic to the Set/GetMaterialAttributes nodes, and if so, skip further processing.
+			if(UMaterialExpressionSetMaterialAttributes* SetAttributes = Cast<UMaterialExpressionSetMaterialAttributes>(Expression))
+			{
+				TArrayView<FExpressionInput*> InputsArray = SetAttributes->GetInputsView();
+				//0 index is always MaterialAttributes so no need to check that entry
+				for(int32 InputIndex = InputsArray.Num() - 1; InputIndex > 0; InputIndex--)
+				{
+					if(SetAttributes->GetInputType(InputIndex) == MCT_Substrate)
+					{
+						return;
+					}
+				}
+			}
+
+			if (UMaterialExpressionGetMaterialAttributes* GetAttributes = Cast<UMaterialExpressionGetMaterialAttributes>(Expression))
+			{
+				//0 index is always MaterialAttributes so no need to check that entry
+				for(int32 OutputIndex = GetAttributes->GetOutputs().Num() - 1; OutputIndex > 0; OutputIndex--)
+				{
+					if (GetAttributes->GetOutputType(OutputIndex) == MCT_Substrate)
+					{
+						return;
+					}
+				}
+			}
 		}
 	}
 
@@ -16006,7 +16142,7 @@ void UMaterialFunction::ConvertExpressionsBetweenLegacyAndSubstrate()
 		NewNode->MaterialExpressionEditorY = ExistingNode->MaterialExpressionEditorY + DownOffset;
 	};
 
-	auto ReplaceNodeAndMoveToTheRight = [](UMaterialExpression* ExistingNode, UMaterialExpression* NewNode, int32 RightOffset = 350)
+	auto ReplaceNodeAndMoveToTheRight = [](UMaterialExpression* ExistingNode, UMaterialExpression* NewNode, int32 RightOffset = 300)
 	{
 		NewNode->MaterialExpressionEditorX = ExistingNode->MaterialExpressionEditorX;
 		NewNode->MaterialExpressionEditorY = ExistingNode->MaterialExpressionEditorY;
@@ -16104,14 +16240,14 @@ void UMaterialFunction::ConvertExpressionsBetweenLegacyAndSubstrate()
 			else if (InputExpression->IsA<UMaterialExpressionSetMaterialAttributes>() || InputExpression->IsA<UMaterialExpressionMakeMaterialAttributes>())
 			{
 				//Same behaviour as above, but for now extends the Make/Set nodes with conversion, so we don't remove the existing node in this case.
-				UMaterialExpressionSubstrateSetAttributes* SetSubstrateNode = ReplacementNodeMapping.Contains(InputExpression) ? Cast<UMaterialExpressionSubstrateSetAttributes>(*ReplacementNodeMapping.Find(InputExpression)) : nullptr;
-				if(!SetSubstrateNode)
+				UMaterialExpressionSetMaterialAttributes* SetAttributesNode = ReplacementNodeMapping.Contains(InputExpression) ? Cast<UMaterialExpressionSetMaterialAttributes>(*ReplacementNodeMapping.Find(InputExpression)) : nullptr;
+				if(!SetAttributesNode)
 				{
-					SetSubstrateNode = Cast<UMaterialExpressionSubstrateSetAttributes>(ReplacementNodeMapping.Add(InputExpression, NewObject<UMaterialExpressionSubstrateSetAttributes>(this)));
-					SetSubstrateNode->Function = this;
-					SetSubstrateNode->NonSubstrateAttributes.Connect(0, InputExpression);
-					Expressions.Add(SetSubstrateNode);
-					RecursedExpressions.Add(SetSubstrateNode);
+					SetAttributesNode = Cast<UMaterialExpressionSetMaterialAttributes>(ReplacementNodeMapping.Add(InputExpression, NewObject<UMaterialExpressionSetMaterialAttributes>(this)));
+					SetAttributesNode->Function = this;
+					SetAttributesNode->ConnectInputAttribute(MP_MaterialAttributes, InputExpression);
+					Expressions.Add(SetAttributesNode);
+					RecursedExpressions.Add(SetAttributesNode);
 
 					UMaterialExpressionSubstrateConvertMaterialAttributes* ConvertNode = NewObject<UMaterialExpressionSubstrateConvertMaterialAttributes>(this);
 					ConvertNode->Function = this;
@@ -16120,15 +16256,15 @@ void UMaterialFunction::ConvertExpressionsBetweenLegacyAndSubstrate()
 					Expressions.Add(ConvertNode);
 					RecursedExpressions.Add(ConvertNode);
 
-					SetSubstrateNode->FrontMaterial.Connect(0, ConvertNode);
+					SetAttributesNode->ConnectInputAttribute(MP_FrontMaterial, ConvertNode);
 
 					PlaceBelowNode(InputExpression, ConvertNode);
-					PlaceBelowNode(ConvertNode, SetSubstrateNode, 250);
+					PlaceBelowNode(ConvertNode, SetAttributesNode, 250);
 				}
 
-				if(SetSubstrateNode)
+				if(SetAttributesNode)
 				{
-					Input->Connect(0, SetSubstrateNode);
+					Input->Connect(0, SetAttributesNode);
 					bBlendConverted = true;
 				}
 			}
@@ -16178,6 +16314,7 @@ void UMaterialFunction::ConvertExpressionsBetweenLegacyAndSubstrate()
 		ConvertAttributesNode->MaterialAttributes.Connect(0, OutputNode->A.Expression);
 		ConvertAttributesNode->ShadingModelOverride = MSM_DefaultLit;
 		ReplaceNodeAndMoveToTheRight(OutputNode, ConvertAttributesNode);
+		PlaceBelowNode(ConvertAttributesNode, ConvertAttributesNode);
 
 		//Add Custom logic connections
 		TArray<class UMaterialExpressionCustomOutput*> CustomOutputExpressions;
@@ -16213,16 +16350,16 @@ void UMaterialFunction::ConvertExpressionsBetweenLegacyAndSubstrate()
 		 * Layer MFs differ from materials in that they only have 1 output,
 		 * so we use the Set Substrate Attributes node to collect the Front Material and MAs to pass to the next function in the layer stack.
 		 */
-		UMaterialExpressionSubstrateSetAttributes* SetAttributesNode = NewObject<UMaterialExpressionSubstrateSetAttributes>(this);
-		SetAttributesNode->Function = this;
-		SetAttributesNode->NonSubstrateAttributes.Connect(0, OutputNode->A.Expression);
-		SetAttributesNode->FrontMaterial.Connect(0, ConvertAttributesNode);
-		ReplaceNodeAndMoveToTheRight(OutputNode, SetAttributesNode);
+		UMaterialExpressionSetMaterialAttributes* SetMatAttributesNode = NewObject<UMaterialExpressionSetMaterialAttributes>(this);
+		SetMatAttributesNode->Function = this;
+		SetMatAttributesNode->ConnectInputAttribute(MP_MaterialAttributes, OutputNode->A.Expression);
+		SetMatAttributesNode->ConnectInputAttribute(MP_FrontMaterial, ConvertAttributesNode);
+		ReplaceNodeAndMoveToTheRight(OutputNode, SetMatAttributesNode);
 
-		Expressions.EmplaceAt(1, SetAttributesNode);
+		Expressions.EmplaceAt(1, SetMatAttributesNode);
 		Expressions.EmplaceAt(2, ConvertAttributesNode);
 
-		OutputNode->A.Connect(0, SetAttributesNode);
+		OutputNode->A.Connect(0, SetMatAttributesNode);
 	}
 	MoveNodeInHorizonalAxis(OutputNode, 50);
 	OutputNode->bCollapsed = true;
@@ -29889,247 +30026,6 @@ bool UMaterialExpressionSubstrateConvertMaterialAttributes::HasSSS() const
 
 #endif // WITH_EDITOR
 
-// Substrate Attributes Start -----
-UMaterialExpressionSubstrateGetAttributes::UMaterialExpressionSubstrateGetAttributes(const FObjectInitializer& ObjectInitializer)
-	: Super(ObjectInitializer)
-{
-#if ENABLE_MATERIAL_LAYER_PROTOTYPE
-#if WITH_EDITORONLY_DATA
-	// Structure to hold one-time initialization
-	struct FConstructorStatics
-	{
-		FText NAME_SubstrateAttributes;
-		FConstructorStatics()
-			: NAME_SubstrateAttributes(LOCTEXT("SubstrateAttributes", "Substrate Attributes"))
-		{
-		}
-	};
-
-	static FConstructorStatics ConstructorStatics;
-
-	bShowOutputNameOnPin = true;
-	bShowMaskColorsOnPin = false;
-
-	MenuCategories.Add(ConstructorStatics.NAME_SubstrateAttributes); //Check this setup for other Substrate nodes
-
-	Outputs.Reset();
-	Outputs.SetNum(ESubstrateAttributeIndex::MSA_MAX);
-	Outputs[ESubstrateAttributeIndex::MSA_NonSubstrateAttributes] = FExpressionOutput(NON_SUBSTRATE_ATTRIBUTES_TEXT);
-	Outputs[ESubstrateAttributeIndex::MSA_FrontMaterial] = FExpressionOutput(FRONT_MATERIAL_ATTRIBUTES_TEXT);
-#endif
-
-#if WITH_EDITOR
-	CachedInputs.Empty();
-	CachedInputs.Add(&MaterialAttributes);
-#endif
-#endif //ENABLE_MATERIAL_LAYER_PROTOTYPE
-}
-
-#if ENABLE_MATERIAL_LAYER_PROTOTYPE
-#if WITH_EDITOR
-void UMaterialExpressionSubstrateGetAttributes::PostEditChangeProperty(FPropertyChangedEvent& PropertyChangedEvent)
-{
-	if (GraphNode && PropertyChangedEvent.Property != nullptr)
-	{
-		GraphNode->ReconstructNode();
-	}
-	Super::PostEditChangeProperty(PropertyChangedEvent);
-}
-
-int32 UMaterialExpressionSubstrateGetAttributes::Compile(class FMaterialCompiler* Compiler, int32 OutputIndex)
-{
-	int32 Ret = INDEX_NONE;
-
-	FGuid AttributeId = Compiler->GetMaterialAttribute();
-	EMaterialProperty Property = FMaterialAttributeDefinitionMap::GetProperty(AttributeId);
-
-	if ((OutputIndex == ESubstrateAttributeIndex::MSA_NonSubstrateAttributes && Property != MP_FrontMaterial)
-		|| (OutputIndex == ESubstrateAttributeIndex::MSA_FrontMaterial && Property == MP_FrontMaterial))
-	{
-		Ret = MaterialAttributes.CompileWithDefault(Compiler, AttributeId);
-	}
-
-	return Ret;
-}
-
-void UMaterialExpressionSubstrateGetAttributes::GatherSubstrateMaterialInfo(FSubstrateMaterialInfo& SubstrateMaterialInfo, int32 OutputIndex)
-{
-	if (MaterialAttributes.Expression && IsResultSubstrateMaterial(OutputIndex))
-	{
-		MaterialAttributes.Expression->GatherSubstrateMaterialInfo(SubstrateMaterialInfo, MaterialAttributes.OutputIndex);
-	}
-}
-
-FSubstrateOperator* UMaterialExpressionSubstrateGetAttributes::SubstrateGenerateMaterialTopologyTree(class FMaterialCompiler* Compiler, class UMaterialExpression* Parent, int32 OutputIndex)
-{
-	if (MaterialAttributes.Expression && IsResultSubstrateMaterial(OutputIndex))
-	{
-		return MaterialAttributes.Expression->SubstrateGenerateMaterialTopologyTree(Compiler, Parent, MaterialAttributes.OutputIndex);
-	}
-
-	return nullptr;
-}
-
-void UMaterialExpressionSubstrateGetAttributes::GetCaption(TArray<FString>& OutCaptions) const
-{
-	OutCaptions.Add(TEXT("Get Substrate Attributes"));
-}
-
-
-FName UMaterialExpressionSubstrateGetAttributes::GetInputName(int32 InputIndex) const
-{
-	if (0 == InputIndex)
-	{
-		return *NSLOCTEXT("GetSubstrateAttributes", "InputName", "Material Attributes").ToString();
-	}
-	return NAME_None;
-}
-
-bool UMaterialExpressionSubstrateGetAttributes::IsInputConnectionRequired(int32 InputIndex) const
-{
-	return true;
-}
-
-uint32 UMaterialExpressionSubstrateGetAttributes::GetInputType(int32 InputIndex)
-{
-	return MCT_MaterialAttributes;
-}
-
-uint32 UMaterialExpressionSubstrateGetAttributes::GetOutputType(int32 OutputIndex)
-{
-	switch (OutputIndex)
-	{
-	case ESubstrateAttributeIndex::MSA_FrontMaterial: return MCT_Substrate;
-	case ESubstrateAttributeIndex::MSA_NonSubstrateAttributes: return MCT_MaterialAttributes;
-	}
-	check(false);
-	return MCT_Float1;
-}
-#endif // WITH_EDITOR
-#endif //ENABLE_MATERIAL_LAYER_PROTOTYPE
-
-UMaterialExpressionSubstrateSetAttributes::UMaterialExpressionSubstrateSetAttributes(const FObjectInitializer& ObjectInitializer)
-	: Super(ObjectInitializer)
-{
-#if ENABLE_MATERIAL_LAYER_PROTOTYPE
-#if WITH_EDITORONLY_DATA
-	// Structure to hold one-time initialization
-	struct FConstructorStatics
-	{
-		FText NAME_SubstrateAttributes;
-		FConstructorStatics()
-			: NAME_SubstrateAttributes(LOCTEXT("SubstrateAttributes", "Substrate Attributes"))
-		{
-		}
-	};
-
-	static FConstructorStatics ConstructorStatics;
-
-	bShowOutputNameOnPin = true;
-	bShowMaskColorsOnPin = false;
-
-	MenuCategories.Add(ConstructorStatics.NAME_SubstrateAttributes);
-
-	Outputs.Reset();
-	Outputs.Add(FExpressionOutput(TEXT("Material Attributes")));
-#endif
-
-#if WITH_EDITOR
-	CachedInputs.Empty();
-	CachedInputs.Add(&FrontMaterial);
-	CachedInputs.Add(&NonSubstrateAttributes);
-#endif
-#endif //ENABLE_MATERIAL_LAYER_PROTOTYPE
-}
-
-#if ENABLE_MATERIAL_LAYER_PROTOTYPE
-#if WITH_EDITOR
-void UMaterialExpressionSubstrateSetAttributes::PostEditChangeProperty(FPropertyChangedEvent& PropertyChangedEvent)
-{
-	if (GraphNode && PropertyChangedEvent.Property != nullptr)
-	{
-		GraphNode->ReconstructNode();
-	}
-	Super::PostEditChangeProperty(PropertyChangedEvent);
-}
-
-int32 UMaterialExpressionSubstrateSetAttributes::Compile(class FMaterialCompiler* Compiler, int32 OutputIndex)
-{
-	int32 Ret = INDEX_NONE;
-	FExpressionInput* Expression = nullptr;
-
-	FGuid AttributeId = Compiler->GetMaterialAttribute();
-	if (AttributeId == FMaterialAttributeDefinitionMap::GetID(MP_FrontMaterial))
-	{
-		Ret = FrontMaterial.Compile(Compiler);
-	}
-	else
-	{
-		Ret = NonSubstrateAttributes.CompileWithDefault(Compiler, AttributeId);
-	}
-
-	return Ret;
-}
-
-void UMaterialExpressionSubstrateSetAttributes::GatherSubstrateMaterialInfo(FSubstrateMaterialInfo& SubstrateMaterialInfo, int32 OutputIndex)
-{
-	if (FrontMaterial.Expression)
-	{
-		FrontMaterial.Expression->GatherSubstrateMaterialInfo(SubstrateMaterialInfo, FrontMaterial.OutputIndex);
-	}
-}
-
-FSubstrateOperator* UMaterialExpressionSubstrateSetAttributes::SubstrateGenerateMaterialTopologyTree(class FMaterialCompiler* Compiler, class UMaterialExpression* Parent, int32 OutputIndex)
-{
-	if (FrontMaterial.Expression)
-	{
-		return FrontMaterial.Expression->SubstrateGenerateMaterialTopologyTree(Compiler, Parent, FrontMaterial.OutputIndex);
-	}
-
-	return nullptr;
-}
-
-void UMaterialExpressionSubstrateSetAttributes::GetCaption(TArray<FString>& OutCaptions) const
-{
-	OutCaptions.Add(TEXT("Set Substrate Attributes"));
-}
-
-
-FName UMaterialExpressionSubstrateSetAttributes::GetInputName(int32 InputIndex) const
-{
-	if (InputIndex == ESubstrateAttributeIndex::MSA_NonSubstrateAttributes)
-	{
-		return FName(NON_SUBSTRATE_ATTRIBUTES_TEXT);
-	}
-	if (InputIndex == ESubstrateAttributeIndex::MSA_FrontMaterial)
-	{
-		return FName(FRONT_MATERIAL_ATTRIBUTES_TEXT);
-	}
-	return NAME_None;
-}
-
-uint32 UMaterialExpressionSubstrateSetAttributes::GetInputType(int32 InputIndex)
-{
-	if (InputIndex == ESubstrateAttributeIndex::MSA_NonSubstrateAttributes)
-	{
-		return MCT_MaterialAttributes;
-	}
-	if (InputIndex == ESubstrateAttributeIndex::MSA_FrontMaterial)
-	{
-		return MCT_Substrate;
-	}
-	return Super::GetInputType(InputIndex);
-}
-
-bool UMaterialExpressionSubstrateSetAttributes::IsInputConnectionRequired(int32 InputIndex) const
-{
-	return true;
-}
-#endif // WITH_EDITOR
-#endif // ENABLE_MATERIAL_LAYER_PROTOTYPE
-
-// Substrate Attributes End -----
-
 UMaterialExpressionExecBegin::UMaterialExpressionExecBegin(const FObjectInitializer& ObjectInitializer)
 	: Super(ObjectInitializer)
 {
diff --git a/Engine/Source/Runtime/Engine/Public/Materials/MaterialExpressionGetMaterialAttributes.h b/Engine/Source/Runtime/Engine/Public/Materials/MaterialExpressionGetMaterialAttributes.h
--- a/Engine/Source/Runtime/Engine/Public/Materials/MaterialExpressionGetMaterialAttributes.h
+++ b/Engine/Source/Runtime/Engine/Public/Materials/MaterialExpressionGetMaterialAttributes.h
@@ -36,6 +36,10 @@ class UMaterialExpressionGetMaterialAttributes : public UMaterialExpression
  
 	//~ Begin UMaterialExpression Interface
 #if WITH_EDITOR
+#if ENABLE_MATERIAL_LAYER_PROTOTYPE
+	ENGINE_API int32 CreateOrGetOutputAttribute(EMaterialProperty Attribute);
+#endif
+
 	virtual int32 Compile(class FMaterialCompiler* Compiler, int32 OutputIndex) override;
 	virtual bool GenerateHLSLExpression(FMaterialHLSLGenerator& Generator, UE::HLSLTree::FScope& Scope, int32 OutputIndex, UE::HLSLTree::FExpression const*& OutExpression) const override;
 	virtual void GetCaption(TArray<FString>& OutCaptions) const override;
diff --git a/Engine/Source/Runtime/Engine/Public/Materials/MaterialExpressionSetMaterialAttributes.h b/Engine/Source/Runtime/Engine/Public/Materials/MaterialExpressionSetMaterialAttributes.h
--- a/Engine/Source/Runtime/Engine/Public/Materials/MaterialExpressionSetMaterialAttributes.h
+++ b/Engine/Source/Runtime/Engine/Public/Materials/MaterialExpressionSetMaterialAttributes.h
@@ -42,7 +42,14 @@ class UMaterialExpressionSetMaterialAttributes : public UMaterialExpression
 	virtual uint32 GetInputType(int32 InputIndex) override;
 	virtual bool IsInputConnectionRequired(int32 InputIndex) const override {return true;}
 	virtual bool IsResultMaterialAttributes(int32 OutputIndex) override {return true;}
+
+#if ENABLE_MATERIAL_LAYER_PROTOTYPE
+	ENGINE_API int32 CreateOrGetInputAttribute(EMaterialProperty Attribute);
+	ENGINE_API bool ConnectInputAttribute(EMaterialProperty Attribute, UMaterialExpression* Expression, int32 OutputIndex = 0);
+	bool GetSubstrateMaterialInputIndex(int32 OutputIndex, int32& InputIndex);
+#endif
 	virtual bool IsResultSubstrateMaterial(int32 OutputIndex) override;
+	virtual void GatherSubstrateMaterialInfo(FSubstrateMaterialInfo& SubstrateMaterialInfo, int32 OutputIndex) override;
 	virtual FSubstrateOperator* SubstrateGenerateMaterialTopologyTree(FMaterialCompiler* Compiler, UMaterialExpression* Parent, int32 OutputIndex) override;
 	virtual void GetExpressionToolTip(TArray<FString>& OutToolTip) override;
 
diff --git a/Engine/Source/Runtime/Engine/Public/Materials/MaterialExpressionSubstrate.h b/Engine/Source/Runtime/Engine/Public/Materials/MaterialExpressionSubstrate.h
--- a/Engine/Source/Runtime/Engine/Public/Materials/MaterialExpressionSubstrate.h
+++ b/Engine/Source/Runtime/Engine/Public/Materials/MaterialExpressionSubstrate.h
@@ -1194,78 +1194,4 @@ class UMaterialExpressionSubstrateThinFilm : public UMaterialExpressionSubstrate
 	virtual void GetExpressionToolTip(TArray<FString>& OutToolTip) override;
 #endif
 	//~ End UMaterialExpression Interface
-};
-
-///////////////////////////////////////////////////////////////////////////////
-// Attributes
-
-#if ENABLE_MATERIAL_LAYER_PROTOTYPE
-enum ESubstrateAttributeIndex : int
-{
-	MSA_FrontMaterial,
-	MSA_NonSubstrateAttributes,
-	MSA_MAX
-};
-
-#define FRONT_MATERIAL_ATTRIBUTES_TEXT "Front Material"
-#define NON_SUBSTRATE_ATTRIBUTES_TEXT "Non-Substrate Attributes"
-
-inline bool CheckIfNonSubstratePropertyType(EMaterialProperty Property)
-{
-	return Property != MP_FrontMaterial;
-}
-
-#endif //#if ENABLE_MATERIAL_LAYER_PROTOTYPE
-UCLASS(MinimalAPI, collapsecategories, hidecategories = Object, DisplayName = "Substrate Get Attributes")
-class UMaterialExpressionSubstrateGetAttributes : public UMaterialExpressionSubstrateUtilityBase
-{
-	GENERATED_UCLASS_BODY()
-
-	UPROPERTY()
-	FMaterialAttributesInput MaterialAttributes;
-
-#if ENABLE_MATERIAL_LAYER_PROTOTYPE
-#if WITH_EDITOR
-	virtual void PostEditChangeProperty(FPropertyChangedEvent& PropertyChangedEvent) override;
-	virtual int32 Compile(class FMaterialCompiler* Compiler, int32 OutputIndex) override;
-	virtual void GetCaption(TArray<FString>& OutCaptions) const override;
-	virtual FName GetInputName(int32 InputIndex) const override;
-	virtual bool IsInputConnectionRequired(int32 InputIndex) const override;
-	virtual uint32 GetInputType(int32 InputIndex) override;
-	virtual uint32 GetOutputType(int32 OutputIndex) override;
-
-	virtual bool IsResultMaterialAttributes(int32 OutputIndex) override { return OutputIndex == MSA_NonSubstrateAttributes; }
-	virtual bool IsResultSubstrateMaterial(int32 OutputIndex) override { return OutputIndex == MSA_FrontMaterial; }
-	virtual void GatherSubstrateMaterialInfo(FSubstrateMaterialInfo& SubstrateMaterialInfo, int32 OutputIndex) override;
-	virtual FSubstrateOperator* SubstrateGenerateMaterialTopologyTree(class FMaterialCompiler* Compiler, class UMaterialExpression* Parent, int32 OutputIndex) override;
-#endif // WITH_EDITOR
-#endif //#if ENABLE_MATERIAL_LAYER_PROTOTYPE
-};
-
-UCLASS(MinimalAPI, collapsecategories, hidecategories = Object, DisplayName = "Substrate Set Attributes")
-class UMaterialExpressionSubstrateSetAttributes : public UMaterialExpressionSubstrateUtilityBase
-{
-	GENERATED_UCLASS_BODY()
-
-	UPROPERTY()
-	FExpressionInput FrontMaterial;
-
-	UPROPERTY()
- 	FMaterialAttributesInput NonSubstrateAttributes;
-
-#if ENABLE_MATERIAL_LAYER_PROTOTYPE
-#if WITH_EDITOR
-	virtual void PostEditChangeProperty(FPropertyChangedEvent& PropertyChangedEvent) override;
-	virtual int32 Compile(class FMaterialCompiler* Compiler, int32 OutputIndex) override;
-	virtual void GetCaption(TArray<FString>& OutCaptions) const override;
-	virtual FName GetInputName(int32 InputIndex) const override;
-	virtual bool IsInputConnectionRequired(int32 InputIndex) const override;
-	virtual uint32 GetInputType(int32 InputIndex) override;
-
-	virtual bool IsResultMaterialAttributes(int32 OutputIndex) override { return true; }
-	virtual bool IsResultSubstrateMaterial(int32 OutputIndex) override { return FrontMaterial.IsConnected(); }
-	virtual void GatherSubstrateMaterialInfo(FSubstrateMaterialInfo& SubstrateMaterialInfo, int32 OutputIndex) override;
-	virtual FSubstrateOperator* SubstrateGenerateMaterialTopologyTree(class FMaterialCompiler* Compiler, class UMaterialExpression* Parent, int32 OutputIndex) override;
-#endif // WITH_EDITOR
-#endif //ENABLE_MATERIAL_LAYER_PROTOTYPE
 };
\ No newline at end of file
