[Shaders] Allow SCW logs to be dumped to disk on crash.

- Enable logging to memory in SCW target script to decide at runtime whether we need to dump the log to disk.

- Add "-DumpLogOnExitCrashOnly" command line argument.
  This is used in combination with DUMP_LOG_ON_EXIT to only dump the memory stored log to file if a critical error was reported (i.e. GIsCriticalError).

- Add CVar "r.ShaderCompiler.DebugDumpWorkerCrashLog" to dump SCW logs to disk when they crash.
  For example, a SCW with worker index 2 will store its log at "Engine/Programs/ShaderCompileWorker/Saved/Logs/ShaderCompileWorker-2.log" when run locally.
  On build machines, the respective path for artifacts is used.

- Make WITH_LOGGING_TO_MEMORY the only condition to allocate a memory backed output device for logging:
  CL 2940738 introduced the preprocessor guard for FOutputDeviceMemory '!IS_PROGRAM && !WITH_EDITORONLY_DATA', while the specific request via WITH_LOGGING_TO_MEMORY was added later.
  For the SCW, this is contradictory if we enable it explicitly but IS_PROGRAM is true for the worker processes.
  We therefore remove the original condition and allow logs to be stored in memory solely on the condition WITH_LOGGING_TO_MEMORY.

- Add "-ShaderCompileWorkerCrashLogsDir=<DIR>" argument to shader compiling manager to allow dumping SCW logs to customized base directory.

#rnx

diff --git a/Engine/Source/Programs/ShaderCompileWorker/ShaderCompileWorker.Target.cs b/Engine/Source/Programs/ShaderCompileWorker/ShaderCompileWorker.Target.cs
--- a/Engine/Source/Programs/ShaderCompileWorker/ShaderCompileWorker.Target.cs
+++ b/Engine/Source/Programs/ShaderCompileWorker/ShaderCompileWorker.Target.cs
@@ -52,6 +52,11 @@ public class ShaderCompileWorkerTarget : TargetRules
 		// Disable logging, as the workers are spawned often and logging will just slow them down
 		GlobalDefinitions.Add("ALLOW_LOG_FILE=0");
 
+		// Allow logging everything to memory so we can decide to dump everything to file when SCW crashed.
+		// This still requires the commandline argument -LogToMemory, so the shader compiling manager can decide to enable/disable it at runtime.
+		bLoggingToMemoryEnabled = true;
+		GlobalDefinitions.Add("USE_LOGGING_IN_SHIPPING=1");
+
 		// Linking against wer.lib/wer.dll causes XGE to bail when the worker is run on a Windows 8 machine, so turn this off.
 		GlobalDefinitions.Add("ALLOW_WINDOWS_ERROR_REPORT_LIB=0");
 
diff --git a/Engine/Source/Runtime/Core/Private/GenericPlatform/GenericPlatformOutputDevices.cpp b/Engine/Source/Runtime/Core/Private/GenericPlatform/GenericPlatformOutputDevices.cpp
--- a/Engine/Source/Runtime/Core/Private/GenericPlatform/GenericPlatformOutputDevices.cpp
+++ b/Engine/Source/Runtime/Core/Private/GenericPlatform/GenericPlatformOutputDevices.cpp
@@ -139,7 +139,6 @@ class FOutputDevice* FGenericPlatformOutputDevices::GetLog()
 		FLogOutputDeviceInitializer()
 		{
 #if WITH_LOGGING_TO_MEMORY
-#if !IS_PROGRAM && !WITH_EDITORONLY_DATA
 			if (!LogDevice
 #if !(UE_BUILD_SHIPPING || UE_BUILD_TEST)
 				&& FParse::Param(FCommandLine::Get(), TEXT("LOGTOMEMORY"))
@@ -150,7 +149,6 @@ class FOutputDevice* FGenericPlatformOutputDevices::GetLog()
 			{
 				LogDevice = MakeUnique<FOutputDeviceMemory>();
 			}
-#endif // !IS_PROGRAM && !WITH_EDITORONLY_DATA
 #endif // WITH_LOGGING_TO_MEMORY
 			if (!LogDevice)
 			{
diff --git a/Engine/Source/Runtime/Core/Private/Misc/OutputDeviceMemory.cpp b/Engine/Source/Runtime/Core/Private/Misc/OutputDeviceMemory.cpp
--- a/Engine/Source/Runtime/Core/Private/Misc/OutputDeviceMemory.cpp
+++ b/Engine/Source/Runtime/Core/Private/Misc/OutputDeviceMemory.cpp
@@ -5,12 +5,16 @@
 =============================================================================*/
 
 #include "Misc/OutputDeviceMemory.h"
-#include "HAL/PlatformTime.h"
-#include "HAL/PlatformOutputDevices.h"
+
+#include "CoreGlobals.h"
 #include "HAL/FileManager.h"
-#include "Misc/ScopeLock.h"
+#include "HAL/PlatformOutputDevices.h"
+#include "HAL/PlatformTime.h"
+#include "Misc/CommandLine.h"
 #include "Misc/OutputDeviceFile.h"
 #include "Misc/OutputDeviceHelper.h"
+#include "Misc/ScopeLock.h"
+#include "Misc/Parse.h"
 
 #define DUMP_LOG_ON_EXIT (!NO_LOGGING && PLATFORM_DESKTOP && (!UE_BUILD_SHIPPING || USE_LOGGING_IN_SHIPPING))
 
@@ -44,13 +48,17 @@ void FOutputDeviceMemory::TearDown()
 	}
 	// Dump on exit
 #if DUMP_LOG_ON_EXIT
-	const FString LogFileName = FPlatformOutputDevices::GetAbsoluteLogFilename();
-	FArchive* LogFile = IFileManager::Get().CreateFileWriter(*LogFileName, FILEWRITE_AllowRead);
-	if (LogFile)
+	const bool bDumpLogOnCrashOnly = FParse::Param(FCommandLine::Get(), TEXT("DumpLogOnExitCrashOnly"));
+	if (GIsCriticalError || !bDumpLogOnCrashOnly)
 	{
-		Dump(*LogFile);
-		LogFile->Flush();
-		delete LogFile;
+		const FString LogFileName = FPlatformOutputDevices::GetAbsoluteLogFilename();
+		FArchive* LogFile = IFileManager::Get().CreateFileWriter(*LogFileName, FILEWRITE_AllowRead);
+		if (LogFile)
+		{
+			Dump(*LogFile);
+			LogFile->Flush();
+			delete LogFile;
+		}
 	}
 #endif // DUMP_LOG_ON_EXIT
 }
diff --git a/Engine/Source/Runtime/Engine/Private/ShaderCompiler/ShaderCompiler.cpp b/Engine/Source/Runtime/Engine/Private/ShaderCompiler/ShaderCompiler.cpp
--- a/Engine/Source/Runtime/Engine/Private/ShaderCompiler/ShaderCompiler.cpp
+++ b/Engine/Source/Runtime/Engine/Private/ShaderCompiler/ShaderCompiler.cpp
@@ -162,6 +162,13 @@ static TAutoConsoleVariable<bool> CVarShaderCompilerDumpDDCKeys(
 	ECVF_Default
 );
 
+bool GDebugDumpWorkerCrashLog = false;
+static FAutoConsoleVariableRef CVarDebugDumpWorkerCrashLog(
+	TEXT("r.ShaderCompiler.DebugDumpWorkerCrashLog"),
+	GDebugDumpWorkerCrashLog,
+	TEXT("If true, the ShaderCompileWorker will dump its entire log to the Saved folder when a crash is detected."),
+	ECVF_ReadOnly
+
 static TAutoConsoleVariable<int> CVarShaderCompilerLogSlowJobThreshold(
 	TEXT("r.ShaderCompiler.LogSlowJobThreshold"),
 	30,
@@ -682,6 +689,11 @@ void FShaderCompilingManager::ReportMemoryUsage()
 	}
 }
 
+static FString GetBuildMachineArtifactBasePath()
+{
+	return FPaths::Combine(*FPaths::EngineDir(), TEXT("Programs"), TEXT("AutomationTool"), TEXT("Saved"), TEXT("Logs"));
+}
+
 FShaderCompilingManager::FShaderCompilingManager() :
 	bCompilingDuringGame(false),
 	NumExternalJobs(0),
@@ -789,11 +801,7 @@ FShaderCompilingManager::FShaderCompilingManager() :
 	AbsoluteShaderBaseWorkingDirectory = AbsoluteBaseDirectory + TEXT("/");
 
 	// Build machines should dump to the AutomationTool/Saved/Logs directory and they will upload as build artifacts via the AutomationTool.
-	FString BaseDebugInfoPath = FPaths::ProjectSavedDir();
-	if (GIsBuildMachine)
-	{
-		BaseDebugInfoPath = FPaths::Combine(*FPaths::EngineDir(), TEXT("Programs"), TEXT("AutomationTool"), TEXT("Saved"), TEXT("Logs"));
-	}
+	const FString BaseDebugInfoPath = GIsBuildMachine ? GetBuildMachineArtifactBasePath() : FPaths::ProjectSavedDir();
 
 	FString AbsoluteDebugInfoDirectory = IFileManager::Get().ConvertToAbsolutePathForExternalAppForWrite(*(BaseDebugInfoPath / TEXT("ShaderDebugInfo")));
 	const FString OverrideShaderDebugDir = CVarShaderOverrideDebugDir.GetValueOnAnyThread();
@@ -1447,6 +1455,26 @@ FProcHandle FShaderCompilingManager::LaunchWorker(const FString& WorkingDirector
 	WorkerParameters += FString(TEXT(" -nothreading "));
 #endif // USE_SHADER_COMPILER_WORKER_TRACE
 
+	if (GDebugDumpWorkerCrashLog)
+	{
+		WorkerParameters += TEXT(" -LogToMemory -DumpLogOnExitCrashOnly ");
+
+		const FString WorkerLogFilename = FString::Printf(TEXT("ShaderCompileWorker-%d.log"), ThreadId);
+		FString CustomCrashLogsDir;
+		if (FParse::Value(FCommandLine::Get(), TEXT("ShaderCompileWorkerCrashLogsDir"), CustomCrashLogsDir))
+		{
+			WorkerParameters += FString::Printf(TEXT("-AbsLog=%s"), *FPaths::Combine(CustomCrashLogsDir, WorkerLogFilename));
+		}
+		else if (GIsBuildMachine)
+		{
+			WorkerParameters += FString::Printf(TEXT("-AbsLog=%s"), *FPaths::Combine(GetBuildMachineArtifactBasePath(), WorkerLogFilename));
+		}
+		else
+		{
+			WorkerParameters += FString::Printf(TEXT("-Log=%s"), *WorkerLogFilename);
+		}
+	}
+
 	// Launch the worker process
 	int32 PriorityModifier = -1; // below normal
 	GConfig->GetInt(TEXT("DevOptions.Shaders"), TEXT("WorkerProcessPriority"), PriorityModifier, GEngineIni);
