Fix compile error from merge conflict as follow up to CL 35144519.

diff --git a/Engine/Source/Runtime/Engine/Private/ShaderCompiler/ShaderCompiler.cpp b/Engine/Source/Runtime/Engine/Private/ShaderCompiler/ShaderCompiler.cpp
--- a/Engine/Source/Runtime/Engine/Private/ShaderCompiler/ShaderCompiler.cpp
+++ b/Engine/Source/Runtime/Engine/Private/ShaderCompiler/ShaderCompiler.cpp
@@ -168,6 +168,7 @@ static FAutoConsoleVariableRef CVarDebugDumpWorkerCrashLog(
 	GDebugDumpWorkerCrashLog,
 	TEXT("If true, the ShaderCompileWorker will dump its entire log to the Saved folder when a crash is detected."),
 	ECVF_ReadOnly
+);
 
 static TAutoConsoleVariable<int> CVarShaderCompilerLogSlowJobThreshold(
 	TEXT("r.ShaderCompiler.LogSlowJobThreshold"),
