Ensure we check for frame counter when checking for CacheBone validity in AnimBP nodes

diff --git a/Engine/Plugins/Animation/BlendStack/Source/Runtime/Private/AnimNode_BlendStack.cpp b/Engine/Plugins/Animation/BlendStack/Source/Runtime/Private/AnimNode_BlendStack.cpp
--- a/Engine/Plugins/Animation/BlendStack/Source/Runtime/Private/AnimNode_BlendStack.cpp
+++ b/Engine/Plugins/Animation/BlendStack/Source/Runtime/Private/AnimNode_BlendStack.cpp
@@ -719,7 +719,7 @@ void UE::BlendStack::FBlendStack_SampleGraphExecutionHelper::EvaluatePlayer(FPos
 void UE::BlendStack::FBlendStack_SampleGraphExecutionHelper::ConditionalCacheBones(const FAnimationBaseContext& Context, FPoseLink& SamplePoseLink)
 {
 	// Only call CacheBones when needed.
-	if (!CacheBoneCounter.IsSynchronized_Counter(Context.AnimInstanceProxy->GetCachedBonesCounter()))
+	if (!CacheBoneCounter.IsSynchronized_All(Context.AnimInstanceProxy->GetCachedBonesCounter()))
 	{
 		// Keep track of samples that have had CacheBones called on.
 		CacheBoneCounter.SynchronizeWith(Context.AnimInstanceProxy->GetCachedBonesCounter());
diff --git a/Engine/Source/Runtime/Engine/Private/Animation/AnimNode_SaveCachedPose.cpp b/Engine/Source/Runtime/Engine/Private/Animation/AnimNode_SaveCachedPose.cpp
--- a/Engine/Source/Runtime/Engine/Private/Animation/AnimNode_SaveCachedPose.cpp
+++ b/Engine/Source/Runtime/Engine/Private/Animation/AnimNode_SaveCachedPose.cpp
@@ -96,7 +96,7 @@ void FAnimNode_SaveCachedPose::Initialize_AnyThread(const FAnimationInitializeCo
 
 void FAnimNode_SaveCachedPose::CacheBones_AnyThread(const FAnimationCacheBonesContext& Context)
 {
-	if (!CachedBonesCounter.IsSynchronized_Counter(Context.AnimInstanceProxy->GetCachedBonesCounter()))
+	if (!CachedBonesCounter.IsSynchronized_All(Context.AnimInstanceProxy->GetCachedBonesCounter()))
 	{
 		CachedBonesCounter.SynchronizeWith(Context.AnimInstanceProxy->GetCachedBonesCounter());
 
diff --git a/Engine/Source/Runtime/Engine/Private/Animation/AnimNode_TransitionPoseEvaluator.cpp b/Engine/Source/Runtime/Engine/Private/Animation/AnimNode_TransitionPoseEvaluator.cpp
--- a/Engine/Source/Runtime/Engine/Private/Animation/AnimNode_TransitionPoseEvaluator.cpp
+++ b/Engine/Source/Runtime/Engine/Private/Animation/AnimNode_TransitionPoseEvaluator.cpp
@@ -38,7 +38,7 @@ void FAnimNode_TransitionPoseEvaluator::Initialize_AnyThread(const FAnimationIni
 
 void FAnimNode_TransitionPoseEvaluator::CacheBones_AnyThread(const FAnimationCacheBonesContext& Context) 
 {
-	if (!CachedBonesCounter.IsSynchronized_Counter(Context.AnimInstanceProxy->GetCachedBonesCounter()))
+	if (!CachedBonesCounter.IsSynchronized_All(Context.AnimInstanceProxy->GetCachedBonesCounter()))
 	{
 		CachedBonesCounter.SynchronizeWith(Context.AnimInstanceProxy->GetCachedBonesCounter());
 
@@ -53,7 +53,7 @@ void FAnimNode_TransitionPoseEvaluator::CacheBones_AnyThread(const FAnimationCac
 
 void FAnimNode_TransitionPoseEvaluator::Update_AnyThread(const FAnimationUpdateContext& Context)
 {
-	if (!CachedBonesCounter.IsSynchronized_Counter(Context.AnimInstanceProxy->GetCachedBonesCounter()))
+	if (!CachedBonesCounter.IsSynchronized_All(Context.AnimInstanceProxy->GetCachedBonesCounter()))
 	{
 		CachedBonesCounter.SynchronizeWith(Context.AnimInstanceProxy->GetCachedBonesCounter());
 
@@ -90,7 +90,7 @@ void FAnimNode_TransitionPoseEvaluator::GatherDebugData(FNodeDebugData& DebugDat
 bool FAnimNode_TransitionPoseEvaluator::InputNodeNeedsUpdate(const FAnimationUpdateContext& Context) const
 {
 	// EM_Standard mode always updates and EM_DelayedFreeze mode only updates if there are cache frames remaining
-	return (EvaluatorMode == EEvaluatorMode::EM_Standard) || ((EvaluatorMode == EEvaluatorMode::EM_DelayedFreeze) && (CacheFramesRemaining > 0)) || !CachedBonesCounter.IsSynchronized_Counter(Context.AnimInstanceProxy->GetCachedBonesCounter());
+	return (EvaluatorMode == EEvaluatorMode::EM_Standard) || ((EvaluatorMode == EEvaluatorMode::EM_DelayedFreeze) && (CacheFramesRemaining > 0)) || !CachedBonesCounter.IsSynchronized_All(Context.AnimInstanceProxy->GetCachedBonesCounter());
 }
 
 bool FAnimNode_TransitionPoseEvaluator::InputNodeNeedsEvaluate() const
