World Partition - Temporary fix to protect OnCellShown/OnCellHidden() delegates from crashing when called on an invalid object

diff --git a/Engine/Source/Runtime/Engine/Private/WorldPartition/WorldPartitionRuntimeLevelStreamingCell.cpp b/Engine/Source/Runtime/Engine/Private/WorldPartition/WorldPartitionRuntimeLevelStreamingCell.cpp
--- a/Engine/Source/Runtime/Engine/Private/WorldPartition/WorldPartitionRuntimeLevelStreamingCell.cpp
+++ b/Engine/Source/Runtime/Engine/Private/WorldPartition/WorldPartitionRuntimeLevelStreamingCell.cpp
@@ -117,6 +117,13 @@ TArray<FName> UWorldPartitionRuntimeLevelStreamingCell::GetActors() const
 
 void UWorldPartitionRuntimeLevelStreamingCell::CreateAndSetLevelStreaming(const FString& InPackageName, const FSoftObjectPath& InWorldAsset)
 {
+	// Temporary test while we investigate issue PLAY-45493
+	if (!ensureMsgf(!LevelStreaming, TEXT("StreamingCell already had an assigned LevelStreaming object '%s'"), *LevelStreaming->GetPathName()))
+	{
+		LevelStreaming->OnLevelShown.RemoveAll(this);
+		LevelStreaming->OnLevelHidden.RemoveAll(this);
+	}
+
 	LevelStreaming = CreateLevelStreaming(InPackageName, InWorldAsset);
 }
 
@@ -690,10 +697,14 @@ void UWorldPartitionRuntimeLevelStreamingCell::OnLevelShown()
 
 void UWorldPartitionRuntimeLevelStreamingCell::OnCellShown() const
 {
-	UWorldPartition* OuterWorldPartition = GetOuterWorld()->GetWorldPartition();
-	if (OuterWorldPartition && OuterWorldPartition->IsInitialized())
+	// Temporary test while we investigate issue PLAY-45493
+	if (ensure(IsValid(this)))
 	{
-		OuterWorldPartition->OnCellShown(this);
+		UWorldPartition* OuterWorldPartition = GetOuterWorld()->GetWorldPartition();
+		if (OuterWorldPartition && OuterWorldPartition->IsInitialized())
+		{
+			OuterWorldPartition->OnCellShown(this);
+		}
 	}
 }
 
@@ -704,9 +715,13 @@ void UWorldPartitionRuntimeLevelStreamingCell::OnLevelHidden()
 
 void UWorldPartitionRuntimeLevelStreamingCell::OnCellHidden() const
 {
-	UWorldPartition* OuterWorldPartition = GetOuterWorld()->GetWorldPartition();
-	if (OuterWorldPartition && OuterWorldPartition->IsInitialized())
+	// Temporary test while we investigate issue PLAY-45493
+	if (ensure(IsValid(this)))
 	{
-		OuterWorldPartition->OnCellHidden(this);
+		UWorldPartition* OuterWorldPartition = GetOuterWorld()->GetWorldPartition();
+		if (OuterWorldPartition && OuterWorldPartition->IsInitialized())
+		{
+			OuterWorldPartition->OnCellHidden(this);
+		}
 	}
 }
