Console: Copy text from console scrollback (history). First version selects/draws and copies full lines.

diff --git a/Engine/Plugins/Runtime/CommonUI/Source/CommonUI/Private/CommonGameViewportClient.cpp b/Engine/Plugins/Runtime/CommonUI/Source/CommonUI/Private/CommonGameViewportClient.cpp
--- a/Engine/Plugins/Runtime/CommonUI/Source/CommonUI/Private/CommonGameViewportClient.cpp
+++ b/Engine/Plugins/Runtime/CommonUI/Source/CommonUI/Private/CommonGameViewportClient.cpp
@@ -101,6 +101,22 @@ bool UCommonGameViewportClient::InputTouch(FViewport* InViewport, int32 Controll
 	return Super::InputTouch(InViewport, ControllerId, Handle, Type, TouchLocation, Force, DeviceTimestamp, TouchpadIndex);
 }
 
+void UCommonGameViewportClient::MouseMove(FViewport* InViewport, int32 X, int32 Y)
+{
+	if (ViewportConsole)
+	{
+		ViewportConsole->MouseMove(InViewport, X, Y);
+	}
+}
+
+void UCommonGameViewportClient::CapturedMouseMove(FViewport* InViewport, int32 X, int32 Y)
+{
+	if (ViewportConsole)
+	{
+		ViewportConsole->CapturedMouseMove(InViewport, X, Y);
+	}
+}
+
 void UCommonGameViewportClient::HandleRerouteInput(FInputDeviceId DeviceId, FKey Key, EInputEvent EventType, FReply& Reply)
 {
 	FPlatformUserId OwningPlatformUser = IPlatformInputDeviceMapper::Get().GetUserForInputDevice(DeviceId);
diff --git a/Engine/Plugins/Runtime/CommonUI/Source/CommonUI/Public/CommonGameViewportClient.h b/Engine/Plugins/Runtime/CommonUI/Source/CommonUI/Public/CommonGameViewportClient.h
--- a/Engine/Plugins/Runtime/CommonUI/Source/CommonUI/Public/CommonGameViewportClient.h
+++ b/Engine/Plugins/Runtime/CommonUI/Source/CommonUI/Public/CommonGameViewportClient.h
@@ -28,6 +28,8 @@ public:
 	virtual bool InputKey(const FInputKeyEventArgs& EventArgs) override;
 	virtual bool InputAxis(FViewport* InViewport, FInputDeviceId InputDevice, FKey Key, float Delta, float DeltaTime, int32 NumSamples, bool bGamepad) override;
 	virtual bool InputTouch(FViewport* InViewport, int32 ControllerId, uint32 Handle, ETouchType::Type Type, const FVector2D& TouchLocation, float Force, FDateTime DeviceTimestamp, uint32 TouchpadIndex) override;
+	virtual void MouseMove(FViewport* InViewport, int32 X, int32 Y) override;
+	virtual void CapturedMouseMove(FViewport* InViewport, int32 X, int32 Y) override;
 	// UGameViewportClient interface end
 
 	FOnRerouteInputDelegate& OnRerouteInput() { return RerouteInput; }
diff --git a/Engine/Source/Runtime/Engine/Classes/Engine/Console.h b/Engine/Source/Runtime/Engine/Classes/Engine/Console.h
--- a/Engine/Source/Runtime/Engine/Classes/Engine/Console.h
+++ b/Engine/Source/Runtime/Engine/Classes/Engine/Console.h
@@ -11,6 +11,7 @@
 
 #include "Console.generated.h"
 
+class FViewport;
 class SWidget;
 struct FAutoCompleteCommand;
 
@@ -136,6 +137,21 @@ class UConsole
 	/** Current list of matching commands for auto-complete, @see UpdateCompleteIndices() */
 	TArray<FAutoCompleteCommand> AutoComplete;
 
+	// Scrollback region selection mouse and state
+	struct ScrollbackSelection
+	{
+		FVector2D MousePosDown = FVector2D(0, 0);
+		FVector2D MousePosUp = FVector2D(0, 0);
+		int32 MousePosX = 0;
+		int32 MousePosY = 0;
+		bool bActive = false;
+		bool bCapture = false;
+		bool bMade = false;
+		float Offset = 0.0f;
+	} Selection;
+
+	float TextH = 15.0f;
+
 	ENGINE_API ~UConsole();
 
 	// UObject Interface
@@ -295,6 +311,9 @@ class UConsole
 	virtual bool InputTouch(int32 ControllerId, uint32 Handle, ETouchType::Type Type, const FVector2D& TouchLocation, float Force, FDateTime DeviceTimestamp, uint32 TouchpadIndex) { return false; }
 	virtual bool InputTouch(FInputDeviceId DevideId, uint32 Handle, ETouchType::Type Type, const FVector2D& TouchLocation, float Force, FDateTime DeviceTimestamp, uint32 TouchpadIndex) { return false; }
 
+	ENGINE_API virtual void MouseMove(FViewport* Viewport, int32 X, int32 Y);
+	ENGINE_API virtual void CapturedMouseMove(FViewport* InViewport, int32 X, int32 Y);
+
 	/** render to the canvas based on the console state */
 	ENGINE_API virtual void PostRender_Console(class UCanvas* Canvas);
 
diff --git a/Engine/Source/Runtime/Engine/Private/UserInterface/Console.cpp b/Engine/Source/Runtime/Engine/Private/UserInterface/Console.cpp
--- a/Engine/Source/Runtime/Engine/Private/UserInterface/Console.cpp
+++ b/Engine/Source/Runtime/Engine/Private/UserInterface/Console.cpp
@@ -660,6 +660,11 @@ void UConsole::OutputTextLine(const FString& Text)
 
 	// Add the line
 	Scrollback.Add(Text);
+
+	if (Selection.bActive)
+	{
+		Selection.Offset += TextH;
+	}
 }
 
 
@@ -1275,6 +1280,25 @@ PRAGMA_ENABLE_DEPRECATION_WARNINGS
 
 bool UConsole::InputKey_Open(FInputDeviceId DeviceId, FKey Key, EInputEvent Event, float AmountDepressed, bool bGamepad)
 {
+	if (Key == EKeys::LeftMouseButton)
+	{
+		if (Event == IE_Pressed)
+		{
+			Selection.MousePosDown = FVector2D(Selection.MousePosX, Selection.MousePosY);
+			Selection.bActive = true;
+			Selection.bMade = false;
+			Selection.Offset = 0.0f;
+			return true;
+		}
+		else if (Event == IE_Released)
+		{
+			Selection.MousePosUp = FVector2D(Selection.MousePosX, Selection.MousePosY);
+			Selection.bCapture = true;
+			Selection.bMade = true;
+			return true;
+		}
+	}
+
 	if (Key == EKeys::PageUp || Key == EKeys::MouseScrollUp)
 	{
 		if (SBPos < Scrollback.Num() - 1)
@@ -1332,33 +1356,34 @@ void UConsole::PostRender_Console_Open(UCanvas* Canvas)
 	const bool bDPIAwareStringMeasurement = true;
 
 	// determine the height of the text
-	float xl, yl;
-	Canvas->StrLen(Font, TEXT("M"),xl,yl, bDPIAwareStringMeasurement);
+	float xl;
+	Canvas->StrLen(Font, TEXT("M"), xl, TextH, bDPIAwareStringMeasurement);
 	xl /= DPIScale;
-	yl /= DPIScale;
-	// Background
-	FLinearColor BackgroundColor = ConsoleDefs::AutocompleteBackgroundColor.ReinterpretAsLinear();
-	BackgroundColor.A = ConsoleSettings->BackgroundOpacityPercentage / 100.0f;
-	FCanvasTileItem ConsoleTile(FVector2D(LeftPos, 0.0f), DefaultTexture_Black->GetResource(), FVector2D(ClipX, Height + TopPos - yl), FVector2D(0.0f, 0.0f), FVector2D(1.0f, 1.0f), BackgroundColor);
+	TextH /= DPIScale;
 
-	// Preserve alpha to allow single-pass composite
-	ConsoleTile.BlendMode = SE_BLEND_AlphaBlend;
+	float y = Height - TextH;
 
-	Canvas->DrawItem(ConsoleTile);
+	if (Scrollback.Num())
+	{
+		FLinearColor BackgroundColorBlack = ConsoleDefs::AutocompleteBackgroundColor.ReinterpretAsLinear();
+		FLinearColor BackgroundColorWhite = FLinearColor::White;
+		BackgroundColorBlack.A = ConsoleSettings->BackgroundOpacityPercentage / 100.0f;
+		BackgroundColorWhite.A = ConsoleSettings->BackgroundOpacityPercentage / 100.0f;
 
-	// figure out which element of the scrollback buffer to should appear first (at the top of the screen)
-	int32 idx = SBHead - SBPos;
+		FCanvasTextItem ConsoleText(FVector2D(LeftPos, TopPos + Height - 5 - TextH), FText::FromString(TEXT("")), Font, ConsoleSettings->InputColor);
 
-	float y = Height - yl;
+		FCanvasTileItem ConsoleTileBlack(FVector2D(LeftPos, 0.0f), DefaultTexture_Black->GetResource(), FVector2D(ClipX, Height + TopPos - TextH), BackgroundColorBlack);
+		FCanvasTileItem ConsoleTileWhite(FVector2D(LeftPos, 0.0f), DefaultTexture_White->GetResource(), FVector2D(ClipX, Height + TopPos - TextH), BackgroundColorWhite);
+		ConsoleTileBlack.BlendMode = SE_BLEND_AlphaBlend;		// Preserve alpha to allow single-pass composite
+		ConsoleTileWhite.BlendMode = SE_BLEND_AlphaBlend;		// Preserve alpha to allow single-pass composite
 
-	if (Scrollback.Num())
-	{
-		FCanvasTextItem ConsoleText(FVector2D(LeftPos, TopPos + Height - 5 - yl), FText::FromString(TEXT("")), Font, ConsoleSettings->InputColor);
-		// change the text color to white
-		ConsoleText.SetColor(FLinearColor::White);
+		// figure out which element of the scrollback buffer should appear first (at the bottom of the scrollback area)
+		int32 idx = SBHead - SBPos;
+
+		TArray<FString> SelectedLines;
 
 		// while we have enough room to draw another line and there are more lines to draw
-		while (y > -yl && idx >= 0)
+		while (y > -TextH && idx >= 0)
 		{
 			float PenX;
 			float PenY;
@@ -1373,21 +1398,96 @@ void UConsole::PostRender_Console_Open(UCanvas* Canvas)
 				Canvas->StrLen(Font, Scrollback[idx], ScrollLineXL, ScrollLineYL, bDPIAwareStringMeasurement);
 				ScrollLineXL /= DPIScale;
 				ScrollLineYL /= DPIScale;
-				if (ScrollLineYL > yl)
+				if (ScrollLineYL > TextH)
 				{
-					y -= (ScrollLineYL - yl);
+					y -= (ScrollLineYL - TextH);
 					PenX = LeftPos;
 					PenY = TopPos + y;
 				}
 
+				bool bLineSelected = false;
+				if (Selection.bActive)
+				{
+					// Check if selection range overlaps current line
+					float x0 = PenY;
+					float x1 = PenY + TextH;
+					float y0 = Selection.MousePosDown.Y;
+					float y1 = Selection.bMade ? Selection.MousePosUp.Y : Selection.MousePosY;
+
+					// Adjust for new lines added since selection was initiated
+					y0 -= Selection.Offset;
+					y1 -= Selection.Offset;
+
+					// Swap to keep range in order, to allow drag selecting down or up
+					if (y0 > y1)
+					{
+						float tmp = y0;
+						y0 = y1;
+						y1 = tmp;
+					}
+
+					auto Overlaps = [](float x0, float x1, float y0, float y1)
+					{
+						return (x0 <= y1) && (x1 >= y0);
+					};
+
+					bLineSelected = Overlaps(x0, x1, y0, y1);
+				}
+
+				if (bLineSelected)
+				{
+					// Selected line (black text on white background)
+					ConsoleTileWhite.Position = FVector2D(LeftPos, PenY);
+					ConsoleTileWhite.Size = FVector2D(ClipX, TextH);
+					Canvas->DrawItem(ConsoleTileWhite);
+
+					ConsoleText.SetColor(FLinearColor::Black);
+
+					if (Selection.bCapture)
+					{
+						SelectedLines.Add(Scrollback[idx]);
+					}
+				}
+				else
+				{
+					// Unselected line (white text on black background)
+					ConsoleTileBlack.Position = FVector2D(LeftPos, PenY);
+					ConsoleTileBlack.Size = FVector2D(ClipX, TextH);
+					Canvas->DrawItem(ConsoleTileBlack);
+
+					ConsoleText.SetColor(FLinearColor::White);
+				}
+
 				ConsoleText.Text = FText::FromString(Scrollback[idx]);
 				Canvas->DrawItem(ConsoleText, PenX, PenY);
 			}
 			idx--;
-			y -= yl;
+			y -= TextH;
+		}
+
+		if (Selection.bCapture && !SelectedLines.IsEmpty())
+		{
+			Algo::Reverse(SelectedLines);
+			const FString SelectedText = FString::Join(SelectedLines, LINE_TERMINATOR);
+
+			FPlatformApplicationMisc::ClipboardCopy(*SelectedText);
+
+			Selection.bCapture = false;		// Finished copy to clipboard when mouse is released
 		}
 	}
 
+	{
+		// Draw any remaining empty scrollback area to the top of viewport
+
+		FLinearColor BackgroundColor = ConsoleDefs::AutocompleteBackgroundColor.ReinterpretAsLinear();
+		BackgroundColor.A = ConsoleSettings->BackgroundOpacityPercentage / 100.0f;
+
+		FCanvasTileItem ConsoleTile(FVector2D(LeftPos, 0.0f), DefaultTexture_Black->GetResource(), FVector2D(ClipX, y + TextH), BackgroundColor);
+		ConsoleTile.BlendMode = SE_BLEND_AlphaBlend;			// Preserve alpha to allow single-pass composite
+
+		Canvas->DrawItem(ConsoleTile);
+	}
+
 	PostRender_InputLine(Canvas, FIntPoint(LeftPos, TopPos + Height + 6));
 }
 
@@ -1462,12 +1562,26 @@ bool UConsole::InputKey(FInputDeviceId DeviceId, FKey Key, EInputEvent Event, fl
 	return bWasConsumed;
 }
 
+void UConsole::MouseMove(FViewport* Viewport, int32 X, int32 Y)
+{
+	Selection.MousePosX = X;
+	Selection.MousePosY = Y;
+}
+
+void UConsole::CapturedMouseMove(FViewport* InViewport, int32 X, int32 Y)
+{
+	Selection.MousePosX = X;
+	Selection.MousePosY = Y;
+}
 
 void UConsole::PostRender_Console(UCanvas* Canvas)
 {
-	if (ConsoleState != NAME_None)
+	if (ConsoleState == NAME_None)
 	{
-		Canvas->ApplySafeZoneTransform();
+		return;
+	}
+
+	Canvas->ApplySafeZoneTransform();
 
 	if (ConsoleState == NAME_Typing)
 	{
@@ -1478,8 +1592,7 @@ void UConsole::PostRender_Console(UCanvas* Canvas)
 		PostRender_Console_Open(Canvas);
 	}
 
-		Canvas->PopSafeZoneTransform();
-	}
+	Canvas->PopSafeZoneTransform();
 }
 
 void UConsole::PostRender_InputLine(UCanvas* Canvas, FIntPoint UserInputLinePos)
@@ -1718,6 +1831,9 @@ void UConsole::FakeGotoState(FName NextStateName)
 	else if (ConsoleState == NAME_Open)
 	{
 		EndState_Open(NextStateName);
+
+		// Restore mouse capture for non-open states
+		GetOuterUGameViewportClient()->SetMouseCaptureMode(EMouseCaptureMode::CapturePermanently);
 	}
 	if (NextStateName == NAME_Typing)
 	{
@@ -1739,6 +1855,13 @@ void UConsole::FakeGotoState(FName NextStateName)
 
 		// Console has opened
 		OnConsoleActivationStateChanged.Broadcast(true);
+
+		// Ensure we receive IE_Pressed with LeftMouseButton
+		GetOuterUGameViewportClient()->SetMouseCaptureMode(EMouseCaptureMode::CaptureDuringMouseDown);
+		Selection.bActive = false;
+		Selection.bCapture = false;
+		Selection.bMade = false;
+		Selection.Offset = 0.0f;
 	}
 	else if (NextStateName == NAME_None)
 	{
