Fixing default SoundWave compression/codec settings

[RN] Moved Default Compression out of Project settings into SoundWave defaults to fix some issues with defaults being written
#tests Sondra M ran shelf and tested locally

diff --git a/Engine/Source/Runtime/Engine/Classes/Sound/AudioSettings.h b/Engine/Source/Runtime/Engine/Classes/Sound/AudioSettings.h
--- a/Engine/Source/Runtime/Engine/Classes/Sound/AudioSettings.h
+++ b/Engine/Source/Runtime/Engine/Classes/Sound/AudioSettings.h
@@ -177,8 +177,8 @@ class UAudioSettings : public UDeveloperSettings
 	EDefaultAudioCompressionType DefaultAudioCompressionType;
 
 	/** The default compression quality (e.g. for new SoundWaves) */
-	UPROPERTY(config, EditAnywhere, Category = "Audio", meta = (ClampMin = 1, UIMin = 1, UIMax = 100))
-	int32 DefaultCompressionQuality = 80;
+	UPROPERTY(config)
+	int32 DefaultCompressionQuality_DEPRECATED = 80;
 
 	/** The amount of audio to send to reverb submixes if no reverb send is setup for the source through attenuation settings. Only used in audio mixer. */
 	UPROPERTY(config)
@@ -283,6 +283,7 @@ public:
 	// Get the quality level settings at the provided level index
 	ENGINE_API const FAudioQualitySettings& GetQualityLevelSettings(int32 QualityLevel) const;
 
+	UE_DEPRECATED(5.5, "Default Compression Quality has moved to SoundWave class defaults")
 	ENGINE_API int32 GetDefaultCompressionQuality() const;
 	
 	// Get the quality name level for a given index
diff --git a/Engine/Source/Runtime/Engine/Classes/Sound/SoundWave.h b/Engine/Source/Runtime/Engine/Classes/Sound/SoundWave.h
--- a/Engine/Source/Runtime/Engine/Classes/Sound/SoundWave.h
+++ b/Engine/Source/Runtime/Engine/Classes/Sound/SoundWave.h
@@ -440,7 +440,7 @@ class USoundWave : public USoundBase, public IAudioProxyDataFactory, public IInt
 private:
 
 	/** Platform agnostic compression quality. 1..100 with 1 being best compression and 100 being best quality. ADPCM and PCM sound asset compression types ignore this parameter. */
-	UPROPERTY(EditAnywhere, Category = "Format|Quality", meta = (DisplayName = "Compression", ClampMin = "1", ClampMax = "100", EditCondition = "SoundAssetCompressionType != ESoundAssetCompressionType::PCM && SoundAssetCompressionType != ESoundAssetCompressionType::ADPCM"), AssetRegistrySearchable)
+	UPROPERTY(Config, EditAnywhere, Category = "Format|Quality", meta = (DisplayName = "Compression", ClampMin = "1", ClampMax = "100", EditCondition = "SoundAssetCompressionType != ESoundAssetCompressionType::PCM && SoundAssetCompressionType != ESoundAssetCompressionType::ADPCM"), AssetRegistrySearchable)
 	int32 CompressionQuality;
 
 public:
@@ -479,7 +479,7 @@ private:
 	uint8 bUseBinkAudio : 1;
 
 	/** The compression type to use for the sound wave asset. */
-	UPROPERTY(EditAnywhere, Category = "Format", AssetRegistrySearchable)
+	UPROPERTY(Config, EditAnywhere, Category = "Format", AssetRegistrySearchable)
 	ESoundAssetCompressionType SoundAssetCompressionType = ESoundAssetCompressionType::PlatformSpecific;
 
 public:
diff --git a/Engine/Source/Runtime/Engine/Private/AudioSettings.cpp b/Engine/Source/Runtime/Engine/Private/AudioSettings.cpp
--- a/Engine/Source/Runtime/Engine/Private/AudioSettings.cpp
+++ b/Engine/Source/Runtime/Engine/Private/AudioSettings.cpp
@@ -189,7 +189,7 @@ const FAudioQualitySettings& UAudioSettings::GetQualityLevelSettings(int32 Quali
 
 int32 UAudioSettings::GetDefaultCompressionQuality() const
 {
-	return FMath::Clamp(DefaultCompressionQuality,1,100);
+	return FMath::Clamp(DefaultCompressionQuality_DEPRECATED,1,100);
 }
 
 int32 UAudioSettings::GetQualityLevelSettingsNum() const
diff --git a/Engine/Source/Runtime/Engine/Private/SoundWave.cpp b/Engine/Source/Runtime/Engine/Private/SoundWave.cpp
--- a/Engine/Source/Runtime/Engine/Private/SoundWave.cpp
+++ b/Engine/Source/Runtime/Engine/Private/SoundWave.cpp
@@ -927,7 +927,6 @@ USoundWave::USoundWave(const FObjectInitializer& ObjectInitializer)
 {
 	Volume = 1.0;
 	Pitch = 1.0;
-	CompressionQuality = -1; // Note: This is set in PostInitProperties, as its unsafe to GetDefault<UAudioSettings> here.
 	SubtitlePriority = DEFAULT_SUBTITLE_PRIORITY;
 	ResourceState = ESoundWaveResourceState::NeedsFree;
 	RawPCMDataSize = 0;
@@ -956,7 +955,8 @@ USoundWave::USoundWave(const FObjectInitializer& ObjectInitializer)
 	bProcedural = false;
 	bRequiresStopFade = false;
 
-	SoundAssetCompressionType = ESoundAssetCompressionType::ProjectDefined;	
+	SoundAssetCompressionType = ESoundAssetCompressionType::ProjectDefined;
+	CompressionQuality = 80;
 
 #if WITH_EDITOR
 	bWasStreamCachingEnabledOnLastCook = FPlatformCompressionUtilities::IsCurrentPlatformUsingStreamCaching();
@@ -1463,11 +1463,6 @@ void USoundWave::PostInitProperties()
 {
 	Super::PostInitProperties();
 
-	// Safe to query the AudioSettings here, instead of the Constructor
-	// These defaults are then overwritten by the UPROPERTY serialization if there's non default property values.
-	CompressionQuality = GetDefault<UAudioSettings>()->GetDefaultCompressionQuality();
-	SoundAssetCompressionType = Audio::ToSoundAssetCompressionType(GetDefault<UAudioSettings>()->DefaultAudioCompressionType);
-
 	if(!IsTemplate())
 	{
 		// Don't rebuild our streaming chunks yet because we may not have loaded the RawPCMData at this point.
