[Material Editor] Removed experimental and unused path to generate more accurate FMaterialCachedExpressionData.

We are moving away from FMaterialCachedExpressionData in the future hopefully and work on new material translator makes this code obsolete.

#rnx

diff --git a/Engine/Source/Runtime/Engine/Private/Materials/MaterialCachedData.cpp b/Engine/Source/Runtime/Engine/Private/Materials/MaterialCachedData.cpp
--- a/Engine/Source/Runtime/Engine/Private/Materials/MaterialCachedData.cpp
+++ b/Engine/Source/Runtime/Engine/Private/Materials/MaterialCachedData.cpp
@@ -50,9 +50,6 @@ const FMaterialCachedExpressionEditorOnlyData FMaterialCachedExpressionEditorOnl
 
 static_assert((uint64)(EMaterialProperty::MP_MaterialAttributes)-1 < (8 * sizeof(FMaterialCachedExpressionData::PropertyConnectedMask)), "PropertyConnectedMask cannot contain entire EMaterialProperty enumeration.");
 
-static bool GExperimentalMaterialCachedDataAnalysisEnabled = false;
-static FAutoConsoleVariableRef CVarExperimentalMaterialCachedDataAnalysisEnabled(TEXT("r.Material.ExperimentalMaterialCachedDataAnalysisEnabled"), GExperimentalMaterialCachedDataAnalysisEnabled, TEXT("Enables material cached data experimental graph based analysis"));
-
 FMaterialCachedExpressionData::FMaterialCachedExpressionData()
 	: FunctionInfosStateCRC(0xffffffff)
 	, bHasMaterialLayers(false)
@@ -556,404 +553,23 @@ void FMaterialCachedExpressionData::UpdateForExpressions(const FMaterialCachedEx
 	}
 }
 
-/**
- * A helper struct used to crawl through the graph starting from material outputs in order to detect which material parameters are
- * written to. This information is used to mark the connected properties in a MaterialCachedData instance.
- * The graph is explored starting from the material outputs and walking backwards to its inputs. Whenever a static switch node is
- * encountered, the algorithm finds out the value of the associated static parameter and continues exploration only along the
- * activated subgraph. Because of this, unused subgraphs aren't visited so that any potential property writes in them don't interfere
- * with the final result.
- */
-struct FMaterialConnectedPropertiesAnalyzer
-{
-	/** The target FMaterialCachedExpressionData to set connected parameters to */
-	FMaterialCachedExpressionData&	CachedExpressionData;
-	
-	/** The current expression (if inside one, otherwise nullptr) */
-	UMaterialFunctionInterface*		CurrentFunction;
-	
-	/** The list of "fringe" material expressions that are yet to be explored */
-	TArray<UMaterialExpression*>	UnexploredExpressions;
-	
-	/** The set of visited expressions, used to avoid visiting the same expressions twice */
-	TSet<UMaterialExpression*>		VisitedExpressions;
-
-	/** Constructor */
-	FMaterialConnectedPropertiesAnalyzer(FMaterialCachedExpressionData& CachedExpressionData, UMaterialFunctionInterface* CurrentFunction = nullptr)
-		: CachedExpressionData(CachedExpressionData)
-		, CurrentFunction(CurrentFunction)
-	{
-	}
-
-	/** Pops an expression from the unexplored list, returning null if list is empty */
-	UMaterialExpression* PopUnvisitedExpression()
-	{
-		if (UnexploredExpressions.IsEmpty())
-		{
-			return nullptr;
-		}
-		UMaterialExpression* Expression = UnexploredExpressions.Last();
-		UnexploredExpressions.Pop();
-		return Expression;
-	}
-
-	/** Pushes an expression to the unexplored list if has not yet been explored */
-	void PushUnexploredExpression(UMaterialExpression* Expression)
-	{
-		if (!Expression || VisitedExpressions.Contains(Expression))
-		{
-			return;
-		}
-		UnexploredExpressions.Add(Expression);
-		VisitedExpressions.Add(Expression);
-	}
-
-	/** Visits all expressions in the UnexploredExpressions list */
-	void VisitUnexploredExpressions(EMaterialParameterAssociation Association, int32 ParameterIndex)
-	{
-		for (;;)
-		{
-			// Pop the current expression.
-			UMaterialExpression* Expression = PopUnvisitedExpression();
-			if (!Expression)
-			{
-				break;
-			}
-
-			// By default explore all input expressions. If static switch nodes are encountered, they
-			// may be able to continue exploration only along the active input and set this flag to false.
-			bool bVisitAllInputExpressions = true;
-
-			// Explore the current expression
-			if (UMaterialExpressionCollectionParameter* ExpressionCollectionParameter = Cast<UMaterialExpressionCollectionParameter>(Expression))
-			{
-				UMaterialParameterCollection* Collection = ExpressionCollectionParameter->Collection;
-				if (Collection)
-				{
-					FMaterialParameterCollectionInfo NewInfo;
-					NewInfo.ParameterCollection = Collection;
-					NewInfo.StateId = Collection->StateId;
-					CachedExpressionData.ParameterCollectionInfos.AddUnique(NewInfo);
-				}
-			}
-			else if (UMaterialExpressionDynamicParameter* ExpressionDynamicParameter = Cast< UMaterialExpressionDynamicParameter>(Expression))
-			{
-				CachedExpressionData.DynamicParameterNames.Empty(ExpressionDynamicParameter->ParamNames.Num());
-				for (const FString& Name : ExpressionDynamicParameter->ParamNames)
-				{
-					CachedExpressionData.DynamicParameterNames.Add(*Name);
-				}
-			}
-			else if (UMaterialExpressionLandscapeGrassOutput* ExpressionGrassOutput = Cast<UMaterialExpressionLandscapeGrassOutput>(Expression))
-			{
-				for (const auto& Type : ExpressionGrassOutput->GrassTypes)
-				{
-					CachedExpressionData.GrassTypes.AddUnique(Type.GrassType);
-				}
-			}
-			else if (Expression->IsA(UMaterialExpressionRuntimeVirtualTextureOutput::StaticClass()))
-			{
-				CachedExpressionData.bHasRuntimeVirtualTextureOutput = true;
-			}
-			else if (Expression->IsA(UMaterialExpressionSceneColor::StaticClass()))
-			{
-				CachedExpressionData.bHasSceneColor = true;
-			}
-			else if (Expression->IsA(UMaterialExpressionPerInstanceRandom::StaticClass()))
-			{
-				CachedExpressionData.bHasPerInstanceRandom = true;
-			}
-			else if (Expression->IsA(UMaterialExpressionPerInstanceCustomData::StaticClass()))
-			{
-				CachedExpressionData.bHasPerInstanceCustomData = true;
-			}
-			else if (Expression->IsA(UMaterialExpressionPerInstanceCustomData3Vector::StaticClass()))
-			{
-				CachedExpressionData.bHasPerInstanceCustomData = true;
-			}
-			else if (Expression->IsA(UMaterialExpressionVertexInterpolator::StaticClass()))
-			{
-				CachedExpressionData.bHasVertexInterpolator = true;
-			}
-			else if (UMaterialExpressionQualitySwitch* ExpressionQualitySwitch = Cast<UMaterialExpressionQualitySwitch>(Expression))
-			{
-				const FExpressionInput DefaultInput = ExpressionQualitySwitch->Default.GetTracedInput();
-				for (int32 InputIndex = 0; InputIndex < EMaterialQualityLevel::Num; InputIndex++)
-				{
-					if (ExpressionQualitySwitch->Inputs[InputIndex].IsConnected())
-					{
-						// We can ignore quality levels that are defined the same way as 'Default'
-						// This avoids compiling a separate explicit quality level resource, that will end up exactly the same as the default resource
-						const FExpressionInput Input = ExpressionQualitySwitch->Inputs[InputIndex].GetTracedInput();
-						if (Input.Expression != DefaultInput.Expression || Input.OutputIndex != DefaultInput.OutputIndex)
-						{
-							CachedExpressionData.QualityLevelsUsed[InputIndex] = true;
-						}
-					}
-				}
-			}
-			else if (UMaterialExpressionSetMaterialAttributes* SetMatAttributes = Cast<UMaterialExpressionSetMaterialAttributes>(Expression))
-			{
-				for (int32 PinIndex = 0; PinIndex < SetMatAttributes->AttributeSetTypes.Num(); ++PinIndex)
-				{
-					// For this material attribute pin do we have something connected?
-					const FGuid& Guid = SetMatAttributes->AttributeSetTypes[PinIndex];
-					const FExpressionInput& AttributeInput = SetMatAttributes->Inputs[PinIndex + 1];
-					const EMaterialProperty MaterialProperty = FMaterialAttributeDefinitionMap::GetProperty(Guid);
-					if (AttributeInput.Expression)
-					{
-						CachedExpressionData.SetPropertyConnected(MaterialProperty);
-					}
-				}
-			}
-			else if (UMaterialExpressionMakeMaterialAttributes* MakeMatAttributes = Cast<UMaterialExpressionMakeMaterialAttributes>(Expression))
-			{
-				auto SetMatAttributeConditionally = [&](EMaterialProperty InMaterialProperty, bool InIsConnected)
-				{
-					if (InIsConnected)
-					{
-						CachedExpressionData.SetPropertyConnected(InMaterialProperty);
-					}
-				};
-
-				SetMatAttributeConditionally(EMaterialProperty::MP_BaseColor, MakeMatAttributes->BaseColor.IsConnected());
-				SetMatAttributeConditionally(EMaterialProperty::MP_Metallic, MakeMatAttributes->Metallic.IsConnected());
-				SetMatAttributeConditionally(EMaterialProperty::MP_Specular, MakeMatAttributes->Specular.IsConnected());
-				SetMatAttributeConditionally(EMaterialProperty::MP_Roughness, MakeMatAttributes->Roughness.IsConnected());
-				SetMatAttributeConditionally(EMaterialProperty::MP_Anisotropy, MakeMatAttributes->Anisotropy.IsConnected());
-				SetMatAttributeConditionally(EMaterialProperty::MP_EmissiveColor, MakeMatAttributes->EmissiveColor.IsConnected());
-				SetMatAttributeConditionally(EMaterialProperty::MP_Opacity, MakeMatAttributes->Opacity.IsConnected());
-				SetMatAttributeConditionally(EMaterialProperty::MP_OpacityMask, MakeMatAttributes->OpacityMask.IsConnected());
-				SetMatAttributeConditionally(EMaterialProperty::MP_Normal, MakeMatAttributes->Normal.IsConnected());
-				SetMatAttributeConditionally(EMaterialProperty::MP_Tangent, MakeMatAttributes->Tangent.IsConnected());
-				SetMatAttributeConditionally(EMaterialProperty::MP_WorldPositionOffset, MakeMatAttributes->WorldPositionOffset.IsConnected());
-				SetMatAttributeConditionally(EMaterialProperty::MP_SubsurfaceColor, MakeMatAttributes->SubsurfaceColor.IsConnected());
-				SetMatAttributeConditionally(EMaterialProperty::MP_CustomData0, MakeMatAttributes->ClearCoat.IsConnected());
-				SetMatAttributeConditionally(EMaterialProperty::MP_CustomData1, MakeMatAttributes->ClearCoatRoughness.IsConnected());
-				SetMatAttributeConditionally(EMaterialProperty::MP_AmbientOcclusion, MakeMatAttributes->AmbientOcclusion.IsConnected());
-				SetMatAttributeConditionally(EMaterialProperty::MP_Refraction, MakeMatAttributes->Refraction.IsConnected());
-				SetMatAttributeConditionally(EMaterialProperty::MP_CustomizedUVs0, MakeMatAttributes->CustomizedUVs[0].IsConnected());
-				SetMatAttributeConditionally(EMaterialProperty::MP_CustomizedUVs1, MakeMatAttributes->CustomizedUVs[1].IsConnected());
-				SetMatAttributeConditionally(EMaterialProperty::MP_CustomizedUVs2, MakeMatAttributes->CustomizedUVs[2].IsConnected());
-				SetMatAttributeConditionally(EMaterialProperty::MP_CustomizedUVs3, MakeMatAttributes->CustomizedUVs[3].IsConnected());
-				SetMatAttributeConditionally(EMaterialProperty::MP_CustomizedUVs4, MakeMatAttributes->CustomizedUVs[4].IsConnected());
-				SetMatAttributeConditionally(EMaterialProperty::MP_CustomizedUVs5, MakeMatAttributes->CustomizedUVs[5].IsConnected());
-				SetMatAttributeConditionally(EMaterialProperty::MP_CustomizedUVs6, MakeMatAttributes->CustomizedUVs[6].IsConnected());
-				SetMatAttributeConditionally(EMaterialProperty::MP_CustomizedUVs7, MakeMatAttributes->CustomizedUVs[7].IsConnected());
-				SetMatAttributeConditionally(EMaterialProperty::MP_PixelDepthOffset, MakeMatAttributes->PixelDepthOffset.IsConnected());
-				SetMatAttributeConditionally(EMaterialProperty::MP_ShadingModel, MakeMatAttributes->ShadingModel.IsConnected());
-				SetMatAttributeConditionally(EMaterialProperty::MP_Displacement, MakeMatAttributes->Displacement.IsConnected());
-			}
-			else if (UMaterialExpressionMaterialAttributeLayers* LayersExpression = Cast<UMaterialExpressionMaterialAttributeLayers>(Expression))
-			{
-				const FMaterialLayersFunctions& LayerFunctions = LayersExpression->DefaultLayers;
-
-				for (int32 LayerIndex = 0; LayerIndex < LayerFunctions.Layers.Num(); ++LayerIndex)
-				{
-					VisitFunction(LayerFunctions.Layers[LayerIndex], LayerParameter, LayerIndex);
-				}
-
-				for (int32 BlendIndex = 0; BlendIndex < LayerFunctions.Blends.Num(); ++BlendIndex)
-				{
-					VisitFunction(LayerFunctions.Blends[BlendIndex], BlendParameter, BlendIndex);
-				}
-			}
-			else if (UMaterialExpressionMaterialFunctionCall* FunctionCall = Cast<UMaterialExpressionMaterialFunctionCall>(Expression))
-			{
-				VisitFunction(FunctionCall->MaterialFunction, Association, ParameterIndex);
-			}
-			else if (UMaterialExpressionStaticSwitchParameter* ExpressionStaticSwitchParameter = Cast<UMaterialExpressionStaticSwitchParameter>(Expression))
-			{
-				// Read the static switch parameter and only continue exploration along active branch.
-				FMaterialParameterMetadata ParameterMeta;
-				if (CachedExpressionData.GetParameterValue(EMaterialParameterType::StaticSwitch, FMemoryImageMaterialParameterInfo(Expression->GetParameterName()), ParameterMeta))
-				{
-					FExpressionInput* ExpressionInput = ParameterMeta.Value.AsStaticSwitch() ? &ExpressionStaticSwitchParameter->A : &ExpressionStaticSwitchParameter->B;
-					PushUnexploredExpression(ExpressionInput->Expression);
-					bVisitAllInputExpressions = false;
-				}
-			}
-			else if (UMaterialExpressionStaticSwitch* ExpressionStaticSwitch = Cast<UMaterialExpressionStaticSwitch>(Expression))
-			{
-				bool ParamValue = (bool)ExpressionStaticSwitch->DefaultValue;
-
-				// If no input expression specified, use the default value otherwise try to evaluate the input expression to the static bool.
-				if (!ExpressionStaticSwitch->Value.Expression || EvaluateStaticBoolExpression(ExpressionStaticSwitch->Value.Expression, ParamValue, Association, ParameterIndex))
-				{
-					// Push the active input expression.
-					FExpressionInput* ExpressionInput = ParamValue ? &ExpressionStaticSwitch->A : &ExpressionStaticSwitch->B;
-					PushUnexploredExpression(ExpressionInput->Expression);
-					bVisitAllInputExpressions = false;
-				}
-			}
-			else if (UMaterialExpressionNamedRerouteUsage* ExpressionRerouteUsage = Cast<UMaterialExpressionNamedRerouteUsage>(Expression))
-			{
-				if (ExpressionRerouteUsage->Declaration)
-				{
-					PushUnexploredExpression(ExpressionRerouteUsage->Declaration->Input.Expression);
-				}
-			}
-
-			// If otherwise specified, explore all input expressions to this node.
-			if (bVisitAllInputExpressions)
-			{
-				for (FExpressionInputIterator It{ Expression }; It; ++It)
-				{
-					PushUnexploredExpression(It->Expression);
-				}
-			}
-		}
-	}
-
-	/** Visits the inner expressions of a function call */
-	void VisitFunction(UMaterialFunctionInterface* Function, EMaterialParameterAssociation Association, int32 ParameterIndex)
-	{
-		if (!Function)
-		{
-			return;
-		}
-
-		// Get function outputs.
-		TArray<FFunctionExpressionInput> Inputs;
-		TArray<FFunctionExpressionOutput> Outputs;
-		Function->GetInputsAndOutputs(Inputs, Outputs);
-
-		// Push all function outputs as unexplored expressions to a new analyzer.
-		// Note: we use a different analyzer here so that multiple function calls to the same function are treated separately, since
-		// their actual state can be different due to static parameter overrides within a function.
-		FMaterialConnectedPropertiesAnalyzer FunctionAnalyzer{ CachedExpressionData, Function };
-		for (auto& Output : Outputs)
-		{
-			FunctionAnalyzer.PushUnexploredExpression(Output.ExpressionOutput);
-		}
-		
-		// Visit inner function expressions.
-		FunctionAnalyzer.VisitUnexploredExpressions(Association, ParameterIndex);
-	}
-
-	/** If expression is of type "static bool" this function walks the graph to find the actual value of the referenced static bool and writes the result to OutValue.
-	 * @return true if the boolean value could be determined.
-	 */
-	bool EvaluateStaticBoolExpression(UMaterialExpression* Expression, bool& bOutValue, EMaterialParameterAssociation Association, int32 ParameterIndex)
-	{
-		if (!Expression)
-		{
-			return false;
-		}
-		else if (UMaterialExpressionStaticBoolParameter* ExpressionStaticBoolParameter = Cast<UMaterialExpressionStaticBoolParameter>(Expression))
-		{
-			FMaterialParameterMetadata ParameterMeta;
-			if (CachedExpressionData.GetParameterValue(EMaterialParameterType::StaticSwitch, FMemoryImageMaterialParameterInfo(Expression->GetParameterName()), ParameterMeta))
-			{
-				bOutValue = ParameterMeta.Value.AsStaticSwitch();
-				return true;
-			}
-			return false;
-		}
-		else if (UMaterialExpressionStaticSwitchParameter* ExpressionStaticSwitchParameter = Cast<UMaterialExpressionStaticSwitchParameter>(Expression))
-		{
-			FMaterialParameterMetadata ParameterMeta;
-			if (CachedExpressionData.GetParameterValue(EMaterialParameterType::StaticSwitch, FMemoryImageMaterialParameterInfo(Expression->GetParameterName()), ParameterMeta))
-			{
-				FExpressionInput* OperandExpressionInput = ParameterMeta.Value.AsStaticSwitch() ? &ExpressionStaticSwitchParameter->A : &ExpressionStaticSwitchParameter->B;
-				return EvaluateStaticBoolExpression(OperandExpressionInput->Expression, bOutValue, Association, ParameterIndex);
-			}
-			return false;
-		}
-		else if (UMaterialExpressionStaticSwitch* ExpressionStaticSwitch = Cast<UMaterialExpressionStaticSwitch>(Expression))
-		{
-			bool Value = (bool)ExpressionStaticSwitch->DefaultValue;
-			// If node has a bool expression connected, try to evaluate that subgraph to a static boolean value.
-			if (ExpressionStaticSwitch->Value.Expression && !EvaluateStaticBoolExpression(ExpressionStaticSwitch->Value.Expression, Value, Association, ParameterIndex))
-			{
-				// Node has the bool expression slot connected, but we could not evaluate it to a static boolean value.
-				return false;
-			}
-			FExpressionInput* OperandExpressionInput = Value ? &ExpressionStaticSwitch->A : &ExpressionStaticSwitch->B;
-			return EvaluateStaticBoolExpression(OperandExpressionInput->Expression, bOutValue, Association, ParameterIndex);
-		}
-		else if (UMaterialExpressionNamedRerouteUsage* ExpressionRerouteUsage = Cast<UMaterialExpressionNamedRerouteUsage>(Expression))
-		{
-			if (!ExpressionRerouteUsage->Declaration)
-			{
-				return false;
-			}
-			return EvaluateStaticBoolExpression(ExpressionRerouteUsage->Declaration->Input.Expression, bOutValue, Association, ParameterIndex);
-		}
-		else if (UMaterialExpressionFunctionInput* ExpressionFunctionInput = Cast<UMaterialExpressionFunctionInput>(Expression))
-		{
-			return EvaluateStaticBoolExpression(ExpressionFunctionInput->Preview.Expression, bOutValue, Association, ParameterIndex);
-		}
-		else if (UMaterialExpressionStaticBool* ExpressionStaticBool = Cast<UMaterialExpressionStaticBool>(Expression))
-		{
-			bOutValue = ExpressionStaticBool->Value;
-			return true;
-		}
-
-		return false;
-	}
-};
-
 void FMaterialCachedExpressionData::AnalyzeMaterial(UMaterial& Material)
 {
-	if (!GExperimentalMaterialCachedDataAnalysisEnabled)
-	{
-		if (!Material.bUseMaterialAttributes)
-		{
-			for (int32 PropertyIndex = 0; PropertyIndex < MP_MAX; ++PropertyIndex)
-			{
-				const EMaterialProperty Property = (EMaterialProperty)PropertyIndex;
-				const FExpressionInput* Input = Material.GetExpressionInputForProperty(Property);
-				if (Input && Input->IsConnected())
-				{
-					SetPropertyConnected(Property);
-				}
-			}
-		}
-
-		FMaterialCachedExpressionContext Context;
-		UpdateForExpressions(Context, Material.GetExpressions(), EMaterialParameterAssociation::GlobalParameter, -1);
-	}
-	else
+	if (!Material.bUseMaterialAttributes)
 	{
-		// Now crawl the graph from outputs to inputs skipping unused subgraphs ("compiled out" by static switches).
-		FMaterialConnectedPropertiesAnalyzer Analyzer{ *this };
-
-		// Add material outputs depending on whether material uses material attributes.
-		if (Material.bUseMaterialAttributes)
+		for (int32 PropertyIndex = 0; PropertyIndex < MP_MAX; ++PropertyIndex)
 		{
-			const FExpressionInput* Input = Material.GetExpressionInputForProperty(MP_MaterialAttributes);
+			const EMaterialProperty Property = (EMaterialProperty)PropertyIndex;
+			const FExpressionInput* Input = Material.GetExpressionInputForProperty(Property);
 			if (Input && Input->IsConnected())
 			{
-				Analyzer.PushUnexploredExpression(Input->Expression);
-			}
-		}
-		else
-		{
-			for (int32 PropertyIndex = 0; PropertyIndex < MP_MAX; ++PropertyIndex)
-			{
-				const EMaterialProperty Property = (EMaterialProperty)PropertyIndex;
-				const FExpressionInput* Input = Material.GetExpressionInputForProperty(Property);
-				if (Input && Input->IsConnected())
-				{
-					Analyzer.PushUnexploredExpression(Input->Expression);
-					SetPropertyConnected(Property);
-				}
-			}
-		}
-	
-		// If there are any connected function output expressions, mark them as unexplored expressions.
-		// This occurs when the user opens a MaterialFunction in the editor.
-		for (UMaterialExpression* Expression : Material.GetExpressions())
-		{
-			if (Expression && Expression->IsA<UMaterialExpressionFunctionOutput>())
-			{
-				Analyzer.PushUnexploredExpression(Expression);
+				SetPropertyConnected(Property);
 			}
 		}
-
-		Analyzer.VisitUnexploredExpressions(GlobalParameter, -1);
 	}
+		
+	FMaterialCachedExpressionContext Context;
+	UpdateForExpressions(Context, Material.GetExpressions(), EMaterialParameterAssociation::GlobalParameter, -1);
 }
 
 
