PR #11664: Adding a "supports UAV" option for render target assets

diff --git a/Engine/Source/Runtime/Engine/Classes/Engine/TextureRenderTarget2D.h b/Engine/Source/Runtime/Engine/Classes/Engine/TextureRenderTarget2D.h
--- a/Engine/Source/Runtime/Engine/Classes/Engine/TextureRenderTarget2D.h
+++ b/Engine/Source/Runtime/Engine/Classes/Engine/TextureRenderTarget2D.h
@@ -148,6 +148,10 @@ class UTextureRenderTarget2D : public UTextureRenderTarget
 	UPROPERTY(EditAnywhere, BlueprintReadOnly, Category=TextureRenderTarget2D, AssetRegistrySearchable)
 	TEnumAsByte<ETextureRenderTargetFormat> RenderTargetFormat;
 
+	/** Whether this render target can be used as an unordered access view */
+	UPROPERTY(EditAnywhere, BlueprintReadOnly, Category = TextureRenderTarget2D, AssetRegistrySearchable)
+	uint8 bSupportsUAV : 1;
+
 	/** Whether to support Mip maps for this render target texture */
 	UPROPERTY(EditAnywhere, BlueprintReadOnly, Category = TextureRenderTarget2D, AssetRegistrySearchable)
 	uint8 bAutoGenerateMips : 1;
diff --git a/Engine/Source/Runtime/Engine/Classes/Engine/TextureRenderTarget2DArray.h b/Engine/Source/Runtime/Engine/Classes/Engine/TextureRenderTarget2DArray.h
--- a/Engine/Source/Runtime/Engine/Classes/Engine/TextureRenderTarget2DArray.h
+++ b/Engine/Source/Runtime/Engine/Classes/Engine/TextureRenderTarget2DArray.h
@@ -47,6 +47,10 @@ class UTextureRenderTarget2DArray : public UTextureRenderTarget
 	UPROPERTY(EditAnywhere, BlueprintReadOnly, Category=TextureRenderTarget2DArray, AssetRegistrySearchable)
 	uint8 bHDR:1;
 
+	/** Whether this render target can be used as an unordered access view */
+	UPROPERTY(EditAnywhere, BlueprintReadOnly, Category=TextureRenderTarget2DArray, AssetRegistrySearchable)
+	uint8 bSupportsUAV : 1;
+
 	/** True to force linear gamma space for this render target */
 	UPROPERTY()
 	uint8 bForceLinearGamma:1;
diff --git a/Engine/Source/Runtime/Engine/Classes/Engine/TextureRenderTargetCube.h b/Engine/Source/Runtime/Engine/Classes/Engine/TextureRenderTargetCube.h
--- a/Engine/Source/Runtime/Engine/Classes/Engine/TextureRenderTargetCube.h
+++ b/Engine/Source/Runtime/Engine/Classes/Engine/TextureRenderTargetCube.h
@@ -40,6 +40,10 @@ class UTextureRenderTargetCube : public UTextureRenderTarget
 	UPROPERTY(EditAnywhere, BlueprintReadOnly, Category=TextureRenderTargetCube, AssetRegistrySearchable)
 	uint8 bHDR:1;
 
+	/** Whether this render target can be used as an unordered access view */
+	UPROPERTY(EditAnywhere, BlueprintReadOnly, Category=TextureRenderTargetCube, AssetRegistrySearchable)
+	uint8 bSupportsUAV : 1;
+
 	/** True to force linear gamma space for this render target */
 	UPROPERTY()
 	uint8 bForceLinearGamma:1;
diff --git a/Engine/Source/Runtime/Engine/Classes/Engine/TextureRenderTargetVolume.h b/Engine/Source/Runtime/Engine/Classes/Engine/TextureRenderTargetVolume.h
--- a/Engine/Source/Runtime/Engine/Classes/Engine/TextureRenderTargetVolume.h
+++ b/Engine/Source/Runtime/Engine/Classes/Engine/TextureRenderTargetVolume.h
@@ -48,6 +48,10 @@ class UTextureRenderTargetVolume : public UTextureRenderTarget
 	UPROPERTY(EditAnywhere, BlueprintReadOnly, Category=TextureRenderTargetVolume, AssetRegistrySearchable)
 	uint8 bHDR:1;
 
+	/** Whether this render target can be used as an unordered access view */
+	UPROPERTY(EditAnywhere, BlueprintReadOnly, Category=TextureRenderTargetVolume, AssetRegistrySearchable)
+	uint8 bSupportsUAV : 1;
+
 	/** True to force linear gamma space for this render target */
 	UPROPERTY()
 	uint8 bForceLinearGamma:1;
diff --git a/Engine/Source/Runtime/Engine/Private/TextureRenderTarget2D.cpp b/Engine/Source/Runtime/Engine/Private/TextureRenderTarget2D.cpp
--- a/Engine/Source/Runtime/Engine/Private/TextureRenderTarget2D.cpp
+++ b/Engine/Source/Runtime/Engine/Private/TextureRenderTarget2D.cpp
@@ -39,6 +39,7 @@ UTextureRenderTarget2D::UTextureRenderTarget2D(const FObjectInitializer& ObjectI
 	SizeY = 1;
 	bHDR_DEPRECATED = true;
 	RenderTargetFormat = RTF_RGBA16f;
+	bSupportsUAV = false;
 	bAutoGenerateMips = false;
 	NumMips = 0;
 	ClearColor = FLinearColor(0.0f, 0.0f, 0.0f, 1.0f);
@@ -86,7 +87,10 @@ bool UTextureRenderTarget2D::IsSRGB() const
 
 FTextureResource* UTextureRenderTarget2D::CreateResource()
 {
-	UWorld* World = GetWorld();
+	if (bSupportsUAV)
+	{
+		bCanCreateUAV = 1;
+	}
 
 	if (bAutoGenerateMips)
 	{
diff --git a/Engine/Source/Runtime/Engine/Private/TextureRenderTarget2DArray.cpp b/Engine/Source/Runtime/Engine/Private/TextureRenderTarget2DArray.cpp
--- a/Engine/Source/Runtime/Engine/Private/TextureRenderTarget2DArray.cpp
+++ b/Engine/Source/Runtime/Engine/Private/TextureRenderTarget2DArray.cpp
@@ -29,6 +29,7 @@ UTextureRenderTarget2DArray::UTextureRenderTarget2DArray(const FObjectInitialize
 	bHDR = true;
 	ClearColor = FLinearColor(0.0f, 1.0f, 0.0f, 1.0f);
 	OverrideFormat = PF_Unknown;
+	bSupportsUAV = false;
 	bForceLinearGamma = true;
 }
 
@@ -120,6 +121,11 @@ void UTextureRenderTarget2DArray::GetResourceSizeEx(FResourceSizeEx& CumulativeR
 
 FTextureResource* UTextureRenderTarget2DArray::CreateResource()
 {
+	if (bSupportsUAV)
+	{
+		bCanCreateUAV = 1;
+	}
+
 	return new FTextureRenderTarget2DArrayResource(this);
 }
 
diff --git a/Engine/Source/Runtime/Engine/Private/TextureRenderTargetCube.cpp b/Engine/Source/Runtime/Engine/Private/TextureRenderTargetCube.cpp
--- a/Engine/Source/Runtime/Engine/Private/TextureRenderTargetCube.cpp
+++ b/Engine/Source/Runtime/Engine/Private/TextureRenderTargetCube.cpp
@@ -35,6 +35,7 @@ UTextureRenderTargetCube::UTextureRenderTargetCube(const FObjectInitializer& Obj
 	bHDR = true;
 	ClearColor = FLinearColor(0.0f, 1.0f, 0.0f, 1.0f);
 	OverrideFormat = PF_Unknown;
+	bSupportsUAV = false;
 	bForceLinearGamma = true;
 	bAutoGenerateMips = false;
 	MipsSamplerFilter = Filter;
@@ -132,6 +133,11 @@ void UTextureRenderTargetCube::GetResourceSizeEx(FResourceSizeEx& CumulativeReso
 
 FTextureResource* UTextureRenderTargetCube::CreateResource()
 {
+	if (bSupportsUAV)
+	{
+		bCanCreateUAV = 1;
+	}
+
 	if (bAutoGenerateMips)
 	{
 		NumMips = FMath::FloorLog2(SizeX) + 1;
diff --git a/Engine/Source/Runtime/Engine/Private/TextureRenderTargetVolume.cpp b/Engine/Source/Runtime/Engine/Private/TextureRenderTargetVolume.cpp
--- a/Engine/Source/Runtime/Engine/Private/TextureRenderTargetVolume.cpp
+++ b/Engine/Source/Runtime/Engine/Private/TextureRenderTargetVolume.cpp
@@ -28,6 +28,7 @@ UTextureRenderTargetVolume::UTextureRenderTargetVolume(const FObjectInitializer&
 	bHDR = true;
 	ClearColor = FLinearColor(0.0f, 1.0f, 0.0f, 1.0f);
 	OverrideFormat = PF_Unknown;
+	bSupportsUAV = false;
 	bForceLinearGamma = true;
 }
 
@@ -123,6 +124,11 @@ void UTextureRenderTargetVolume::GetResourceSizeEx(FResourceSizeEx& CumulativeRe
 
 FTextureResource* UTextureRenderTargetVolume::CreateResource()
 {
+	if (bSupportsUAV)
+	{
+		bCanCreateUAV = 1;
+	}
+
 	return new FTextureRenderTargetVolumeResource(this);
 }
 
