Revised previously altered behaviour of MaterialMayModifyMeshPosition to capture additional instances of when a material may output a vertex modification to avoid skipping Depth and Velocity shader compilations for materials (which meant the depth buffer did not apply WPO whilst the base pass did, leading to a Depth/Color mismatch).

diff --git a/Engine/Shaders/Public/ShaderVersion.ush b/Engine/Shaders/Public/ShaderVersion.ush
--- a/Engine/Shaders/Public/ShaderVersion.ush
+++ b/Engine/Shaders/Public/ShaderVersion.ush
@@ -3,4 +3,4 @@
 // Each time the console command is executed it generates a new GUID. As this file is included
 // in Platform.ush (which should be included in any shader) it allows to invalidate the shader DDC.
 // 
-#pragma message("UESHADERMETADATA_VERSION 60ae06fd-b60f-4424-80dd-25b05e4faaa9")
+#pragma message("UESHADERMETADATA_VERSION 1367B0EB-44F3-3974-1744-30ADFAA1DE69")
\ No newline at end of file
diff --git a/Engine/Source/Runtime/Engine/Private/Materials/MaterialCachedData.cpp b/Engine/Source/Runtime/Engine/Private/Materials/MaterialCachedData.cpp
--- a/Engine/Source/Runtime/Engine/Private/Materials/MaterialCachedData.cpp
+++ b/Engine/Source/Runtime/Engine/Private/Materials/MaterialCachedData.cpp
@@ -49,6 +49,7 @@ const FMaterialCachedExpressionData FMaterialCachedExpressionData::EmptyData{};
 const FMaterialCachedExpressionEditorOnlyData FMaterialCachedExpressionEditorOnlyData::EmptyData{};
 
 static_assert((uint64)(EMaterialProperty::MP_MaterialAttributes)-1 < (8 * sizeof(FMaterialCachedExpressionData::PropertyConnectedMask)), "PropertyConnectedMask cannot contain entire EMaterialProperty enumeration.");
+static_assert((uint64)(EMaterialProperty::MP_MaterialAttributes)-1 < (8 * sizeof(FMaterialCachedExpressionData::PropertyDefaultAltered)), "PropertyDefaultAltered cannot contain entire EMaterialProperty enumeration.");
 
 FMaterialCachedExpressionData::FMaterialCachedExpressionData()
 	: FunctionInfosStateCRC(0xffffffff)
@@ -560,10 +561,17 @@ void FMaterialCachedExpressionData::AnalyzeMaterial(UMaterial& Material)
 		for (int32 PropertyIndex = 0; PropertyIndex < MP_MAX; ++PropertyIndex)
 		{
 			const EMaterialProperty Property = (EMaterialProperty)PropertyIndex;
-			const FExpressionInput* Input = Material.GetExpressionInputForProperty(Property);
-			if (Input && Input->IsConnected())
+			FMaterialInputDescription Description;
+			if (Material.GetExpressionInputDescription(Property, Description))
 			{
-				SetPropertyConnected(Property);
+				if (Description.Input && Description.Input->IsConnected())
+				{
+					SetPropertyConnected(Property);
+				}
+				else if (Description.ConstantValue.AsLinearColor() != FLinearColor(FMaterialAttributeDefinitionMap::GetDefaultValue(Property)))
+				{
+					SetPropertyDefaultAltered(Property);
+				}
 			}
 		}
 	}
diff --git a/Engine/Source/Runtime/Engine/Private/Materials/MaterialShared.cpp b/Engine/Source/Runtime/Engine/Private/Materials/MaterialShared.cpp
--- a/Engine/Source/Runtime/Engine/Private/Materials/MaterialShared.cpp
+++ b/Engine/Source/Runtime/Engine/Private/Materials/MaterialShared.cpp
@@ -1193,7 +1193,15 @@ bool FMaterial::MaterialMayModifyMeshPosition() const
 {
 	// Conservative estimate when called before material translation has occurred. 
 	// This function is only intended for use in deciding whether or not shader permutations are required.
-	return HasVertexPositionOffsetConnected() || HasPixelDepthOffsetConnected() || HasDisplacementConnected();
+	const FMaterialCachedExpressionData& CachedData = GetCachedExpressionData();
+		// First check if the property is connected
+	return CachedData.IsPropertyConnected(MP_WorldPositionOffset)
+		|| CachedData.IsPropertyConnected(MP_PixelDepthOffset)
+		|| CachedData.IsPropertyConnected(MP_Displacement)
+		// If not, check that the default values have changed (ensures Depth pass shaders compile for WPO etc.)
+		|| CachedData.IsPropertyDefaultAltered(MP_WorldPositionOffset)
+		|| CachedData.IsPropertyDefaultAltered(MP_PixelDepthOffset)
+		|| CachedData.IsPropertyDefaultAltered(MP_Displacement);
 }
 
 bool FMaterial::MaterialUsesPixelDepthOffset_GameThread() const
diff --git a/Engine/Source/Runtime/Engine/Public/MaterialCachedData.h b/Engine/Source/Runtime/Engine/Public/MaterialCachedData.h
--- a/Engine/Source/Runtime/Engine/Public/MaterialCachedData.h
+++ b/Engine/Source/Runtime/Engine/Public/MaterialCachedData.h
@@ -268,6 +268,16 @@ struct FMaterialCachedExpressionData
 		PropertyConnectedMask |= (1ull << (uint64)Property);
 	}
 
+	bool IsPropertyDefaultAltered(EMaterialProperty Property) const
+	{
+		return ((PropertyDefaultAltered >> (uint64)Property) & 0x1) != 0;
+	}
+
+	void SetPropertyDefaultAltered(EMaterialProperty Property)
+	{
+		PropertyDefaultAltered |= (1ull << (uint64)Property);
+	}
+
 	bool Serialize(FArchive& Ar);
 	void PostSerialize(const FArchive& Ar);
 
@@ -375,6 +385,10 @@ struct FMaterialCachedExpressionData
 	/** Each bit corresponds to EMaterialProperty connection status. */
 	UPROPERTY()
 	uint64 PropertyConnectedMask = 0;
+	
+	/** Each bit corresponds to EMaterialProperty default value status. */
+	UPROPERTY()
+	uint64 PropertyDefaultAltered = 0;
 
 #if WITH_EDITOR
 	/** Array of errors reporting a parameter being set multiple times to distinct values. */
