[Material Editor] Moved code that cleans up null expressions in Material Functions from ::Serialize() to ::PreSave().

diff --git a/Engine/Source/Runtime/Engine/Private/Materials/MaterialExpressions.cpp b/Engine/Source/Runtime/Engine/Private/Materials/MaterialExpressions.cpp
--- a/Engine/Source/Runtime/Engine/Private/Materials/MaterialExpressions.cpp
+++ b/Engine/Source/Runtime/Engine/Private/Materials/MaterialExpressions.cpp
@@ -15872,6 +15872,25 @@ UMaterialFunctionMaterialLayerBlendInstance::UMaterialFunctionMaterialLayerBlend
 {
 }
 
+///////////////////////////////////////////////////////////////////////////////
+// UMaterialFunctionEditorOnlyData
+///////////////////////////////////////////////////////////////////////////////
+void UMaterialFunctionEditorOnlyData::PreSave(FObjectPreSaveContext ObjectSaveContext)
+{
+#if WITH_EDITORONLY_DATA
+	// If the collection of expressions got some null expressions remove them now, but warn the user about it.
+	if (ExpressionCollection.Expressions.Remove(nullptr))
+	{
+		UE_LOG(LogMaterial, Warning, TEXT(
+			"Material Function %s editor only data contained null expression and some expressions may be missing. "
+			"Please close and reopen this Material Function and verify it is still valid."),
+			*GetFullName());
+	}
+#endif
+
+	UMaterialFunctionInterfaceEditorOnlyData::PreSave(ObjectSaveContext);
+}
+
 ///////////////////////////////////////////////////////////////////////////////
 // UMaterialFunction
 ///////////////////////////////////////////////////////////////////////////////
@@ -15990,23 +16009,23 @@ void UMaterialFunction::ForceRecompileForRendering(FMaterialUpdateContext& Updat
 
 #endif // WITH_EDITOR
 
-void UMaterialFunction::Serialize(FArchive& Ar)
+void UMaterialFunction::PreSave(FObjectPreSaveContext ObjectSaveContext)
 {
-#if WITH_EDITORONLY_DATA
-	UMaterialFunctionEditorOnlyData* EditorOnly = GetEditorOnlyData();
-	if (EditorOnly && Ar.IsSaving() && !Ar.IsCooking())
+	#if WITH_EDITORONLY_DATA
+	if (DependentFunctionExpressionCandidates.Remove(nullptr))
 	{
-		// If the collection of expressions got some null expressions remove them now, but warn the user about it.
-		if (EditorOnly->ExpressionCollection.Expressions.Remove(nullptr))
-		{
-			FText Message = FText::Format(NSLOCTEXT("MaterialExpressions", "Error_NullExpressionsInMaterialFunction",
-				"Material Function {0} editor only data contained null expression and some expressions may be missing."
-				"\n\nPlease close and repoen this Material Function and verify it is still valid."), FText::FromString(GetFullName()));
-			FMessageDialog::Open(EAppMsgType::Ok, Message);
-		}
+		UE_LOG(LogMaterial, Warning, TEXT(
+			"Material Function %s contained some null dependent function expression calls. "
+			"Please close and reopen this Material Function and verify it is still valid."),
+			   *GetFullName());
 	}
-#endif
+	#endif
+
+	UMaterialFunctionInterface::PreSave(ObjectSaveContext);
+}
 
+void UMaterialFunction::Serialize(FArchive& Ar)
+{
 	Super::Serialize(Ar);
 
 #if WITH_EDITOR
diff --git a/Engine/Source/Runtime/Engine/Public/Materials/MaterialFunction.h b/Engine/Source/Runtime/Engine/Public/Materials/MaterialFunction.h
--- a/Engine/Source/Runtime/Engine/Public/Materials/MaterialFunction.h
+++ b/Engine/Source/Runtime/Engine/Public/Materials/MaterialFunction.h
@@ -21,9 +21,14 @@ UCLASS(MinimalAPI, Optional)
 class UMaterialFunctionEditorOnlyData : public UMaterialFunctionInterfaceEditorOnlyData
 {
 	GENERATED_BODY()
+
 public:
 	UPROPERTY()
 	FMaterialExpressionCollection ExpressionCollection;
+
+	//~ Begin UObject Interface.
+	ENGINE_API virtual void PreSave(FObjectPreSaveContext ObjectSaveContext) override;
+	//~ End UObject Interface.
 };
 
 /**
@@ -115,6 +120,7 @@ public:
 #if WITH_EDITOR
 	virtual void PostEditChangeProperty(FPropertyChangedEvent& PropertyChangedEvent) override;
 #endif // WITH_EDITOR
+	virtual void PreSave(FObjectPreSaveContext ObjectSaveContext) override;
 	virtual void Serialize(FArchive& Ar) override;
 	virtual void PostLoad() override;
 #if WITH_EDITORONLY_DATA
