Timeline change to fix bIsPlaying not being replicated if you called it from the server after it immediately finished playing. Instead of just a single bit being used to track if we are playing, we will use a whole uint8. Use the first bit as a bool flag, and modify the upper 7 bits any time we call play so that we ensure the value is dirtied and thus actually replicated to clients.

diff --git a/Engine/Source/Runtime/Core/Public/UObject/FortniteMainBranchObjectVersion.h b/Engine/Source/Runtime/Core/Public/UObject/FortniteMainBranchObjectVersion.h
--- a/Engine/Source/Runtime/Core/Public/UObject/FortniteMainBranchObjectVersion.h
+++ b/Engine/Source/Runtime/Core/Public/UObject/FortniteMainBranchObjectVersion.h
@@ -506,6 +506,10 @@ struct FFortniteMainBranchObjectVersion
 
 		// PCG changes around the ApplyOnActor node, where we collapsed the TargetActor to the input pin.
 		PCGApplyOnActorNodeMoveTargetActorEdgeToInput,
+
+		// Deprecation of the bPlaying flag on FTimeline struct types in favor of a better
+		// PlayingStateTracker type to improve replication reliability 
+		TimelinePlayingStateTrackerDeprecation,
 		
 		// -----<new versions can be added above this line>-------------------------------------------------
 		VersionPlusOne,
diff --git a/Engine/Source/Runtime/Engine/Classes/Components/TimelineComponent.h b/Engine/Source/Runtime/Engine/Classes/Components/TimelineComponent.h
--- a/Engine/Source/Runtime/Engine/Classes/Components/TimelineComponent.h
+++ b/Engine/Source/Runtime/Engine/Classes/Components/TimelineComponent.h
@@ -185,8 +185,8 @@ struct FTimelineLinearColorTrack
 USTRUCT()
 struct FTimeline
 {
-	GENERATED_USTRUCT_BODY()
-		
+	GENERATED_BODY()
+	
 private:
 	/** Specified how the timeline determines its own length (e.g. specified length, last keyframe) */
 	UPROPERTY(NotReplicated)
@@ -199,10 +199,23 @@ private:
 	/** If playback should move the current position backwards instead of forwards */
 	UPROPERTY()
 	uint8 bReversePlayback:1;
-
-	/** Are we currently playing (moving Position) */
+	
+	UE_DEPRECATED(5.5, "Use PlayingStateTracker instead")
+	UPROPERTY(Meta = (DeprecatedProperty, DeprecationMessage = "This property is no longer supported. Use PlayingStateTracker instead"))
+	uint8 bPlaying_DEPRECATED:1;
+	
+	/**
+	* If the first bit is set to 1 (PlayingStateTracker & 0x01 == 1), then we are playing
+	* 
+	* The rest of the bits in the uint8 are reserved for keeping track of the "dirty" state,
+	* being incremented when our state is modified. This ensures that the value is replicated
+	* if it changes multiple times in one frame, such as calling "Play From Start" in the resulting
+	* "Finished" delegate.
+	*
+	* You should modify this value using the "ChangeMarkPlayingState" function.
+	*/
 	UPROPERTY()
-	uint8 bPlaying:1;
+	uint8 PlayingStateTracker;
 
 	/** How long the timeline is, will stop or loop at the end */
  	UPROPERTY(NotReplicated)
@@ -255,11 +268,15 @@ private:
 	FProperty* DirectionProperty;
 
 public:
+	
+PRAGMA_DISABLE_DEPRECATION_WARNINGS
+	
 	FTimeline()
 	: LengthMode( TL_LastKeyFrame )
 	, bLooping( false )
 	, bReversePlayback( false )
-	, bPlaying( false )
+	, bPlaying_DEPRECATED ( false )
+	, PlayingStateTracker( 2 )
 	, Length( 5.f )
 	, PlayRate( 1.f )
 	, Position( 0.0f )	
@@ -268,6 +285,15 @@ public:
 	{
 	}
 
+	// Note: We need to explicitly disable warnings on these constructors/operators for clang to be happy with deprecated variables
+	~FTimeline() = default;
+	FTimeline(const FTimeline&) = default;
+	FTimeline(FTimeline&&) = default;
+	FTimeline& operator=(const FTimeline&) = default;
+	FTimeline& operator=(FTimeline&&) = default;
+	
+PRAGMA_ENABLE_DEPRECATION_WARNINGS
+
 	/** Helper function to get to the timeline direction enum */
 	ENGINE_API static UEnum* GetTimelineDirectionEnum();
 
@@ -376,9 +402,30 @@ public:
 
 	/** Get all curves used by the Timeline */
 	void GetAllCurves(TSet<class UCurveBase*>& InOutCurves) const;
+
+	ENGINE_API bool Serialize(FArchive& Ar);
+	ENGINE_API void PostSerialize(const FArchive& Ar);
+	
 private:
 	/** Returns the time value of the last keyframe in any of the timeline's curves */
 	float GetLastKeyframeTime() const;
+
+	/**
+	 * Mark the playing state as the given value
+	 *
+	 * @param bIsPlaying	True if the timeline should be playing, false if it should be stopped.
+	 */
+	void ChangeMarkPlayingState(const bool bIsPlaying);
+};
+
+// Need to customize serialization to safely update the old bPlaying flag to the newer play state tracer
+template<> struct TStructOpsTypeTraits<FTimeline> : public TStructOpsTypeTraitsBase2<FTimeline>
+{
+	enum
+	{
+		WithSerialize = true,
+		WithPostSerialize = true,
+	};
 };
 
 /** 
@@ -392,6 +439,7 @@ class UTimelineComponent : public UActorComponent
 	GENERATED_UCLASS_BODY()
 
 private:
+	
 	/** The actual timeline structure */
 	UPROPERTY(ReplicatedUsing=OnRep_Timeline)
 	FTimeline	TheTimeline;
diff --git a/Engine/Source/Runtime/Engine/Private/Timeline.cpp b/Engine/Source/Runtime/Engine/Private/Timeline.cpp
--- a/Engine/Source/Runtime/Engine/Private/Timeline.cpp
+++ b/Engine/Source/Runtime/Engine/Private/Timeline.cpp
@@ -14,6 +14,7 @@
 #include "Components/TimelineComponent.h"
 #include "Engine/World.h"
 #include "Misc/App.h"
+#include "UObject/FortniteMainBranchObjectVersion.h"
 
 DEFINE_LOG_CATEGORY_STATIC(LogTimeline, Log, All);
 
@@ -37,7 +38,8 @@ UEnum* FTimeline::GetTimelineDirectionEnum()
 void FTimeline::Play()
 {
 	bReversePlayback = false;
-	bPlaying = true;
+
+	ChangeMarkPlayingState(/*bIsPlaying=*/ true);
 }
 
 void FTimeline::PlayFromStart()
@@ -49,7 +51,8 @@ void FTimeline::PlayFromStart()
 void FTimeline::Reverse()
 {
 	bReversePlayback = true;
-	bPlaying = true;
+	
+	ChangeMarkPlayingState(/*bIsPlaying=*/ true);
 }
 
 void FTimeline::ReverseFromEnd()
@@ -60,12 +63,13 @@ void FTimeline::ReverseFromEnd()
 
 void FTimeline::Stop()
 {
-	bPlaying = false;
+	ChangeMarkPlayingState(/*bIsPlaying=*/false);
 }
 
 bool FTimeline::IsPlaying() const
 {
-	return bPlaying;
+	// Consider us playing if the first bit is set to 1
+	return PlayingStateTracker & 0x01;
 }
 
 void FTimeline::AddEvent(float Time, FOnTimelineEvent Event)
@@ -496,7 +500,7 @@ bool FTimeline::IsLooping() const
 
 bool FTimeline::IsReversing() const
 {
-	return bPlaying && bReversePlayback;
+	return IsPlaying() && bReversePlayback;
 }
 
 void FTimeline::SetPlayRate(float NewRate)
@@ -513,7 +517,7 @@ void FTimeline::TickTimeline(float DeltaTime)
 {
 	bool bIsFinished = false;
 
-	if(bPlaying)
+	if (IsPlaying())
 	{
 		const float TimelineLength = GetTimelineLength();
 		float EffectiveDeltaTime = DeltaTime * (bReversePlayback ? (-PlayRate) : (PlayRate));
@@ -707,6 +711,27 @@ float FTimeline::GetLastKeyframeTime() const
 	return MaxTime;
 }
 
+void FTimeline::ChangeMarkPlayingState(const bool bIsPlaying)
+{
+	// We will use the upper bits of the PlayingStateTracker as "dirty" bits to ensure that the
+	// value gets replicated properly in case you call "Play" on the same frame as the timeline
+	// finishes.
+	if (bIsPlaying)
+	{
+		// Increment by 2 to "dirty" the number without touching the first bit (2 is 0010)
+        // wrapping it to 255 (max uint8 - 1)
+        PlayingStateTracker = FMath::WrapExclusive(PlayingStateTracker + 2, 2, 255);
+    
+        // Make sure that we always set the 1 bit because we are playing.
+        PlayingStateTracker |= 0x01;
+	}
+	else
+	{
+		// To flag the state as stopped, reset the first bit to zero
+		PlayingStateTracker &= ~0x01;
+	}
+}
+
 void FTimeline::GetAllCurves(TSet<class UCurveBase*>& InOutCurves) const
 {
 	for (auto& Track : InterpVectors)
@@ -723,6 +748,31 @@ void FTimeline::GetAllCurves(TSet<class UCurveBase*>& InOutCurves) const
 	}
 }
 
+bool FTimeline::Serialize(FArchive& Ar)
+{
+	// Need a custom version to handle the bPlaying deprecation
+	Ar.UsingCustomVersion(FFortniteMainBranchObjectVersion::GUID);
+
+	// Returning false will just run the normal serialization
+	return false;
+}
+
+void FTimeline::PostSerialize(const FArchive& Ar)
+{
+	// If we are loading code from before the serialization change...
+	const int32 CustomVersion = Ar.CustomVer(FFortniteMainBranchObjectVersion::GUID);
+	if (Ar.IsLoading() && CustomVersion < FFortniteMainBranchObjectVersion::TimelinePlayingStateTrackerDeprecation)
+	{
+PRAGMA_DISABLE_DEPRECATION_WARNINGS
+
+		// ...populate the new PlayingStateTracker with the value of the
+		// old bPlaying flag. 
+		ChangeMarkPlayingState(bPlaying_DEPRECATED);
+
+PRAGMA_ENABLE_DEPRECATION_WARNINGS
+	}
+}
+
 //////////////////////////////////////////////////////////////////////////
 // UTimelineComponent
 
