Fix missing "R" at the end of the type trait "WithSerializer" for FTimeline types.

#rnx

diff --git a/Engine/Source/Runtime/Engine/Classes/Components/TimelineComponent.h b/Engine/Source/Runtime/Engine/Classes/Components/TimelineComponent.h
--- a/Engine/Source/Runtime/Engine/Classes/Components/TimelineComponent.h
+++ b/Engine/Source/Runtime/Engine/Classes/Components/TimelineComponent.h
@@ -423,7 +423,7 @@ template<> struct TStructOpsTypeTraits<FTimeline> : public TStructOpsTypeTraitsB
 {
 	enum
 	{
-		WithSerialize = true,
+		WithSerializer = true,
 		WithPostSerialize = true,
 	};
 };
