Fix to set transient Skeletal Mesh Asset property when value is changed on an actor in level viewport

diff --git a/Engine/Source/Runtime/Engine/Private/Components/SkeletalMeshComponentInstanceData.cpp b/Engine/Source/Runtime/Engine/Private/Components/SkeletalMeshComponentInstanceData.cpp
--- a/Engine/Source/Runtime/Engine/Private/Components/SkeletalMeshComponentInstanceData.cpp
+++ b/Engine/Source/Runtime/Engine/Private/Components/SkeletalMeshComponentInstanceData.cpp
@@ -7,17 +7,18 @@
 
 FSkeletalMeshComponentInstanceData::FSkeletalMeshComponentInstanceData(const USkeletalMeshComponent* SourceComponent)
 	: FSceneComponentInstanceData(SourceComponent)
+	, SkeletalMeshAsset(nullptr)
 	, bUpdateAnimationInEditor(0)
 	, bUpdateClothInEditor(0)
-
 {
 #if WITH_EDITOR
-	// Only Blueprint components would reset tansient values when the construction script is re-run.
+	// Only Blueprint components would reset transient values when the construction script is re-run.
 	// Hence, we only need to apply instance cache for Blueprint created components.
 	const bool bIsBlueprintCreatedComponent = SourceComponent->CreationMethod == EComponentCreationMethod::SimpleConstructionScript
 		|| SourceComponent->CreationMethod == EComponentCreationMethod::UserConstructionScript;
 	if (bIsBlueprintCreatedComponent)
 	{
+		SkeletalMeshAsset = SourceComponent->GetSkeletalMeshAsset();
 		bUpdateAnimationInEditor = SourceComponent->GetUpdateAnimationInEditor();
 		bUpdateClothInEditor = SourceComponent->GetUpdateClothInEditor();
 	}
@@ -26,7 +27,7 @@ FSkeletalMeshComponentInstanceData::FSkeletalMeshComponentInstanceData(const USk
 
 bool FSkeletalMeshComponentInstanceData::ContainsData() const
 {
-	return bUpdateAnimationInEditor || bUpdateClothInEditor;
+	return bUpdateAnimationInEditor || bUpdateClothInEditor || SkeletalMeshAsset;
 }
 
 void FSkeletalMeshComponentInstanceData::ApplyToComponent(UActorComponent* Component, const ECacheApplyPhase CacheApplyPhase)
@@ -36,6 +37,7 @@ void FSkeletalMeshComponentInstanceData::ApplyToComponent(UActorComponent* Compo
 #if WITH_EDITOR
 	if (USkeletalMeshComponent * SkeletalMesh = Cast<USkeletalMeshComponent>(Component))
 	{
+		SkeletalMesh->SetSkeletalMeshAsset(SkeletalMeshAsset);
 		SkeletalMesh->SetUpdateAnimationInEditor(bUpdateAnimationInEditor);
 		SkeletalMesh->SetUpdateClothInEditor(bUpdateClothInEditor);
 	}
diff --git a/Engine/Source/Runtime/Engine/Private/Components/SkeletalMeshComponentInstanceData.h b/Engine/Source/Runtime/Engine/Private/Components/SkeletalMeshComponentInstanceData.h
--- a/Engine/Source/Runtime/Engine/Private/Components/SkeletalMeshComponentInstanceData.h
+++ b/Engine/Source/Runtime/Engine/Private/Components/SkeletalMeshComponentInstanceData.h
@@ -6,6 +6,7 @@
 #include "SkeletalMeshComponentInstanceData.generated.h"
 
 class USkeletalMeshComponent;
+class USkeletalMesh;
 
 /** Saves internal SkeletalMesh (transient) state that gets lost at reconstruction for Blueprint created components. */
 USTRUCT()
@@ -19,6 +20,7 @@ struct ENGINE_API FSkeletalMeshComponentInstanceData : public FSceneComponentIns
 	virtual bool ContainsData() const override;
 	virtual void ApplyToComponent(UActorComponent* Component, const ECacheApplyPhase CacheApplyPhase) override;
 
+	USkeletalMesh* SkeletalMeshAsset;
 	uint8 bUpdateAnimationInEditor : 1;
 	uint8 bUpdateClothInEditor : 1;
 };
