Level Streaming: Restored previous behavior for how often the GEngine->BeginStreamingPauseDelegate gets called. MoviePipeline: Fixed flickering User Interface in shipping builds.

Before CL 29229382 this delegate was only called when IsVisibilityRequestPending() || IsAsyncLoading() was true, which it wouldn't be if there was no work to do. The change made it so it was getting called every frame, even without work.

Calling this function every frame creates a performance hit when flushing level streaming (as it flushes the rendering pipeline in FStreamingPauseRenderingModule::BeginStreamingPause) which can generate a stall (~35ms). This is normally not an issue during gameplay as you rarely flush level streaming, but in Movie Render Queue it is flushed every frame to ensure frames are not rendered with non-fully loaded levels, leading to a performance hit in MRQ renders with shipping builds. Additionally, this causes the user interface of the game to flicker (during MRQ renders) repeatedly.

This restores previous behavior so it should be pretty safe but specifically tested regular levels without sub-levels, and a regular level with sub-levels that are enabled by the Sequencer Level Visibility track.

diff --git a/Engine/Source/Runtime/Engine/Private/World.cpp b/Engine/Source/Runtime/Engine/Private/World.cpp
--- a/Engine/Source/Runtime/Engine/Private/World.cpp
+++ b/Engine/Source/Runtime/Engine/Private/World.cpp
@@ -4238,12 +4238,12 @@ void UWorld::BlockTillLevelStreamingCompleted()
 		if (IsVisibilityRequestPending() || HasAsyncLevelRequests())
 		{
 			WorkToDo = 2;
-		}
-
-		if (!bIsStreamingPaused && GEngine->GameViewport && GEngine->BeginStreamingPauseDelegate && GEngine->BeginStreamingPauseDelegate->IsBound())
-		{
-			GEngine->BeginStreamingPauseDelegate->Execute(GEngine->GameViewport->Viewport);
-			bIsStreamingPaused = true;
+			// Only call the streaming pause delegates if there is actually work to do.
+			if (!bIsStreamingPaused && GEngine->GameViewport && GEngine->BeginStreamingPauseDelegate && GEngine->BeginStreamingPauseDelegate->IsBound())
+			{
+				GEngine->BeginStreamingPauseDelegate->Execute(GEngine->GameViewport->Viewport);
+				bIsStreamingPaused = true;
+			}
 		}
 
 		// Flush level streaming requests, blocking till completion.
