Fix for cook warning related to skeletal mesh component instance data introduced by 35191113

diff --git a/Engine/Source/Runtime/Engine/Private/Components/SkeletalMeshComponentInstanceData.cpp b/Engine/Source/Runtime/Engine/Private/Components/SkeletalMeshComponentInstanceData.cpp
--- a/Engine/Source/Runtime/Engine/Private/Components/SkeletalMeshComponentInstanceData.cpp
+++ b/Engine/Source/Runtime/Engine/Private/Components/SkeletalMeshComponentInstanceData.cpp
@@ -35,9 +35,16 @@ void FSkeletalMeshComponentInstanceData::ApplyToComponent(UActorComponent* Compo
 	Super::ApplyToComponent(Component, CacheApplyPhase);
 
 #if WITH_EDITOR
-	if (USkeletalMeshComponent * SkeletalMesh = Cast<USkeletalMeshComponent>(Component))
+	const bool bIsBlueprintCreatedComponent = Component->CreationMethod == EComponentCreationMethod::SimpleConstructionScript
+		|| Component->CreationMethod == EComponentCreationMethod::UserConstructionScript;
+	USkeletalMeshComponent* SkeletalMesh = Cast<USkeletalMeshComponent>(Component);
+
+	if (SkeletalMesh && bIsBlueprintCreatedComponent)
 	{
-		SkeletalMesh->SetSkeletalMeshAsset(SkeletalMeshAsset);
+		if (SkeletalMeshAsset)
+		{
+			SkeletalMesh->SetSkeletalMeshAsset(SkeletalMeshAsset);
+		}
 		SkeletalMesh->SetUpdateAnimationInEditor(bUpdateAnimationInEditor);
 		SkeletalMesh->SetUpdateClothInEditor(bUpdateClothInEditor);
 	}
