Make UAnimationAsset a blueprint type

diff --git a/Engine/Source/Runtime/Engine/Classes/Animation/AnimationAsset.h b/Engine/Source/Runtime/Engine/Classes/Animation/AnimationAsset.h
--- a/Engine/Source/Runtime/Engine/Classes/Animation/AnimationAsset.h
+++ b/Engine/Source/Runtime/Engine/Classes/Animation/AnimationAsset.h
@@ -989,7 +989,7 @@ struct FAnimationGroupReference
 	}
 };
 
-UCLASS(abstract, MinimalAPI)
+UCLASS(abstract, BlueprintType, MinimalAPI)
 class UAnimationAsset : public UObject, public IInterface_AssetUserData, public IInterface_PreviewMeshProvider
 {
 	GENERATED_BODY()
