OnDemandShaderCompilation: Fix spurious failure to connect to server message (and crash) when running -cookonthefly. The ODSCThread does not have its own connection in that case and should not treat lack of its own connection as a failure.
#rnx

diff --git a/Engine/Source/Runtime/Engine/Private/ODSC/ODSCThread.cpp b/Engine/Source/Runtime/Engine/Private/ODSC/ODSCThread.cpp
--- a/Engine/Source/Runtime/Engine/Private/ODSC/ODSCThread.cpp
+++ b/Engine/Source/Runtime/Engine/Private/ODSC/ODSCThread.cpp
@@ -464,7 +464,14 @@ void FODSCThread::Process()
 		}
 	}
 
-	bIsConnectedToODSCServer = CheckODSCConnection();
+	if (bHasDefaultConnection)
+	{
+		bIsConnectedToODSCServer = true;
+	}
+	else
+	{
+		bIsConnectedToODSCServer = CheckODSCConnection();
+	}
 
 	ON_SCOPE_EXIT
 	{
