Fix potential editor crash in triggering anim notifies dues to null skeletal mesh component
#rnx

diff --git a/Engine/Source/Runtime/Engine/Private/Animation/AnimInstance.cpp b/Engine/Source/Runtime/Engine/Private/Animation/AnimInstance.cpp
--- a/Engine/Source/Runtime/Engine/Private/Animation/AnimInstance.cpp
+++ b/Engine/Source/Runtime/Engine/Private/Animation/AnimInstance.cpp
@@ -1547,7 +1547,7 @@ void UAnimInstance::TriggerAnimNotifies(float DeltaSeconds)
 		{
 #if WITH_EDITOR
 			// Prevent firing notifies in animation editors if requested 
-			if(!SkelMeshComp->IsA<UDebugSkelMeshComponent>() || AnimNotifyEvent.NotifyStateClass->ShouldFireInEditor())
+			if((SkelMeshComp && !SkelMeshComp->IsA<UDebugSkelMeshComponent>()) || AnimNotifyEvent.NotifyStateClass->ShouldFireInEditor())
 #endif
 			{
 				TRACE_ANIM_NOTIFY(this, AnimNotifyEvent, End);
@@ -1573,7 +1573,7 @@ void UAnimInstance::TriggerAnimNotifies(float DeltaSeconds)
 		{
 #if WITH_EDITOR
 			// Prevent firing notifies in animation editors if requested 
-			if(!SkelMeshComp->IsA<UDebugSkelMeshComponent>() || AnimNotifyEvent->NotifyStateClass->ShouldFireInEditor())
+			if((SkelMeshComp && !SkelMeshComp->IsA<UDebugSkelMeshComponent>()) || AnimNotifyEvent->NotifyStateClass->ShouldFireInEditor())
 #endif
 			{
 				TRACE_ANIM_NOTIFY(this, *AnimNotifyEvent, Begin);
@@ -1594,7 +1594,7 @@ void UAnimInstance::TriggerAnimNotifies(float DeltaSeconds)
 		{
 #if WITH_EDITOR
 			// Prevent firing notifies in animation editors if requested 
-			if(!SkelMeshComp->IsA<UDebugSkelMeshComponent>() || AnimNotifyEvent.NotifyStateClass->ShouldFireInEditor())
+			if((SkelMeshComp && !SkelMeshComp->IsA<UDebugSkelMeshComponent>()) || AnimNotifyEvent.NotifyStateClass->ShouldFireInEditor())
 #endif
 			{
 				TRACE_ANIM_NOTIFY(this, AnimNotifyEvent, Tick);
