[Shaders] Ensure directory to dump SCW crash log exists when absolute log path is used.

- Fixes reading "-ShaderCompileWorkerCrashLogsDir" command line argument, which must include the '=' character.
- Create directory when FShaderCompilingManager is initialized if it doesn't already exit.
- Save directory path in variable to avoid parsing command line every time a worker launches.

#rnx

diff --git a/Engine/Source/Runtime/Engine/Private/ShaderCompiler/ShaderCompiler.cpp b/Engine/Source/Runtime/Engine/Private/ShaderCompiler/ShaderCompiler.cpp
--- a/Engine/Source/Runtime/Engine/Private/ShaderCompiler/ShaderCompiler.cpp
+++ b/Engine/Source/Runtime/Engine/Private/ShaderCompiler/ShaderCompiler.cpp
@@ -863,6 +863,30 @@ FShaderCompilingManager::FShaderCompilingManager() :
 
 	FAssetCompilingManager::Get().RegisterManager(this);
 
+	// Ensure directory for dumping worker crash log exits before launching workers
+	if (GDebugDumpWorkerCrashLog)
+	{
+		FString CustomCrashLogsDir;
+		if (FParse::Value(FCommandLine::Get(), TEXT("ShaderCompileWorkerCrashLogsDir="), CustomCrashLogsDir))
+		{
+			WorkerCrashLogBaseDirectory = MoveTemp(CustomCrashLogsDir);
+		}
+		else if (GIsBuildMachine)
+		{
+			WorkerCrashLogBaseDirectory = GetBuildMachineArtifactBasePath();
+		}
+
+		// If this is empty, fall back to relative paths and the default log directory
+		if (!WorkerCrashLogBaseDirectory.IsEmpty() && !IFileManager::Get().DirectoryExists(*WorkerCrashLogBaseDirectory))
+		{
+			if (!IFileManager::Get().MakeDirectory(*WorkerCrashLogBaseDirectory, true))
+			{
+				const uint32 ErrorCode = FPlatformMisc::GetLastError();
+				UE_LOG(LogShaderCompilers, Warning, TEXT("Failed to create directory for ShaderCompileWorker crash logs '%s' (Error Code: %u)"), *WorkerCrashLogBaseDirectory, ErrorCode);
+			}
+		}
+	}
+
 #if WITH_EDITOR
 	static const bool bAllowShaderRecompileOnSave = CVarRecompileShadersOnSave.GetValueOnAnyThread();
 	if (bAllowShaderRecompileOnSave)
@@ -1461,14 +1485,9 @@ FProcHandle FShaderCompilingManager::LaunchWorker(const FString& WorkingDirector
 		WorkerParameters += TEXT(" -LogToMemory -DumpLogOnExitCrashOnly ");
 
 		const FString WorkerLogFilename = FString::Printf(TEXT("ShaderCompileWorker-%d.log"), ThreadId);
-		FString CustomCrashLogsDir;
-		if (FParse::Value(FCommandLine::Get(), TEXT("ShaderCompileWorkerCrashLogsDir"), CustomCrashLogsDir))
-		{
-			WorkerParameters += FString::Printf(TEXT("-AbsLog=%s"), *FPaths::Combine(CustomCrashLogsDir, WorkerLogFilename));
-		}
-		else if (GIsBuildMachine)
+		if (!WorkerCrashLogBaseDirectory.IsEmpty())
 		{
-			WorkerParameters += FString::Printf(TEXT("-AbsLog=%s"), *FPaths::Combine(GetBuildMachineArtifactBasePath(), WorkerLogFilename));
+			WorkerParameters += FString::Printf(TEXT("-AbsLog=%s"), *FPaths::Combine(WorkerCrashLogBaseDirectory, WorkerLogFilename));
 		}
 		else
 		{
diff --git a/Engine/Source/Runtime/Engine/Public/ShaderCompiler.h b/Engine/Source/Runtime/Engine/Public/ShaderCompiler.h
--- a/Engine/Source/Runtime/Engine/Public/ShaderCompiler.h
+++ b/Engine/Source/Runtime/Engine/Public/ShaderCompiler.h
@@ -805,6 +805,8 @@ private:
 	FString ShaderCompileWorkerName;
 	/** Last value of GetNumRemainingAssets */
 	int32 LastNumRemainingAssets = 0;
+	/** If dumping crash logs for workers is enabled and an absolute path is used (i.e. -AbsLog), this contains the base directory path. */
+	FString WorkerCrashLogBaseDirectory;
 
 	/** 
 	 * Tracks the total time that shader compile workers have been busy since startup.  
