Fix for a crash when rebuilding textures where mip streaming would try to access invalid mip data during force rebuilds.

diff --git a/Engine/Source/Runtime/Engine/Private/TextureDerivedData.cpp b/Engine/Source/Runtime/Engine/Private/TextureDerivedData.cpp
--- a/Engine/Source/Runtime/Engine/Private/TextureDerivedData.cpp
+++ b/Engine/Source/Runtime/Engine/Private/TextureDerivedData.cpp
@@ -4051,6 +4051,9 @@ void UTexture::ForceRebuildPlatformData(uint8 InEncodeSpeedOverride /* =255 ETex
 		// Cache() will clear FTexturePlatformData::Mips which can be accessed by the streaming update
 		WaitForPendingInitOrStreaming();
 
+		// Make sure the flush actually releases our resource.
+		ReleaseResource();
+
 		FTexturePlatformData *&PlatformDataLink = *PlatformDataLinkPtr;
 		FlushRenderingCommands();
 
@@ -4085,6 +4088,9 @@ void UTexture::ForceRebuildPlatformData(uint8 InEncodeSpeedOverride /* =255 ETex
 			uint32(ETextureCacheFlags::ForceRebuild),
 			nullptr
 			);
+
+		// The build was synchronous but we still need to complete the compilation.
+		BlockOnAnyAsyncBuild();
 	}
 }
 
