Unshelved from pending changelist '35220222':

PR #12189: Erroneous description comment for the UpdateTextureRegions

#rnx

diff --git a/Engine/Source/Runtime/Engine/Classes/Engine/Texture2D.h b/Engine/Source/Runtime/Engine/Classes/Engine/Texture2D.h
--- a/Engine/Source/Runtime/Engine/Classes/Engine/Texture2D.h
+++ b/Engine/Source/Runtime/Engine/Classes/Engine/Texture2D.h
@@ -287,7 +287,7 @@ public:
 	*	@param SrcPitch - the pitch of the source data in bytes
 	*	@param SrcBpp - the size one pixel data in bytes
 	*	@param SrcData - the source data
-	*  @param bFreeData - if true, the SrcData and Regions pointers will be freed after the update.
+	*	@param DataCleanupFunc - a function that will be called to clean up the SrcData and Regions after the update.
 	*/
 	ENGINE_API void UpdateTextureRegions(int32 MipIndex, uint32 NumRegions, const FUpdateTextureRegion2D* Regions, uint32 SrcPitch, uint32 SrcBpp, uint8* SrcData, TFunction<void(uint8* SrcData, const FUpdateTextureRegion2D* Regions)> DataCleanupFunc = [](uint8*, const FUpdateTextureRegion2D*){});
 
