Save 16 bytes in UMaterialInterface (200 -> 186)
UMaterialInstanceDynamic goes from 776 bytes to 760, saving 600 kb
#rnx

diff --git a/Engine/Source/Runtime/Engine/Public/Materials/MaterialInterface.h b/Engine/Source/Runtime/Engine/Public/Materials/MaterialInterface.h
--- a/Engine/Source/Runtime/Engine/Public/Materials/MaterialInterface.h
+++ b/Engine/Source/Runtime/Engine/Public/Materials/MaterialInterface.h
@@ -283,10 +283,6 @@ public:
 	UPROPERTY(EditAnywhere, BlueprintReadOnly, Category = PostProcessMaterial, meta = (DisplayName = "Neural Profile"))
 	TObjectPtr<class UNeuralProfile> NeuralProfile;
 
-	/** Whether this material interface is included in the base game (and not in a DLC) */
-	UPROPERTY()
-	uint8 bIncludedInBaseGame : 1;
-
 	/** Event triggered when the base material is set */
 	FOnBaseMaterialIsSet OnBaseMaterialSetEvent;
 
@@ -321,6 +317,10 @@ private:
 	/** Feature levels to force to compile. */
 	uint32 FeatureLevelsToForceCompile;
 
+	/** Whether this material interface is included in the base game (and not in a DLC) */
+	UPROPERTY(meta=(DisplayAfter="NeuralProfile"))
+	uint8 bIncludedInBaseGame : 1;
+
 public:
 
 	ENGINE_API UMaterialInterface();
@@ -1194,6 +1194,9 @@ protected:
 
 	void UpdateMaterialRenderProxy(FMaterialRenderProxy& Proxy);
 
+	/** Set if CachedExpressionData was loaded from disk, should typically be true when running with cooked data, and false in the editor */
+	bool bLoadedCachedExpressionData = false;
+	
 	/**
 	 * Cached data generated from the material's expressions, may be nullptr
 	 * UMaterials should always have cached data
@@ -1201,9 +1204,6 @@ protected:
 	 */
 	TUniquePtr<FMaterialCachedExpressionData> CachedExpressionData;
 
-	/** Set if CachedExpressionData was loaded from disk, should typically be true when running with cooked data, and false in the editor */
-	bool bLoadedCachedExpressionData = false;
-
 #if WITH_EDITOR
 	TUniquePtr<FMaterialCachedHLSLTree> CachedHLSLTree;
 #endif // WITH_EDITOR
