Ensure internal objects (skeleton & mesh component etc.) are cached when calling DispatchParallelTickPose

#tests

diff --git a/Engine/Source/Runtime/Engine/Private/Components/SkeletalMeshComponent.cpp b/Engine/Source/Runtime/Engine/Private/Components/SkeletalMeshComponent.cpp
--- a/Engine/Source/Runtime/Engine/Private/Components/SkeletalMeshComponent.cpp
+++ b/Engine/Source/Runtime/Engine/Private/Components/SkeletalMeshComponent.cpp
@@ -2691,6 +2691,9 @@ void USkeletalMeshComponent::DispatchParallelTickPose(FActorComponentTickFunctio
 				AnimEvaluationContext.bDuplicateToCacheCurve = false;
 				AnimEvaluationContext.bDuplicateToCachedAttributes = false;
 
+				// Ensure that internal objects are cached before we run any anim graph logic
+				DoInstancePreEvaluation();
+
 				if(bDoParallelUpdate)
 				{
 					DispatchParallelEvaluationTasks(TickFunction);
