APawn 792 -> 784 bytes
ACharacter 1632 -> 1536 bytes (with 8 bytes padding that subclasses vtables will take)
Other character related tweaks

diff --git a/Engine/Source/Runtime/Core/Public/Serialization/Archive.h b/Engine/Source/Runtime/Core/Public/Serialization/Archive.h
--- a/Engine/Source/Runtime/Core/Public/Serialization/Archive.h
+++ b/Engine/Source/Runtime/Core/Public/Serialization/Archive.h
@@ -1583,6 +1583,13 @@ public:
 #endif
 	}
 
+#define SerializeBitfield(Ar, BitField) \
+	{ \
+		bool PREPROCESSOR_JOIN(TempBitfield, __LINE__) = BitField; \
+		Ar.SerializeBits(&PREPROCESSOR_JOIN(TempBitfield, __LINE__), 1); \
+		BitField = PREPROCESSOR_JOIN(TempBitfield, __LINE__); \
+	}
+
 	virtual void SerializeBits(void* V, int64 LengthBits)
 	{
 		Serialize(V, (LengthBits + 7) / 8);
diff --git a/Engine/Source/Runtime/Engine/Classes/GameFramework/Character.h b/Engine/Source/Runtime/Engine/Classes/GameFramework/Character.h
--- a/Engine/Source/Runtime/Engine/Classes/GameFramework/Character.h
+++ b/Engine/Source/Runtime/Engine/Classes/GameFramework/Character.h
@@ -17,6 +17,7 @@
 #include "GameFramework/RootMotionSource.h"
 #include "Character.generated.h"
 
+class ACharacter;
 class AController;
 class FDebugDisplayInfo;
 class UAnimMontage;
@@ -32,8 +33,8 @@
 
 DECLARE_DYNAMIC_MULTICAST_DELEGATE_ThreeParams(FMovementModeChangedSignature, class ACharacter*, Character, EMovementMode, PrevMovementMode, uint8, PreviousCustomMode);
 DECLARE_DYNAMIC_MULTICAST_DELEGATE_ThreeParams(FCharacterMovementUpdatedSignature, float, DeltaSeconds, FVector, OldLocation, FVector, OldVelocity);
-DECLARE_DYNAMIC_MULTICAST_DELEGATE(FCharacterReachedApexSignature);
-DECLARE_DYNAMIC_MULTICAST_DELEGATE_OneParam(FLandedSignature, const FHitResult&, Hit);
+DECLARE_DYNAMIC_MULTICAST_SPARSE_DELEGATE(FCharacterReachedApexSignature, ACharacter, OnReachedJumpApex);
+DECLARE_DYNAMIC_MULTICAST_SPARSE_DELEGATE_OneParam(FLandedSignature, ACharacter, LandedDelegate, const FHitResult&, Hit);
 
 // CVars
 namespace CharacterCVars
@@ -47,10 +48,6 @@
 {
 	GENERATED_USTRUCT_BODY()
 
-	/** Whether this has useful/active data. */
-	UPROPERTY()
-	bool bIsActive = false;
-
 #if WITH_EDITORONLY_DATA
 	/** AnimMontage providing Root Motion */
 	UPROPERTY(meta = (DeprecatedProperty, DeprecationMessage = "Use the GetAnimMontage function instead"))
@@ -61,6 +58,18 @@
 	UPROPERTY()
 	TObjectPtr<UAnimSequenceBase> Animation = nullptr;
 
+	/** Whether this has useful/active data. */
+	UPROPERTY()
+	bool bIsActive = false;
+
+	/** Additional replicated flag, if MovementBase can't be resolved on the client. So we don't use wrong data. */
+	UPROPERTY()
+	bool bRelativePosition = false;
+
+	/** Whether rotation is relative or absolute. */
+	UPROPERTY()
+	bool bRelativeRotation = false;
+
 	/** Track position of Montage */
 	UPROPERTY()
 	float Position = 0.f;
@@ -81,14 +90,6 @@
 	UPROPERTY()
 	FName MovementBaseBoneName;
 
-	/** Additional replicated flag, if MovementBase can't be resolved on the client. So we don't use wrong data. */
-	UPROPERTY()
-	bool bRelativePosition = false;
-
-	/** Whether rotation is relative or absolute. */
-	UPROPERTY()
-	bool bRelativeRotation = false;
-
 	/** State of Root Motion Sources on Authority */
 	UPROPERTY()
 	FRootMotionSourceGroup AuthoritativeRootMotion;
@@ -186,14 +187,26 @@
 	UPROPERTY()
 	uint16 BaseID = 0;
 
-	/** Component we are based on */
+	/** Whether the server says that there is a base. On clients, the component may not have resolved yet. */
 	UPROPERTY()
-	TObjectPtr<UPrimitiveComponent> MovementBase = nullptr;
+	uint8 bServerHasBaseComponent:1 = false;
 
+	/** Whether rotation is relative to the base or absolute. It can only be relative if location is also relative. */
+	UPROPERTY()
+	uint8 bRelativeRotation:1 = false;
+
+	/** Whether there is a velocity on the server. Used for forcing replication when velocity goes to zero. */
+	UPROPERTY()
+	uint8 bServerHasVelocity:1 = false;
+
 	/** Bone name on component, for skeletal meshes. NAME_None if not a skeletal mesh or if bone is invalid. */
 	UPROPERTY()
 	FName BoneName;
 
+	/** Component we are based on */
+	UPROPERTY()
+	TObjectPtr<UPrimitiveComponent> MovementBase = nullptr;
+
 	/** Location relative to MovementBase. Only valid if HasRelativeLocation() is true. */
 	UPROPERTY()
 	FVector_NetQuantize100 Location;
@@ -202,17 +215,5 @@
 	UPROPERTY()
 	FRotator Rotation = FRotator(0.f);
 
-	/** Whether the server says that there is a base. On clients, the component may not have resolved yet. */
-	UPROPERTY()
-	bool bServerHasBaseComponent = false;
-
-	/** Whether rotation is relative to the base or absolute. It can only be relative if location is also relative. */
-	UPROPERTY()
-	bool bRelativeRotation = false;
-
-	/** Whether there is a velocity on the server. Used for forcing replication when velocity goes to zero. */
-	UPROPERTY()
-	bool bServerHasVelocity = false;
-
 	/** Is location relative? */
 	FORCEINLINE bool HasRelativeLocation() const
@@ -439,10 +440,6 @@
 	UPROPERTY(ReplicatedUsing=OnRep_ReplicatedBasedMovement)
 	struct FBasedMovementInfo ReplicatedBasedMovement;
 
-	/** Scale to apply to root motion translation on this Character */
-	UPROPERTY(Replicated)
-	float AnimRootMotionTranslationScale;
-
 public:
 	/** Rep notify for ReplicatedBasedMovement */
 	UFUNCTION()
@@ -455,17 +452,6 @@
 	/** Whether this Character should include acceleration data in its replicated movement */
 	ENGINE_API virtual bool ShouldReplicateAcceleration() const { return CharacterCVars::EnableCharacterAccelerationReplication != 0; }
 
-	/** Saved translation offset of mesh. */
-	UPROPERTY()
-	FVector BaseTranslationOffset;
-
-	/** Saved rotation offset of mesh. */
-	UPROPERTY()
-	FQuat BaseRotationOffset;
-
-	/** Event called after actor's base changes (if SetBase was requested to notify us with bNotifyPawn). */
-	ENGINE_API virtual void BaseChange();
-
 	/** CharacterMovement ServerLastTransformUpdateTimeStamp value, replicated to simulated proxies. */
 	UPROPERTY(Replicated)
 	float ReplicatedServerLastTransformUpdateTimeStamp;
@@ -473,20 +459,28 @@
 	UPROPERTY(ReplicatedUsing=OnRep_ReplayLastTransformUpdateTimeStamp)
 	float ReplayLastTransformUpdateTimeStamp;
 
-	/** CharacterMovement MovementMode (and custom mode) replicated for simulated proxies. Use CharacterMovementComponent::UnpackNetworkMovementMode() to translate it. */
-	UPROPERTY(Replicated)
-	uint8 ReplicatedMovementMode;
+	/** Saved rotation offset of mesh. */
+	UPROPERTY()
+	FQuat BaseRotationOffset;
 
+	/** Saved translation offset of mesh. */
+	UPROPERTY()
+	FVector BaseTranslationOffset;
+
+	/** Event called after actor's base changes (if SetBase was requested to notify us with bNotifyPawn). */
+	ENGINE_API virtual void BaseChange();
+
 	/** CharacterMovement Custom gravity direction replicated for simulated proxies. */
 	UPROPERTY(Replicated)
 	FVector_NetQuantizeNormal ReplicatedGravityDirection;
 
-	/** Flag that we are receiving replication of the based movement. */
-	UPROPERTY()
-	bool bInBaseReplication;
-
 	/** Cached version of the replicated gravity direction before replication. Used to compare if the value was changed as a result of replication. */
 	FVector_NetQuantizeNormal PreNetReceivedGravityDirection;
+
+	/** Scale to apply to root motion translation on this Character */
+	UPROPERTY(Replicated)
+	float AnimRootMotionTranslationScale;
+
 public:
 	UFUNCTION()
 	ENGINE_API void OnRep_ReplayLastTransformUpdateTimeStamp();
@@ -530,13 +524,19 @@
 	UPROPERTY(EditAnywhere, BlueprintReadWrite, Category=Camera)
 	float CrouchedEyeHeight;
 
+protected:
+	/** Flag that we are receiving replication of the based movement. */
+	UPROPERTY()
+	uint8 bInBaseReplication:1;
+
+public:
 	/** Set by character movement to specify that this Character is currently crouched. */
 	UPROPERTY(BlueprintReadOnly, replicatedUsing=OnRep_IsCrouched, Category=Character)
-	uint32 bIsCrouched:1;
+	uint8 bIsCrouched:1;
 
 	/** Set to indicate that this Character is currently under the force of a jump (if JumpMaxHoldTime is non-zero). IsJumpProvidingForce() handles this as well. */
 	UPROPERTY(Transient, Replicated)
-	uint32 bProxyIsJumpForceApplied : 1;
+	uint8 bProxyIsJumpForceApplied : 1;
 
 	/** Handle Crouching replicated from server */
 	UFUNCTION()
@@ -544,39 +544,45 @@
 
 	/** When true, player wants to jump */
 	UPROPERTY(BlueprintReadOnly, Category=Character)
-	uint32 bPressedJump:1;
+	uint8 bPressedJump:1;
 
 	/** When true, applying updates to network client (replaying saved moves for a locally controlled character) */
 	UPROPERTY(Transient)
-	uint32 bClientUpdating:1;
+	uint8 bClientUpdating:1;
 
 	/** True if Pawn was initially falling when started to replay network moves. */
 	UPROPERTY(Transient)
-	uint32 bClientWasFalling:1; 
+	uint8 bClientWasFalling:1; 
 
 	/** If server disagrees with root motion track position, client has to resimulate root motion from last AckedMove. */
 	UPROPERTY(Transient)
-	uint32 bClientResimulateRootMotion:1;
+	uint8 bClientResimulateRootMotion:1;
 
 	/** If server disagrees with root motion state, client has to resimulate root motion from last AckedMove. */
 	UPROPERTY(Transient)
-	uint32 bClientResimulateRootMotionSources:1;
+	uint8 bClientResimulateRootMotionSources:1;
 
 	/** Disable simulated gravity (set when character encroaches geometry on client, to keep it from falling through floors) */
 	UPROPERTY()
-	uint32 bSimGravityDisabled:1;
+	uint8 bSimGravityDisabled:1;
 
 	UPROPERTY(Transient)
-	uint32 bClientCheckEncroachmentOnNetUpdate:1;
+	uint8 bClientCheckEncroachmentOnNetUpdate:1;
 
 	/** Disable root motion on the server. When receiving a DualServerMove, where the first move is not root motion and the second is. */
 	UPROPERTY(Transient)
-	uint32 bServerMoveIgnoreRootMotion:1;
+	uint8 bServerMoveIgnoreRootMotion:1;
 
 	/** Tracks whether or not the character was already jumping last frame. */
 	UPROPERTY(VisibleInstanceOnly, BlueprintReadOnly, Transient, Category=Character)
-	uint32 bWasJumping : 1;
+	uint8 bWasJumping : 1;
 
+protected:
+	/** CharacterMovement MovementMode (and custom mode) replicated for simulated proxies. Use CharacterMovementComponent::UnpackNetworkMovementMode() to translate it. */
+	UPROPERTY(Replicated)
+	uint8 ReplicatedMovementMode;
+
+public:
 	/** 
 	 * Jump key Held Time.
 	 * This is the time that the player has held the jump key, in seconds.
diff --git a/Engine/Source/Runtime/Engine/Classes/GameFramework/CharacterMovementComponent.h b/Engine/Source/Runtime/Engine/Classes/GameFramework/CharacterMovementComponent.h
--- a/Engine/Source/Runtime/Engine/Classes/GameFramework/CharacterMovementComponent.h
+++ b/Engine/Source/Runtime/Engine/Classes/GameFramework/CharacterMovementComponent.h
@@ -2885,9 +2885,14 @@ FORCEINLINE uint32 UCharacterMovementComponent::PackYawAndPitchTo32(const float
 	return Rotation32;
 }
 
+class FSavedMove_Character_FixLayout
+{
+public:
+	virtual ~FSavedMove_Character_FixLayout() = default;
+};
 
 /** FSavedMove_Character represents a saved move on the client that has been sent to the server and might need to be played back. */
-class FSavedMove_Character
+class FSavedMove_Character : public FSavedMove_Character_FixLayout
 {
 public:
 	ENGINE_API FSavedMove_Character();
@@ -2901,17 +2906,24 @@ public:
 
 	ACharacter* CharacterOwner;
 
-	uint32 bPressedJump:1;
-	uint32 bWantsToCrouch:1;
-	uint32 bForceMaxAccel:1;
+	uint8 bPressedJump:1;
+	uint8 bWantsToCrouch:1;
+	uint8 bForceMaxAccel:1;
 
 	/** If true, can't combine this move with another move. */
-	uint32 bForceNoCombine:1;
+	uint8 bForceNoCombine:1;
 
 	/** If true this move is using an old TimeStamp, before a reset occurred. */
-	uint32 bOldTimeStampBeforeReset:1;
+	uint8 bOldTimeStampBeforeReset:1;
 
-	uint32 bWasJumping:1;
+	uint8 bWasJumping:1;
+
+	UE_DEPRECATED_FORGAME(4.20, "This property is deprecated, use StartPackedMovementMode or EndPackedMovementMode instead.")
+	uint8 MovementMode;
+
+	// Information at the start of the move
+	uint8 StartPackedMovementMode;
+	uint8 EndPackedMovementMode;
 
 	float TimeStamp;    // Time of this move.
 	float DeltaTime;    // amount of time for this move
@@ -2921,11 +2933,6 @@ public:
 	int32 JumpMaxCount;
 	int32 JumpCurrentCount;
 	
-	UE_DEPRECATED_FORGAME(4.20, "This property is deprecated, use StartPackedMovementMode or EndPackedMovementMode instead.")
-	uint8 MovementMode;
-
-	// Information at the start of the move
-	uint8 StartPackedMovementMode;
 	FVector StartLocation;
 	FVector StartRelativeLocation;
 	FVector StartVelocity;
@@ -2945,7 +2952,6 @@ public:
 	FRotator StartAttachRelativeRotation;
 
 	// Information after the move has been performed
-	uint8 EndPackedMovementMode;
 	FVector SavedLocation;
 	FRotator SavedRotation;
 	FVector SavedVelocity;
diff --git a/Engine/Source/Runtime/Engine/Private/Character.cpp b/Engine/Source/Runtime/Engine/Private/Character.cpp
--- a/Engine/Source/Runtime/Engine/Private/Character.cpp
+++ b/Engine/Source/Runtime/Engine/Private/Character.cpp
@@ -1225,7 +1225,7 @@ void ACharacter::OnRep_ReplicatedBasedMovement()
 	}
 
 	CharacterMovement->bNetworkUpdateReceived = true;
-	TGuardValue<bool> bInBaseReplicationGuard(bInBaseReplication, true);
+	FGuardValue_Bitfield(bInBaseReplication, true);
 
 	const bool bBaseChanged = (BasedMovement.MovementBase != ReplicatedBasedMovement.MovementBase || BasedMovement.BoneName != ReplicatedBasedMovement.BoneName);
 	if (bBaseChanged)
