Skeletal Mesh: make sure that the mesh description bones data reflects the ref skeleton
   - also fixed the static to skeletal mesh  importer not setting the mesh description bones data

#rnx

diff --git a/Engine/Plugins/Experimental/MeshModelingToolsetExp/Source/MeshModelingToolsEditorOnlyExp/Private/SkeletalMesh/SkinWeightsPaintTool.cpp b/Engine/Plugins/Experimental/MeshModelingToolsetExp/Source/MeshModelingToolsEditorOnlyExp/Private/SkeletalMesh/SkinWeightsPaintTool.cpp
--- a/Engine/Plugins/Experimental/MeshModelingToolsetExp/Source/MeshModelingToolsEditorOnlyExp/Private/SkeletalMesh/SkinWeightsPaintTool.cpp
+++ b/Engine/Plugins/Experimental/MeshModelingToolsetExp/Source/MeshModelingToolsEditorOnlyExp/Private/SkeletalMesh/SkinWeightsPaintTool.cpp
@@ -2661,6 +2661,15 @@ void USkinWeightsPaintTool::TransferWeights()
 	{
 		TargetAttributes->CopyBoneAttributes(*SourceMesh.Attributes());
 	}
+	else
+	{
+		const USkeletalMeshComponent* Component = GetSkeletalMeshComponent(Target);
+		if (ensure(Component))
+		{
+			const FReferenceSkeleton& RefSkeleton = Component->GetSkeletalMeshAsset()->GetRefSkeleton();
+			ensure(TargetAttributes->GetNumBones() == RefSkeleton.GetRawBoneNum());
+		}
+	}
 
 	// NOTE should we expose all the options?
 	// 	TransferBoneWeights.NormalThreshold;
diff --git a/Engine/Source/Developer/SkeletalMeshUtilitiesCommon/Private/StaticToSkeletalMeshConverter.cpp b/Engine/Source/Developer/SkeletalMeshUtilitiesCommon/Private/StaticToSkeletalMeshConverter.cpp
--- a/Engine/Source/Developer/SkeletalMeshUtilitiesCommon/Private/StaticToSkeletalMeshConverter.cpp
+++ b/Engine/Source/Developer/SkeletalMeshUtilitiesCommon/Private/StaticToSkeletalMeshConverter.cpp
@@ -261,6 +261,29 @@ static bool AddLODFromStaticMeshSourceModel(
 		FSkeletalMeshAttributes SkeletalMeshAttributes(SkeletalMeshGeometry);
 		SkeletalMeshAttributes.Register();
 		
+		// Fill Bones data.
+		const FReferenceSkeleton RefSkeleton = InSkeletalMesh->GetRefSkeleton();
+		const int32 NumRefBones = InSkeletalMesh->GetRefSkeleton().GetRawBoneNum();
+		
+		FSkeletalMeshAttributes::FBoneArray& Bones = SkeletalMeshAttributes.Bones();
+		Bones.Reset(NumRefBones);
+	
+		FSkeletalMeshAttributes::FBoneNameAttributesRef BoneNames = SkeletalMeshAttributes.GetBoneNames();
+		FSkeletalMeshAttributes::FBoneParentIndexAttributesRef BoneParentIndices = SkeletalMeshAttributes.GetBoneParentIndices();
+		FSkeletalMeshAttributes::FBonePoseAttributesRef BonePoses = SkeletalMeshAttributes.GetBonePoses();
+		
+		for (int Index = 0; Index < NumRefBones; ++Index)
+		{
+			const FMeshBoneInfo& BoneInfo = RefSkeleton.GetRawRefBoneInfo()[Index];
+			const FTransform& BoneTransform = RefSkeleton.GetRawRefBonePose()[Index];
+
+			const FBoneID BoneID = SkeletalMeshAttributes.CreateBone();
+
+			BoneNames.Set(BoneID, BoneInfo.Name);
+			BoneParentIndices.Set(BoneID, BoneInfo.ParentIndex);
+			BonePoses.Set(BoneID, BoneTransform);
+		}
+		
 		// Full binding to the root bone.
 		FSkinWeightsVertexAttributesRef SkinWeights = SkeletalMeshAttributes.GetVertexSkinWeights();
 		UE::AnimationCore::FBoneWeight RootInfluence(InBoneIndex, 1.0f);
diff --git a/Engine/Source/Runtime/Engine/Classes/Engine/SkeletalMeshSourceModel.h b/Engine/Source/Runtime/Engine/Classes/Engine/SkeletalMeshSourceModel.h
--- a/Engine/Source/Runtime/Engine/Classes/Engine/SkeletalMeshSourceModel.h
+++ b/Engine/Source/Runtime/Engine/Classes/Engine/SkeletalMeshSourceModel.h
@@ -173,6 +173,12 @@ private:
 	 *  and then unloading the bulkd data, to conserve memory.
 	 */
 	void UpdateCachedMeshStatisticsFromBulkIfNeeded();
+	
+	/**
+	 *  Update the MeshDescription bones data if needed so that it matches the skeletal mesh ref skeleton.  
+     */
+	void UpdateBonesDataIfNeeded(FMeshDescription& InOutMeshDescription) const;
+    
 #endif
 	
 	UPROPERTY()
diff --git a/Engine/Source/Runtime/Engine/Private/SkeletalMeshSourceModel.cpp b/Engine/Source/Runtime/Engine/Private/SkeletalMeshSourceModel.cpp
--- a/Engine/Source/Runtime/Engine/Private/SkeletalMeshSourceModel.cpp
+++ b/Engine/Source/Runtime/Engine/Private/SkeletalMeshSourceModel.cpp
@@ -130,7 +130,9 @@ FMeshDescription* FSkeletalMeshSourceModel::GetMeshDescription() const
 	
 	if (MeshDescriptionBulkData->HasCachedMeshDescription())
 	{
-		return &MeshDescriptionBulkData->GetMeshDescription()->GetMeshDescription();
+		FMeshDescription& MeshDescription = MeshDescriptionBulkData->GetMeshDescription()->GetMeshDescription();
+		UpdateBonesDataIfNeeded(MeshDescription);
+		return &MeshDescription;
 	}
 
 	return nullptr;
@@ -421,5 +423,44 @@ void FSkeletalMeshSourceModel::UpdateCachedMeshStatisticsFromBulkIfNeeded()
 	}
 }
 
+void FSkeletalMeshSourceModel::UpdateBonesDataIfNeeded(FMeshDescription& InOutMeshDescription) const
+{
+	FSkeletalMeshAttributes MeshAttributes(InOutMeshDescription);
+	if (!MeshAttributes.HasBones())
+	{
+		MeshAttributes.Register(true);
+	}
+
+	const USkeletalMesh* SkeletalMesh = GetOwner();
+	if (!SkeletalMesh)
+	{
+		return;
+	}
+	
+	const int32 NumMeshDescBones = MeshAttributes.GetNumBones();
+	const FReferenceSkeleton RefSkeleton = SkeletalMesh->GetRefSkeleton();
+	const int32 NumRefBones = SkeletalMesh->GetRefSkeleton().GetRawBoneNum();
+	if (NumMeshDescBones != NumRefBones && NumRefBones > 0)
+	{
+		MeshAttributes.Bones().Reset(NumRefBones);
+	
+		FSkeletalMeshAttributes::FBoneNameAttributesRef BoneNames = MeshAttributes.GetBoneNames();
+		FSkeletalMeshAttributes::FBoneParentIndexAttributesRef BoneParentIndices = MeshAttributes.GetBoneParentIndices();
+		FSkeletalMeshAttributes::FBonePoseAttributesRef BonePoses = MeshAttributes.GetBonePoses();
+
+
+		for (int Index = 0; Index < NumRefBones; ++Index)
+		{
+			const FMeshBoneInfo& BoneInfo = RefSkeleton.GetRawRefBoneInfo()[Index];
+			const FTransform& BoneTransform = RefSkeleton.GetRawRefBonePose()[Index];
+
+			const FBoneID BoneID = MeshAttributes.CreateBone();
+
+			BoneNames.Set(BoneID, BoneInfo.Name);
+			BoneParentIndices.Set(BoneID, BoneInfo.ParentIndex);
+			BonePoses.Set(BoneID, BoneTransform);
+		}
+	}
+}
 
 #endif
