Use MaxAssetSize (texture max resolution) instead of MAX_TEXTURE_SIZE (which is 8k) to do proper early out when iterating over streamable render asset in memory

#rnx

diff --git a/Engine/Source/Runtime/Engine/Private/Streaming/AsyncTextureStreaming.cpp b/Engine/Source/Runtime/Engine/Private/Streaming/AsyncTextureStreaming.cpp
--- a/Engine/Source/Runtime/Engine/Private/Streaming/AsyncTextureStreaming.cpp
+++ b/Engine/Source/Runtime/Engine/Private/Streaming/AsyncTextureStreaming.cpp
@@ -121,6 +121,7 @@ void FAsyncRenderAssetStreamingData::UpdatePerfectWantedMips_Async(FStreamingRen
 	bool bLooksLowRes = false;
 
 	const float MaxAllowedSize = StreamingRenderAsset.GetMaxAllowedSize(MaxScreenSizeOverAllViews);
+	const float MaxAssetSize = StreamingRenderAsset.IsMesh() ? MAX_TEXTURE_SIZE : MaxAllowedSize;
 
 #if !UE_BUILD_SHIPPING
 	if (Settings.bStressTest)
@@ -149,7 +150,7 @@ void FAsyncRenderAssetStreamingData::UpdatePerfectWantedMips_Async(FStreamingRen
 	else
 	{
 		const EStreamableRenderAssetType AssetType = StreamingRenderAsset.RenderAssetType;
-		DynamicInstancesView.GetRenderAssetScreenSize(AssetType, RenderAsset, MaxSize, MaxSize_VisibleOnly, MaxNumForcedLODs, bOutputToLog ? TEXT("Dynamic") : nullptr);
+		DynamicInstancesView.GetRenderAssetScreenSize(AssetType, RenderAsset, MaxSize, MaxSize_VisibleOnly, MaxNumForcedLODs, MaxAssetSize, bOutputToLog ? TEXT("Dynamic") : nullptr);
 
 		bool bCulled = false;
 		if (Settings.bMipCalculationEnablePerLevelList)
@@ -183,7 +184,7 @@ void FAsyncRenderAssetStreamingData::UpdatePerfectWantedMips_Async(FStreamingRen
 				}
 
 				// No need to iterate more if render asset is already at maximum resolution.
-				if (MaxSize_VisibleOnly >= MAX_TEXTURE_SIZE || MaxNumForcedLODs >= StreamingRenderAsset.MaxAllowedMips)
+				if (MaxSize_VisibleOnly >= MaxAssetSize || MaxNumForcedLODs >= StreamingRenderAsset.MaxAllowedMips)
 				{
 					break;
 				}
@@ -192,7 +193,7 @@ void FAsyncRenderAssetStreamingData::UpdatePerfectWantedMips_Async(FStreamingRen
 				float TmpMaxVisibleOnly = MaxSize_VisibleOnly;
 				int32 TmpMaxNumForcedLODs = MaxNumForcedLODs;
 				
-				StaticInstancesView.GetRenderAssetScreenSize(AssetType, RenderAsset, TmpMaxSize, TmpMaxVisibleOnly, TmpMaxNumForcedLODs, bOutputToLog ? TEXT("Static") : nullptr);
+				StaticInstancesView.GetRenderAssetScreenSize(AssetType, RenderAsset, TmpMaxSize, TmpMaxVisibleOnly, TmpMaxNumForcedLODs, MaxAssetSize, bOutputToLog ? TEXT("Static") : nullptr);
 
 				MaxSize = FMath::Max(TmpMaxSize, MaxSize);
 				MaxSize_VisibleOnly = FMath::Max(TmpMaxVisibleOnly, MaxSize_VisibleOnly);
@@ -214,7 +215,7 @@ void FAsyncRenderAssetStreamingData::UpdatePerfectWantedMips_Async(FStreamingRen
 					break;
 				}
 
-				StaticInstancesView.GetRenderAssetScreenSize(AssetType, RenderAsset, MaxSize, MaxSize_VisibleOnly, MaxNumForcedLODs, bOutputToLog ? TEXT("Static") : nullptr);
+				StaticInstancesView.GetRenderAssetScreenSize(AssetType, RenderAsset, MaxSize, MaxSize_VisibleOnly, MaxNumForcedLODs, MaxAssetSize, bOutputToLog ? TEXT("Static") : nullptr);
 			}
 		}
 
diff --git a/Engine/Source/Runtime/Engine/Private/Streaming/TextureInstanceView.cpp b/Engine/Source/Runtime/Engine/Private/Streaming/TextureInstanceView.cpp
--- a/Engine/Source/Runtime/Engine/Private/Streaming/TextureInstanceView.cpp
+++ b/Engine/Source/Runtime/Engine/Private/Streaming/TextureInstanceView.cpp
@@ -480,6 +480,7 @@ void FRenderAssetInstanceAsyncView::GetRenderAssetScreenSize(
 	float& MaxSize,
 	float& MaxSize_VisibleOnly,
 	int32& MaxNumForcedLODs,
+	const float MaxAssetSize,
 	const TCHAR* LogPrefix) const
 {
 	// No need to iterate more if texture is already at maximum resolution.
@@ -501,7 +502,7 @@ void FRenderAssetInstanceAsyncView::GetRenderAssetScreenSize(
 				const FRenderAssetInstanceView::FCompiledElement* CompiledElementData = CompiledElements->GetData();
 
 				int32 CompiledElementIndex = 0;
-				while (CompiledElementIndex < NumCompiledElements && MaxSize_VisibleOnly < MAX_TEXTURE_SIZE)
+				while (CompiledElementIndex < NumCompiledElements && MaxSize_VisibleOnly < MaxAssetSize)
 				{
 					const FRenderAssetInstanceView::FCompiledElement& CompiledElement = CompiledElementData[CompiledElementIndex];
 					if (ensure(BoundsViewInfo.IsValidIndex(CompiledElement.BoundsIndex)))
@@ -527,7 +528,7 @@ void FRenderAssetInstanceAsyncView::GetRenderAssetScreenSize(
 					++CompiledElementIndex;
 				}
 
-				if (MaxSize_VisibleOnly >= MAX_TEXTURE_SIZE && CompiledElementIndex > 1)
+				if (MaxSize_VisibleOnly >= MaxAssetSize && CompiledElementIndex > 1)
 				{
 					// This does not realloc anything but moves the closest element at head, making the next update find it immediately and early exit.
 					FRenderAssetInstanceView::FCompiledElement* SwapElementData = const_cast<FRenderAssetInstanceView::FCompiledElement*>(CompiledElementData);
@@ -537,7 +538,7 @@ void FRenderAssetInstanceAsyncView::GetRenderAssetScreenSize(
 		}
 		else
 		{
-			for (auto It = View->GetElementIterator(InAsset); It && (AssetType != EStreamableRenderAssetType::Texture || MaxSize_VisibleOnly < MAX_TEXTURE_SIZE || LogPrefix); ++It)
+			for (auto It = View->GetElementIterator(InAsset); It && (AssetType != EStreamableRenderAssetType::Texture || MaxSize_VisibleOnly < MaxAssetSize || LogPrefix); ++It)
 			{
 				// Only handle elements that are in bounds.
 				if (ensure(BoundsViewInfo.IsValidIndex(It.GetBoundsIndex())))
diff --git a/Engine/Source/Runtime/Engine/Private/Streaming/TextureInstanceView.h b/Engine/Source/Runtime/Engine/Private/Streaming/TextureInstanceView.h
--- a/Engine/Source/Runtime/Engine/Private/Streaming/TextureInstanceView.h
+++ b/Engine/Source/Runtime/Engine/Private/Streaming/TextureInstanceView.h
@@ -268,6 +268,7 @@ public:
 		float& MaxSize,
 		float& MaxSize_VisibleOnly,
 		int32& MaxNumForcedLODs,
+		const float MaxAssetSize,
 		const TCHAR* LogPrefix) const;
 
 	FORCEINLINE bool HasRenderAssetReferences(const UStreamableRenderAsset* InAsset) const
