Crunch primitive components.

#rnx

diff --git a/Engine/Source/Runtime/Engine/Classes/Components/InstancedStaticMeshComponent.h b/Engine/Source/Runtime/Engine/Classes/Components/InstancedStaticMeshComponent.h
--- a/Engine/Source/Runtime/Engine/Classes/Components/InstancedStaticMeshComponent.h
+++ b/Engine/Source/Runtime/Engine/Classes/Components/InstancedStaticMeshComponent.h
@@ -125,6 +125,12 @@ class UInstancedStaticMeshComponent : public UStaticMeshComponent, public ISMIns
 	UPROPERTY(EditAnywhere, Category=Instances, AdvancedDisplay)
 	int32 NumCustomDataFloats;
 
+	/** Value used to seed the random number stream that generates random numbers for each of this mesh's instances.
+	The random number is stored in a buffer accessible to materials through the PerInstanceRandom expression. If
+	this is set to zero (default), it will be populated automatically by the editor. */
+	UPROPERTY(EditAnywhere, BlueprintReadWrite, Category=InstancedStaticMeshComponent)
+	int32 InstancingRandomSeed=0;
+
 	/** Array of custom data for instances. This will contains NumCustomDataFloats*InstanceCount entries. The entries are represented sequantially, in instance order. Can be read in a material and manipulated through Blueprints.
 	 *	Example: If NumCustomDataFloats is 1, then each entry will belong to an instance. Custom data 0 will belong to Instance 0. Custom data 1 will belong to Instance 1 etc.
 	 *	Example: If NumCustomDataFloats is 2, then each pair of sequential entries belong to an instance. Custom data 0 and 1 will belong to Instance 0. Custom data 2 and 3 will belong to Instance 2 etc.
@@ -132,12 +138,6 @@ class UInstancedStaticMeshComponent : public UStaticMeshComponent, public ISMIns
 	UPROPERTY(EditAnywhere, EditFixedSize, SkipSerialization, DisplayName="Custom data", Category=Instances, AdvancedDisplay, meta=(EditFixedOrder))
 	TArray<float> PerInstanceSMCustomData;
 
-	/** Value used to seed the random number stream that generates random numbers for each of this mesh's instances.
-	The random number is stored in a buffer accessible to materials through the PerInstanceRandom expression. If
-	this is set to zero (default), it will be populated automatically by the editor. */
-	UPROPERTY(EditAnywhere, BlueprintReadWrite, Category=InstancedStaticMeshComponent)
-	int32 InstancingRandomSeed=0;
-
 	/** Additional random seeds ranges. Each seed entry will be applied from AdditionalRandomSeeds[i].StartInstanceIndex to AdditionalRandomSeeds[i+1].StartInstanceIndex -1 */
 	UPROPERTY()
 	TArray<FInstancedStaticMeshRandomSeed> AdditionalRandomSeeds;
@@ -165,14 +165,30 @@ class UInstancedStaticMeshComponent : public UStaticMeshComponent, public ISMIns
 	UPROPERTY()
 	uint8 bInheritPerInstanceData : 1;
 
+	/** Don't create any collision when this bool is set */
+	UPROPERTY()
+	uint8 bDisableCollision : 1;
+
+	/** 
+	* Flag for using RemoveAtSwap on instance removal. 
+	* The implementation is free to ignore this flag, but should honor whatever behavior is being returned by SupportsRemoveSwap().
+	*/
+	uint8 bSupportRemoveAtSwap : 1;
+
+protected:
+	uint8 bIsInstanceDataApplyCompleted : 1;
+
+private:
+	uint8 bHasPreviousTransforms : 1;
+
+	/** Flag for whether we are using conservative bounds. */
+	uint8 bUseConservativeBounds : 1;
+
+public:
 	/** Mapping from PerInstanceSMData order to instance render buffer order. If empty, the PerInstanceSMData order is used. */
 	UPROPERTY()
 	TArray<int32> InstanceReorderTable;
 	
-	/** Don't create any collision when this bool is set */
-	UPROPERTY()
-	bool bDisableCollision;
-
 	/** Tracks outstanding proxysize, as this is a bit hard to do with the fire-and-forget grass. */
 	SIZE_T ProxySize;
 
@@ -416,10 +432,6 @@ class UInstancedStaticMeshComponent : public UStaticMeshComponent, public ISMIns
 	ENGINE_API FInstanceDataFlags MakeInstanceDataFlags(bool bAnyMaterialHasPerInstanceRandom, bool bAnyMaterialHasPerInstanceCustomData) const;
 
 private:
-	bool bHasPreviousTransforms = false;
-
-	/** Flag for whether we are using conservative bounds. */
-	bool bUseConservativeBounds = false;
 	/** Current cached conservativ bounds. */
 	FBox CachedConservativeInstanceBounds;
 public:
@@ -585,12 +597,6 @@ public:
 	/** Request to navigation system to update only part of navmesh occupied specified instances transforms. */
 	ENGINE_API virtual void PartialNavigationUpdates(TConstArrayView<FTransform> InstanceTransforms);
 
-	/** 
-	 * Flag for using RemoveAtSwap on instance removal. 
-	 * The implementation is free to ignore this flag, but should honor whatever behavior is being returned by SupportsRemoveSwap().
-	 */
-	bool bSupportRemoveAtSwap = false;
-
 	ENGINE_API TSharedPtr<FISMCInstanceDataSceneProxy, ESPMode::ThreadSafe> GetOrCreateInstanceDataSceneProxy();
 
 	/**
@@ -633,7 +639,10 @@ private:
 	ENGINE_API bool BatchUpdateInstancesTransformsInternal(int32 StartInstanceIndex, TArrayView<const FTransform> NewInstancesTransforms, bool bWorldSpace, bool bMarkRenderStateDirty, bool bTeleport);
 
 protected:
-	bool bIsInstanceDataApplyCompleted = true;
+
+	/** Number of pending lightmaps still to be calculated (Apply()'d). */
+	UPROPERTY(Transient, DuplicateTransient, TextExportTransient)
+	int32 NumPendingLightmaps;
 
 	FPrimitiveInstanceDataManager PrimitiveInstanceDataManager;
 
@@ -679,10 +688,6 @@ protected:
 	/** Initializes the body instance for the specified instance of the static mesh. */
 	ENGINE_API void InitInstanceBody(int32 InstanceIdx, FBodyInstance* InBodyInstance);
 
-	/** Number of pending lightmaps still to be calculated (Apply()'d). */
-	UPROPERTY(Transient, DuplicateTransient, TextExportTransient)
-	int32 NumPendingLightmaps;
-
 	/** The mappings for all the instances of this component. */
 	UPROPERTY(Transient, DuplicateTransient, TextExportTransient)
 	TArray<FInstancedStaticMeshMappingInfo> CachedMappings;
diff --git a/Engine/Source/Runtime/Engine/Classes/Components/PrimitiveComponent.h b/Engine/Source/Runtime/Engine/Classes/Components/PrimitiveComponent.h
--- a/Engine/Source/Runtime/Engine/Classes/Components/PrimitiveComponent.h
+++ b/Engine/Source/Runtime/Engine/Classes/Components/PrimitiveComponent.h
@@ -295,6 +295,10 @@ public:
 	UPROPERTY(EditAnywhere, AdvancedDisplay, BlueprintReadWrite, Category=HLOD, meta=(DisplayName="HLOD Batching Policy", DisplayAfter="bEnableAutoLODGeneration", EditConditionHides, EditCondition="bEnableAutoLODGeneration"))
 	EHLODBatchingPolicy HLODBatchingPolicy;
 
+	/** Control shadow invalidation behavior, in particular with respect to Virtual Shadow Maps and material effects like World Position Offset. */
+	UPROPERTY(EditAnywhere, BlueprintReadOnly, Category=Lighting, AdvancedDisplay, meta=(EditCondition="CastShadow"))
+	EShadowCacheInvalidationBehavior ShadowCacheInvalidationBehavior;
+
 	/** Whether to include this component in HLODs or not. */
 	UPROPERTY(EditAnywhere, BlueprintReadWrite, Category = HLOD, meta=(DisplayName="Include Component in HLOD"))
 	uint8 bEnableAutoLODGeneration : 1;
@@ -500,10 +504,6 @@ public:
 	UPROPERTY(EditAnywhere, BlueprintReadOnly, Category=Lighting, AdvancedDisplay, meta=(EditCondition="CastShadow", DisplayName = "Static Shadow"))
 	uint8 bCastStaticShadow:1;
 
-	/** Control shadow invalidation behavior, in particular with respect to Virtual Shadow Maps and material effects like World Position Offset. */
-	UPROPERTY(EditAnywhere, BlueprintReadOnly, Category=Lighting, AdvancedDisplay, meta=(EditCondition="CastShadow"))
-	EShadowCacheInvalidationBehavior ShadowCacheInvalidationBehavior;
-
 	/** 
 	 * Whether the object should cast a volumetric translucent shadow.
 	 * Volumetric translucent shadows are useful for primitives with smoothly changing opacity like particles representing a volume, 
@@ -677,12 +677,6 @@ protected:
 	/** PSOs requested priority */
 	EPSOPrecachePriority PSOPrecacheRequestPriority : 2;
 	static_assert((int)EPSOPrecachePriority::Highest < 1 << 2);
-
-	/** Cached array of material PSO requests which can be used to boost the priority */
-	TArray<FMaterialPSOPrecacheRequestID> MaterialPSOPrecacheRequestIDs;
-
-	/** Graph event used to track all the PSO precache events */
-	FGraphEventRef PSOPrecacheCompileEvent;
 #endif
 
 	uint8 bIgnoreBoundsForEditorFocus : 1;
@@ -759,6 +753,14 @@ public:
 	UPROPERTY(EditAnywhere, AdvancedDisplay, BlueprintReadOnly, Category=Lighting)
 	FLightingChannels LightingChannels;
 
+#if UE_WITH_PSO_PRECACHING
+	/** Cached array of material PSO requests which can be used to boost the priority */
+	TArray<FMaterialPSOPrecacheRequestID> MaterialPSOPrecacheRequestIDs;
+
+	/** Graph event used to track all the PSO precache events */
+	FGraphEventRef PSOPrecacheCompileEvent;
+#endif
+
 	/**
 	 * Defines run-time groups of components. For example allows to assemble multiple parts of a building at runtime.
 	 * -1 means that component doesn't belong to any group.
@@ -774,7 +776,21 @@ public:
 	UPROPERTY(EditAnywhere, AdvancedDisplay, BlueprintReadOnly, Category=Rendering,  meta=(UIMin = "0", UIMax = "255", editcondition = "bRenderCustomDepth", DisplayName = "CustomDepth Stencil Value"))
 	int32 CustomDepthStencilValue;
 
+	/**
+	* Defines how quickly it should be culled. For example buildings should have a low priority, but small dressing should have a high priority.
+	*/
+	UPROPERTY(EditAnywhere, AdvancedDisplay, BlueprintReadOnly, Category = RayTracing)
+	ERayTracingGroupCullingPriority RayTracingGroupCullingPriority;
+
+	/** Mask used for stencil buffer writes. */
+	UPROPERTY(EditAnywhere, AdvancedDisplay, BlueprintReadOnly, Category = "Rendering", meta = (editcondition = "bRenderCustomDepth"))
+	ERendererStencilMask CustomDepthStencilWriteMask;
+
 private:
+	/** Which specific HLOD levels this component should be excluded from */
+	UPROPERTY(EditAnywhere, AdvancedDisplay, Category = HLOD, meta = (Bitmask, BitmaskEnum = "/Script/Engine.EHLODLevelExclusion", DisplayName = "Exclude from HLOD Levels", DisplayAfter = "bEnableAutoLODGeneration", EditConditionHides, EditCondition = "AllowHLODLevelsExclusion()"))
+	uint8 ExcludeFromHLODLevels;
+
 	/** Optional user defined default values for the custom primitive data of this primitive */
 	UPROPERTY(EditAnywhere, Category=Rendering, meta = (DisplayName = "Custom Primitive Data Defaults"))
 	FCustomPrimitiveData CustomPrimitiveData;
@@ -1389,16 +1405,6 @@ public:
 	UPROPERTY(BlueprintAssignable, Category="Input|Touch Input")
 	FComponentEndTouchOverSignature OnInputTouchLeave;
 
-	/**
-	 * Defines how quickly it should be culled. For example buildings should have a low priority, but small dressing should have a high priority.
-	 */
-	UPROPERTY(EditAnywhere, AdvancedDisplay, BlueprintReadOnly, Category = RayTracing)
-	ERayTracingGroupCullingPriority RayTracingGroupCullingPriority;
-
-	/** Mask used for stencil buffer writes. */
-	UPROPERTY(EditAnywhere, AdvancedDisplay, BlueprintReadOnly, Category = "Rendering", meta = (editcondition = "bRenderCustomDepth"))
-	ERendererStencilMask CustomDepthStencilWriteMask;
-
 	/** Scale the bounds of this object, used for frustum culling. Useful for features like WorldPositionOffset. */
 	UFUNCTION(BlueprintCallable, Category = "Rendering")
 	ENGINE_API void SetBoundsScale(float NewBoundsScale=1.f);
@@ -1984,10 +1990,6 @@ public:
 	FRenderCommandFence DetachFence;
 
 private:
-	/** Which specific HLOD levels this component should be excluded from */
-	UPROPERTY(EditAnywhere, AdvancedDisplay, Category = HLOD, meta = (Bitmask, BitmaskEnum = "/Script/Engine.EHLODLevelExclusion", DisplayName = "Exclude from HLOD Levels", DisplayAfter = "bEnableAutoLODGeneration", EditConditionHides, EditCondition = "AllowHLODLevelsExclusion()"))
-	uint8 ExcludeFromHLODLevels;
-
 	/** LOD parent primitive to draw instead of this one (multiple UPrim's will point to the same LODParent ) */
 	UPROPERTY(NonPIEDuplicateTransient)
 	TObjectPtr<class UPrimitiveComponent> LODParentPrimitive;
diff --git a/Engine/Source/Runtime/Engine/Classes/Components/SplineComponent.h b/Engine/Source/Runtime/Engine/Classes/Components/SplineComponent.h
--- a/Engine/Source/Runtime/Engine/Classes/Components/SplineComponent.h
+++ b/Engine/Source/Runtime/Engine/Classes/Components/SplineComponent.h
@@ -76,8 +76,10 @@ struct FSplineCurves
 	UPROPERTY()
 	FInterpCurveFloat ReparamTable;
 
+#if WITH_EDITORONLY_DATA
 	UPROPERTY()
 	TObjectPtr<USplineMetadata> Metadata_DEPRECATED = nullptr;
+#endif
 
 	UPROPERTY(transient)
 	uint32 Version = 0xffffffff;
@@ -198,6 +200,7 @@ class USplineComponent : public UPrimitiveComponent
 	UPROPERTY(EditAnywhere, Replicated, Category=Points)
 	FSplineCurves SplineCurves;
 
+#if WITH_EDITORONLY_DATA
 	/** Deprecated - please use GetSplinePointsPosition() to fetch this FInterpCurve */
 	UPROPERTY()
 	FInterpCurveVector SplineInfo_DEPRECATED;
@@ -215,6 +218,7 @@ class USplineComponent : public UPrimitiveComponent
 
 	UPROPERTY()
 	bool bAllowSplineEditingPerInstance_DEPRECATED;
+#endif
 
 	/** Number of steps per spline segment to place in the reparameterization table */
 	UPROPERTY(EditAnywhere, Replicated, AdvancedDisplay, Category = Spline, meta=(ClampMin=4, UIMin=4, ClampMax=100, UIMax=100))
diff --git a/Engine/Source/Runtime/Engine/Classes/Components/SplineMeshComponent.h b/Engine/Source/Runtime/Engine/Classes/Components/SplineMeshComponent.h
--- a/Engine/Source/Runtime/Engine/Classes/Components/SplineMeshComponent.h
+++ b/Engine/Source/Runtime/Engine/Classes/Components/SplineMeshComponent.h
@@ -55,6 +55,10 @@ struct FSplineMeshParams
 	UPROPERTY(EditAnywhere, Category=SplineMesh, AdvancedDisplay)
 	float StartRoll;
 
+	/** Roll around spline applied at end, in radians. */
+	UPROPERTY(EditAnywhere, Category=SplineMesh, AdvancedDisplay, meta=(DisplayAfter="EndTangent"))
+	float EndRoll;
+
 	/** Starting offset of the mesh from the spline, in component space. */
 	UPROPERTY(EditAnywhere, Category=SplineMesh, AdvancedDisplay)
 	FVector2D StartOffset;
@@ -71,10 +75,6 @@ struct FSplineMeshParams
 	UPROPERTY(EditAnywhere, Category = SplineMesh)
 	FVector EndTangent;
 
-	/** Roll around spline applied at end, in radians. */
-	UPROPERTY(EditAnywhere, Category=SplineMesh, AdvancedDisplay)
-	float EndRoll;
-
 	/** Ending offset of the mesh from the spline, in component space. */
 	UPROPERTY(EditAnywhere, Category=SplineMesh, AdvancedDisplay)
 	FVector2D EndOffset;
@@ -93,11 +93,11 @@ struct FSplineMeshParams
 		, StartTangent(ForceInit)
 		, StartScale(ForceInit)
 		, StartRoll(0)
+		, EndRoll(0)
 		, StartOffset(ForceInit)
 		, EndPos(ForceInit)
 		, EndScale(ForceInit)
 		, EndTangent(ForceInit)
-		, EndRoll(0)
 		, EndOffset(ForceInit)
 		, NaniteClusterBoundsScale(1.0f)
 	{
@@ -130,6 +130,13 @@ class USplineMeshComponent : public UStaticMeshComponent, public IInterface_Coll
 	UPROPERTY()
 	FGuid CachedMeshBodySetupGuid;
 
+	/** 
+	* The max draw distance to use in the main pass when also rendering to a runtime virtual texture. 
+	* This is only exposed to the user through the same setting on ULandscapeSplineSegment. 
+	*/
+	UPROPERTY()
+	float VirtualTextureMainPassMaxDrawDistance = 0.f;
+
 	// Navigation bounds can differ from primitive bounds since NavCollision can hold more geometry
 	FBox CachedNavigationBounds;
 
@@ -141,6 +148,10 @@ class USplineMeshComponent : public UStaticMeshComponent, public IInterface_Coll
 	UPROPERTY(EditAnywhere, Category = SplineMesh, AdvancedDisplay)
 	float SplineBoundaryMax;
 
+	/** Chooses the forward axis for the spline mesh orientation */
+	UPROPERTY(EditAnywhere, Category=SplineMesh)
+	TEnumAsByte<ESplineMeshAxis::Type> ForwardAxis;
+
 	/** If true, spline mesh properties - StartPos, EndPos, StartTangent and EndTangent- may be edited per instance in the level viewport. Otherwise, the spline mesh should be initialized in the construction script. */
 	UPROPERTY(EditDefaultsOnly, Category = Spline)
 	uint8 bAllowSplineEditingPerInstance:1;
@@ -153,17 +164,6 @@ class USplineMeshComponent : public UStaticMeshComponent, public IInterface_Coll
 	UPROPERTY(transient)
 	uint8 bMeshDirty : 1;
 
-	/** Chooses the forward axis for the spline mesh orientation */
-	UPROPERTY(EditAnywhere, Category=SplineMesh)
-	TEnumAsByte<ESplineMeshAxis::Type> ForwardAxis;
-
-	/** 
-	 * The max draw distance to use in the main pass when also rendering to a runtime virtual texture. 
-	 * This is only exposed to the user through the same setting on ULandscapeSplineSegment. 
-	 */
-	UPROPERTY()
-	float VirtualTextureMainPassMaxDrawDistance = 0.f;
-
 #if WITH_EDITORONLY_DATA
 	UPROPERTY(transient)
 	uint8 bSelected:1;
diff --git a/Engine/Source/Runtime/Engine/Classes/Components/StaticMeshComponent.h b/Engine/Source/Runtime/Engine/Classes/Components/StaticMeshComponent.h
--- a/Engine/Source/Runtime/Engine/Classes/Components/StaticMeshComponent.h
+++ b/Engine/Source/Runtime/Engine/Classes/Components/StaticMeshComponent.h
@@ -105,10 +105,6 @@ class UStaticMeshComponent : public UMeshComponent
 	UPROPERTY(EditAnywhere, AdvancedDisplay, BlueprintReadOnly, Category=LOD)
 	int32 ForcedLodModel;
 
-	/** LOD that was desired for rendering this StaticMeshComponent last frame. */
-	UPROPERTY(meta = (DeprecatedProperty, DeprecationMessage = "This property is deprecated and no longer supported."))
-	int32 PreviousLODLevel_DEPRECATED;
-
 	/** 
 	 * Specifies the smallest LOD that will be used for this component.  
 	 * This is ignored if ForcedLodModel is enabled.
@@ -120,8 +116,12 @@ class UStaticMeshComponent : public UMeshComponent
 	UPROPERTY()
 	int32 SubDivisionStepSize;
 
-	/** The static mesh that this component uses to render */
+	/** Wireframe color to use if bOverrideWireframeColor is true */
+	UPROPERTY(EditAnywhere, AdvancedDisplay, BlueprintReadOnly, Category=Rendering, meta=(editcondition = "bOverrideWireframeColor"))
+	FColor WireframeColorOverride;
+
 private:
+	/** The static mesh that this component uses to render */
 	UPROPERTY(EditAnywhere, BlueprintReadOnly, Category=StaticMesh, ReplicatedUsing=OnRep_StaticMesh, meta=(AllowPrivateAccess="true"))
 	TObjectPtr<class UStaticMesh> StaticMesh;
 
@@ -140,9 +140,11 @@ public:
 	UFUNCTION()
 	ENGINE_API void OnRep_StaticMesh(class UStaticMesh *OldStaticMesh);
 
-	/** Wireframe color to use if bOverrideWireframeColor is true */
-	UPROPERTY(EditAnywhere, AdvancedDisplay, BlueprintReadOnly, Category=Rendering, meta=(editcondition = "bOverrideWireframeColor"))
-	FColor WireframeColorOverride;
+	/**
+	* Distance at which to disable World Position Offset for an entire instance (0 = Never disable WPO).
+	**/
+	UPROPERTY(EditAnywhere, AdvancedDisplay, BlueprintReadOnly, Category = Rendering)
+	int32 WorldPositionOffsetDisableDistance = 0;
 
 	/** Forces this component to always use Nanite for masked materials, even if FNaniteSettings::bAllowMaskedMaterials=false */
 	UPROPERTY(EditAnywhere, AdvancedDisplay, BlueprintReadWrite, Category = Rendering)
@@ -177,12 +179,6 @@ public:
 	UPROPERTY(EditAnywhere, AdvancedDisplay, BlueprintReadWrite, Category = RayTracing)
 	uint8 bEvaluateWorldPositionOffsetInRayTracing : 1;
 
-	/**
-	 * Distance at which to disable World Position Offset for an entire instance (0 = Never disable WPO).
-	 **/
-	UPROPERTY(EditAnywhere, AdvancedDisplay, BlueprintReadOnly, Category = Rendering)
-	int32 WorldPositionOffsetDisableDistance = 0;
-
 protected:
 	/** Initial value of bEvaluateWorldPositionOffset when BeginPlay() was called. Can be useful if we want to reset to initial state. */
 	uint8 bInitialEvaluateWorldPositionOffset : 1;
diff --git a/Engine/Source/Runtime/Engine/Classes/PhysicsEngine/BodyInstance.h b/Engine/Source/Runtime/Engine/Classes/PhysicsEngine/BodyInstance.h
--- a/Engine/Source/Runtime/Engine/Classes/PhysicsEngine/BodyInstance.h
+++ b/Engine/Source/Runtime/Engine/Classes/PhysicsEngine/BodyInstance.h
@@ -295,6 +295,14 @@ struct FBodyInstance : public FBodyInstanceCore
 	/** When we are a body within a SkeletalMeshComponent, we cache the index of the bone we represent, to speed up sync'ing physics to anim. */
 	int16 InstanceBoneIndex;
 
+	/** [PhysX Only] This physics body's solver iteration count for position. Increasing this will be more CPU intensive, but better stabilized.  */
+	UPROPERTY(EditAnywhere, AdvancedDisplay, BlueprintReadOnly, Category=Physics)
+	uint8 PositionSolverIterationCount;
+
+	/** [PhysX Only] This physics body's solver iteration count for velocity. Increasing this will be more CPU intensive, but better stabilized. */
+	UPROPERTY(EditAnywhere, AdvancedDisplay, BlueprintReadOnly, Category = Physics)
+	uint8 VelocitySolverIterationCount;
+
 private:
 	/** Enum indicating what type of object this should be considered as when it moves */
 	UPROPERTY(EditAnywhere, Category=Custom)
@@ -316,14 +324,14 @@ private:
 	TEnumAsByte<ECollisionEnabled::Type> CollisionEnabled;
 
 	/** When per-shape collision is changed at runtime, state is stored in an optional array of per-shape collision state.
-	*	Before this array's IsSet state is true, collision values from the BodySetup's AggGeom are used in GetShapeCollisionEnabled.
+	*	Before bShapeCollisionEnabledIsSet state is true, collision values from the BodySetup's AggGeom are used in GetShapeCollisionEnabled.
 	*/
-	TOptional<TArray<TEnumAsByte<ECollisionEnabled::Type>>> ShapeCollisionEnabled;
+	TArray<TEnumAsByte<ECollisionEnabled::Type>> ShapeCollisionEnabled;
 
 	/** When per-shape collision responses are changed at runtime, state is stored in an optional array of per-shape
-	*	collision response settings. If this is not set, the base body instance's CollisionResponses member is used for all shapes.
+	*	collision response settings. If bShapeCollisionResponsesIsSet is false, the base body instance's CollisionResponses member is used for all shapes.
 	*/
-	TOptional<TArray<TPair<int32, FCollisionResponse>>> ShapeCollisionResponses;
+	TArray<TPair<int32, FCollisionResponse>> ShapeCollisionResponses;
 
 public:
 	// Current state of the physics body for tracking deferred addition and removal.
@@ -342,6 +350,10 @@ public:
 	uint8 bUseCCD : 1;
 
 private:
+
+	uint8 bShapeCollisionEnabledIsSet : 1;
+	uint8 bShapeCollisionResponsesIsSet : 1;
+
 	/** [EXPERIMENTAL] If true Motion-Aware Collision Detection (MACD) will be used for this component */
 	UPROPERTY(EditAnywhere, AdvancedDisplay, Category = Collision)
 	uint8 bUseMACD : 1;
@@ -475,6 +487,10 @@ public:
 private:
 	void UpdateSolverAsyncDeltaTime();
 
+	/** Collision Profile Name **/
+	UPROPERTY(EditAnywhere, Category=Custom)
+	FName CollisionProfileName;
+
 public:
 	/** Current scale of physics - used to know when and how physics must be rescaled to match current transform of OwnerComponent. */
 	FVector Scale3D;
@@ -488,22 +504,6 @@ public:
 	struct FCollisionResponseContainer ResponseToChannels_DEPRECATED;
 #endif // WITH_EDITORONLY_DATA
 
-private:
-
-	/** Collision Profile Name **/
-	UPROPERTY(EditAnywhere, Category=Custom)
-	FName CollisionProfileName;
-
-public:
-
-	/** [PhysX Only] This physics body's solver iteration count for position. Increasing this will be more CPU intensive, but better stabilized.  */
-	UPROPERTY(EditAnywhere, AdvancedDisplay, BlueprintReadOnly, Category=Physics)
-	uint8 PositionSolverIterationCount;
-
-	/** [PhysX Only] This physics body's solver iteration count for velocity. Increasing this will be more CPU intensive, but better stabilized. */
-	UPROPERTY(EditAnywhere, AdvancedDisplay, BlueprintReadOnly, Category = Physics)
-	uint8 VelocitySolverIterationCount;
-
 private:
 	/** Custom Channels for Responses*/
 	UPROPERTY(EditAnywhere, Category = Custom)
diff --git a/Engine/Source/Runtime/Engine/Private/Components/SplineComponent.cpp b/Engine/Source/Runtime/Engine/Private/Components/SplineComponent.cpp
--- a/Engine/Source/Runtime/Engine/Private/Components/SplineComponent.cpp
+++ b/Engine/Source/Runtime/Engine/Private/Components/SplineComponent.cpp
@@ -36,7 +36,9 @@ USplineMetadata::USplineMetadata(const FObjectInitializer& ObjectInitializer)
 
 USplineComponent::USplineComponent(const FObjectInitializer& ObjectInitializer)
 	: Super(ObjectInitializer)
+#if WITH_EDITORONLY_DATA
 	, bAllowSplineEditingPerInstance_DEPRECATED(true)
+#endif
 	, ReparamStepsPerSegment(10)
 	, Duration(1.0f)
 	, bStationaryEndpoints(false)
@@ -67,18 +69,19 @@ USplineComponent::USplineComponent(const FObjectInitializer& ObjectInitializer)
 
 	UpdateSpline();
 
+#if WITH_EDITORONLY_DATA
 	// Set these deprecated values up so that old assets with default values load correctly (and are subsequently upgraded during Serialize)
 	SplineInfo_DEPRECATED = SplineCurves.Position;
 	SplineRotInfo_DEPRECATED = SplineCurves.Rotation;
 	SplineScaleInfo_DEPRECATED = SplineCurves.Scale;
 	SplineReparamTable_DEPRECATED = SplineCurves.ReparamTable;
+#endif
 }
 
 void USplineComponent::ResetToDefault()
 {
 	SetDefaultSpline();
 
-	bAllowSplineEditingPerInstance_DEPRECATED = true;
 	ReparamStepsPerSegment = 10;
 	Duration = 1.0f;
 	bStationaryEndpoints = false;
@@ -89,6 +92,7 @@ void USplineComponent::ResetToDefault()
 	bClosedLoop = false;
 	DefaultUpVector = FVector::UpVector;
 #if WITH_EDITORONLY_DATA
+	bAllowSplineEditingPerInstance_DEPRECATED = true;
 	EditorUnselectedSplineSegmentColor = FStyleColors::White.GetSpecifiedColor();
 	EditorSelectedSplineSegmentColor = FStyleColors::AccentOrange.GetSpecifiedColor();
 	EditorTangentColor = FLinearColor(0.718f, 0.589f, 0.921f);
@@ -157,6 +161,7 @@ void USplineComponent::Serialize(FArchive& Ar)
 
 	Ar.UsingCustomVersion(FEditorObjectVersion::GUID);
 
+#if WITH_EDITORONLY_DATA
 	// Move points to new properties
 	if (Ar.CustomVer(FEditorObjectVersion::GUID) < FEditorObjectVersion::SplineComponentCurvesInStruct &&
 		Ar.IsLoading())
@@ -166,6 +171,7 @@ void USplineComponent::Serialize(FArchive& Ar)
 		SplineCurves.Scale = SplineScaleInfo_DEPRECATED;
 		SplineCurves.ReparamTable = SplineReparamTable_DEPRECATED;
 	}
+#endif
 
 	// Support old resources which don't have the rotation and scale splines present
 	const FPackageFileVersion ArchiveUEVersion = Ar.UEVer();
diff --git a/Engine/Source/Runtime/Engine/Private/InstancedStaticMesh.cpp b/Engine/Source/Runtime/Engine/Private/InstancedStaticMesh.cpp
--- a/Engine/Source/Runtime/Engine/Private/InstancedStaticMesh.cpp
+++ b/Engine/Source/Runtime/Engine/Private/InstancedStaticMesh.cpp
@@ -2057,6 +2057,7 @@ UInstancedStaticMeshComponent::UInstancedStaticMeshComponent(const FObjectInitia
 
 UInstancedStaticMeshComponent::UInstancedStaticMeshComponent(FVTableHelper& Helper)
 	: Super(Helper)
+	, bIsInstanceDataApplyCompleted(true)
 	, PrimitiveInstanceDataManager(this)
 {
 }
diff --git a/Engine/Source/Runtime/Engine/Private/InstancedStaticMesh/ISMInstanceDataManager.cpp b/Engine/Source/Runtime/Engine/Private/InstancedStaticMesh/ISMInstanceDataManager.cpp
--- a/Engine/Source/Runtime/Engine/Private/InstancedStaticMesh/ISMInstanceDataManager.cpp
+++ b/Engine/Source/Runtime/Engine/Private/InstancedStaticMesh/ISMInstanceDataManager.cpp
@@ -61,7 +61,18 @@ FArchive& operator<<(FArchive& Ar, FInstanceSceneDataBuffers::FCompressedSpatial
 }
 
 FPrimitiveInstanceDataManager::FPrimitiveInstanceDataManager(UPrimitiveComponent* InPrimitiveComponent) 
-	: PrimitiveComponent(InPrimitiveComponent) 
+	: bNumCustomDataChanged(false)
+	, bBakedLightingDataChanged(false)
+	, bTransformChangedAllInstances(false)
+#if WITH_EDITOR
+	, bAnyEditorDataChanged(false)
+#endif
+	, bPrimitiveTransformChanged(false)
+	, bAnyInstanceChange(false)
+	, bComponentMarkedDirty(false)
+	, bEnableTracking(false)
+	, bFirstFlush(true)
+	, PrimitiveComponent(InPrimitiveComponent) 
 {
 	// Don't do anything if this is not a "real" ISM being tracked (this logic shopuld move out).
 	if (PrimitiveComponent.IsValid() && PrimitiveComponent->HasAnyFlags(RF_ClassDefaultObject | RF_ArchetypeObject))
diff --git a/Engine/Source/Runtime/Engine/Private/PhysicsEngine/BodyInstance.cpp b/Engine/Source/Runtime/Engine/Private/PhysicsEngine/BodyInstance.cpp
--- a/Engine/Source/Runtime/Engine/Private/PhysicsEngine/BodyInstance.cpp
+++ b/Engine/Source/Runtime/Engine/Private/PhysicsEngine/BodyInstance.cpp
@@ -329,6 +329,8 @@ bool FCollisionResponse::operator==(const FCollisionResponse& Other) const
 FBodyInstance::FBodyInstance()
 	: InstanceBodyIndex(INDEX_NONE)
 	, InstanceBoneIndex(INDEX_NONE)
+	, PositionSolverIterationCount(8)
+	, VelocitySolverIterationCount(1)
 	, ObjectType(ECC_WorldStatic)
 	, MaskFilter(0)
 	, CollisionEnabled(ECollisionEnabled::QueryAndPhysics)
@@ -336,6 +338,8 @@ FBodyInstance::FBodyInstance()
 	, SleepFamily(ESleepFamily::Normal)
 	, DOFMode(0)
 	, bUseCCD(false)
+	, bShapeCollisionEnabledIsSet(false)
+	, bShapeCollisionResponsesIsSet(false)
 	, bUseMACD(false)
 	, bIgnoreAnalyticCollisions(false)
 	, bNotifyRigidBodyCollision(false)
@@ -358,10 +362,8 @@ FBodyInstance::FBodyInstance()
 	, bOneWayInteraction(false)
 	, bOverrideSolverAsyncDeltaTime(false)
 	, SolverAsyncDeltaTime(1.f / 60)
-	, Scale3D(1.0f)
 	, CollisionProfileName(UCollisionProfile::CustomCollisionProfileName)
-	, PositionSolverIterationCount(8)
-	, VelocitySolverIterationCount(1)
+	, Scale3D(1.0f)
 	, MaxDepenetrationVelocity(0.f)
 	, MassInKgOverride(100.f)
 	, ExternalCollisionProfileBodySetup(nullptr)
@@ -588,19 +590,18 @@ bool FBodyInstance::SetResponseToChannels(const FCollisionResponseContainer& New
 
 bool FBodyInstance::SetShapeResponseToChannels(const int32 ShapeIndex, const FCollisionResponseContainer& NewResponses)
 {
-	if (!ShapeCollisionResponses.IsSet())
+	if (!bShapeCollisionResponsesIsSet)
 	{
-		ShapeCollisionResponses = TArray<TPair<int32, FCollisionResponse>>();
+		ShapeCollisionResponses.Reset();
+		bShapeCollisionResponsesIsSet = true;
 	}
 
-	TArray<TPair<int32, FCollisionResponse>>& ShapeCollisionResponsesValue = ShapeCollisionResponses.GetValue();
-
 	bool bIndexExists = false;
 	int32 ResponseIndex = 0;
-	const int32 ResponseNum = ShapeCollisionResponsesValue.Num();
+	const int32 ResponseNum = ShapeCollisionResponses.Num();
 	for (; ResponseIndex < ResponseNum; ++ResponseIndex)
 	{
-		if (ShapeCollisionResponsesValue[ResponseIndex].Key == ShapeIndex)
+		if (ShapeCollisionResponses[ResponseIndex].Key == ShapeIndex)
 		{
 			break;
 		}
@@ -608,10 +609,10 @@ bool FBodyInstance::SetShapeResponseToChannels(const int32 ShapeIndex, const FCo
 
 	if (ResponseIndex == ResponseNum)
 	{
-		ShapeCollisionResponsesValue.Add(TPair<int32, FCollisionResponse>(ShapeIndex, FCollisionResponse()));
+		ShapeCollisionResponses.Add(TPair<int32, FCollisionResponse>(ShapeIndex, FCollisionResponse()));
 	}
 
-	if (ShapeCollisionResponsesValue[ResponseIndex].Value.SetCollisionResponseContainer(NewResponses))
+	if (ShapeCollisionResponses[ResponseIndex].Value.SetCollisionResponseContainer(NewResponses))
 	{
 		UpdatePhysicsFilterData();
 		return true;
@@ -628,13 +629,13 @@ const FCollisionResponseContainer& FBodyInstance::GetShapeResponseToChannels(con
 const FCollisionResponseContainer& FBodyInstance::GetShapeResponseToChannels(const int32 ShapeIndex, const FCollisionResponseContainer& DefaultResponseContainer) const
 {
 	// Return per-shape collision response override if there is one
-	if (ShapeCollisionResponses.IsSet())
+	if (bShapeCollisionResponsesIsSet)
 	{
-		for (int32 ResponseIndex = 0; ResponseIndex < ShapeCollisionResponses.GetValue().Num(); ++ResponseIndex)
+		for (int32 ResponseIndex = 0; ResponseIndex < ShapeCollisionResponses.Num(); ++ResponseIndex)
 		{
-			if (ShapeCollisionResponses.GetValue()[ResponseIndex].Key == ShapeIndex)
+			if (ShapeCollisionResponses[ResponseIndex].Key == ShapeIndex)
 			{
-				return ShapeCollisionResponses.GetValue()[ResponseIndex].Value.GetResponseContainer();
+				return ShapeCollisionResponses[ResponseIndex].Value.GetResponseContainer();
 			}
 		}
 	}
@@ -749,17 +750,18 @@ void FBodyInstance::SetShapeCollisionEnabled(const int32 ShapeIndex, ECollisionE
 		if (OldType != NewType)
 		{
 			// If ShapeCollisionEnabled wasn't set up yet, copy values from BodySetup into it
-			if (!ShapeCollisionEnabled.IsSet())
+			if (!bShapeCollisionEnabledIsSet)
 			{
 				const int32 ShapeCount = GetBodySetup()->AggGeom.GetElementCount();
-				ShapeCollisionEnabled = TArray<TEnumAsByte<ECollisionEnabled::Type>>();
-				ShapeCollisionEnabled.GetValue().SetNum(ShapeCount);
+				ShapeCollisionEnabled.Reset();
+				bShapeCollisionEnabledIsSet = true;
+				ShapeCollisionEnabled.SetNum(ShapeCount);
 				for (int32 OptionalShapeIndex = 0; OptionalShapeIndex < ShapeCount; ++OptionalShapeIndex)
 				{
-					ShapeCollisionEnabled.GetValue()[OptionalShapeIndex] = GetBodySetup()->AggGeom.GetElement(OptionalShapeIndex)->GetCollisionEnabled();
+					ShapeCollisionEnabled[OptionalShapeIndex] = GetBodySetup()->AggGeom.GetElement(OptionalShapeIndex)->GetCollisionEnabled();
 				}
 			}
-			ShapeCollisionEnabled.GetValue()[ShapeIndex] = NewType;
+			ShapeCollisionEnabled[ShapeIndex] = NewType;
 
 			if (bUpdatePhysicsFilterData)
 			{
@@ -918,11 +920,11 @@ ECollisionEnabled::Type FBodyInstance::GetShapeCollisionEnabled(const int32 Shap
 {
 	// If any runtime shape collision overrides have been set, return that.
 	// Otherwise, get it from the bodysetup.
-	if (ShapeCollisionEnabled.IsSet())
+	if (bShapeCollisionEnabledIsSet)
 	{
-		if (ensure(ShapeCollisionEnabled.GetValue().IsValidIndex(ShapeIndex)))
+		if (ensure(ShapeCollisionEnabled.IsValidIndex(ShapeIndex)))
 		{
-			return ShapeCollisionEnabled.GetValue()[ShapeIndex];
+			return ShapeCollisionEnabled[ShapeIndex];
 		}
 	}
 
diff --git a/Engine/Source/Runtime/Engine/Public/InstancedStaticMesh/ISMInstanceDataManager.h b/Engine/Source/Runtime/Engine/Public/InstancedStaticMesh/ISMInstanceDataManager.h
--- a/Engine/Source/Runtime/Engine/Public/InstancedStaticMesh/ISMInstanceDataManager.h
+++ b/Engine/Source/Runtime/Engine/Public/InstancedStaticMesh/ISMInstanceDataManager.h
@@ -49,7 +49,7 @@ public:
 	/**
 	 * Current tracking state, 
 	 */
-	enum class ETrackingState
+	enum class ETrackingState : uint8
 	{
 		Initial, // In the initial state, there is no proxy and therefore changes do not need to be tracked, e.g., during initial setup of an ISM component.
 		Tracked,
@@ -57,7 +57,7 @@ public:
 		Optimized, // In the optimized state there's no need to track any delta changes, but if anything changes at all we must rebuild.
 	};
 
-	enum class EMode
+	enum class EMode : uint8
 	{
 		Default,
 		Legacy, // In this mode, we create a legacy supporting proxy.
@@ -206,31 +206,33 @@ private:
 	 */
 	static FISMPrecomputedSpatialHashData PrecomputeOptimizationData(FInstanceUpdateComponentDesc &&ComponentData, TStridedView<FMatrix> InstanceTransforms);
 #endif
-	EMode Mode = EMode::Default;
-	ETrackingState TrackingState = ETrackingState::Initial;
+
+	// Change set.
+	FInstanceAttributeTracker InstanceUpdateTracker;
 
 	// Id allocation tracking
 	TBitArray<> ValidInstanceIdMask;
 	int32 IdSearchStartIndex = 0;
 
-	// Change set.
-	FInstanceAttributeTracker InstanceUpdateTracker;
+	EMode Mode = EMode::Default;
+	ETrackingState TrackingState = ETrackingState::Initial;
 
-	bool bNumCustomDataChanged = false;
-	bool bBakedLightingDataChanged = false;
-	bool bTransformChangedAllInstances = false;
+	uint8 bNumCustomDataChanged : 1;
+	uint8 bBakedLightingDataChanged : 1;
+	uint8 bTransformChangedAllInstances : 1;
 #if WITH_EDITOR
-	bool bAnyEditorDataChanged = false;
+	uint8 bAnyEditorDataChanged : 1;
 #endif	
-	bool bPrimitiveTransformChanged = false;
-	bool bAnyInstanceChange = false;
+	uint8 bPrimitiveTransformChanged : 1;
+	uint8 bAnyInstanceChange : 1;
+
+	uint8 bComponentMarkedDirty : 1;
+	uint8 bEnableTracking : 1;
+	uint8 bFirstFlush : 1;
 
 	TSharedPtr<FISMCInstanceDataSceneProxy, ESPMode::ThreadSafe> Proxy;
 	TWeakObjectPtr<UPrimitiveComponent> PrimitiveComponent = nullptr;
 
-	bool bComponentMarkedDirty = false;
-	bool bEnableTracking = false;
-
 	TPimplPtr<struct FLegacyBuildData> LegacyBuildData;
 
 	using PrecomputedOptimizationDataPtr = FISMCInstanceDataSceneProxy::FISMPrecomputedSpatialHashDataPtr;
@@ -244,5 +246,4 @@ private:
 	int32 NumCustomDataFloats = 0;
 	float AbsMaxDisplacement = 0.0f;
 	FRenderBounds StaticMeshBounds;
-	bool bFirstFlush = true;
 };
diff --git a/Engine/Source/Runtime/PhysicsCore/Public/BodyInstanceCore.h b/Engine/Source/Runtime/PhysicsCore/Public/BodyInstanceCore.h
--- a/Engine/Source/Runtime/PhysicsCore/Public/BodyInstanceCore.h
+++ b/Engine/Source/Runtime/PhysicsCore/Public/BodyInstanceCore.h
@@ -20,7 +20,6 @@ struct FBodyInstanceCore
 	TWeakObjectPtr<UBodySetupCore> BodySetup;
 
 	PHYSICSCORE_API FBodyInstanceCore();
-	virtual ~FBodyInstanceCore() {}
 
 	/** 
 	 * If true, this body will use simulation. If false, will be 'fixed' (ie kinematic) and move where it is told. 
