Add missing debug name to the animation compression derived data task

It's missing a debug name, which means that out of memory issues don't show what task is taking up memory.

#rnx

diff --git a/Engine/Source/Runtime/Engine/Private/Animation/AnimationCompressionDerivedData.cpp b/Engine/Source/Runtime/Engine/Private/Animation/AnimationCompressionDerivedData.cpp
--- a/Engine/Source/Runtime/Engine/Private/Animation/AnimationCompressionDerivedData.cpp
+++ b/Engine/Source/Runtime/Engine/Private/Animation/AnimationCompressionDerivedData.cpp
@@ -111,7 +111,7 @@ void FAnimationSequenceAsyncCacheTask::BeginCache(const FIoHash& KeyHash)
 
 		check(BuildTask == nullptr);
 		BuildTask = MakeUnique<FAnimationSequenceAsyncBuildTask>(this, KeyHash);
-		BuildTask->StartBackgroundTask(ThreadPool, BasePriority, EQueuedWorkFlags::DoNotRunInsideBusyWait, RequiredMemory);
+		BuildTask->StartBackgroundTask(ThreadPool, BasePriority, EQueuedWorkFlags::DoNotRunInsideBusyWait, RequiredMemory, TEXT("AnimationCompression"));
 	}
 }
 	
