Add CVar to disable Niagara and Cascade graphics PSO precaching.

This disables asset-level precaching, and by extension component-level precaching too for Niagara and Cascade components which rely on the asset-level PSO precaching events.

Note this only disables graphics PSO precaching and does not affect compute PSOs.

diff --git a/Engine/Plugins/FX/Niagara/Source/Niagara/Private/NiagaraSystem.cpp b/Engine/Plugins/FX/Niagara/Source/Niagara/Private/NiagaraSystem.cpp
--- a/Engine/Plugins/FX/Niagara/Source/Niagara/Private/NiagaraSystem.cpp
+++ b/Engine/Plugins/FX/Niagara/Source/Niagara/Private/NiagaraSystem.cpp
@@ -1347,7 +1347,8 @@ void UNiagaraSystem::DeclareConstructClasses(TArray<FTopLevelAssetPath>& OutCons
 
 void UNiagaraSystem::PrecachePSOs()
 {
-	if (!IsComponentPSOPrecachingEnabled() && !IsResourcePSOPrecachingEnabled())
+	// Only precache if asset precaching is enabled and at least one of component or resource precaching is enabled.
+	if (!IsAssetPSOPrecachingEnabled() || (!IsComponentPSOPrecachingEnabled() && !IsResourcePSOPrecachingEnabled()))
 	{
 		return;
 	}
diff --git a/Engine/Source/Runtime/Engine/Private/Components/PrimitiveComponent.cpp b/Engine/Source/Runtime/Engine/Private/Components/PrimitiveComponent.cpp
--- a/Engine/Source/Runtime/Engine/Private/Components/PrimitiveComponent.cpp
+++ b/Engine/Source/Runtime/Engine/Private/Components/PrimitiveComponent.cpp
@@ -4854,7 +4854,7 @@ bool UPrimitiveComponent::ShouldRenderProxyFallbackToDefaultMaterial() const
 bool UPrimitiveComponent::CheckPSOPrecachingAndBoostPriority(EPSOPrecachePriority NewPSOPrecachePriority)
 {
 #if UE_WITH_PSO_PRECACHING
-	ensure(!IsComponentPSOPrecachingEnabled() || bPSOPrecacheCalled);
+	ensure(!IsAssetPSOPrecachingEnabled() || !IsComponentPSOPrecachingEnabled() || bPSOPrecacheCalled);
 	check(NewPSOPrecachePriority == EPSOPrecachePriority::High || NewPSOPrecachePriority == EPSOPrecachePriority::Highest);
 
 	if (PSOPrecacheCompileEvent && !PSOPrecacheCompileEvent->IsComplete())
diff --git a/Engine/Source/Runtime/Engine/Private/PSOPrecache.cpp b/Engine/Source/Runtime/Engine/Private/PSOPrecache.cpp
--- a/Engine/Source/Runtime/Engine/Private/PSOPrecache.cpp
+++ b/Engine/Source/Runtime/Engine/Private/PSOPrecache.cpp
@@ -31,6 +31,14 @@ static FAutoConsoleVariableRef CVarPSOPrecacheResources(
 	ECVF_ReadOnly
 );
 
+int32 GPSOPrecacheAssets = 1;
+static FAutoConsoleVariableRef CVarPSOPrecacheAssets(
+	TEXT("r.PSOPrecache.Assets"),
+	GPSOPrecacheAssets,
+	TEXT("Precache all possible used PSOs by assets during Postload (default 1 if PSOPrecaching is enabled). This affects Niagara and Cascade assets."),
+	ECVF_ReadOnly
+);
+
 int32 GPSOProxyCreationWhenPSOReady = 1;
 static FAutoConsoleVariableRef CVarPSOProxyCreationWhenPSOReady(
 	TEXT("r.PSOPrecache.ProxyCreationWhenPSOReady"),
@@ -61,6 +69,11 @@ static FAutoConsoleVariableRef CVarPSOComponentBoostStrategy(
 	ECVF_ReadOnly
 );
 
+bool IsAssetPSOPrecachingEnabled()
+{
+	return FApp::CanEverRender() && PipelineStateCache::IsPSOPrecachingEnabled() && GPSOPrecacheAssets && !GIsEditor;
+}
+
 bool IsComponentPSOPrecachingEnabled()
 {
 	return FApp::CanEverRender() && PipelineStateCache::IsPSOPrecachingEnabled() && GPSOPrecacheComponents && !GIsEditor;
diff --git a/Engine/Source/Runtime/Engine/Private/Particles/ParticleComponents.cpp b/Engine/Source/Runtime/Engine/Private/Particles/ParticleComponents.cpp
--- a/Engine/Source/Runtime/Engine/Private/Particles/ParticleComponents.cpp
+++ b/Engine/Source/Runtime/Engine/Private/Particles/ParticleComponents.cpp
@@ -2687,7 +2687,8 @@ void UParticleSystem::PostLoad()
 
 void UParticleSystem::PrecachePSOs()
 {
-	if (!IsComponentPSOPrecachingEnabled() && !IsResourcePSOPrecachingEnabled())
+	// Only precache if asset precaching is enabled and at least one of component or resource precaching is enabled.
+	if (!IsAssetPSOPrecachingEnabled() || (!IsComponentPSOPrecachingEnabled() && !IsResourcePSOPrecachingEnabled()))
 	{
 		return;
 	}
@@ -3540,7 +3541,7 @@ bool UFXSystemComponent::RequiresLWCTileRecache(const FVector3f CurrentTile, con
 void UFXSystemComponent::PrecacheAssetPSOs(UFXSystemAsset* FXSystemAsset)
 {
 #if UE_WITH_PSO_PRECACHING
-	if (!FApp::CanEverRender() || !IsComponentPSOPrecachingEnabled() || FXSystemAsset == nullptr)
+	if (!FApp::CanEverRender() || !IsAssetPSOPrecachingEnabled() || FXSystemAsset == nullptr)
 	{
 		return;
 	}
diff --git a/Engine/Source/Runtime/Engine/Public/PSOPrecache.h b/Engine/Source/Runtime/Engine/Public/PSOPrecache.h
--- a/Engine/Source/Runtime/Engine/Public/PSOPrecache.h
+++ b/Engine/Source/Runtime/Engine/Public/PSOPrecache.h
@@ -248,6 +248,11 @@ struct FMaterialPSOPrecacheParams
 	}
 };
 
+/**
+ * Precaching PSOs for assets?
+ */
+extern ENGINE_API bool IsAssetPSOPrecachingEnabled();
+
 /**
  * Precaching PSOs for components?
  */
