Fix missing Material Parameter Groups

- The underlying data still existed, but was removed from the property window (UI) so wasn't exposed to artists.
- This broke in the change that converted material editor only data into a separate sub object.
- When the sub object was moved it lost all its property specifiers that made it visible in the property window.
- Since this is a sub object there was special markup we needed to make it edit-inline and display correctly.

- Added the following property specifiers to make this work correctly:

**VisibleAnywhere**:  Indicates that this property is visible in all property windows, but cannot be edited. This Specifier is incompatible with the "Edit" Specifiers.
**Instanced**: Object ( UCLASS ) properties only. When an instance of this class is created, it will be given a unique copy of the Object assigned to this property in defaults. Used for instancing subobjects defined in class default properties. Implies EditInline and Export.
**AllowEditInlineCustomization**: Used for edit inline UObject* properties, enables detail class customization for instanced UObject properties.
**ShowInnerProperties**: Since this is a object pointer property, we want to show the inner properties in the property window (instead of a pointer to the UObject).

- Fixed a bug with the parameter groups which was adding a separate group for _each_ parameter that had that group name.  So if you had 10 parameters in the _Inputs_ group, you would end up with 10 groups in the array all named _Inputs_.  Now we add unique groups to the array.

#tests Duplicating materials, adding parameter groups, renaming groups, sorting groups

diff --git a/Engine/Source/Editor/MaterialEditor/Private/MaterialEditor.cpp b/Engine/Source/Editor/MaterialEditor/Private/MaterialEditor.cpp
--- a/Engine/Source/Editor/MaterialEditor/Private/MaterialEditor.cpp
+++ b/Engine/Source/Editor/MaterialEditor/Private/MaterialEditor.cpp
@@ -1123,7 +1123,7 @@ void FMaterialEditor::GetAllMaterialExpressionGroups(TArray<FString>* OutGroups)
 				{
 					return GroupName == DataElement.GroupName;
 				});
-				UpdatedGroups.Add(FParameterGroupData(GroupName, ParameterGroupDataElement->GroupSortPriority));
+				UpdatedGroups.AddUnique(FParameterGroupData(GroupName, ParameterGroupDataElement->GroupSortPriority));
 			}
 		}
 	}
diff --git a/Engine/Source/Runtime/Engine/Public/Materials/Material.h b/Engine/Source/Runtime/Engine/Public/Materials/Material.h
--- a/Engine/Source/Runtime/Engine/Public/Materials/Material.h
+++ b/Engine/Source/Runtime/Engine/Public/Materials/Material.h
@@ -296,6 +296,11 @@ struct FParameterGroupData
 		GroupSortPriority(InSortPriority)
 	{
 	}
+
+	bool operator==(const FParameterGroupData& Other) const
+	{
+		return GroupName == Other.GroupName && GroupSortPriority == Other.GroupSortPriority;
+	}
 };
 
 UCLASS(MinimalAPI, Optional)
diff --git a/Engine/Source/Runtime/Engine/Public/Materials/MaterialInterface.h b/Engine/Source/Runtime/Engine/Public/Materials/MaterialInterface.h
--- a/Engine/Source/Runtime/Engine/Public/Materials/MaterialInterface.h
+++ b/Engine/Source/Runtime/Engine/Public/Materials/MaterialInterface.h
@@ -258,7 +258,7 @@ class UMaterialInterface : public UObject, public IBlendableInterface, public II
 protected:
 	friend class UMaterialInterfaceEditorOnlyData;
 
-	UPROPERTY()
+	UPROPERTY(VisibleAnywhere, Instanced, Category = Material, meta = (AllowEditInlineCustomization, ShowInnerProperties))
 	TObjectPtr<UMaterialInterfaceEditorOnlyData> EditorOnlyData;
 
 	ENGINE_API virtual const UClass* GetEditorOnlyDataClass() const;
