Replace asset-loading time PSO precaching CVar with two CVars for Niagara and Cascade that allow finer control.

r.PSOPrecache.NiagaraPrecachePSOAtAssetLoadingTime: 0 to disable asset-loading time precaching, 1 to enable (default).
r.PSOPrecache.CascadePrecachingTime: 0 to disable precaching fully, 1 to enable asset-loading time precaching (default), 2 to enable component-loading time precaching, 3 to enable render proxy creation time precaching.

diff --git a/Engine/Plugins/FX/Niagara/Source/Niagara/Private/NiagaraSystem.cpp b/Engine/Plugins/FX/Niagara/Source/Niagara/Private/NiagaraSystem.cpp
--- a/Engine/Plugins/FX/Niagara/Source/Niagara/Private/NiagaraSystem.cpp
+++ b/Engine/Plugins/FX/Niagara/Source/Niagara/Private/NiagaraSystem.cpp
@@ -139,6 +139,13 @@ static FAutoConsoleVariableRef CVarNiagaraCompileWaitLoggingTerminationCap(
 	ECVF_Default
 );
 
+int32 GNiagaraPrecachePSOAtAssetLoadingTime = 1;
+FAutoConsoleVariableRef CVarNiagaraPrecachePSOAtAssetLoadingTime(
+	TEXT("r.PSOPrecache.NiagaraPrecachePSOAtAssetLoadingTime"),
+	GNiagaraPrecachePSOAtAssetLoadingTime,
+	TEXT("Controls whether Niagara systems start PSO precaching at asset loading time (1 = default) or only when activated (0)."),
+	ECVF_Default);
+
 #if WITH_EDITORONLY_DATA
 static int GNiagaraOnDemandCompileEnabled = 1;
 static FAutoConsoleVariableRef CVarNiagaraOnDemandCompileEnabled(
@@ -1291,7 +1298,10 @@ void UNiagaraSystem::PostLoad()
 	}
 #endif // WITH_EDITORONLY_DATA
 
-	PrecachePSOs();
+	if (GNiagaraPrecachePSOAtAssetLoadingTime)
+	{
+		PrecachePSOs();
+	}
 }
 
 void UNiagaraSystem::PostDuplicate(bool bDuplicateForPIE)
@@ -1347,8 +1357,7 @@ void UNiagaraSystem::DeclareConstructClasses(TArray<FTopLevelAssetPath>& OutCons
 
 void UNiagaraSystem::PrecachePSOs()
 {
-	// Only precache if asset precaching is enabled and at least one of component or resource precaching is enabled.
-	if (!IsAssetPSOPrecachingEnabled() || (!IsComponentPSOPrecachingEnabled() && !IsResourcePSOPrecachingEnabled()))
+	if (HasLaunchedPSOPrecaching() || (!IsComponentPSOPrecachingEnabled() && !IsResourcePSOPrecachingEnabled()))
 	{
 		return;
 	}
diff --git a/Engine/Source/Runtime/Engine/Classes/Particles/ParticleSystem.h b/Engine/Source/Runtime/Engine/Classes/Particles/ParticleSystem.h
--- a/Engine/Source/Runtime/Engine/Classes/Particles/ParticleSystem.h
+++ b/Engine/Source/Runtime/Engine/Classes/Particles/ParticleSystem.h
@@ -137,6 +137,7 @@ public:
 	mutable FParticlePerfStats* ParticlePerfStats = nullptr;
 #endif
 
+	const bool HasLaunchedPSOPrecaching() const { return PSOPrecachingLaunched; }
 	const FGraphEventRef& GetPrecachePSOsEvent() const { return PrecachePSOsEvent; }
 	const TArray<FMaterialPSOPrecacheRequestID>& GetMaterialPSOPrecacheRequestIDs() const { return MaterialPSOPrecacheRequestIDs; }
 
@@ -144,6 +145,7 @@ protected:
 	
 	ENGINE_API void LaunchPSOPrecaching(const FMaterialInterfacePSOPrecacheParamsList& VFsPerMaterials);
 
+	bool PSOPrecachingLaunched = false;
 	FGraphEventRef PrecachePSOsEvent;
 	TArray<FMaterialPSOPrecacheRequestID> MaterialPSOPrecacheRequestIDs;
 };
diff --git a/Engine/Source/Runtime/Engine/Private/Components/PrimitiveComponent.cpp b/Engine/Source/Runtime/Engine/Private/Components/PrimitiveComponent.cpp
--- a/Engine/Source/Runtime/Engine/Private/Components/PrimitiveComponent.cpp
+++ b/Engine/Source/Runtime/Engine/Private/Components/PrimitiveComponent.cpp
@@ -4854,7 +4854,7 @@ bool UPrimitiveComponent::ShouldRenderProxyFallbackToDefaultMaterial() const
 bool UPrimitiveComponent::CheckPSOPrecachingAndBoostPriority(EPSOPrecachePriority NewPSOPrecachePriority)
 {
 #if UE_WITH_PSO_PRECACHING
-	ensure(!IsAssetPSOPrecachingEnabled() || !IsComponentPSOPrecachingEnabled() || bPSOPrecacheCalled);
+	ensure(!IsComponentPSOPrecachingEnabled() || bPSOPrecacheCalled);
 	check(NewPSOPrecachePriority == EPSOPrecachePriority::High || NewPSOPrecachePriority == EPSOPrecachePriority::Highest);
 
 	if (PSOPrecacheCompileEvent && !PSOPrecacheCompileEvent->IsComplete())
diff --git a/Engine/Source/Runtime/Engine/Private/PSOPrecache.cpp b/Engine/Source/Runtime/Engine/Private/PSOPrecache.cpp
--- a/Engine/Source/Runtime/Engine/Private/PSOPrecache.cpp
+++ b/Engine/Source/Runtime/Engine/Private/PSOPrecache.cpp
@@ -31,14 +31,6 @@ static FAutoConsoleVariableRef CVarPSOPrecacheResources(
 	ECVF_ReadOnly
 );
 
-int32 GPSOPrecacheAssets = 1;
-static FAutoConsoleVariableRef CVarPSOPrecacheAssets(
-	TEXT("r.PSOPrecache.Assets"),
-	GPSOPrecacheAssets,
-	TEXT("Precache all possible used PSOs by assets during Postload (default 1 if PSOPrecaching is enabled). This affects Niagara and Cascade assets."),
-	ECVF_ReadOnly
-);
-
 int32 GPSOProxyCreationWhenPSOReady = 1;
 static FAutoConsoleVariableRef CVarPSOProxyCreationWhenPSOReady(
 	TEXT("r.PSOPrecache.ProxyCreationWhenPSOReady"),
@@ -69,11 +61,6 @@ static FAutoConsoleVariableRef CVarPSOComponentBoostStrategy(
 	ECVF_ReadOnly
 );
 
-bool IsAssetPSOPrecachingEnabled()
-{
-	return FApp::CanEverRender() && PipelineStateCache::IsPSOPrecachingEnabled() && GPSOPrecacheAssets && !GIsEditor;
-}
-
 bool IsComponentPSOPrecachingEnabled()
 {
 	return FApp::CanEverRender() && PipelineStateCache::IsPSOPrecachingEnabled() && GPSOPrecacheComponents && !GIsEditor;
diff --git a/Engine/Source/Runtime/Engine/Private/Particles/ParticleComponents.cpp b/Engine/Source/Runtime/Engine/Private/Particles/ParticleComponents.cpp
--- a/Engine/Source/Runtime/Engine/Private/Particles/ParticleComponents.cpp
+++ b/Engine/Source/Runtime/Engine/Private/Particles/ParticleComponents.cpp
@@ -163,6 +163,8 @@ void UFXSystemAsset::LaunchPSOPrecaching(const FMaterialInterfacePSOPrecachePara
 		PrecachePSOsEvent = ReleasePrecachePSOsEventTask->GetCompletionEvent();
 		ReleasePrecachePSOsEventTask->Unlock();
 	}
+
+	PSOPrecachingLaunched = true;
 }
 
 //////////////////////////////////////////////////////////////////////////
@@ -209,6 +211,17 @@ FAutoConsoleVariableRef CVarFXSkipZeroDeltaTime(
 	TEXT("This fixes issue like PSA_Velocity aligned sprites, but could cause issues with things that rely on accurate velocities (i.e. TSR)."),
 	ECVF_Default);
 
+int32 GCascadePSOPrecachingTime = 1;
+FAutoConsoleVariableRef CVarCascadePSOPrecachingTime(
+	TEXT("r.PSOPrecache.CascadePrecachingTime"),
+	GCascadePSOPrecachingTime,
+	TEXT("Controls when PSO precaching happens for Cascade systems:\n")
+	TEXT("	0: no precaching\n")
+	TEXT("	1: precaching at asset loading time (default)\n")
+	TEXT("	2: precaching at component loading time\n")
+	TEXT("	3: precaching at component proxy creation time"),
+	ECVF_Default);
+
 /** Whether to allow particle systems to perform work. */
 ENGINE_API bool GIsAllowingParticles = true;
 
@@ -2583,7 +2596,10 @@ void UParticleSystem::PostLoad()
 		}
 	}
 
-	PrecachePSOs();
+	if (GCascadePSOPrecachingTime == 1)
+	{
+		PrecachePSOs();
+	}
 
 #if WITH_EDITOR
 	// Due to there still being some ways that LODLevel counts get mismatched,
@@ -2687,8 +2703,7 @@ void UParticleSystem::PostLoad()
 
 void UParticleSystem::PrecachePSOs()
 {
-	// Only precache if asset precaching is enabled and at least one of component or resource precaching is enabled.
-	if (!IsAssetPSOPrecachingEnabled() || (!IsComponentPSOPrecachingEnabled() && !IsResourcePSOPrecachingEnabled()))
+	if (HasLaunchedPSOPrecaching() || (!IsComponentPSOPrecachingEnabled() && !IsResourcePSOPrecachingEnabled()))
 	{
 		return;
 	}
@@ -3541,7 +3556,7 @@ bool UFXSystemComponent::RequiresLWCTileRecache(const FVector3f CurrentTile, con
 void UFXSystemComponent::PrecacheAssetPSOs(UFXSystemAsset* FXSystemAsset)
 {
 #if UE_WITH_PSO_PRECACHING
-	if (!FApp::CanEverRender() || !IsAssetPSOPrecachingEnabled() || FXSystemAsset == nullptr)
+	if (!FApp::CanEverRender() || !IsComponentPSOPrecachingEnabled() || FXSystemAsset == nullptr)
 	{
 		return;
 	}
@@ -3950,6 +3965,12 @@ void UParticleSystemComponent::PostLoad()
 	{
 		PrimaryComponentTick.bStartWithTickEnabled = false;
 	}
+
+	if (Template && GCascadePSOPrecachingTime == 2)
+	{
+		Template->ConditionalPostLoad();
+		Template->PrecachePSOs();
+	}
 }
 
 void UParticleSystemComponent::Serialize( FArchive& Ar )
diff --git a/Engine/Source/Runtime/Engine/Private/Particles/ParticleSystemRender.cpp b/Engine/Source/Runtime/Engine/Private/Particles/ParticleSystemRender.cpp
--- a/Engine/Source/Runtime/Engine/Private/Particles/ParticleSystemRender.cpp
+++ b/Engine/Source/Runtime/Engine/Private/Particles/ParticleSystemRender.cpp
@@ -101,6 +101,9 @@ float GMinParticleDrawTimeToTrack = .0001f;
 /** Whether to do LOD calculation on GameThread in game */
 extern bool GbEnableGameThreadLODCalculation;
 
+/** When to precache Cascade systems' PSOs */
+extern int32 GCascadePSOPrecachingTime;
+
 ///////////////////////////////////////////////////////////////////////////////
 
 /**
@@ -6947,6 +6950,10 @@ FPrimitiveSceneProxy* UParticleSystemComponent::CreateSceneProxy()
 #if UE_WITH_PSO_PRECACHING
 		if (!bPSOPrecacheCalled)
 		{
+			if (GCascadePSOPrecachingTime == 3)
+			{
+				Template->PrecachePSOs();
+			}
 			PrecacheAssetPSOs(Template);
 		}
 
diff --git a/Engine/Source/Runtime/Engine/Public/PSOPrecache.h b/Engine/Source/Runtime/Engine/Public/PSOPrecache.h
--- a/Engine/Source/Runtime/Engine/Public/PSOPrecache.h
+++ b/Engine/Source/Runtime/Engine/Public/PSOPrecache.h
@@ -248,11 +248,6 @@ struct FMaterialPSOPrecacheParams
 	}
 };
 
-/**
- * Precaching PSOs for assets?
- */
-extern ENGINE_API bool IsAssetPSOPrecachingEnabled();
-
 /**
  * Precaching PSOs for components?
  */
