Backout ISM changes while determining why things fail with them
-----------------------------------------------------------------
Crunch primitive components.

#rnx

diff --git a/Engine/Source/Runtime/Engine/Classes/Components/InstancedStaticMeshComponent.h b/Engine/Source/Runtime/Engine/Classes/Components/InstancedStaticMeshComponent.h
--- a/Engine/Source/Runtime/Engine/Classes/Components/InstancedStaticMeshComponent.h
+++ b/Engine/Source/Runtime/Engine/Classes/Components/InstancedStaticMeshComponent.h
@@ -125,12 +125,6 @@ class UInstancedStaticMeshComponent : public UStaticMeshComponent, public ISMIns
 	UPROPERTY(EditAnywhere, Category=Instances, AdvancedDisplay)
 	int32 NumCustomDataFloats;
 
-	/** Value used to seed the random number stream that generates random numbers for each of this mesh's instances.
-	The random number is stored in a buffer accessible to materials through the PerInstanceRandom expression. If
-	this is set to zero (default), it will be populated automatically by the editor. */
-	UPROPERTY(EditAnywhere, BlueprintReadWrite, Category=InstancedStaticMeshComponent)
-	int32 InstancingRandomSeed=0;
-
 	/** Array of custom data for instances. This will contains NumCustomDataFloats*InstanceCount entries. The entries are represented sequantially, in instance order. Can be read in a material and manipulated through Blueprints.
 	 *	Example: If NumCustomDataFloats is 1, then each entry will belong to an instance. Custom data 0 will belong to Instance 0. Custom data 1 will belong to Instance 1 etc.
 	 *	Example: If NumCustomDataFloats is 2, then each pair of sequential entries belong to an instance. Custom data 0 and 1 will belong to Instance 0. Custom data 2 and 3 will belong to Instance 2 etc.
@@ -138,6 +132,12 @@ class UInstancedStaticMeshComponent : public UStaticMeshComponent, public ISMIns
 	UPROPERTY(EditAnywhere, EditFixedSize, SkipSerialization, DisplayName="Custom data", Category=Instances, AdvancedDisplay, meta=(EditFixedOrder))
 	TArray<float> PerInstanceSMCustomData;
 
+	/** Value used to seed the random number stream that generates random numbers for each of this mesh's instances.
+	The random number is stored in a buffer accessible to materials through the PerInstanceRandom expression. If
+	this is set to zero (default), it will be populated automatically by the editor. */
+	UPROPERTY(EditAnywhere, BlueprintReadWrite, Category=InstancedStaticMeshComponent)
+	int32 InstancingRandomSeed=0;
+
 	/** Additional random seeds ranges. Each seed entry will be applied from AdditionalRandomSeeds[i].StartInstanceIndex to AdditionalRandomSeeds[i+1].StartInstanceIndex -1 */
 	UPROPERTY()
 	TArray<FInstancedStaticMeshRandomSeed> AdditionalRandomSeeds;
@@ -165,30 +165,14 @@ class UInstancedStaticMeshComponent : public UStaticMeshComponent, public ISMIns
 	UPROPERTY()
 	uint8 bInheritPerInstanceData : 1;
 
-	/** Don't create any collision when this bool is set */
-	UPROPERTY()
-	uint8 bDisableCollision : 1;
-
-	/** 
-	* Flag for using RemoveAtSwap on instance removal. 
-	* The implementation is free to ignore this flag, but should honor whatever behavior is being returned by SupportsRemoveSwap().
-	*/
-	uint8 bSupportRemoveAtSwap : 1;
-
-protected:
-	uint8 bIsInstanceDataApplyCompleted : 1;
-
-private:
-	uint8 bHasPreviousTransforms : 1;
-
-	/** Flag for whether we are using conservative bounds. */
-	uint8 bUseConservativeBounds : 1;
-
-public:
 	/** Mapping from PerInstanceSMData order to instance render buffer order. If empty, the PerInstanceSMData order is used. */
 	UPROPERTY()
 	TArray<int32> InstanceReorderTable;
 	
+	/** Don't create any collision when this bool is set */
+	UPROPERTY()
+	bool bDisableCollision;
+
 	/** Tracks outstanding proxysize, as this is a bit hard to do with the fire-and-forget grass. */
 	SIZE_T ProxySize;
 
@@ -432,6 +416,10 @@ public:
 	ENGINE_API FInstanceDataFlags MakeInstanceDataFlags(bool bAnyMaterialHasPerInstanceRandom, bool bAnyMaterialHasPerInstanceCustomData) const;
 
 private:
+	bool bHasPreviousTransforms = false;
+
+	/** Flag for whether we are using conservative bounds. */
+	bool bUseConservativeBounds = false;
 	/** Current cached conservativ bounds. */
 	FBox CachedConservativeInstanceBounds;
 public:
@@ -597,6 +585,12 @@ public:
 	/** Request to navigation system to update only part of navmesh occupied specified instances transforms. */
 	ENGINE_API virtual void PartialNavigationUpdates(TConstArrayView<FTransform> InstanceTransforms);
 
+	/** 
+	 * Flag for using RemoveAtSwap on instance removal. 
+	 * The implementation is free to ignore this flag, but should honor whatever behavior is being returned by SupportsRemoveSwap().
+	 */
+	bool bSupportRemoveAtSwap = false;
+
 	ENGINE_API TSharedPtr<FISMCInstanceDataSceneProxy, ESPMode::ThreadSafe> GetOrCreateInstanceDataSceneProxy();
 
 	/**
@@ -639,10 +633,7 @@ private:
 	ENGINE_API bool BatchUpdateInstancesTransformsInternal(int32 StartInstanceIndex, TArrayView<const FTransform> NewInstancesTransforms, bool bWorldSpace, bool bMarkRenderStateDirty, bool bTeleport);
 
 protected:
-
-	/** Number of pending lightmaps still to be calculated (Apply()'d). */
-	UPROPERTY(Transient, DuplicateTransient, TextExportTransient)
-	int32 NumPendingLightmaps;
+	bool bIsInstanceDataApplyCompleted = true;
 
 	FPrimitiveInstanceDataManager PrimitiveInstanceDataManager;
 
@@ -688,6 +679,10 @@ protected:
 	/** Initializes the body instance for the specified instance of the static mesh. */
 	ENGINE_API void InitInstanceBody(int32 InstanceIdx, FBodyInstance* InBodyInstance);
 
+	/** Number of pending lightmaps still to be calculated (Apply()'d). */
+	UPROPERTY(Transient, DuplicateTransient, TextExportTransient)
+	int32 NumPendingLightmaps;
+
 	/** The mappings for all the instances of this component. */
 	UPROPERTY(Transient, DuplicateTransient, TextExportTransient)
 	TArray<FInstancedStaticMeshMappingInfo> CachedMappings;
diff --git a/Engine/Source/Runtime/Engine/Private/InstancedStaticMesh.cpp b/Engine/Source/Runtime/Engine/Private/InstancedStaticMesh.cpp
--- a/Engine/Source/Runtime/Engine/Private/InstancedStaticMesh.cpp
+++ b/Engine/Source/Runtime/Engine/Private/InstancedStaticMesh.cpp
@@ -2057,7 +2057,6 @@ UInstancedStaticMeshComponent::UInstancedStaticMeshComponent(const FObjectInitia
 
 UInstancedStaticMeshComponent::UInstancedStaticMeshComponent(FVTableHelper& Helper)
 	: Super(Helper)
-	, bIsInstanceDataApplyCompleted(true)
 	, PrimitiveInstanceDataManager(this)
 {
 }
diff --git a/Engine/Source/Runtime/Engine/Private/InstancedStaticMesh/ISMInstanceDataManager.cpp b/Engine/Source/Runtime/Engine/Private/InstancedStaticMesh/ISMInstanceDataManager.cpp
--- a/Engine/Source/Runtime/Engine/Private/InstancedStaticMesh/ISMInstanceDataManager.cpp
+++ b/Engine/Source/Runtime/Engine/Private/InstancedStaticMesh/ISMInstanceDataManager.cpp
@@ -61,18 +61,7 @@ FArchive& operator<<(FArchive& Ar, FInstanceSceneDataBuffers::FCompressedSpatial
 }
 
 FPrimitiveInstanceDataManager::FPrimitiveInstanceDataManager(UPrimitiveComponent* InPrimitiveComponent) 
-	: bNumCustomDataChanged(false)
-	, bBakedLightingDataChanged(false)
-	, bTransformChangedAllInstances(false)
-#if WITH_EDITOR
-	, bAnyEditorDataChanged(false)
-#endif
-	, bPrimitiveTransformChanged(false)
-	, bAnyInstanceChange(false)
-	, bComponentMarkedDirty(false)
-	, bEnableTracking(false)
-	, bFirstFlush(true)
-	, PrimitiveComponent(InPrimitiveComponent) 
+	: PrimitiveComponent(InPrimitiveComponent) 
 {
 	// Don't do anything if this is not a "real" ISM being tracked (this logic shopuld move out).
 	if (PrimitiveComponent.IsValid() && PrimitiveComponent->HasAnyFlags(RF_ClassDefaultObject | RF_ArchetypeObject))
diff --git a/Engine/Source/Runtime/Engine/Public/InstancedStaticMesh/ISMInstanceDataManager.h b/Engine/Source/Runtime/Engine/Public/InstancedStaticMesh/ISMInstanceDataManager.h
--- a/Engine/Source/Runtime/Engine/Public/InstancedStaticMesh/ISMInstanceDataManager.h
+++ b/Engine/Source/Runtime/Engine/Public/InstancedStaticMesh/ISMInstanceDataManager.h
@@ -49,7 +49,7 @@ public:
 	/**
 	 * Current tracking state, 
 	 */
-	enum class ETrackingState : uint8
+	enum class ETrackingState
 	{
 		Initial, // In the initial state, there is no proxy and therefore changes do not need to be tracked, e.g., during initial setup of an ISM component.
 		Tracked,
@@ -57,7 +57,7 @@ public:
 		Optimized, // In the optimized state there's no need to track any delta changes, but if anything changes at all we must rebuild.
 	};
 
-	enum class EMode : uint8
+	enum class EMode
 	{
 		Default,
 		Legacy, // In this mode, we create a legacy supporting proxy.
@@ -206,33 +206,31 @@ private:
 	 */
 	static FISMPrecomputedSpatialHashData PrecomputeOptimizationData(FInstanceUpdateComponentDesc &&ComponentData, TStridedView<FMatrix> InstanceTransforms);
 #endif
-
-	// Change set.
-	FInstanceAttributeTracker InstanceUpdateTracker;
+	EMode Mode = EMode::Default;
+	ETrackingState TrackingState = ETrackingState::Initial;
 
 	// Id allocation tracking
 	TBitArray<> ValidInstanceIdMask;
 	int32 IdSearchStartIndex = 0;
 
-	EMode Mode = EMode::Default;
-	ETrackingState TrackingState = ETrackingState::Initial;
+	// Change set.
+	FInstanceAttributeTracker InstanceUpdateTracker;
 
-	uint8 bNumCustomDataChanged : 1;
-	uint8 bBakedLightingDataChanged : 1;
-	uint8 bTransformChangedAllInstances : 1;
+	bool bNumCustomDataChanged = false;
+	bool bBakedLightingDataChanged = false;
+	bool bTransformChangedAllInstances = false;
 #if WITH_EDITOR
-	uint8 bAnyEditorDataChanged : 1;
+	bool bAnyEditorDataChanged = false;
 #endif	
-	uint8 bPrimitiveTransformChanged : 1;
-	uint8 bAnyInstanceChange : 1;
-
-	uint8 bComponentMarkedDirty : 1;
-	uint8 bEnableTracking : 1;
-	uint8 bFirstFlush : 1;
+	bool bPrimitiveTransformChanged = false;
+	bool bAnyInstanceChange = false;
 
 	TSharedPtr<FISMCInstanceDataSceneProxy, ESPMode::ThreadSafe> Proxy;
 	TWeakObjectPtr<UPrimitiveComponent> PrimitiveComponent = nullptr;
 
+	bool bComponentMarkedDirty = false;
+	bool bEnableTracking = false;
+
 	TPimplPtr<struct FLegacyBuildData> LegacyBuildData;
 
 	using PrecomputedOptimizationDataPtr = FISMCInstanceDataSceneProxy::FISMPrecomputedSpatialHashDataPtr;
@@ -246,4 +244,5 @@ private:
 	int32 NumCustomDataFloats = 0;
 	float AbsMaxDisplacement = 0.0f;
 	FRenderBounds StaticMeshBounds;
+	bool bFirstFlush = true;
 };
