UStaticMesh 584 -> 576
Saves 0.3 MB with 3100 instances (moves from 672 to 576 mempool)
#rnx

diff --git a/Engine/Source/Runtime/Engine/Classes/Engine/StaticMesh.h b/Engine/Source/Runtime/Engine/Classes/Engine/StaticMesh.h
--- a/Engine/Source/Runtime/Engine/Classes/Engine/StaticMesh.h
+++ b/Engine/Source/Runtime/Engine/Classes/Engine/StaticMesh.h
@@ -858,6 +858,16 @@ public:
 #endif
 	}
 
+protected:
+	/**
+	 * Index of an element to ignore while gathering streaming texture factors.
+	 * This is useful to disregard automatically generated vertex data which breaks texture factor heuristics.
+	 */
+	UPROPERTY(meta = (DisplayAfter = "ExtendedBounds"))
+	int32 ElementToIgnoreForTexFactor;
+
+public:
+
 #if WITH_EDITORONLY_DATA
 	/**
 	 * Returns true if this SM should have Nanite built for it.
@@ -1300,12 +1310,6 @@ public:
 #endif
 
 protected:
-	/**
-	 * Index of an element to ignore while gathering streaming texture factors.
-	 * This is useful to disregard automatically generated vertex data which breaks texture factor heuristics.
-	 */
-	UPROPERTY()
-	int32 ElementToIgnoreForTexFactor;
 
 	/** Array of user data stored with the asset */
 	UPROPERTY(EditAnywhere, AdvancedDisplay, Instanced, Category = StaticMesh)
