Timeline: Fix an assert in custom version when doing reference gathering.

#rnx

diff --git a/Engine/Source/Runtime/Engine/Private/Timeline.cpp b/Engine/Source/Runtime/Engine/Private/Timeline.cpp
--- a/Engine/Source/Runtime/Engine/Private/Timeline.cpp
+++ b/Engine/Source/Runtime/Engine/Private/Timeline.cpp
@@ -760,8 +760,7 @@ bool FTimeline::Serialize(FArchive& Ar)
 void FTimeline::PostSerialize(const FArchive& Ar)
 {
 	// If we are loading code from before the serialization change...
-	const int32 CustomVersion = Ar.CustomVer(FFortniteMainBranchObjectVersion::GUID);
-	if (Ar.IsLoading() && CustomVersion < FFortniteMainBranchObjectVersion::TimelinePlayingStateTrackerDeprecation)
+	if (Ar.IsLoading() && Ar.CustomVer(FFortniteMainBranchObjectVersion::GUID) < FFortniteMainBranchObjectVersion::TimelinePlayingStateTrackerDeprecation)
 	{
 PRAGMA_DISABLE_DEPRECATION_WARNINGS
 
