RADA/Seektable fix

#tests Locally tested in editor and verified fixed by DanT, additionally tested by SondraM in Main.

diff --git a/Engine/Source/Developer/AudioFormatRad/Private/AudioFormatRad.cpp b/Engine/Source/Developer/AudioFormatRad/Private/AudioFormatRad.cpp
--- a/Engine/Source/Developer/AudioFormatRad/Private/AudioFormatRad.cpp
+++ b/Engine/Source/Developer/AudioFormatRad/Private/AudioFormatRad.cpp
@@ -41,7 +41,7 @@ class FAudioFormatRad : public IAudioFormat
 	enum
 	{
 		/** Version for RAD Audio format, this becomes part of the DDC key. */
-		UE_AUDIO_RAD_VER = 3,
+		UE_AUDIO_RAD_VER = 4,
 	};
 
 public:
diff --git a/Engine/Source/Runtime/Engine/Private/StreamedAudioChunkSeekTable.cpp b/Engine/Source/Runtime/Engine/Private/StreamedAudioChunkSeekTable.cpp
--- a/Engine/Source/Runtime/Engine/Private/StreamedAudioChunkSeekTable.cpp
+++ b/Engine/Source/Runtime/Engine/Private/StreamedAudioChunkSeekTable.cpp
@@ -19,10 +19,10 @@ namespace StreamedAudioChunkSeekTable_Private
 			}
 		}
 	}
-	static void DeltaDecode(const TArray<uint16>& InDeltas, TArray<uint32>& OutValues)
+	static void DeltaDecode(const TArray<uint16>& InDeltas, TArray<uint32>& OutValues, uint32 StartingValue = 0)
 	{
 		OutValues.SetNum(InDeltas.Num());
-		uint32 Sum = 0;
+		uint32 Sum = StartingValue;
 		for (int32 i = 0; i < InDeltas.Num(); ++i)
 		{
 			Sum += InDeltas[i];
@@ -207,7 +207,7 @@ struct FVariableRateSeekTable : public FStreamedAudioChunkSeekTable::ISeekTableI
 	static int32 CalcSize(int32 NumEntries)
 	{
 		// Not including super class.
-		static const int32 SizeOfHeader = 0;
+		static const int32 SizeOfHeader = sizeof(uint32);
 		static const int32 SizePerEntry = sizeof(uint16) + sizeof(uint16);
 
 		// Header + TArray<uint16> (each entry, plus count).
@@ -227,17 +227,25 @@ struct FVariableRateSeekTable : public FStreamedAudioChunkSeekTable::ISeekTableI
 	uint32 FindOffset(uint32 InTimeInAudioFrames) const override
 	{
 		check(Offsets.Num() == Times.Num());
-		int32 Index = Algo::LowerBound(Times, InTimeInAudioFrames);
-		if (Index < Times.Num())
+		const int32 Index = Algo::LowerBound(Times, InTimeInAudioFrames);
+		if (Offsets.IsValidIndex(Index))
 		{
-			return Offsets[Index];
+			// Minor hack here. RADA is the only customer of this currently and needs
+			// the one block behind the one we've found. 
+			if (Index > 0)
+			{
+				return Offsets[Index-1];
+			}
+			return Offsets[0];
 		}
 		return INDEX_NONE;
 	}
+
 	uint32 FindTime(uint32 InOffset) const override
 	{
-		int32 Index = Algo::LowerBound(Offsets,InOffset);
-		if (Index < Offsets.Num())
+		check(Offsets.Num() == Times.Num());
+		const int32 Index = Algo::LowerBound(Offsets, InOffset);
+		if (Offsets.IsValidIndex(Index))
 		{
 			return Times[Index];
 		}
@@ -265,6 +273,9 @@ struct FVariableRateSeekTable : public FStreamedAudioChunkSeekTable::ISeekTableI
 		TArray<uint16> DeltaOffsets;
 		StreamedAudioChunkSeekTable_Private::DeltaEncode(Offsets,DeltaOffsets);
 		Ar << DeltaOffsets;
+
+		uint32 FirstTimeItem = Times.Num() > 0 ? Times[0] : 0;
+		Ar << FirstTimeItem;
 		
 		TArray<uint16> DeltaTimes;
 		StreamedAudioChunkSeekTable_Private::DeltaEncode(Times, DeltaTimes);
@@ -276,10 +287,13 @@ struct FVariableRateSeekTable : public FStreamedAudioChunkSeekTable::ISeekTableI
 		TArray<uint16> DeltaOffsets;
 		Ar << DeltaOffsets;
 		StreamedAudioChunkSeekTable_Private::DeltaDecode(DeltaOffsets, Offsets);
+
+		uint32 FirstTimeItem = 0;
+		Ar << FirstTimeItem;
 		
-		TArray<uint16> DeltaSamples;
-		Ar << DeltaSamples;
-		StreamedAudioChunkSeekTable_Private::DeltaDecode(DeltaSamples, Times);
+		TArray<uint16> DeltaTimes;
+		Ar << DeltaTimes;
+		StreamedAudioChunkSeekTable_Private::DeltaDecode(DeltaTimes, Times, FirstTimeItem);
 		return true;
 	}
 };
@@ -323,4 +337,4 @@ int32 FStreamedAudioChunkSeekTable::CalcSize(int32 InNumEntries, EChunkSeekTable
 	}
 	checkNoEntry();
 	return 0;
-}
\ No newline at end of file
+}
