#UE Use UE_ASSET_LOG to better format the warning about MinLODs for build processes to parse the affected asset

diff --git a/Engine/Source/Runtime/Engine/Private/StaticMesh.cpp b/Engine/Source/Runtime/Engine/Private/StaticMesh.cpp
--- a/Engine/Source/Runtime/Engine/Private/StaticMesh.cpp
+++ b/Engine/Source/Runtime/Engine/Private/StaticMesh.cpp
@@ -7050,7 +7050,7 @@ void UStaticMesh::ExecutePostLoadInternal(FStaticMeshPostLoadContext& Context)
 		{
 			for (const FText& MinLODError : MinLODErrors)
 			{
-				UE_LOG(LogStaticMesh, Warning, TEXT("%s %s"), *GetNameSafe(this), *MinLODError.ToString());
+				UE_ASSET_LOG(LogStaticMesh, Warning, this, TEXT("%s"), *MinLODError.ToString());
 			}
 		}
 		else
