[LevelInstance] Instead of Unloading immediatly in OnActorDeleted, request the unload of the Level Instance so it can get processed in the LevelInstanceSubsystem::Tick

- PCG can delete actors inside Tick which happens while the FTickTaskLevels are registered, so deleting a Level Instance from PCG and unloading the Level immediately was causing a failed check in FTickTaskManager::FreeTickTaskLevel

#rnx

diff --git a/Engine/Source/Runtime/Engine/Private/LevelInstance/LevelInstanceSubsystem.cpp b/Engine/Source/Runtime/Engine/Private/LevelInstance/LevelInstanceSubsystem.cpp
--- a/Engine/Source/Runtime/Engine/Private/LevelInstance/LevelInstanceSubsystem.cpp
+++ b/Engine/Source/Runtime/Engine/Private/LevelInstance/LevelInstanceSubsystem.cpp
@@ -593,10 +593,10 @@ void ULevelInstanceSubsystem::UnloadLevelInstance(const FLevelInstanceID& LevelI
 {
 	if (GetWorld()->IsGameWorld())
 	{
-		FLevelInstance LevelInstance;
-		if (LoadedLevelInstances.RemoveAndCopyValue(LevelInstanceID, LevelInstance))
+		FLevelInstance LoadedLevelInstance;
+		if (LoadedLevelInstances.RemoveAndCopyValue(LevelInstanceID, LoadedLevelInstance))
 		{
-			ULevelStreamingLevelInstance::UnloadInstance(LevelInstance.LevelStreaming);
+			ULevelStreamingLevelInstance::UnloadInstance(LoadedLevelInstance.LevelStreaming);
 		}
 	}
 #if WITH_EDITOR
@@ -610,10 +610,10 @@ void ULevelInstanceSubsystem::UnloadLevelInstance(const FLevelInstanceID& LevelI
 			LevelsToRemoveScope.Reset(new FLevelsToRemoveScope(this));
 		}
 
-		FLevelInstance LevelInstance;
-		if (LoadedLevelInstances.RemoveAndCopyValue(LevelInstanceID, LevelInstance))
+		FLevelInstance LoadedLevelInstance;
+		if (LoadedLevelInstances.RemoveAndCopyValue(LevelInstanceID, LoadedLevelInstance))
 		{
-			if (ULevel* LoadedLevel = LevelInstance.LevelStreaming->GetLoadedLevel())
+			if (ULevel* LoadedLevel = LoadedLevelInstance.LevelStreaming->GetLoadedLevel())
 			{
 				ForEachActorInLevel(LoadedLevel, [this](AActor* LevelActor)
 				{
@@ -628,7 +628,7 @@ void ULevelInstanceSubsystem::UnloadLevelInstance(const FLevelInstanceID& LevelI
 				});
 			}
 
-			ULevelStreamingLevelInstance::UnloadInstance(LevelInstance.LevelStreaming);
+			ULevelStreamingLevelInstance::UnloadInstance(LoadedLevelInstance.LevelStreaming);
 		}
 
 		if (bReleaseScope)
@@ -1997,24 +1997,8 @@ void ULevelInstanceSubsystem::OnActorDeleted(AActor* Actor)
 		{
 			CommitLevelInstance(LevelInstance);
 		}
-		else
-		{
-			// We are ending editing. Discard Non dirty child edits
-			ForEachLevelInstanceChild(LevelInstance, /*bRecursive=*/true, [this](const ILevelInstanceInterface* ChildLevelInstance)
-			{
-				if (const FLevelInstanceEdit* ChildLevelInstanceEdit = GetLevelInstanceEdit(ChildLevelInstance))
-				{
-					check(!IsLevelInstanceEditDirty(ChildLevelInstanceEdit));
-					LevelInstanceEdit.Reset();
-					return false;
-				}
-				return true;
-			});
-		}
-
-		LevelInstancesToLoadOrUpdate.Remove(LevelInstance);
-				
-		UnloadLevelInstance(LevelInstance->GetLevelInstanceID());
+		
+		RequestUnloadLevelInstance(LevelInstance);
 		
 		// Remove from root so it gets collected on the next GC if it can be.
 		if (!bAlreadyRooted)
