Rename static variable for wp.editor.GenerateStreamingDescriptor cmd to avoid clash with a similarly named one in worldpartition.cpp

diff --git a/Engine/Source/Runtime/Engine/Private/WorldPartition/WorldPartitionStreamingDescriptor.cpp b/Engine/Source/Runtime/Engine/Private/WorldPartition/WorldPartitionStreamingDescriptor.cpp
--- a/Engine/Source/Runtime/Engine/Private/WorldPartition/WorldPartitionStreamingDescriptor.cpp
+++ b/Engine/Source/Runtime/Engine/Private/WorldPartition/WorldPartitionStreamingDescriptor.cpp
@@ -18,7 +18,7 @@
 #include "HAL/FileManager.h"
 #include "Editor.h"
 
-static FAutoConsoleCommand DumpActorDesc(
+static FAutoConsoleCommand GenerateStreamingDescriptorCmd(
 	TEXT("wp.editor.GenerateStreamingDescriptor"),
 	TEXT("Generate the streaming descriptor for the current world in the speficied file."),
 	FConsoleCommandWithArgsDelegate::CreateLambda([](const TArray<FString>& Args)
