Engine NetTest
* Tick delta time now set at init
* Added utility functions to tick and drop or delay the packets generated during the tick.

NetDriver
* Packet simulation settings now have a FrameDelay settings that delays packets for a specific number of frames.

#rnx

diff --git a/Engine/Source/Runtime/Engine/Classes/Engine/NetConnection.h b/Engine/Source/Runtime/Engine/Classes/Engine/NetConnection.h
--- a/Engine/Source/Runtime/Engine/Classes/Engine/NetConnection.h
+++ b/Engine/Source/Runtime/Engine/Classes/Engine/NetConnection.h
@@ -194,7 +194,13 @@ struct FDelayedPacket
 	FOutPacketTraits Traits;
 
 	/** The time at which to send the packet */
-	double SendTime;
+	double SendTime = 0.0;
+
+	/** The number of frames we delay sending the packet */
+	uint32 DelayFrameCount = 0;
+
+	/** Flag telling the packet was sent and can be cleaned up */
+	bool bSent = false;
 
 public:
 
@@ -202,7 +208,6 @@ public:
 		: Data()
 		, SizeBits(InSizeBits)
 		, Traits(InTraits)
-		, SendTime(0.0)
 	{
 		int32 SizeBytes = FMath::DivideAndRoundUp(SizeBits, 8);
 
@@ -223,6 +228,9 @@ struct FDelayedIncomingPacket
 	/** Time at which the packet should be reinjected into the connection */
 	double ReinjectionTime = 0.0;
 
+	/** Number of frames until we can reinject the incoming packet */
+	uint32 ReinjectionFrameCount = 0;
+
 	void CountBytes(FArchive& Ar) const
 	{
 		if (PacketData.IsValid())
@@ -430,6 +438,7 @@ public:
 	int32			PacketOverhead;			// Bytes overhead per packet sent.
 	FString			Challenge;				// Server-generated challenge.
 	FString			ClientResponse;			// Client-generated response.
+
 	int32			ResponseId;				// Id assigned by the server for linking responses to connections upon authentication
 	FString			RequestURL;				// URL requested by client
 
@@ -491,7 +500,7 @@ public:
 	float			StatPeriod;
 
 	/** Average lag seen during the last StatPeriod */
-	float AvgLag;
+	float 			AvgLag;
 
 	/** Total accumulated lag values during the current StatPeriod */
 	double			LagAcc;
@@ -867,6 +876,9 @@ private:
 	/** Process incoming packets that have been delayed for long enough */
 	void ReinjectDelayedPackets();
 
+	/** Update local packets we queued to add latency and send those that are due */
+	void UpdateDelayedPackets(const double CurrentRealtimeSeconds);
+
 #endif //#if DO_ENABLE_NET_TEST
 
 public:
diff --git a/Engine/Source/Runtime/Engine/Classes/Engine/NetDriver.h b/Engine/Source/Runtime/Engine/Classes/Engine/NetDriver.h
--- a/Engine/Source/Runtime/Engine/Classes/Engine/NetDriver.h
+++ b/Engine/Source/Runtime/Engine/Classes/Engine/NetDriver.h
@@ -555,6 +555,18 @@ struct FPacketSimulationSettings
 	UPROPERTY(EditAnywhere, Category = "Simulation Settings")
 	int32	PktJitter = 0;
 
+	/**
+	 * Delays sending packets for a specific number of ticks
+	 */
+	UPROPERTY(EditAnywhere, Category = "Simulation Settings")
+	int32	PktFrameDelay = 0;
+
+	/**
+	 * Delays processing received packets for a specific number of ticks
+	 */
+	UPROPERTY(EditAnywhere, Category = "Simulation Settings")
+	int32	PktIncomingFrameDelay = 0;
+
 	/** reads in settings from the .ini file 
 	 * @note: overwrites all previous settings
 	 */
@@ -1381,7 +1393,7 @@ public:
 	 */
 	ENGINE_API void SetPacketSimulationSettings(const FPacketSimulationSettings& NewSettings);
 
-	void OnPacketSimulationSettingsChanged();
+	ENGINE_API void OnPacketSimulationSettingsChanged();
 #endif
 
 	// Constructors.
@@ -1660,7 +1672,7 @@ public:
 	bool HandlePropertyConditionsMemCommand(const TCHAR* Cmd, FOutputDevice& Ar);
 #endif
 
-	void HandlePacketLossBurstCommand( int32 DurationInMilliseconds );
+	ENGINE_API void HandlePacketLossBurstCommand( int32 DurationInMilliseconds );
 
 	// ---------------------------------------------------------------
 	//	Game code API for updating server Actor Replication State
diff --git a/Engine/Source/Runtime/Engine/Private/NetConnection.cpp b/Engine/Source/Runtime/Engine/Private/NetConnection.cpp
--- a/Engine/Source/Runtime/Engine/Private/NetConnection.cpp
+++ b/Engine/Source/Runtime/Engine/Private/NetConnection.cpp
@@ -317,8 +317,7 @@ UNetConnection::UNetConnection(const FObjectInitializer& ObjectInitializer)
 ,	Handler()
 ,	StatelessConnectComponent()
 ,	PacketOverhead		( 0 )
-,	ResponseId			( 0 )
-
+,	ResponseId			( 0 ) // variable is unused and will be deleted
 ,	QueuedBits			( 0 )
 ,	TickCount			( 0 )
 ,	LastProcessedFrame	( 0 )
@@ -2180,9 +2179,17 @@ void UNetConnection::ReinjectDelayedPackets()
 		TGuardValue<bool> ReinjectingGuard(bIsReinjectingDelayedPackets, true);
 
 		uint32 NbReinjected(0);
-		for (const FDelayedIncomingPacket& DelayedPacket : DelayedIncomingPackets)
+		for (FDelayedIncomingPacket& DelayedPacket : DelayedIncomingPackets)
 		{
-			if (DelayedPacket.ReinjectionTime > CurrentTime)
+			if (DelayedPacket.ReinjectionFrameCount > 0)
+			{
+				if (DelayedPacket.ReinjectionFrameCount > 1)
+				{
+					DelayedPacket.ReinjectionFrameCount--;
+					break;
+				}
+			}			
+			else if (DelayedPacket.ReinjectionTime > CurrentTime)
 			{
 				break;
 			}
@@ -2458,6 +2465,13 @@ bool UNetConnection::CheckOutgoingPacketEmulation(FOutPacketTraits& Traits)
 		
 		return true;
 	}
+	else if (PacketSimulationSettings.PktFrameDelay > 0)
+	{
+		FDelayedPacket& B = *(new(Delayed)FDelayedPacket(SendBuffer.GetData(), SendBuffer.GetNumBits(), Traits));
+		B.DelayFrameCount = PacketSimulationSettings.PktFrameDelay;
+		
+		return true;
+	}
 
 	return false;
 }
@@ -2916,6 +2930,17 @@ void UNetConnection::ReceivedPacket( FBitReader& Reader, bool bIsReinjectedPacke
 			UE_LOG(LogNet, VeryVerbose, TEXT("Delaying incoming packet for %f seconds"), ExtraLagInSec);
 			return;
 		}
+		else if (PacketSimulationSettings.PktIncomingFrameDelay)
+		{
+			FDelayedIncomingPacket DelayedPacket;
+			DelayedPacket.PacketData = MakeUnique<FBitReader>(Reader);
+			DelayedPacket.ReinjectionFrameCount = PacketSimulationSettings.PktIncomingFrameDelay;
+
+			DelayedIncomingPackets.Emplace(MoveTemp(DelayedPacket));
+
+			UE_LOG(LogNet, VeryVerbose, TEXT("Delaying incoming packet for %u frames"), PacketSimulationSettings.PktIncomingFrameDelay);
+			return;
+		}
 	}
 #endif //#if DO_ENABLE_NET_TEST
 
@@ -4404,43 +4429,7 @@ void UNetConnection::Tick(float DeltaSeconds)
 
 	// Lag simulation.
 #if DO_ENABLE_NET_TEST
-	if (Delayed.Num() > 0)
-	{
-		if (bSendDelayedPacketsOutofOrder)
-		{
-			for (int Idx=0; Idx < Delayed.Num(); ++Idx)
-			{
-				if (CurrentRealtimeSeconds > Delayed[Idx].SendTime)
-				{
-					LowLevelSend((char*)&(Delayed[Idx].Data[0]), Delayed[Idx].SizeBits, Delayed[Idx].Traits);
-				}
-			}
-
-			Delayed.RemoveAll([&CurrentRealtimeSeconds](const FDelayedPacket& rhs)
-			{
-				return CurrentRealtimeSeconds > rhs.SendTime;
-			});
-		}
-		else
-		{
-			uint32 NbPacketsSent(0);
-			for (FDelayedPacket& DelayedPacket : Delayed)
-			{
-				if (CurrentRealtimeSeconds > DelayedPacket.SendTime)
-				{
-					LowLevelSend((char*)&DelayedPacket.Data[0], DelayedPacket.SizeBits, DelayedPacket.Traits);
-					++NbPacketsSent;
-				}
-				else
-				{
-					// Break now instead of continuing to iterate through the list. Otherwise may cause out of order sends
-					break;
-				}
-			}
-
-			Delayed.RemoveAt(0, NbPacketsSent, EAllowShrinking::No);
-		}
-	}
+	UpdateDelayedPackets(CurrentRealtimeSeconds);
 #endif
 
 	// if this is 0 it's our first tick since init, so start our real-time tracking from here
@@ -4796,6 +4785,49 @@ void UNetConnection::Tick(float DeltaSeconds)
 	}
 }
 
+#if DO_ENABLE_NET_TEST
+void UNetConnection::UpdateDelayedPackets(const double CurrentRealtimeSeconds)
+{
+	for (int32 Index = 0; Index < Delayed.Num(); ++Index)
+	{
+		FDelayedPacket& Packet = Delayed[Index];
+		bool bSendPacket = false;
+
+		if (Packet.DelayFrameCount > 0)
+		{
+			if (Packet.DelayFrameCount == 1)
+			{
+				bSendPacket = true;
+			}
+			else
+			{
+				Packet.DelayFrameCount--;
+			}
+		}
+		else if (CurrentRealtimeSeconds > Packet.SendTime)
+		{
+			bSendPacket = true;
+		}
+
+		if (bSendPacket)
+		{
+			LowLevelSend((char*)&(Packet.Data[0]), Packet.SizeBits, Packet.Traits);
+			Packet.bSent = true;
+		}
+		else if (!bSendDelayedPacketsOutofOrder)
+		{
+			// Break now instead of continuing to iterate through the list. Otherwise it causes out of order packets due to the random variance added to each latency
+			break;
+		}
+	}
+
+	Delayed.RemoveAll([](const FDelayedPacket& Packet)
+	{
+		return Packet.bSent;
+	});
+}
+#endif // DO_ENABLE_NET_TEST
+
 void UNetConnection::HandleConnectionTimeout(const FString& Error)
 {
 	if (!bPendingDestroy)
diff --git a/Engine/Source/Runtime/Experimental/Iris/Core/Public/Iris/ReplicationSystem/ReplicationSystemTypes.h b/Engine/Source/Runtime/Experimental/Iris/Core/Public/Iris/ReplicationSystem/ReplicationSystemTypes.h
--- a/Engine/Source/Runtime/Experimental/Iris/Core/Public/Iris/ReplicationSystem/ReplicationSystemTypes.h
+++ b/Engine/Source/Runtime/Experimental/Iris/Core/Public/Iris/ReplicationSystem/ReplicationSystemTypes.h
@@ -89,4 +89,4 @@ inline const TCHAR* LexToString(EDependentObjectSchedulingHint SchedulingHint)
 	}
 }
 
-}
+} // end namespace UE::Net
