Fix warning C6237: (<zero> && <expression>) is always zero.  <expression> is never evaluated and might have side effects. by untangling the complicated if and laying it out slightly more clearly in a lambda

#rnx

diff --git a/Engine/Source/Runtime/Engine/Private/GameMode.cpp b/Engine/Source/Runtime/Engine/Private/GameMode.cpp
--- a/Engine/Source/Runtime/Engine/Private/GameMode.cpp
+++ b/Engine/Source/Runtime/Engine/Private/GameMode.cpp
@@ -671,22 +671,41 @@ void AGameMode::AddInactivePlayer(APlayerState* PlayerState, APlayerController*
 bool AGameMode::FindInactivePlayer(APlayerController* PC)
 {
 	check(PC && PC->PlayerState);
+
+	if (InactivePlayerArray.Num() == 0)
+	{
+		return false;
+	}
+
 	// don't bother for spectators
 	if (MustSpectate(PC))
 	{
 		return false;
 	}
 
-	// On console, we have to check the unique net id as network address isn't valid
-	const bool bIsConsole = !PLATFORM_DESKTOP;
-	// Assume valid unique ids means comparison should be via this method
-	const bool bHasValidUniqueId = PC->PlayerState->GetUniqueId().IsValid();
-	// Don't accidentally compare empty network addresses (already issue with two clients on same machine during development)
-	const bool bHasValidNetworkAddress = !PC->PlayerState->SavedNetworkAddress.IsEmpty();
-	const bool bUseUniqueIdCheck = bIsConsole || bHasValidUniqueId;
-
-	const FString NewNetworkAddress = PC->PlayerState->SavedNetworkAddress;
+	#if PLATFORM_DESKTOP
 	const FString NewName = PC->PlayerState->GetPlayerName();
+	auto EvaluatePlayerState = [PC, &NewName](APlayerState* CurrentPlayerState)
+	#else
+	auto EvaluatePlayerState = [PC](APlayerState* CurrentPlayerState)
+	#endif
+	{
+		// On console, we have to check the unique net id as network address isn't valid
+		#if PLATFORM_DESKTOP
+		// Assume valid unique ids means comparison should be via this method
+		if (PC->PlayerState->GetUniqueId().IsValid())
+		#endif
+		{
+			return CurrentPlayerState->GetUniqueId() == PC->PlayerState->GetUniqueId();
+		}
+
+		#if PLATFORM_DESKTOP
+		const FString& NewNetworkAddress = PC->PlayerState->SavedNetworkAddress;
+		const bool bHasValidNetworkAddress = !PC->PlayerState->SavedNetworkAddress.IsEmpty();
+		return bHasValidNetworkAddress && (FCString::Stricmp(*CurrentPlayerState->SavedNetworkAddress, *NewNetworkAddress) == 0) && (FCString::Stricmp(*CurrentPlayerState->GetPlayerName(), *NewName) == 0);
+		#endif
+	};
+
 	for (int32 i=0; i < InactivePlayerArray.Num(); i++)
 	{
 		APlayerState* CurrentPlayerState = InactivePlayerArray[i];
@@ -695,8 +714,7 @@ bool AGameMode::FindInactivePlayer(APlayerController* PC)
 			InactivePlayerArray.RemoveAt(i,1);
 			i--;
 		}
-		else if ((bUseUniqueIdCheck && (CurrentPlayerState->GetUniqueId() == PC->PlayerState->GetUniqueId())) ||
-				 (!bUseUniqueIdCheck && bHasValidNetworkAddress && (FCString::Stricmp(*CurrentPlayerState->SavedNetworkAddress, *NewNetworkAddress) == 0) && (FCString::Stricmp(*CurrentPlayerState->GetPlayerName(), *NewName) == 0)))
+		else if (EvaluatePlayerState(CurrentPlayerState))
 		{
 			// found it!
 			APlayerState* OldPlayerState = PC->PlayerState;
