Adding DeterminesOutputType/DynamicOutputParam to GetOverlappingActors so that users don't have to cast the results to their class filter type

#tests Verified existing uses of node still function in Windows editor & new nodes work properly
#rnx

diff --git a/Engine/Source/Runtime/Engine/Classes/GameFramework/Actor.h b/Engine/Source/Runtime/Engine/Classes/GameFramework/Actor.h
--- a/Engine/Source/Runtime/Engine/Classes/GameFramework/Actor.h
+++ b/Engine/Source/Runtime/Engine/Classes/GameFramework/Actor.h
@@ -2190,7 +2190,7 @@ public:
 	 * @param OverlappingActors		[out] Returned list of overlapping actors
 	 * @param ClassFilter			[optional] If set, only returns actors of this class or subclasses
 	 */
-	UFUNCTION(BlueprintCallable, Category="Collision", meta=(UnsafeDuringActorConstruction="true"))
+	UFUNCTION(BlueprintCallable, Category="Collision", meta=(UnsafeDuringActorConstruction="true", DeterminesOutputType="ClassFilter",  DynamicOutputParam="OverlappingActors"))
 	ENGINE_API void GetOverlappingActors(TArray<AActor*>& OverlappingActors, TSubclassOf<AActor> ClassFilter=nullptr) const;
 
 	/** 
