[NetPhysics] Refactor physics replication cache to run async.
- Also makes GeometryCollectionComponent take solver frame from the marshaling manager instead of the replication cache since it's not using the cache for the states.

diff --git a/Engine/Source/Runtime/Engine/Private/PhysicsEngine/Experimental/PhysScene_Chaos.cpp b/Engine/Source/Runtime/Engine/Private/PhysicsEngine/Experimental/PhysScene_Chaos.cpp
--- a/Engine/Source/Runtime/Engine/Private/PhysicsEngine/Experimental/PhysScene_Chaos.cpp
+++ b/Engine/Source/Runtime/Engine/Private/PhysicsEngine/Experimental/PhysScene_Chaos.cpp
@@ -11,6 +11,7 @@
 #include "Components/SkeletalMeshComponent.h"
 #include "Physics/PhysicsInterfaceUtils.h"
 #include "PhysicsReplication.h"
+#include "Physics/PhysicsReplicationCache.h"
 #include "PhysicsEngine/ClusterUnionComponent.h"
 #include "PhysicsEngine/ConstraintInstance.h"
 #include "PhysicsEngine/PhysicsCollisionHandler.h"
@@ -60,9 +61,6 @@ FAutoConsoleVariableRef CVar_KinematicDeferralUpdateExternalAccelerationStructur
 bool GKinematicDeferralLogInvalidBodies = false;
 FAutoConsoleVariableRef CVar_KinematicDeferralLogInvalidBodies(TEXT("p.KinematicDeferralLogInvalidBodies"), GKinematicDeferralLogInvalidBodies, TEXT("If true and p.KinematicDeferralCheckValidBodies is true, log when an invalid body is found on kinematic update."));
 
-float GReplicationCacheLingerForNSeconds = 3.f;
-FAutoConsoleVariableRef CVar_ReplicationCacheLingerForNSeconds(TEXT("np2.ReplicationCache.LingerForNSeconds"), GReplicationCacheLingerForNSeconds, TEXT("How long to keep data in the replication cache without the actor accessing it, after this we stop caching the actors state until it tries to access it again."));
-
 bool bGClusterUnionSyncBodiesMoveNewComponents = true;
 FAutoConsoleVariableRef CVar_GClusterUnionSyncBodiesCheckDirtyFlag(TEXT("p.ClusterUnion.SyncBodiesMoveNewComponents"), bGClusterUnionSyncBodiesMoveNewComponents, TEXT("Enable a fix to ensure new components in a cluster union are moved once on add (even if the cluster is not moving)."));
 
@@ -590,7 +588,7 @@ FPhysScene_Chaos::~FPhysScene_Chaos()
 	
 	// Make sure physics replication is cleared before we're fully destructed
 	PhysicsReplication.Reset();
-	ReplicationCache.Reset();
+	PhysicsReplicationCache.Reset();
 
 	FPhysicsDelegates::OnPhysSceneTerm.Broadcast(this);
 
@@ -1323,103 +1321,38 @@ void FPhysScene_Chaos::AddToComponentMaps(UPrimitiveComponent* Component, IPhysi
 	}
 }
 
+PRAGMA_DISABLE_DEPRECATION_WARNINGS
 // FReplicationCacheData constructor needs to be in .cpp due to UPrimitiveComponent being forward declared in the header which TWeakObjectPtr doesn't handle
 FPhysScene_Chaos::FReplicationCacheData::FReplicationCacheData(UPrimitiveComponent* InRootComponent, Chaos::FReal InAccessTime)
 	: RootComponent(InRootComponent)
 	, AccessTime(InAccessTime)
 	, bValidStateCached(false)
 {}
+PRAGMA_ENABLE_DEPRECATION_WARNINGS
 
 const FRigidBodyState* FPhysScene_Chaos::GetStateFromReplicationCache(UPrimitiveComponent* RootComponent, int& ServerFrame)
 {
-	if (!GetSolver()->GetRewindCallback())
-	{
-		// We only populate replication cache through the RewindCallback
-		ServerFrame = GetSolver()->GetCurrentFrame();
-		return nullptr;
-	}
-
-	ServerFrame = ReplicationCache.ServerFrame;
+	// Create the physics replication cache if not already created
+	CreatePhysicsReplicationCache();
 
-	const FObjectKey Key(RootComponent);
-	if (!ReplicationCache.Map.Contains(Key))
-	{
-		RegisterForReplicationCache(RootComponent);
-	}
-
-	FRigidBodyState* ReplicationState = nullptr;
-	if (FReplicationCacheData* ReplicationData = ReplicationCache.Map.Find(Key))
-	{
-		if (ReplicationData->IsCached())
-		{
-			ReplicationData->SetAccessTime(GetSolver()->GetSolverTime());
-			ReplicationState = &ReplicationData->GetState();
-		}
-	}
-	return ReplicationState;
+	return PhysicsReplicationCache->GetStateFromReplicationCache(RootComponent, ServerFrame);
 }
 
 void FPhysScene_Chaos::RegisterForReplicationCache(UPrimitiveComponent* RootComponent)
 {
-	// ToDo, remove call to EnableAsyncPhysicsTickCallback when ReplicationCache is refactored to run with non-GT Frozen async callback
-	EnableAsyncPhysicsTickCallback();
-
-	const FObjectKey Key(RootComponent);
-	ReplicationCache.Map.Add(Key, FReplicationCacheData(RootComponent, GetSolver()->GetSolverTime()));
+	// Create the physics replication cache if not already created
+	CreatePhysicsReplicationCache();
+	
+	PhysicsReplicationCache->RegisterForReplicationCache(RootComponent);
 }
 
-void FPhysScene_Chaos::PopulateReplicationCache(const int32 PhysicsStep)
+void FPhysScene_Chaos::CreatePhysicsReplicationCache()
 {
-	auto ReplicationCacheHelper = [this](auto& Handle, FReplicationCacheData& ReplicationData, bool& StateWasCached)
+	if (!PhysicsReplicationCache)
 	{
-		// If the component reference has lingered in the replication cache for too long without being accessed, remove it and stop caching data.
-		const Chaos::FReal CacheLingerTime = GetSolver()->GetSolverTime() - ReplicationData.GetAccessTime();
-		if (CacheLingerTime > GReplicationCacheLingerForNSeconds)
-		{
-			StateWasCached = false;
-		}
-		else
-		{
-			FRigidBodyState& ReplicationState = ReplicationData.GetState();
-			ReplicationState.Position = Handle->GetX();
-			ReplicationState.Quaternion = Handle->GetR();
-			ReplicationState.LinVel = Handle->GetV();
-			ReplicationState.AngVel = FMath::RadiansToDegrees(Handle->GetW());
-			ReplicationState.Flags = Handle->ObjectState() == Chaos::EObjectStateType::Sleeping ? ERigidBodyFlags::Sleeping : 0;
-			StateWasCached = true;
-		}
-		ReplicationData.SetIsCached(StateWasCached);
-	};
-	
-	ReplicationCache.ServerFrame = PhysicsStep;
-	bool StateWasCached;
-	for (auto It = ReplicationCache.Map.CreateIterator(); It; ++It)
-	{
-		StateWasCached = false;
-		FReplicationCacheData& ReplicationData = It.Value();
-		UPrimitiveComponent* RootComponent = ReplicationData.GetRootComponent();
-		if (RootComponent)
-		{
-			if (FBodyInstanceAsyncPhysicsTickHandle BIHandle = RootComponent->GetBodyInstanceAsyncPhysicsTickHandle())
-			{
-				ReplicationCacheHelper(BIHandle, ReplicationData, StateWasCached);
-			}
-			else if (Chaos::FPhysicsObjectHandle PhysicsObject = RootComponent->GetPhysicsObjectByName(NAME_None))
-			{
-				Chaos::FReadPhysicsObjectInterface_Internal Interface = Chaos::FPhysicsObjectInternalInterface::GetRead();
-				if (Chaos::FPBDRigidParticleHandle* POHandle = Interface.GetRigidParticle(PhysicsObject))
-				{
-					ReplicationCacheHelper(POHandle, ReplicationData, StateWasCached);
-				}
-			}
-		}
-
-		if (!StateWasCached)
-		{
-			// Deregister actor from ReplicationCache
-			It.RemoveCurrent();
-		}
+		PhysicsReplicationCache = MakeUnique<FPhysicsReplicationCache>(this);
 	}
+	check(PhysicsReplicationCache);
 }
 
 void FPhysScene_Chaos::RemoveFromComponentMaps(IPhysicsProxyBase* InObject)
diff --git a/Engine/Source/Runtime/Engine/Private/PhysicsEngine/NetworkPhysicsComponent.cpp b/Engine/Source/Runtime/Engine/Private/PhysicsEngine/NetworkPhysicsComponent.cpp
--- a/Engine/Source/Runtime/Engine/Private/PhysicsEngine/NetworkPhysicsComponent.cpp
+++ b/Engine/Source/Runtime/Engine/Private/PhysicsEngine/NetworkPhysicsComponent.cpp
@@ -176,22 +176,6 @@ bool FNetworkPhysicsRewindDataImportantStateProxy::NetSerialize(FArchive& Ar, cl
 
 // --------------------------- Network Physics Callback ---------------------------
 
-/** After physics has solved and output is created
-* NOTE: GameThread is frozen during this callback */
-void FNetworkPhysicsCallback::ApplyCallbacks_Internal(int32 PhysicsStep, const TArray<Chaos::ISimCallbackObject*>& SimCallbackObjects)
-{
-	QUICK_SCOPE_CYCLE_COUNTER(NetworkPhysicsComponent_ApplyCallbacks_Internal);
-	UpdateNetMode();
-
-	if ((NetMode == NM_ListenServer) || (NetMode == NM_DedicatedServer))
-	{
-		if (FPhysScene_Chaos* Scene = static_cast<FPhysScene_Chaos*>(World->GetPhysicsScene()))
-		{
-			Scene->PopulateReplicationCache(PhysicsStep);
-		}
-	}
-}
-
 // Before PreSimulate_Internal
 void FNetworkPhysicsCallback::ProcessInputs_Internal(int32 PhysicsStep, const TArray<Chaos::FSimCallbackInputAndObject>& SimCallbacks)
 {
@@ -230,11 +214,8 @@ int32 FNetworkPhysicsCallback::TriggerRewindIfNeeded_Internal(int32 LatestStepCo
 
 	if (RewindData)
 	{
-		if (NetMode == NM_Client)
-		{
-			const int32 ReplicationFrame = RewindData->GetResimFrame();
-			ResimFrame = (ResimFrame == INDEX_NONE) ? ReplicationFrame : (ReplicationFrame == INDEX_NONE) ? ResimFrame : FMath::Min(ReplicationFrame, ResimFrame);
-		}
+		const int32 ReplicationFrame = RewindData->GetResimFrame();
+		ResimFrame = (ResimFrame == INDEX_NONE) ? ReplicationFrame : (ReplicationFrame == INDEX_NONE) ? ResimFrame : FMath::Min(ReplicationFrame, ResimFrame);
 
 		if (ResimFrame != INDEX_NONE)
 		{
@@ -359,7 +340,7 @@ void FNetworkPhysicsCallback::ProcessInputs_External(int32 PhysicsStep, const TA
 	if (InputCmdCVars::bCmdOffsetEnabled)
 	{
 		PRAGMA_DISABLE_DEPRECATION_WARNINGS
-		if (NetMode == NM_Client)
+		if (World && World->GetNetMode() == NM_Client)
 		{
 			UpdateClientPlayer_External(PhysicsStep);
 		}
diff --git a/Engine/Source/Runtime/Engine/Private/PhysicsEngine/PhysicsReplicationCache.cpp b/Engine/Source/Runtime/Engine/Private/PhysicsEngine/PhysicsReplicationCache.cpp
new file mode 100644
--- /dev/null
+++ b/Engine/Source/Runtime/Engine/Private/PhysicsEngine/PhysicsReplicationCache.cpp
@@ -0,0 +1,267 @@
+// Copyright Epic Games, Inc. All Rights Reserved.
+
+#include "Physics/PhysicsReplicationCache.h"
+#include "Physics/Experimental/PhysScene_Chaos.h"
+#include "PBDRigidsSolver.h"
+#include "Chaos/PhysicsObjectInternalInterface.h"
+#include "Components/PrimitiveComponent.h"
+#include "Physics/NetworkPhysicsComponent.h"
+
+namespace ReplicationCacheCvars
+{
+	float LingerForSeconds = 1.0f;
+	FAutoConsoleVariableRef CVar_ReplicationCacheLingerForSeconds(TEXT("np2.ReplicationCache.LingerForNSeconds"), LingerForSeconds, TEXT("How long to keep data in the replication cache without the actor accessing it, after this we stop caching the actors state until it tries to access it again."));
+}
+
+
+
+// -------------- Game Thread --------------
+
+FPhysicsReplicationCache::FPhysicsReplicationCache(FPhysScene_Chaos* InPhysicsScene)
+{
+	using namespace Chaos;
+	
+	PhysicsScene = InPhysicsScene;
+	check(PhysicsScene);
+
+	// Create and register the async flow
+	AsyncPhysicsReplicationCache = nullptr;
+	if (Chaos::FPhysicsSolver* Solver = PhysicsScene->GetSolver())
+	{
+		AsyncPhysicsReplicationCache = Solver->CreateAndRegisterSimCallbackObject_External<FPhysicsReplicationCacheAsync>();
+	}
+}
+
+FPhysicsReplicationCache::~FPhysicsReplicationCache()
+{
+	// Unregister and free the async flow
+	if (AsyncPhysicsReplicationCache && PhysicsScene)
+	{
+		if (Chaos::FPhysicsSolver* Solver = PhysicsScene->GetSolver())
+		{
+			Solver->UnregisterAndFreeSimCallbackObject_External(AsyncPhysicsReplicationCache);
+		}
+	}
+}
+
+const FRigidBodyState* FPhysicsReplicationCache::GetStateFromReplicationCache(UPrimitiveComponent* RootComponent, int32& OutSolverFrame)
+{
+	FRigidBodyState* ReplicationState = nullptr;
+	OutSolverFrame = 0;
+
+	if (!AsyncPhysicsReplicationCache)
+	{
+		return ReplicationState;
+	}
+
+	// Process async output to get the latest cache update
+	ProcessAsyncOutput();
+
+	// Get the state for specified component or register for replication cache
+	if (Chaos::FPhysicsObjectHandle PhysicsObject = RootComponent->GetPhysicsObjectByName(NAME_None))
+	{
+		if (FPhysicsReplicationCacheAsyncInput* AsyncInput = AsyncPhysicsReplicationCache->GetProducerInputData_External())
+		{
+			// Register physics object in internal replication cache
+			AsyncInput->AccessedObjects.Add(PhysicsObject);
+		}
+
+		// Get the cached state if there is one, returns nullptr if no state is cached, which mainly happens on the first call since we haven't started caching this particles state yet at that point.
+		ReplicationState = ReplicationCache_External.Find(PhysicsObject);
+		
+		if (ReplicationState)
+		{
+			// Return the solver frame through the out parameter
+			OutSolverFrame = SolverFrame;
+		}
+	}
+
+	return ReplicationState;
+}
+
+/** Process async output to populate the replication cache on the game thread.
+* Called from FPhysicsReplicationCacheAsync::ProcessOutputs_External() which gets called as soon as a new output is available */
+void FPhysicsReplicationCache::ProcessAsyncOutput()
+{
+	if (!AsyncPhysicsReplicationCache)
+	{
+		return;
+	}
+
+	// Receive state from Physics Thread
+	while (Chaos::TSimCallbackOutputHandle<FPhysicsReplicationCacheAsyncOutput> AsyncOutput = AsyncPhysicsReplicationCache->PopFutureOutputData_External())
+	{
+		// We only care about the latest data from the physics thread
+		if (!AsyncPhysicsReplicationCache->IsOutputQueueEmpty_External())
+		{
+			continue;
+		}
+
+		if (AsyncOutput->SolverFrame <= SolverFrame)
+		{
+			continue;
+		}
+
+		// The key array and value array should always be populated with the same count
+		check(AsyncOutput->ReplicationCache_Key_Marshal.Num() == AsyncOutput->ReplicationCache_Value_Marshal.Num());
+
+		const int32 CacheSize = AsyncOutput->ReplicationCache_Key_Marshal.Num();
+
+		SolverFrame = AsyncOutput->SolverFrame;
+		ReplicationCache_External.Empty(CacheSize);
+
+		for (int32 Idx = 0; Idx < CacheSize; Idx++)
+		{
+			Chaos::FConstPhysicsObjectHandle PhysicsObject = AsyncOutput->ReplicationCache_Key_Marshal[Idx];
+			ReplicationCache_External.Emplace(PhysicsObject, AsyncOutput->ReplicationCache_Value_Marshal[Idx]);
+		}
+	}
+}
+
+void FPhysicsReplicationCache::RegisterForReplicationCache(UPrimitiveComponent* RootComponent)
+{
+	if (!RootComponent || !AsyncPhysicsReplicationCache)
+	{
+		return;
+	}
+
+	if (Chaos::FPhysicsObjectHandle PhysicsObject = RootComponent->GetPhysicsObjectByName(NAME_None))
+	{
+		if (FPhysicsReplicationCacheAsyncInput* AsyncInput = AsyncPhysicsReplicationCache->GetProducerInputData_External())
+		{
+			// Register physics object in internal replication cache
+			AsyncInput->AccessedObjects.Add(PhysicsObject);
+		}
+	}
+}
+
+void FPhysicsReplicationCache::UnregisterForReplicationCache(UPrimitiveComponent* RootComponent)
+{
+	if (!RootComponent || !AsyncPhysicsReplicationCache)
+	{
+		return;
+	}
+
+	if (Chaos::FPhysicsObjectHandle PhysicsObject = RootComponent->GetPhysicsObjectByName(NAME_None))
+	{
+		// Clear physics object from external replication cache
+		ReplicationCache_External.Remove(PhysicsObject);
+
+		// Clear physics object from internal replication cache
+		if (FPhysicsReplicationCacheAsyncInput* AsyncInput = AsyncPhysicsReplicationCache->GetProducerInputData_External())
+		{
+			AsyncInput->UnregisterObjects.Add(PhysicsObject);
+		}
+	}
+}
+
+
+
+// -------------- Physics Thread --------------
+
+void FPhysicsReplicationCacheAsync::OnPreSimulate_Internal()
+{
+	ProcessAsyncInputs();
+	PopulateReplicationCache_Internal();
+}
+
+void FPhysicsReplicationCacheAsync::OnPhysicsObjectUnregistered_Internal(Chaos::FConstPhysicsObjectHandle PhysicsObject)
+{
+	// Unregister physics object from caching state
+	ReplicationCache_Internal.Remove(PhysicsObject);
+	bUpdateAfterRemoval = true; // Ensure we produce an async output after this removal, even if the cache is empty
+}
+
+void FPhysicsReplicationCacheAsync::ProcessAsyncInputs()
+{
+	// Process async inputs from the game thread
+	if (const FPhysicsReplicationCacheAsyncInput* AsyncInput = GetConsumerInput_Internal())
+	{
+		if (Chaos::FPBDRigidsSolver* RigidSolver = static_cast<Chaos::FPBDRigidsSolver*>(GetSolver()))
+		{
+			for (const Chaos::FConstPhysicsObjectHandle& PhysicsObject : AsyncInput->AccessedObjects)
+			{
+				// Register physics object to cache state
+				FPhysicsReplicationCacheData& CacheData = ReplicationCache_Internal.FindOrAdd(PhysicsObject);
+
+				// Register time this data was accessed, used to stop caching data if the object stops getting accessed
+				CacheData.SetAccessTime(RigidSolver->GetSolverTime());
+			}
+		}
+
+		for (const Chaos::FConstPhysicsObjectHandle& PhysicsObject : AsyncInput->UnregisterObjects)
+		{
+			// Unregister physics object from caching state
+			ReplicationCache_Internal.Remove(PhysicsObject);
+			bUpdateAfterRemoval = true; // Ensure we produce an async output after this removal, even if the cache is empty
+		}
+	}
+}
+
+void FPhysicsReplicationCacheAsync::PopulateReplicationCache_Internal()
+{
+	// Early out if the cache is empty, unless we have recently removed an entry from the cache, then we need to send one last update even if the cache is empty to clear the game thread cache
+	if (ReplicationCache_Internal.Num() == 0 && !bUpdateAfterRemoval)
+	{
+		return;
+	}
+
+	Chaos::FPBDRigidsSolver* RigidSolver = static_cast<Chaos::FPBDRigidsSolver*>(GetSolver());
+	if (!RigidSolver)
+	{
+		return;
+	}
+
+	FPhysicsReplicationCacheAsyncOutput& AsyncOutput = GetProducerOutputData_Internal();
+	AsyncOutput.SolverFrame = RigidSolver->GetCurrentFrame();
+
+	// Setup the arrays to marshal replication cache with correct allocation size
+	AsyncOutput.ReplicationCache_Key_Marshal.SetNum(ReplicationCache_Internal.Num());
+	AsyncOutput.ReplicationCache_Value_Marshal.SetNum(ReplicationCache_Internal.Num());
+	int32 MarshalIndex = 0;
+
+	// Iterate over all physics objects in the cache and marshal their state back to the game thread
+	for (auto It = ReplicationCache_Internal.CreateIterator(); It; ++It)
+	{
+		FPhysicsReplicationCacheData& ReplicationCacheData = It.Value();
+
+		// Check if the object has lingered too long without being accessed
+		if (ReplicationCacheData.GetAccessTime() < (RigidSolver->GetSolverTime() - ReplicationCacheCvars::LingerForSeconds))
+		{
+			// Unregister physics object from caching state
+			It.RemoveCurrent();
+		
+			// Set empty entry in output since it's already allocated
+			AsyncOutput.ReplicationCache_Key_Marshal[MarshalIndex] = nullptr;
+			MarshalIndex++;
+		
+			continue;
+		}
+
+		// Get current state and marshal to game thread
+		const Chaos::FConstPhysicsObjectHandle& PhysicsObject = It.Key();
+		if (PhysicsObject)
+		{
+			Chaos::FReadPhysicsObjectInterface_Internal Interface = Chaos::FPhysicsObjectInternalInterface::GetRead();
+			if (Chaos::FPBDRigidParticleHandle* Handle = Interface.GetRigidParticle(PhysicsObject))
+			{
+				FRigidBodyState& ReplicationState = ReplicationCacheData.GetState();
+				ReplicationState.Position = Handle->GetP(); // Note, X is not yet updated in PostSolve_Internal so using P which is up to date
+				ReplicationState.Quaternion = Handle->GetQ(); // Note, R is not yet updated in PostSolve_Internal so using Q which is up to date
+				ReplicationState.LinVel = Handle->GetV();
+				ReplicationState.AngVel = FMath::RadiansToDegrees(Handle->GetW());
+				ReplicationState.Flags = Handle->ObjectState() == Chaos::EObjectStateType::Sleeping ? ERigidBodyFlags::Sleeping : 0;
+
+				// Marshal state
+				AsyncOutput.ReplicationCache_Key_Marshal[MarshalIndex] = PhysicsObject;
+				AsyncOutput.ReplicationCache_Value_Marshal[MarshalIndex] = ReplicationState;
+				MarshalIndex++;
+			}
+		}
+	}
+
+	bUpdateAfterRemoval = false;
+
+	// Mark async output as final, meaning it can be read on the game thread instantly instead of at the end of the physics tick
+	FinalizeOutputData_Internal();
+}
\ No newline at end of file
diff --git a/Engine/Source/Runtime/Engine/Public/Physics/Experimental/PhysScene_Chaos.h b/Engine/Source/Runtime/Engine/Public/Physics/Experimental/PhysScene_Chaos.h
--- a/Engine/Source/Runtime/Engine/Public/Physics/Experimental/PhysScene_Chaos.h
+++ b/Engine/Source/Runtime/Engine/Public/Physics/Experimental/PhysScene_Chaos.h
@@ -27,6 +27,7 @@ class UPrimitiveComponent;
 
 class AdvanceOneTimeStepTask;
 class IPhysicsReplication;
+class FPhysicsReplicationCache;
 class FPhysInterface_Chaos;
 class FChaosSolversModule;
 struct FForceFieldProxy;
@@ -255,44 +256,78 @@ public:
 	/** Register a component for physics replication state caching, the component will deregister automatically if cache is not accessed within timelimit set by CVar: np2.ReplicationCache.LingerForNSeconds */
 	ENGINE_API void RegisterForReplicationCache(UPrimitiveComponent* RootComponent);
 
-	/** Populate the replication cache from the list of registered components */
-	ENGINE_API void PopulateReplicationCache(const int32 PhysicsStep);
+	UE_DEPRECATED(5.5, "Deprecated, ReplicationCache has been moved to FPhysicsReplicationCache and it's getting populated automatically. Use GetStateFromReplicationCache() and RegisterForReplicationCache() to use the replication cache. ")
+	ENGINE_API void PopulateReplicationCache(const int32 PhysicsStep) { }
 
-	struct FReplicationCacheData
+	/** Get the FPhysicsReplicationCache */
+	ENGINE_API FPhysicsReplicationCache* GetPhysicsReplicationCache()
 	{
-		FReplicationCacheData(UPrimitiveComponent* InRootComponent, Chaos::FReal InAccessTime);
+		// Create the physics replication cache if not already created
+		CreatePhysicsReplicationCache();
 
-		UPrimitiveComponent* GetRootComponent()	{ return RootComponent.Get(); }
-		FRigidBodyState& GetState()	{ return State; }
+		return PhysicsReplicationCache.Get();
+	}
+
+	/** Create the replication cache if one doesn't exist */
+	ENGINE_API void CreatePhysicsReplicationCache();
+
+	struct UE_DEPRECATED(5.5, "Deprecated, ReplicationCache has been moved to FPhysicsReplicationCache")
+		FReplicationCacheData
+	{
+		PRAGMA_DISABLE_DEPRECATION_WARNINGS
+		FReplicationCacheData() = default;
+		FReplicationCacheData(UPrimitiveComponent * InRootComponent, Chaos::FReal InAccessTime);
+		~FReplicationCacheData() = default;
+		FReplicationCacheData(const FReplicationCacheData&) = default;
+		FReplicationCacheData(FReplicationCacheData&&) = default;
+		FReplicationCacheData& operator=(const FReplicationCacheData&) = default;
+		FReplicationCacheData& operator=(FReplicationCacheData&&) = default;
+		PRAGMA_ENABLE_DEPRECATION_WARNINGS
+
+		PRAGMA_DISABLE_DEPRECATION_WARNINGS
+		UPrimitiveComponent* GetRootComponent() { return RootComponent.Get(); }
+		FRigidBodyState& GetState() { return State; }
 		void SetAccessTime(Chaos::FReal Time) { AccessTime = Time; }
 		Chaos::FReal GetAccessTime() { return AccessTime; }
 		void SetIsCached(bool InIsCached) { bValidStateCached = InIsCached; }
 		bool IsCached() { return bValidStateCached; }
+		PRAGMA_ENABLE_DEPRECATION_WARNINGS
 
 	private:
+		PRAGMA_DISABLE_DEPRECATION_WARNINGS
 		TWeakObjectPtr<UPrimitiveComponent> RootComponent;
 		Chaos::FReal AccessTime;
 		bool bValidStateCached;
 		FRigidBodyState State;
+		PRAGMA_ENABLE_DEPRECATION_WARNINGS
 	};
 
-	// Storage structure for replication data
-	// probably should just expose read/write API not the structure directly itself like this.
-	struct FPrimitiveComponentReplicationCache
+	struct UE_DEPRECATED(5.5, "Deprecated, ReplicationCache has been moved to FPhysicsReplicationCache")
+		FPrimitiveComponentReplicationCache
 	{
+		PRAGMA_DISABLE_DEPRECATION_WARNINGS
+		FPrimitiveComponentReplicationCache() = default;
+		~FPrimitiveComponentReplicationCache() = default;
+		FPrimitiveComponentReplicationCache(const FPrimitiveComponentReplicationCache&) = default;
+		FPrimitiveComponentReplicationCache(FPrimitiveComponentReplicationCache&&) = default;
+		FPrimitiveComponentReplicationCache& operator=(const FPrimitiveComponentReplicationCache&) = default;
+		FPrimitiveComponentReplicationCache& operator=(FPrimitiveComponentReplicationCache&&) = default;
+		PRAGMA_ENABLE_DEPRECATION_WARNINGS
+
+		PRAGMA_DISABLE_DEPRECATION_WARNINGS
 		int32 ServerFrame = 0;
 		TMap<FObjectKey, FReplicationCacheData> Map;
-
 		void Reset()
 		{
 			ServerFrame = 0;
 			Map.Reset();
 		}
+		PRAGMA_ENABLE_DEPRECATION_WARNINGS
 	};
 
-	FPrimitiveComponentReplicationCache ReplicationCache;
-
 private:
+	TUniquePtr<FPhysicsReplicationCache> PhysicsReplicationCache;
+
 	TSet<UPrimitiveComponent*> CollisionEventRegistrations;
 	TSet<UPrimitiveComponent*> GlobalCollisionEventRegistrations;
 	TSet<UPrimitiveComponent*> GlobalRemovalEventRegistrations;
diff --git a/Engine/Source/Runtime/Engine/Public/Physics/NetworkPhysicsComponent.h b/Engine/Source/Runtime/Engine/Public/Physics/NetworkPhysicsComponent.h
--- a/Engine/Source/Runtime/Engine/Public/Physics/NetworkPhysicsComponent.h
+++ b/Engine/Source/Runtime/Engine/Public/Physics/NetworkPhysicsComponent.h
@@ -463,9 +463,7 @@ struct TStructOpsTypeTraits<FNetworkPhysicsRewindDataImportantStateProxy> : publ
 struct FNetworkPhysicsCallback : public Chaos::IRewindCallback
 {
 	FNetworkPhysicsCallback(UWorld* InWorld) : World(InWorld) 
-	{
-		UpdateNetMode();
-	}
+	{ }
 
 	// Delegate on the internal inputs process
 	FOnPreProcessInputsInternal PreProcessInputsInternal;
@@ -477,7 +475,6 @@ struct FNetworkPhysicsCallback : public Chaos::IRewindCallback
 	virtual void InjectInputs_External(int32 PhysicsStep, int32 NumSteps) override;
 	virtual void ProcessInputs_External(int32 PhysicsStep, const TArray<Chaos::FSimCallbackInputAndObject>& SimCallbackInputs);
 	virtual void ProcessInputs_Internal(int32 PhysicsStep, const TArray<Chaos::FSimCallbackInputAndObject>& SimCallbackInputs) override;
-	virtual void ApplyCallbacks_Internal(int32 PhysicsStep, const TArray<Chaos::ISimCallbackObject*>& SimCallbackObjects) override;
 	virtual void PreResimStep_Internal(int32 PhysicsStep, bool bFirst) override;
 	virtual void PostResimStep_Internal(int32 PhysicsStep) override;
 	virtual int32 TriggerRewindIfNeeded_Internal(int32 LatestStepCompleted) override;
@@ -502,18 +499,9 @@ struct FNetworkPhysicsCallback : public Chaos::IRewindCallback
 	UE_DEPRECATED(5.4, "Physics frame offset is handled by the PlayerController automatically, it's recommended to use APlayerController::GetAsyncPhysicsTimestamp() to get the ServerFrame and LocalFrame on both client and server. Also disable the deprecated flow by setting p.net.CmdOffsetEnabled = 0")
 	void UpdateServerPlayer_External(int32 PhysicsStep);
 
-	// Cache the current netmode for use in PT
-	void UpdateNetMode()
-	{
-		NetMode = World ? World->GetNetMode() : ENetMode::NM_Client;
-	}
-
 	// World owning that callback
 	UWorld* World = nullptr;
 
-	// Current NetMode
-	ENetMode NetMode;
-
 	// List of rewindable sim callback objects
 	TArray<Chaos::ISimCallbackObject*> RewindableCallbackObjects;
 };
diff --git a/Engine/Source/Runtime/Engine/Public/Physics/PhysicsReplicationCache.h b/Engine/Source/Runtime/Engine/Public/Physics/PhysicsReplicationCache.h
new file mode 100644
--- /dev/null
+++ b/Engine/Source/Runtime/Engine/Public/Physics/PhysicsReplicationCache.h
@@ -0,0 +1,120 @@
+// Copyright Epic Games, Inc. All Rights Reserved.
+
+/*=============================================================================
+	PhysicsReplicationCache.h 
+	Cache physics state from PT on GT for actors using physics replication.
+=============================================================================*/
+
+#pragma once
+
+#include "Chaos/SimCallbackObject.h"
+#include "Engine/ReplicatedState.h"
+
+class FPhysScene_Chaos;
+class UPrimitiveComponent;
+
+// -------------- Game Thread --------------
+
+class FPhysicsReplicationCache
+{
+public:
+	FPhysicsReplicationCache(FPhysScene_Chaos* InPhysicsScene);
+	
+	virtual ~FPhysicsReplicationCache();
+
+	/** Get cached state for replication, if no state is cached RegisterForReplicationCache() is called */
+	const FRigidBodyState* GetStateFromReplicationCache(UPrimitiveComponent* RootComponent, int32& OutSolverFrame);
+
+	/** Register a component for physics replication state caching, the component will unregister automatically if cache is not accessed within time limit set by CVar: np2.ReplicationCache.LingerForNSeconds */
+	void RegisterForReplicationCache(UPrimitiveComponent* RootComponent);
+	
+	/** Unregister a component from physics replication state caching */
+	void UnregisterForReplicationCache(UPrimitiveComponent* RootComponent);
+
+private:
+	/** Process marshaled data from the physics thread */
+	void ProcessAsyncOutput();
+
+private:
+	int32 SolverFrame = 0;
+	TMap<Chaos::FConstPhysicsObjectHandle, FRigidBodyState> ReplicationCache_External;
+	
+	FPhysScene_Chaos* PhysicsScene;
+	FDelegateHandle DelegateInjectInputs_External;
+
+	/** Async part of the replication cache, only access on Physics Thread except for initialize and uninitialize */
+	class FPhysicsReplicationCacheAsync* AsyncPhysicsReplicationCache;
+};
+
+
+
+// -------------- Async Marshaling --------------
+
+/** Async Marshal Input */
+struct FPhysicsReplicationCacheAsyncInput : public Chaos::FSimCallbackInput
+{
+	TArray<Chaos::FConstPhysicsObjectHandle> AccessedObjects;
+	TArray<Chaos::FConstPhysicsObjectHandle> UnregisterObjects;
+
+	void Reset()
+	{
+		AccessedObjects.Reset();
+		UnregisterObjects.Reset();
+	}
+};
+
+/** Async Marshal Output */
+struct FPhysicsReplicationCacheAsyncOutput : public Chaos::FSimCallbackOutput
+{
+	int32 SolverFrame;
+	TArray<Chaos::FConstPhysicsObjectHandle> ReplicationCache_Key_Marshal;
+	TArray<FRigidBodyState> ReplicationCache_Value_Marshal;
+
+	void Reset()
+	{
+		SolverFrame = 0;
+	}
+};
+
+
+
+// -------------- Physics Thread --------------
+
+/** Replication cache data holding state and access time */
+struct FPhysicsReplicationCacheData
+{
+public:
+	FPhysicsReplicationCacheData()
+	{ }
+
+	void SetAccessTime(double Time) { AccessTime = Time; }
+	const double GetAccessTime() const { return AccessTime; }
+	FRigidBodyState& GetState() { return StateData; }
+
+private:
+	double AccessTime;
+	FRigidBodyState StateData;
+};
+
+/** Async replication cache class */
+class FPhysicsReplicationCacheAsync : public Chaos::TSimCallbackObject<
+	FPhysicsReplicationCacheAsyncInput,
+	FPhysicsReplicationCacheAsyncOutput,
+	Chaos::ESimCallbackOptions::Presimulate | Chaos::ESimCallbackOptions::PhysicsObjectUnregister>
+{
+	virtual void OnPreSimulate_Internal() override;
+
+	/** Callback when a physics object is removed */
+	virtual void OnPhysicsObjectUnregistered_Internal(Chaos::FConstPhysicsObjectHandle PhysicsObject) override;
+
+	/** Process incoming async input data from GameThread, register and unregister particles to/from cache */
+	void ProcessAsyncInputs();
+
+	/** Populate replication cache by iterating over particles and populating async output with data */
+	void PopulateReplicationCache_Internal();
+
+private:
+	TMap<Chaos::FConstPhysicsObjectHandle, FPhysicsReplicationCacheData> ReplicationCache_Internal;
+	bool bUpdateAfterRemoval = false;
+};
+
diff --git a/Engine/Source/Runtime/Experimental/Chaos/Public/Chaos/SimCallbackObject.h b/Engine/Source/Runtime/Experimental/Chaos/Public/Chaos/SimCallbackObject.h
--- a/Engine/Source/Runtime/Experimental/Chaos/Public/Chaos/SimCallbackObject.h
+++ b/Engine/Source/Runtime/Experimental/Chaos/Public/Chaos/SimCallbackObject.h
@@ -185,9 +185,15 @@ public:
 		return Solver; 
 	}
 
-	// Rewind API
+	UE_DEPRECATED(5.5, "This callback was never called. Instead look at either ISimCallbackObject::ProcessInputs_External or FNetworkPhysicsCallback::InjectInputsExternal.")
 	virtual void InjectInputs_External(int32 PhysicsStep, int32 NumSteps) {}
+
+	/** Called before any sim callbacks are triggered but after async inputs has marshaled over
+	* Register with ESimCallbackOptions::Rewind to get callback */
 	virtual void ProcessInputs_Internal(int32 PhysicsStep) {}
+
+	/** Called before async inputs are marshaled over to the physics thread, only called if an async input is produced to get marshaled 
+	* Register with ESimCallbackOptions::Rewind to get callback */
 	virtual void ProcessInputs_External(int32 PhysicsStep) {}
 	
 	virtual int32 TriggerRewindIfNeeded_Internal(int32 LastCompletedStep)
@@ -572,6 +578,15 @@ public:
 		return TSimCallbackOutputHandle<TOutputType>(Output, this);
 	}
 
+	/**
+	* Check if the output queue is empty of data
+	* Can be used while iterating the queue through PopOutputData_External and PopFutureOutputData_External to check if the current data is the last data.
+	*/
+	bool IsOutputQueueEmpty_External()
+	{
+		return OutputQueue.IsEmpty();
+	}
+
 	/**
 	* Gets the current producer output data. This is what the callback generates. If multiple callbacks are triggered in one step, the same output is used
 	*/
diff --git a/Engine/Source/Runtime/Experimental/GeometryCollectionEngine/Private/GeometryCollection/GeometryCollectionComponent.cpp b/Engine/Source/Runtime/Experimental/GeometryCollectionEngine/Private/GeometryCollection/GeometryCollectionComponent.cpp
--- a/Engine/Source/Runtime/Experimental/GeometryCollectionEngine/Private/GeometryCollection/GeometryCollectionComponent.cpp
+++ b/Engine/Source/Runtime/Experimental/GeometryCollectionEngine/Private/GeometryCollection/GeometryCollectionComponent.cpp
@@ -2332,10 +2332,8 @@ void UGeometryCollectionComponent::UpdateRepData()
 
 		if (LocalRepData.HasChanged(RepData))
 		{
-			if (Owner->GetWorld() && Owner->GetWorld()->GetPhysicsScene())
-			{
-				LocalRepData.ServerFrame = Owner->GetWorld()->GetPhysicsScene()->ReplicationCache.ServerFrame;
-			}
+			LocalRepData.ServerFrame = Solver->GetMarshallingManager().GetInternalStep_External();
+
 			INC_DWORD_STAT_BY(STAT_GCReplicatedClusters, LocalRepData.Clusters.Num());
 
 			FlushNetDormancyIfNeeded();
