[NetPhysics] Force correct state when setting resim object to sleep.

diff --git a/Engine/Source/Runtime/Engine/Private/PhysicsEngine/PhysicsReplication.cpp b/Engine/Source/Runtime/Engine/Private/PhysicsEngine/PhysicsReplication.cpp
--- a/Engine/Source/Runtime/Engine/Private/PhysicsEngine/PhysicsReplication.cpp
+++ b/Engine/Source/Runtime/Engine/Private/PhysicsEngine/PhysicsReplication.cpp
@@ -2032,6 +2032,12 @@ bool FPhysicsReplicationAsync::ResimulationReplication(Chaos::FPBDRigidParticleH
 	}
 #endif
 
+	// Wake up if is sleeping and should not sleep
+	if (Handle->IsSleeping() && !bShouldSleep)
+	{
+		RigidsSolver->GetEvolution()->SetParticleObjectState(Handle, Chaos::EObjectStateType::Dynamic);
+	}
+
 	if (ShouldTriggerResim && Target.TickCount == 0 && LocalFrame > RewindData->GetBlockedResimFrame())
 	{
 		// Trigger resimulation
@@ -2043,12 +2049,6 @@ bool FPhysicsReplicationAsync::ResimulationReplication(Chaos::FPBDRigidParticleH
 	}
 	else if (SettingsCurrent.ResimulationSettings.GetRuntimeCorrectionEnabled())
 	{
-		// Wake up if is sleeping and should not sleep
-		if (Handle->IsSleeping() && !bShouldSleep)
-		{
-			RigidsSolver->GetEvolution()->SetParticleObjectState(Handle, Chaos::EObjectStateType::Dynamic);
-		}
-
 		const int32 NumPredictedFrames = RigidsSolver->GetCurrentFrame() - LocalFrame - Target.TickCount;
 
 		if (Target.TickCount <= NumPredictedFrames && NumPredictedFrames > 0)
@@ -2100,6 +2100,9 @@ bool FPhysicsReplicationAsync::ResimulationReplication(Chaos::FPBDRigidParticleH
 	// Set sleep state if we are about to clear the target from memory and the target is set to sleep
 	if (bClearTarget && bShouldSleep)
 	{
+		// Snap object into correct state, it should already be at that state or very close to it
+		RigidsSolver->GetEvolution()->ApplyParticleTransformCorrection(Handle, Target.TargetState.Position, Target.TargetState.Quaternion, /*bApplyToConnectedBodies*/true, /*bInRecalculateFrictionOnConnectedBodies*/true, ReplicatedParticleIDs);
+
 		RigidsSolver->GetEvolution()->SetParticleObjectState(Handle, Chaos::EObjectStateType::Sleeping);
 		if (PhysicsReplicationCVars::PredictiveInterpolationCVars::bSleepConnectedBodies)
 		{
