StaticBools are not cached correctly for MP_FrontMaterial connections being resolved between expression compilation paths (rather than being written to the shader). Added a new function in expressions to check when a recompile should be forced (which is overridable to allow further extension in future if needed).
Also fixed static bool parameters not being resolved correctly through material function inputs by adding the compile push/pop function state calls to ensure the required data is present for resolving the child function call.

diff --git a/Engine/Source/Runtime/Engine/Private/Materials/HLSLMaterialTranslator.cpp b/Engine/Source/Runtime/Engine/Private/Materials/HLSLMaterialTranslator.cpp
--- a/Engine/Source/Runtime/Engine/Private/Materials/HLSLMaterialTranslator.cpp
+++ b/Engine/Source/Runtime/Engine/Private/Materials/HLSLMaterialTranslator.cpp
@@ -4475,17 +4475,23 @@ int32 FHLSLMaterialTranslator::CallExpression(FMaterialExpressionKey ExpressionK
 		ExpressionKey.MaterialAttributeID = FGuid(0, 0, 0, 0);
 	}
 
-	// Some expressions can discard output indices and share compiles with a swizzle/mask
-	if (ExpressionKey.Expression && ExpressionKey.Expression->CanIgnoreOutputIndex())
-	{
-		ExpressionKey.OutputIndex = INDEX_NONE;
-	}
 
 	// Substrate BSDF expression should not be de-duplicated using expression output hash. 
 	// This is automatically handled via the compiler SubstrateTreeStack.
 	// It means that a node can be blended at multiple point of the graph (allowing acyclic graph instead of tree, e.g. a Slab can be used into multiple input).
 	// We do this for all SubstrateData which can be output from BSDF nodes (slabs) or other nodes such as material functions.
-	const bool bExpressionIsSubstrate = ExpressionKey.Expression && ExpressionKey.Expression->IsResultSubstrateMaterial(ExpressionKey.OutputIndex);
+	bool bExpressionIsSubstrate = false;
+	if(Substrate::IsSubstrateEnabled() && ExpressionKey.Expression)
+	{
+		bExpressionIsSubstrate = ExpressionKey.Expression->IsResultSubstrateMaterial(ExpressionKey.OutputIndex) 
+							|| ExpressionKey.Expression->ForceSubstrateExpressionRecompile(Compiler, ExpressionKey.OutputIndex);
+	}
+
+	// Some expressions can discard output indices and share compiles with a swizzle/mask
+	if (ExpressionKey.Expression && ExpressionKey.Expression->CanIgnoreOutputIndex())
+	{
+		ExpressionKey.OutputIndex = INDEX_NONE;
+	}
 
 	// Check if this expression has already been translated.
 	check(ShaderFrequency < SF_NumFrequencies);
diff --git a/Engine/Source/Runtime/Engine/Private/Materials/MaterialExpressions.cpp b/Engine/Source/Runtime/Engine/Private/Materials/MaterialExpressions.cpp
--- a/Engine/Source/Runtime/Engine/Private/Materials/MaterialExpressions.cpp
+++ b/Engine/Source/Runtime/Engine/Private/Materials/MaterialExpressions.cpp
@@ -2382,6 +2382,16 @@ bool UMaterialExpression::IsUsingNewHLSLGenerator() const
 	return false;
 }
 
+bool UMaterialExpression::ForceSubstrateExpressionRecompile(FMaterialCompiler* Compiler, int32 OutputIndex)
+{	
+	//Currently this is only needed for FM compilation for Static Bool resolutions, legacy compilation is unaffected.
+	if(Compiler->GetMaterialAttribute() == FMaterialAttributeDefinitionMap::GetID(MP_FrontMaterial))
+	{
+		return GetOutputType(OutputIndex) == MCT_StaticBool;
+	}
+	return false;
+}
+
 FSubstrateOperator* UMaterialExpression::SubstrateGenerateMaterialTopologyTree(class FMaterialCompiler* Compiler, class UMaterialExpression* Parent, int32 OutputIndex)
 {
 	Compiler->Errorf(TEXT("Missing SubstrateGenerateMaterialTopologyTree implementation for node %s."), *GetClass()->GetName());
@@ -19673,7 +19683,12 @@ FSubstrateOperator* UMaterialExpressionMaterialFunctionCall::SubstrateGenerateMa
 	if (OutputIndex >= 0 && OutputIndex < FunctionOutputs.Num() && FunctionOutputs[OutputIndex].ExpressionOutput)
 	{
 		this->LinkFunctionIntoCaller(nullptr);
+		FMaterialFunctionCompileState LocalState(this);
+		Compiler->PushFunction(SharedCompileState ? SharedCompileState : &LocalState);
+
 		FSubstrateOperator* ResultingOperator = FunctionOutputs[OutputIndex].ExpressionOutput->SubstrateGenerateMaterialTopologyTree(Compiler, Parent, 0);
+
+		Compiler->PopFunction();
 		this->UnlinkFunctionFromCaller(nullptr);
 		return ResultingOperator;
 	}
diff --git a/Engine/Source/Runtime/Engine/Public/Materials/MaterialExpression.h b/Engine/Source/Runtime/Engine/Public/Materials/MaterialExpression.h
--- a/Engine/Source/Runtime/Engine/Public/Materials/MaterialExpression.h
+++ b/Engine/Source/Runtime/Engine/Public/Materials/MaterialExpression.h
@@ -328,6 +328,12 @@ class UMaterialExpression : public UObject
 
 	ENGINE_API bool IsUsingNewHLSLGenerator() const;
 
+	/**
+	* For now this call is only needed for resolving Static values, but can be expanded in future as needed.
+	* Essentially this forces the HLSLMaterialTranslator to skip cached code chunks and use the static bool logic in place for resolving switches etc.
+	*/
+	ENGINE_API virtual bool ForceSubstrateExpressionRecompile(class FMaterialCompiler* Compiler, int32 OutputIndex);
+
 	/**
 	 * Builds this expression's material IR.
 	 * 
