Enable Hardware VRS support for XSX.
This will trigger an XSX shader recompile for any projects with "Support Hardware Variable Rate Shading" enabled under Rendering -> Optimizations.

diff --git a/Engine/Source/Runtime/Engine/Private/Materials/MaterialShared.cpp b/Engine/Source/Runtime/Engine/Private/Materials/MaterialShared.cpp
--- a/Engine/Source/Runtime/Engine/Private/Materials/MaterialShared.cpp
+++ b/Engine/Source/Runtime/Engine/Private/Materials/MaterialShared.cpp
@@ -2811,12 +2811,6 @@ void FMaterial::SetupMaterialEnvironment(
 		SET_SHADER_DEFINE(OutEnvironment, USE_STENCIL_LOD_DITHER_DEFAULT, CVar->GetValueOnAnyThread() != 0 ? 1 : 0);
 	}
 
-	if (FDataDrivenShaderPlatformInfo::GetSupportsVariableRateShading(Platform) && IsVariableRateShadingAllowed())
-	{
-		static const auto CVar = IConsoleManager::Get().FindTConsoleVariableDataInt(TEXT("r.VRS.Support"));
-		OutEnvironment.SetCompileArgument(TEXT("USING_VARIABLE_RATE_SHADING"), CVar->GetValueOnAnyThread() != 0 ? 1 : 0);
-	}
-
 	{
 		switch (GetMaterialDomain())
 		{
@@ -5706,6 +5700,7 @@ FMaterialShaderParameters::FMaterialShaderParameters(const FMaterial* InMaterial
 	bIsMobileSeparateTranslucencyEnabled = InMaterial->IsMobileSeparateTranslucencyEnabled();
 	bAlwaysEvaluateWorldPositionOffset = InMaterial->ShouldAlwaysEvaluateWorldPositionOffset();
 	bDisablePreExposureScale = InMaterial->GetDisablePreExposureScale();
+	bAllowVariableRateShading = InMaterial->IsVariableRateShadingAllowed();
 }
 
 #if WITH_EDITOR
diff --git a/Engine/Source/Runtime/Engine/Public/MaterialShared.h b/Engine/Source/Runtime/Engine/Public/MaterialShared.h
--- a/Engine/Source/Runtime/Engine/Public/MaterialShared.h
+++ b/Engine/Source/Runtime/Engine/Public/MaterialShared.h
@@ -1057,6 +1057,7 @@ struct FMaterialShaderParameters
 			uint64 bIsMobileSeparateTranslucencyEnabled : 1;
 			uint64 bAlwaysEvaluateWorldPositionOffset : 1;
 			uint64 bDisablePreExposureScale : 1;
+			uint64 bAllowVariableRateShading : 1;
 		};
 	};
 
diff --git a/Engine/Source/Runtime/Renderer/Private/BasePassRendering.h b/Engine/Source/Runtime/Renderer/Private/BasePassRendering.h
--- a/Engine/Source/Runtime/Renderer/Private/BasePassRendering.h
+++ b/Engine/Source/Runtime/Renderer/Private/BasePassRendering.h
@@ -274,6 +274,11 @@ public:
 	{
 		LightMapPolicyType::ModifyCompilationEnvironment(Parameters, OutEnvironment);
 		Super::ModifyCompilationEnvironment(Parameters, OutEnvironment);
+
+		if (HardwareVariableRateShadingSupportedByPlatform(Parameters.Platform) && Parameters.MaterialParameters.bAllowVariableRateShading)
+		{
+			OutEnvironment.SetCompileArgument(TEXT("USING_VARIABLE_RATE_SHADING"), true);
+		}
 	}
 };
 
@@ -515,6 +520,11 @@ public:
 	{
 		FMeshMaterialShader::ModifyCompilationEnvironment(Parameters, OutEnvironment);
 		FForwardLightingParameters::ModifyCompilationEnvironment(Parameters.Platform, OutEnvironment);
+
+		if (HardwareVariableRateShadingSupportedByPlatform(Parameters.Platform) && Parameters.MaterialParameters.bAllowVariableRateShading)
+		{
+			OutEnvironment.SetCompileArgument(TEXT("USING_VARIABLE_RATE_SHADING"), true);
+		}
 	}
 
 	static bool ValidateCompiledResult(EShaderPlatform Platform, const FShaderParameterMap& ParameterMap, TArray<FString>& OutError)
