World Partition - Handle the rare case where a streaming level outlives its world
* Since those three objects (runtime cell, streaming level, world) are independant, they can possibly have different lifetime.
* The OnCellShown() call will be skipped if the level streaming is alive but its cell is not, as the delegate IsBound() test will reject it
* The crash occurs if both the level streaming object and the runtime cell are alive, but the world is not
#tests Forced the rare objects states required to repro

diff --git a/Engine/Source/Runtime/Engine/Private/WorldPartition/WorldPartitionRuntimeLevelStreamingCell.cpp b/Engine/Source/Runtime/Engine/Private/WorldPartition/WorldPartitionRuntimeLevelStreamingCell.cpp
--- a/Engine/Source/Runtime/Engine/Private/WorldPartition/WorldPartitionRuntimeLevelStreamingCell.cpp
+++ b/Engine/Source/Runtime/Engine/Private/WorldPartition/WorldPartitionRuntimeLevelStreamingCell.cpp
@@ -697,10 +697,13 @@ void UWorldPartitionRuntimeLevelStreamingCell::OnLevelShown()
 
 void UWorldPartitionRuntimeLevelStreamingCell::OnCellShown() const
 {
-	// Temporary test while we investigate issue PLAY-45493
-	if (ensure(IsValid(this)))
+	// Test if the outer world is valid to handle the rare case where a streaming level outlives its world
+	// * Since those three objects are independant, they can possibly have different lifetime.
+	// * The OnCellShown() call will be skipped if the level streaming is alive but its cell is not, as the delegate IsBound() test will reject it
+	// * A crash would occurs if both the level streaming object and the runtime cell are alive, but the world is not
+	if (UWorld* OuterWorld = GetOuterWorld())
 	{
-		UWorldPartition* OuterWorldPartition = GetOuterWorld()->GetWorldPartition();
+		UWorldPartition* OuterWorldPartition = OuterWorld->GetWorldPartition();
 		if (OuterWorldPartition && OuterWorldPartition->IsInitialized())
 		{
 			OuterWorldPartition->OnCellShown(this);
@@ -715,10 +718,13 @@ void UWorldPartitionRuntimeLevelStreamingCell::OnLevelHidden()
 
 void UWorldPartitionRuntimeLevelStreamingCell::OnCellHidden() const
 {
-	// Temporary test while we investigate issue PLAY-45493
-	if (ensure(IsValid(this)))
+	// Test if the outer world is valid to handle the rare case where a streaming level outlives its world
+	// * Since those three objects are independant, they can possibly have different lifetime.
+	// * The OnCellShown() call will be skipped if the level streaming is alive but its cell is not, as the delegate IsBound() test will reject it
+	// * A crash would occurs if both the level streaming object and the runtime cell are alive, but the world is not
+	if (UWorld* OuterWorld = GetOuterWorld())
 	{
-		UWorldPartition* OuterWorldPartition = GetOuterWorld()->GetWorldPartition();
+		UWorldPartition* OuterWorldPartition = OuterWorld->GetWorldPartition();
 		if (OuterWorldPartition && OuterWorldPartition->IsInitialized())
 		{
 			OuterWorldPartition->OnCellHidden(this);
