Revised substrate compiler logic to ensure the function stacks are reset prior to compiling the FrontMaterial, otherwise a mismatch can occur between anticipated stacks and the compiled code chunks. Reverts previous verison of this logic which forced a recompile more directly.

diff --git a/Engine/Source/Runtime/Engine/Private/Materials/HLSLMaterialTranslator.cpp b/Engine/Source/Runtime/Engine/Private/Materials/HLSLMaterialTranslator.cpp
--- a/Engine/Source/Runtime/Engine/Private/Materials/HLSLMaterialTranslator.cpp
+++ b/Engine/Source/Runtime/Engine/Private/Materials/HLSLMaterialTranslator.cpp
@@ -1435,6 +1435,18 @@ void FHLSLMaterialTranslator::TranslateMaterial()
 			bSubstrateWritesAmbientOcclusion |= SubstrateCtx.bSubstrateWritesAmbientOcclusion;
 		}
 		CurrentSubstrateCompilationContext = ESubstrateCompilationContext::SCC_Default;
+
+		//Clear the stacks as they may be referencing scoped chunks that are lost before the compile step.
+		//Skips the Vertex shader because this is handled whilst resetting custom vertex interpolation below.
+		for (int iFrequency = EShaderFrequency::SF_Vertex + 1; iFrequency < EShaderFrequency::SF_NumFrequencies; iFrequency++)
+		{
+			while (FunctionStacks[iFrequency].Num() > 1)
+			{
+				FMaterialFunctionCompileState* Stack = FunctionStacks[iFrequency].Pop(EAllowShrinking::No);
+				delete Stack;
+			}
+			FunctionStacks[iFrequency][0]->Reset();
+		}
 	}
 
 	// Generate code:
@@ -4475,24 +4487,18 @@ int32 FHLSLMaterialTranslator::CallExpression(FMaterialExpressionKey ExpressionK
 		ExpressionKey.MaterialAttributeID = FGuid(0, 0, 0, 0);
 	}
 
-
-	// Substrate BSDF expression should not be de-duplicated using expression output hash. 
-	// This is automatically handled via the compiler SubstrateTreeStack.
-	// It means that a node can be blended at multiple point of the graph (allowing acyclic graph instead of tree, e.g. a Slab can be used into multiple input).
-	// We do this for all SubstrateData which can be output from BSDF nodes (slabs) or other nodes such as material functions.
-	bool bExpressionIsSubstrate = false;
-	if(Substrate::IsSubstrateEnabled() && ExpressionKey.Expression)
-	{
-		bExpressionIsSubstrate = ExpressionKey.Expression->IsResultSubstrateMaterial(ExpressionKey.OutputIndex) 
-							|| ExpressionKey.Expression->ForceSubstrateExpressionRecompile(Compiler, ExpressionKey.OutputIndex);
-	}
-
 	// Some expressions can discard output indices and share compiles with a swizzle/mask
 	if (ExpressionKey.Expression && ExpressionKey.Expression->CanIgnoreOutputIndex())
 	{
 		ExpressionKey.OutputIndex = INDEX_NONE;
 	}
 
+	// Substrate BSDF expression should not be de-duplicated using expression output hash. 
+	// This is automatically handled via the compiler SubstrateTreeStack.
+	// It means that a node can be blended at multiple point of the graph (allowing acyclic graph instead of tree, e.g. a Slab can be used into multiple input).
+	// We do this for all SubstrateData which can be output from BSDF nodes (slabs) or other nodes such as material functions.
+	const bool bExpressionIsSubstrate = ExpressionKey.Expression && ExpressionKey.Expression->IsResultSubstrateMaterial(ExpressionKey.OutputIndex);
+
 	// Check if this expression has already been translated.
 	check(ShaderFrequency < SF_NumFrequencies);
 	auto& CurrentFunctionStack = FunctionStacks[ShaderFrequency];
diff --git a/Engine/Source/Runtime/Engine/Private/Materials/MaterialExpressions.cpp b/Engine/Source/Runtime/Engine/Private/Materials/MaterialExpressions.cpp
--- a/Engine/Source/Runtime/Engine/Private/Materials/MaterialExpressions.cpp
+++ b/Engine/Source/Runtime/Engine/Private/Materials/MaterialExpressions.cpp
@@ -2382,16 +2382,6 @@ bool UMaterialExpression::IsUsingNewHLSLGenerator() const
 	return false;
 }
 
-bool UMaterialExpression::ForceSubstrateExpressionRecompile(FMaterialCompiler* Compiler, int32 OutputIndex)
-{	
-	//Currently this is only needed for FM compilation for Static Bool resolutions, legacy compilation is unaffected.
-	if(Compiler->GetMaterialAttribute() == FMaterialAttributeDefinitionMap::GetID(MP_FrontMaterial))
-	{
-		return GetOutputType(OutputIndex) == MCT_StaticBool;
-	}
-	return false;
-}
-
 FSubstrateOperator* UMaterialExpression::SubstrateGenerateMaterialTopologyTree(class FMaterialCompiler* Compiler, class UMaterialExpression* Parent, int32 OutputIndex)
 {
 	Compiler->Errorf(TEXT("Missing SubstrateGenerateMaterialTopologyTree implementation for node %s."), *GetClass()->GetName());
diff --git a/Engine/Source/Runtime/Engine/Public/Materials/MaterialExpression.h b/Engine/Source/Runtime/Engine/Public/Materials/MaterialExpression.h
--- a/Engine/Source/Runtime/Engine/Public/Materials/MaterialExpression.h
+++ b/Engine/Source/Runtime/Engine/Public/Materials/MaterialExpression.h
@@ -328,12 +328,6 @@ class UMaterialExpression : public UObject
 
 	ENGINE_API bool IsUsingNewHLSLGenerator() const;
 
-	/**
-	* For now this call is only needed for resolving Static values, but can be expanded in future as needed.
-	* Essentially this forces the HLSLMaterialTranslator to skip cached code chunks and use the static bool logic in place for resolving switches etc.
-	*/
-	ENGINE_API virtual bool ForceSubstrateExpressionRecompile(class FMaterialCompiler* Compiler, int32 OutputIndex);
-
 	/**
 	 * Builds this expression's material IR.
 	 * 
