buildfix when logging is disabled

diff --git a/Engine/Source/Runtime/Engine/Private/ODSC/ODSCManager.cpp b/Engine/Source/Runtime/Engine/Private/ODSC/ODSCManager.cpp
--- a/Engine/Source/Runtime/Engine/Private/ODSC/ODSCManager.cpp
+++ b/Engine/Source/Runtime/Engine/Private/ODSC/ODSCManager.cpp
@@ -39,7 +39,7 @@ class FODSCManagerAccess
 public:
 	static void ODSCLogMissedMaterials(const TArray<FString>& Args)
 	{
-#if WITH_ODSC
+#if WITH_ODSC && !NO_LOGGING
 		if (FODSCManager::IsODSCActive())
 		{
 			TArray<FString> MaterialPaths;
