Const correctness for some physics API (no change to functionality)

diff --git a/Engine/Source/Runtime/Engine/Classes/PhysicsEngine/PhysicsAsset.h b/Engine/Source/Runtime/Engine/Classes/PhysicsEngine/PhysicsAsset.h
--- a/Engine/Source/Runtime/Engine/Classes/PhysicsEngine/PhysicsAsset.h
+++ b/Engine/Source/Runtime/Engine/Classes/PhysicsEngine/PhysicsAsset.h
@@ -290,19 +290,19 @@ public:
 	//~ End UObject Interface
 
 	// Find the index of the physics bone that is controlling this graphics bone.
-	ENGINE_API int32		FindControllingBodyIndex(class USkeletalMesh* skelMesh, int32 BoneIndex);
-	ENGINE_API int32		FindParentBodyIndex(class USkeletalMesh * skelMesh, int32 StartBoneIndex) const;
+	ENGINE_API int32		FindControllingBodyIndex(const class USkeletalMesh* skelMesh, int32 BoneIndex) const;
+	ENGINE_API int32		FindParentBodyIndex(const class USkeletalMesh * skelMesh, int32 StartBoneIndex) const;
 	ENGINE_API int32		FindParentBodyIndex(const FReferenceSkeleton& RefSkeleton, const int32 StartBoneIndex) const;
-	ENGINE_API int32		FindConstraintIndex(FName ConstraintName);
-	ENGINE_API int32		FindConstraintIndex(FName Bone1Name, FName Bone2Name);
-	FName					FindConstraintBoneName(int32 ConstraintIndex);
-	ENGINE_API int32		FindMirroredBone(class USkeletalMesh* skelMesh, int32 BoneIndex);
+	ENGINE_API int32		FindConstraintIndex(FName ConstraintName) const;
+	ENGINE_API int32		FindConstraintIndex(FName Bone1Name, FName Bone2Name) const;
+	FName					FindConstraintBoneName(int32 ConstraintIndex) const;
+	ENGINE_API int32		FindMirroredBone(const class USkeletalMesh* skelMesh, int32 BoneIndex) const;
 
 	/** Utility for getting indices of all bodies below (and including) the one with the supplied name. */
-	ENGINE_API void			GetBodyIndicesBelow(TArray<int32>& OutBodyIndices, FName InBoneName, USkeletalMesh* InSkelMesh, bool bIncludeParent = true);
-	ENGINE_API void			GetBodyIndicesBelow(TArray<int32>& OutBodyIndices, const FName InBoneName, const FReferenceSkeleton& RefSkeleton, const bool bIncludeParent = true);
+	ENGINE_API void			GetBodyIndicesBelow(TArray<int32>& OutBodyIndices, FName InBoneName, const USkeletalMesh* InSkelMesh, bool bIncludeParent = true) const;
+	ENGINE_API void			GetBodyIndicesBelow(TArray<int32>& OutBodyIndices, const FName InBoneName, const FReferenceSkeleton& RefSkeleton, const bool bIncludeParent = true) const;
 
-	ENGINE_API void			GetNearestBodyIndicesBelow(TArray<int32> & OutBodyIndices, FName InBoneName, USkeletalMesh * InSkelMesh);
+	ENGINE_API void			GetNearestBodyIndicesBelow(TArray<int32> & OutBodyIndices, FName InBoneName, const USkeletalMesh* InSkelMesh) const;
 
 	ENGINE_API FBox			CalcAABB(const class USkinnedMeshComponent* MeshComponent, const FTransform& LocalToWorld) const;
 
diff --git a/Engine/Source/Runtime/Engine/Private/PhysicsEngine/PhysicsAsset.cpp b/Engine/Source/Runtime/Engine/Private/PhysicsEngine/PhysicsAsset.cpp
--- a/Engine/Source/Runtime/Engine/Private/PhysicsEngine/PhysicsAsset.cpp
+++ b/Engine/Source/Runtime/Engine/Private/PhysicsEngine/PhysicsAsset.cpp
@@ -401,7 +401,7 @@ bool UPhysicsAsset::CanCalculateValidAABB(const USkinnedMeshComponent* MeshComp,
 }
 #endif //WITH_EDITOR
 
-int32	UPhysicsAsset::FindControllingBodyIndex(class USkeletalMesh* skelMesh, int32 StartBoneIndex)
+int32 UPhysicsAsset::FindControllingBodyIndex(const class USkeletalMesh* skelMesh, int32 StartBoneIndex) const
 {
 	int32 BoneIndex = StartBoneIndex;
 	while(BoneIndex!=INDEX_NONE)
@@ -423,7 +423,7 @@ int32	UPhysicsAsset::FindControllingBodyIndex(class USkeletalMesh* skelMesh, int
 	return INDEX_NONE; // Shouldn't reach here.
 }
 
-int32 UPhysicsAsset::FindParentBodyIndex(class USkeletalMesh* skelMesh, int32 StartBoneIndex) const
+int32 UPhysicsAsset::FindParentBodyIndex(const class USkeletalMesh* skelMesh, int32 StartBoneIndex) const
 {
 	if (skelMesh)
 	{
@@ -463,7 +463,7 @@ int32 UPhysicsAsset::FindBodyIndex(FName bodyName) const
 	return INDEX_NONE;
 }
 
-int32 UPhysicsAsset::FindConstraintIndex(FName ConstraintName)
+int32 UPhysicsAsset::FindConstraintIndex(FName ConstraintName) const
 {
 	for(int32 i=0; i<ConstraintSetup.Num(); i++)
 	{
@@ -476,7 +476,7 @@ int32 UPhysicsAsset::FindConstraintIndex(FName ConstraintName)
 	return INDEX_NONE;
 }
 
-int32 UPhysicsAsset::FindConstraintIndex(FName Bone1Name, FName Bone2Name)
+int32 UPhysicsAsset::FindConstraintIndex(FName Bone1Name, FName Bone2Name) const
 {
 	for (int32 i = 0; i < ConstraintSetup.Num(); i++)
 	{
@@ -490,7 +490,7 @@ int32 UPhysicsAsset::FindConstraintIndex(FName Bone1Name, FName Bone2Name)
 	return INDEX_NONE;
 }
 
-FName UPhysicsAsset::FindConstraintBoneName(int32 ConstraintIndex)
+FName UPhysicsAsset::FindConstraintBoneName(int32 ConstraintIndex) const
 {
 	if ( (ConstraintIndex < 0) || (ConstraintIndex >= ConstraintSetup.Num()) )
 	{
@@ -500,7 +500,7 @@ FName UPhysicsAsset::FindConstraintBoneName(int32 ConstraintIndex)
 	return ConstraintSetup[ConstraintIndex]->DefaultInstance.GetChildBoneName();
 }
 
-int32 UPhysicsAsset::FindMirroredBone(USkeletalMesh* SkelMesh,  int32 BoneIndex)
+int32 UPhysicsAsset::FindMirroredBone(const USkeletalMesh* SkelMesh,  int32 BoneIndex) const 
 {
 	if (SkelMesh)
 	{
@@ -512,7 +512,7 @@ int32 UPhysicsAsset::FindMirroredBone(USkeletalMesh* SkelMesh,  int32 BoneIndex)
 	return INDEX_NONE; 
 }
 
-void UPhysicsAsset::GetBodyIndicesBelow(TArray<int32>& OutBodyIndices, FName InBoneName, USkeletalMesh* SkelMesh, bool bIncludeParent /*= true*/)
+void UPhysicsAsset::GetBodyIndicesBelow(TArray<int32>& OutBodyIndices, FName InBoneName, const USkeletalMesh* SkelMesh, bool bIncludeParent /*= true*/) const
 {
 	if (SkelMesh)
 	{
@@ -520,7 +520,7 @@ void UPhysicsAsset::GetBodyIndicesBelow(TArray<int32>& OutBodyIndices, FName InB
 	}
 }
 
-void UPhysicsAsset::GetBodyIndicesBelow(TArray<int32>& OutBodyIndices, const FName InBoneName, const FReferenceSkeleton& RefSkeleton, const bool bIncludeParent /*= true*/)
+void UPhysicsAsset::GetBodyIndicesBelow(TArray<int32>& OutBodyIndices, const FName InBoneName, const FReferenceSkeleton& RefSkeleton, const bool bIncludeParent /*= true*/) const
 {
 	const int32 BaseIndex = RefSkeleton.FindBoneIndex(InBoneName);
 
@@ -529,7 +529,7 @@ void UPhysicsAsset::GetBodyIndicesBelow(TArray<int32>& OutBodyIndices, const FNa
 		// Iterate over all other bodies, looking for 'children' of this one
 		for (int32 i = 0; i < SkeletalBodySetups.Num(); i++)
 		{
-			UBodySetup* BS = SkeletalBodySetups[i];
+			const UBodySetup* BS = SkeletalBodySetups[i];
 			if (!ensure(BS))
 			{
 				continue;
@@ -545,7 +545,7 @@ void UPhysicsAsset::GetBodyIndicesBelow(TArray<int32>& OutBodyIndices, const FNa
 	}
 }
 
-void UPhysicsAsset::GetNearestBodyIndicesBelow(TArray<int32> & OutBodyIndices, FName InBoneName, USkeletalMesh * InSkelMesh)
+void UPhysicsAsset::GetNearestBodyIndicesBelow(TArray<int32> & OutBodyIndices, FName InBoneName, const USkeletalMesh* InSkelMesh) const
 {
 	TArray<int32> AllBodiesBelow;
 	GetBodyIndicesBelow(AllBodiesBelow, InBoneName, InSkelMesh, false);
@@ -563,7 +563,7 @@ void UPhysicsAsset::GetNearestBodyIndicesBelow(TArray<int32> & OutBodyIndices, F
 		int32 BodyIndex = AllBodiesBelow[i];
 		if (Nearest[BodyIndex] == false) continue;
 
-		UBodySetup * Body = SkeletalBodySetups[BodyIndex];
+		const UBodySetup * Body = SkeletalBodySetups[BodyIndex];
 		if (!ensure(Body))
 		{
 			continue;
