Deprecate / remove legacy audio streaming manager.
#rnx

diff --git a/Engine/Source/Editor/AudioEditor/Private/Factories/SoundFactory.cpp b/Engine/Source/Editor/AudioEditor/Private/Factories/SoundFactory.cpp
--- a/Engine/Source/Editor/AudioEditor/Private/Factories/SoundFactory.cpp
+++ b/Engine/Source/Editor/AudioEditor/Private/Factories/SoundFactory.cpp
@@ -722,11 +722,10 @@ UObject* USoundFactory::CreateObject
 		}
 				
 		// Compressed data is now out of date.
-		const bool bRebuildStreamingChunks = FPlatformCompressionUtilities::IsCurrentPlatformUsingStreamCaching();
-		Sound->InvalidateCompressedData(true /* bFreeResources */, bRebuildStreamingChunks);
+		Sound->InvalidateCompressedData(true /* bFreeResources */);
 
 		// If stream caching is enabled, we need to make sure this asset is ready for playback.
-		if (bRebuildStreamingChunks && Sound->IsStreaming(nullptr))
+		if (Sound->IsStreaming(nullptr))
 		{
 			Sound->LoadZerothChunk();
 		}
diff --git a/Engine/Source/Runtime/AudioMixer/Private/AudioMixerBlueprintLibrary.cpp b/Engine/Source/Runtime/AudioMixer/Private/AudioMixerBlueprintLibrary.cpp
--- a/Engine/Source/Runtime/AudioMixer/Private/AudioMixerBlueprintLibrary.cpp
+++ b/Engine/Source/Runtime/AudioMixer/Private/AudioMixerBlueprintLibrary.cpp
@@ -608,12 +608,6 @@ void UAudioMixerBlueprintLibrary::PrimeSoundForPlayback(USoundWave* SoundWave, c
 	{
 		UE_LOG(LogAudioMixer, Warning, TEXT("Prime Sound For Playback called with a null SoundWave pointer."));
 	}
-	else if (!FPlatformCompressionUtilities::IsCurrentPlatformUsingStreamCaching())
-	{
-		UE_LOG(LogAudioMixer, Warning, TEXT("Prime Sound For Playback doesn't do anything unless Audio Load On Demand is enabled."));
-		
-		OnLoadCompletion.ExecuteIfBound(SoundWave, false);
-	}
 	else
 	{
 		IStreamingManager::Get().GetAudioStreamingManager().RequestChunk(SoundWave->CreateSoundWaveProxy(), 1, [OnLoadCompletion, SoundWave](EAudioChunkLoadResult InResult)
diff --git a/Engine/Source/Runtime/AudioMixer/Private/AudioMixerBuffer.cpp b/Engine/Source/Runtime/AudioMixer/Private/AudioMixerBuffer.cpp
--- a/Engine/Source/Runtime/AudioMixer/Private/AudioMixerBuffer.cpp
+++ b/Engine/Source/Runtime/AudioMixer/Private/AudioMixerBuffer.cpp
@@ -46,11 +46,6 @@ namespace Audio
 
 		if (DecompressionState)
 		{
-			if (BufferType == EBufferType::Streaming)
-			{
-				IStreamingManager::Get().GetAudioStreamingManager().RemoveDecoder(DecompressionState);
-			}
-
 			delete DecompressionState;
 			DecompressionState = nullptr;
 		}
diff --git a/Engine/Source/Runtime/AudioMixer/Private/AudioMixerSource.cpp b/Engine/Source/Runtime/AudioMixer/Private/AudioMixerSource.cpp
--- a/Engine/Source/Runtime/AudioMixer/Private/AudioMixerSource.cpp
+++ b/Engine/Source/Runtime/AudioMixer/Private/AudioMixerSource.cpp
@@ -1234,11 +1234,6 @@ namespace Audio
 							check(MixerSourceBuffer.IsValid());
 
 							ICompressedAudioInfo* Decoder = MixerBuffer->GetDecompressionState(false);
-							if (BufferType == EBufferType::Streaming)
-							{
-								IStreamingManager::Get().GetAudioStreamingManager().AddDecoder(Decoder);
-							}
-
 							MixerSourceBuffer->ReadMoreRealtimeData(Decoder, 0, EBufferReadMode::Asynchronous);
 
 							// not ready
@@ -1385,8 +1380,6 @@ namespace Audio
 
 		InitializationState = EMixerSourceInitializationState::NotInitialized;
 
-		IStreamingManager::Get().GetAudioStreamingManager().RemoveStreamingSoundSource(this);
-
 		bIsStopping = false;
 
 		if (WaveInstance)
diff --git a/Engine/Source/Runtime/AudioMixer/Private/AudioMixerSourceBuffer.cpp b/Engine/Source/Runtime/AudioMixer/Private/AudioMixerSourceBuffer.cpp
--- a/Engine/Source/Runtime/AudioMixer/Private/AudioMixerSourceBuffer.cpp
+++ b/Engine/Source/Runtime/AudioMixer/Private/AudioMixerSourceBuffer.cpp
@@ -152,10 +152,6 @@ namespace Audio
 		if (DecompressionState == nullptr)
 		{
 			DecompressionState = InCompressedAudioInfo;
-			if (BufferType == EBufferType::Streaming)
-			{
-				IStreamingManager::Get().GetAudioStreamingManager().AddDecoder(DecompressionState);
-			}
 		}
 	}
 
@@ -544,11 +540,6 @@ namespace Audio
 		// Clean up decompression state after things have been finished using it
 		if (DecompressionState)
 		{
-			if (BufferType == EBufferType::Streaming)
-			{
-				IStreamingManager::Get().GetAudioStreamingManager().RemoveDecoder(DecompressionState);
-			}
-
 			delete DecompressionState;
 			DecompressionState = nullptr;
 		}
diff --git a/Engine/Source/Runtime/Engine/Classes/Sound/SoundWave.h b/Engine/Source/Runtime/Engine/Classes/Sound/SoundWave.h
--- a/Engine/Source/Runtime/Engine/Classes/Sound/SoundWave.h
+++ b/Engine/Source/Runtime/Engine/Classes/Sound/SoundWave.h
@@ -720,8 +720,6 @@ private:
 	// Whether or not the thumbnail supports generation
 	uint8 bNeedsThumbnailGeneration : 1;
 
-	// Whether this was previously cooked with stream caching enabled.
-	uint8 bWasStreamCachingEnabledOnLastCook : 1;
 	// Whether this asset is loaded from cooked data.
 	uint8 bLoadedFromCookedData : 1;
 #endif // !WITH_EDITOR
diff --git a/Engine/Source/Runtime/Engine/Private/AudioCompressionSettingsUtils.cpp b/Engine/Source/Runtime/Engine/Private/AudioCompressionSettingsUtils.cpp
--- a/Engine/Source/Runtime/Engine/Private/AudioCompressionSettingsUtils.cpp
+++ b/Engine/Source/Runtime/Engine/Private/AudioCompressionSettingsUtils.cpp
@@ -320,11 +320,6 @@ const FPlatformAudioCookOverrides* FPlatformCompressionUtilities::GetCookOverrid
 	return &Overrides;
 }
 
-bool FPlatformCompressionUtilities::IsCurrentPlatformUsingStreamCaching()
-{
-	return true;
-}
-
 const FAudioStreamCachingSettings& FPlatformCompressionUtilities::GetStreamCachingSettingsForCurrentPlatform()
 {
 	const FPlatformAudioCookOverrides* Settings = GetCookOverrides();
diff --git a/Engine/Source/Runtime/Engine/Private/AudioDerivedData.cpp b/Engine/Source/Runtime/Engine/Private/AudioDerivedData.cpp
--- a/Engine/Source/Runtime/Engine/Private/AudioDerivedData.cpp
+++ b/Engine/Source/Runtime/Engine/Private/AudioDerivedData.cpp
@@ -1349,14 +1349,7 @@ int32 FStreamedAudioPlatformData::DeserializeChunkFromDDC(TArray<uint8> TempData
 		Ar.Serialize(*OutChunkData, ChunkSize);
 	}
 
-	if (FPlatformCompressionUtilities::IsCurrentPlatformUsingStreamCaching())
-	{
-		return AudioDataSize;
-	}
-	else
-	{
-		return ChunkSize;
-	}
+	return AudioDataSize;
 }
 
 int32 FStreamedAudioPlatformData::GetChunkFromDDC(int32 ChunkIndex, uint8** OutChunkData, bool bMakeSureChunkIsLoaded /* = false */)
diff --git a/Engine/Source/Runtime/Engine/Private/AudioDevice.cpp b/Engine/Source/Runtime/Engine/Private/AudioDevice.cpp
--- a/Engine/Source/Runtime/Engine/Private/AudioDevice.cpp
+++ b/Engine/Source/Runtime/Engine/Private/AudioDevice.cpp
@@ -4468,9 +4468,7 @@ void FAudioDevice::StartSources(TArray<FWaveInstance*>& WaveInstances, int32 Fir
 		if (!WaveInstance->ShouldStopDueToMaxConcurrency() && (bGameTicking || WaveInstance->bIsUISound))
 		{
 			FSoundSource* Source = WaveInstanceSourceMap.FindRef(WaveInstance);
-			if (!Source &&
-				(!WaveInstance->IsStreaming() ||
-				IStreamingManager::Get().GetAudioStreamingManager().CanCreateSoundSource(WaveInstance)))
+			if (!Source)
 			{
 				// Check for full sources and stop the oldest stopping source
 				if (!FreeSources.Num())
@@ -4523,7 +4521,6 @@ void FAudioDevice::StartSources(TArray<FWaveInstance*>& WaveInstances, int32 Fir
 				// If we succeeded above then we need to map the wave instance to the source
 				if (bSuccess)
 				{
-					IStreamingManager::Get().GetAudioStreamingManager().AddStreamingSoundSource(Source);
 					// Associate wave instance with it which is used earlier in this function.
 					WaveInstanceSourceMap.Add(WaveInstance, Source);
 				}
@@ -4540,7 +4537,7 @@ void FAudioDevice::StartSources(TArray<FWaveInstance*>& WaveInstances, int32 Fir
 					WaveInstanceSourceMap.Remove(WaveInstance);
 				}
 			}
-			else if (Source)
+			else // Source is valid
 			{
 				if (!Source->IsInitialized() && Source->IsPreparedToInit())
 				{
@@ -4569,12 +4566,6 @@ void FAudioDevice::StartSources(TArray<FWaveInstance*>& WaveInstances, int32 Fir
 					}
 				}
 			}
-			else
-			{
-				// This can happen if the streaming manager determines that this sound should not be started.
-				// We stop the wave instance to prevent it from attempting to initialize every frame
-				WaveInstance->StopWithoutNotification();
-			}
 		}
 	}
 
diff --git a/Engine/Source/Runtime/Engine/Private/AudioStreaming.cpp b/Engine/Source/Runtime/Engine/Private/AudioStreaming.cpp
--- a/Engine/Source/Runtime/Engine/Private/AudioStreaming.cpp
+++ b/Engine/Source/Runtime/Engine/Private/AudioStreaming.cpp
@@ -11,7 +11,6 @@ AudioStreaming.cpp: Implementation of audio streaming classes.
 #include "Serialization/MemoryReader.h"
 #include "HAL/PlatformFile.h"
 #include "Async/AsyncFileHandle.h"
-#include "AudioDecompress.h"
 
 static int32 SpoofFailedStreamChunkLoad = 0;
 FAutoConsoleVariableRef CVarSpoofFailedStreamChunkLoad(
@@ -85,888 +84,4 @@ void FAsyncStreamDerivedChunkWorker::DoWork()
 	UE_LOG(LogAudio, Verbose, TEXT("End of Async DDC Chunk Load. DDC Key: %s"), *DerivedDataKey);
 }
 
-#endif // #if WITH_EDITORONLY_DATA
-
-////////////////////////
-// FStreamingWaveData //
-////////////////////////
-
-FStreamingWaveData::FStreamingWaveData()
-	: SoundWave(NULL)
-	, AudioStreamingManager(nullptr)
-{
-}
-
-FStreamingWaveData::~FStreamingWaveData()
-{
-
-}
-
-void FStreamingWaveData::FreeResources()
-{
-	// Make sure there are no pending requests in flight.
-	for (int32 Pass = 0; Pass < 3; Pass++)
-	{
-		BlockTillAllRequestsFinished();
-		if (!UpdateStreamingStatus())
-		{
-			break;
-		}
-		check(Pass < 2); // we should be done after two passes. Pass 0 will start anything we need and pass 1 will complete those requests
-	}
-
-	for (FLoadedAudioChunk& LoadedChunk : LoadedChunks)
-	{
-		FreeLoadedChunk(LoadedChunk);
-	}
-}
-
-bool FStreamingWaveData::Initialize(const FSoundWaveProxyPtr& InSoundWave, FLegacyAudioStreamingManager* InAudioStreamingManager)
-{
-	if (!ensure(InSoundWave.IsValid()) && !InSoundWave->GetNumChunks())
-	{
-#if WITH_EDITOR
-		UE_LOG(LogAudio, Display, TEXT("Failed to initialize streaming wave data due to lack of serialized stream chunks. Error during stream cooking."));
-#else
-		UE_LOG(LogAudio, Warning, TEXT("Failed to initialize streaming wave data due to lack of serialized stream chunks. Error during stream cooking."));
-#endif
-		return false;
-	}
-
-	SoundWave = InSoundWave;
-	AudioStreamingManager = InAudioStreamingManager;
-
-	// Always get the first chunk of data so we can play immediately
-	check(LoadedChunks.Num() == 0);
-	check(LoadedChunkIndices.Num() == 0);
-
-	// Prepare 4 chunks of streaming wave data in loaded chunks array
-	LoadedChunks.Reset(4);
-
-	const FStreamedAudioChunk& Chunk = SoundWave->GetChunk(0);
-	const int32 DataSize = Chunk.DataSize;
-	const int32 AudioDataSize = Chunk.AudioDataSize;
-	const int32 FirstLoadedChunkIndex = AddNewLoadedChunk(DataSize, AudioDataSize);
-
-	FLoadedAudioChunk* FirstChunk = &LoadedChunks[FirstLoadedChunkIndex];
-	FirstChunk->Index = 0;
-
-	// Make sure we have loaded the 0th chunk before proceeding:
-	if (!SoundWave->GetChunkData(0, &FirstChunk->Data, true))
-	{
-		// Error/warning logging will have already been performed in the GetChunkData function
-		return false;
-	}
-
-	// Set up the loaded/requested indices to be identical
-	LoadedChunkIndices.Add(0);
-	CurrentRequest.RequiredIndices.Add(0);
-
-	return true;
-}
-
-bool FStreamingWaveData::UpdateStreamingStatus()
-{
-	bool	bHasPendingRequestInFlight = true;
-	int32	RequestStatus = PendingChunkChangeRequestStatus.GetValue();
-	TArray<uint32> IndicesToLoad;
-	TArray<uint32> IndicesToFree;
-
-	if (!HasPendingRequests(IndicesToLoad, IndicesToFree))
-	{
-		check(RequestStatus == AudioState_ReadyFor_Requests);
-		bHasPendingRequestInFlight = false;
-	}
-	// Pending request in flight, though we might be able to finish it.
-	else
-	{
-		if (RequestStatus == AudioState_ReadyFor_Finalization)
-		{
-			if (UE_LOG_ACTIVE(LogAudio, Log) && IndicesToLoad.Num() > 0)
-			{
-				FString LogString = FString::Printf(TEXT("Finalised loading of chunk(s) %d"), IndicesToLoad[0]);
-				for (int32 Index = 1; Index < IndicesToLoad.Num(); ++Index)
-				{
-					LogString += FString::Printf(TEXT(", %d"), IndicesToLoad[Index]);
-				}
-				LogString += FString::Printf(TEXT(" from SoundWave'%s'"), *SoundWave->GetFName().ToString());
-				UE_LOG(LogAudio, Log, TEXT("%s"), *LogString);
-			}
-
-			bool bFailedRequests = false;
-#if WITH_EDITORONLY_DATA
-			bFailedRequests = FinishDDCRequests();
-#endif //WITH_EDITORONLY_DATA
-
-			// could maybe iterate over the things we know are done, but I couldn't tell if that was IndicesToLoad or not.
-			for (FLoadedAudioChunk& LoadedChunk : LoadedChunks)
-			{
-				if (LoadedChunk.IORequest != nullptr)
-				{
-					LoadedChunk.IORequest->WaitCompletion();
-
-					delete LoadedChunk.IORequest;
-					LoadedChunk.IORequest = nullptr;
-				}
-			}
-
-			PendingChunkChangeRequestStatus.Decrement();
-			bHasPendingRequestInFlight = false;
-			LoadedChunkIndices = CurrentRequest.RequiredIndices;
-		}
-		else if (RequestStatus == AudioState_ReadyFor_Requests) // odd that this is an else, probably we should start requests right now
-		{
-			BeginPendingRequests(IndicesToLoad, IndicesToFree);
-		}
-	}
-
-	return bHasPendingRequestInFlight;
-}
-
-void FStreamingWaveData::UpdateChunkRequests(FWaveRequest& InWaveRequest)
-{
-	// Might change this but ensures chunk 0 stays loaded for now
-	check(InWaveRequest.RequiredIndices.Contains(0));
-	check(PendingChunkChangeRequestStatus.GetValue() == AudioState_ReadyFor_Requests);
-
-	CurrentRequest = InWaveRequest;
-
-}
-
-bool FStreamingWaveData::HasPendingRequests(TArray<uint32>& IndicesToLoad, TArray<uint32>& IndicesToFree) const
-{
-	IndicesToLoad.Empty();
-	IndicesToFree.Empty();
-
-	// Find indices that aren't loaded
-	for (auto NeededIndex : CurrentRequest.RequiredIndices)
-	{
-		if (!LoadedChunkIndices.Contains(NeededIndex))
-		{
-			IndicesToLoad.AddUnique(NeededIndex);
-		}
-	}
-
-	// Find indices that aren't needed anymore
-	for (auto CurrentIndex : LoadedChunkIndices)
-	{
-		if (!CurrentRequest.RequiredIndices.Contains(CurrentIndex))
-		{
-			IndicesToFree.AddUnique(CurrentIndex);
-		}
-	}
-
-	return IndicesToLoad.Num() > 0 || IndicesToFree.Num() > 0;
-}
-
-void FStreamingWaveData::BeginPendingRequests(const TArray<uint32>& IndicesToLoad, const TArray<uint32>& IndicesToFree)
-{
-	if (UE_LOG_ACTIVE(LogAudio, Log) && IndicesToLoad.Num() > 0)
-	{
-		FString LogString = FString::Printf(TEXT("Requesting ASync load of chunk(s) %d"), IndicesToLoad[0]);
-		for (int32 Index = 1; Index < IndicesToLoad.Num(); ++Index)
-		{
-			LogString += FString::Printf(TEXT(", %d"), IndicesToLoad[Index]);
-		}
-		LogString += FString::Printf(TEXT(" from SoundWave'%s'"), *SoundWave->GetFName().ToString());
-		UE_LOG(LogAudio, Log, TEXT("%s"), *LogString);
-	}
-
-	TArray<uint32> FreeChunkIndices;
-
-	// Mark Chunks for removal in case they can be reused
-	{
-		for (auto Index : IndicesToFree)
-		{
-			for (int32 ChunkIndex = 0; ChunkIndex < LoadedChunks.Num(); ++ChunkIndex)
-			{
-				check(Index != 0);
-				if (LoadedChunks[ChunkIndex].Index == Index)
-				{
-					FreeLoadedChunk(LoadedChunks[ChunkIndex]);
-					LoadedChunks.RemoveAt(ChunkIndex);
-					break;
-				}
-			}
-		}
-	}
-
-	if (IndicesToLoad.Num() > 0)
-	{
-		PendingChunkChangeRequestStatus.Set(AudioState_InProgress_Loading);
-
-		// Set off all IO Requests
-		for (auto Index : IndicesToLoad)
-		{
-			const FStreamedAudioChunk& Chunk = SoundWave->GetChunk(Index);
-			int32 ChunkSize = Chunk.DataSize;
-
-			int32 LoadedChunkStorageIndex = AddNewLoadedChunk(ChunkSize, Chunk.AudioDataSize);
-			FLoadedAudioChunk* ChunkStorage = &LoadedChunks[LoadedChunkStorageIndex];
-
-			ChunkStorage->Index = Index;
-
-			check(LoadedChunkStorageIndex != INDEX_NONE);
-
-			// Pass the request on to the async io manager after increasing the request count. The request count 
-			// has been pre-incremented before fielding the update request so we don't have to worry about file
-			// I/O immediately completing and the game thread kicking off again before this function
-			// returns.
-			PendingChunkChangeRequestStatus.Increment();
-
-			EAsyncIOPriorityAndFlags AsyncIOPriority = AIOP_High | AIOP_FLAG_DONTCACHE;
-
-			// Load and decompress async.
-#if WITH_EDITORONLY_DATA
-			if (Chunk.DerivedDataKey.IsEmpty() == false)
-			{
-				ChunkStorage->Data = static_cast<uint8*>(FMemory::Malloc(ChunkSize));
-				INC_DWORD_STAT_BY(STAT_AudioMemorySize, ChunkSize);
-				INC_DWORD_STAT_BY(STAT_AudioMemory, ChunkSize);
-
-				// This streaming manager does not use the callback when the DDC load
-				// is completed:
-				TFunction<void(bool)> NullOnLoadCompletedCallback = [](bool) {};
-
-				FAsyncStreamDerivedChunkTask* Task = new FAsyncStreamDerivedChunkTask(
-					Chunk.DerivedDataKey,
-					ChunkStorage->Data,
-					ChunkSize,
-					&PendingChunkChangeRequestStatus,
-					MoveTemp(NullOnLoadCompletedCallback)
-				);
-				PendingAsyncStreamDerivedChunkTasks.Add(Task);
-				// This task may perform a long synchronous DDC request. Using DoNotRunInsideBusyWait prevents potentially delaying foreground tasks.
-				Task->StartBackgroundTask(GThreadPool, EQueuedWorkPriority::Normal, EQueuedWorkFlags::DoNotRunInsideBusyWait);
-			}
-			else
-#endif // #if WITH_EDITORONLY_DATA
-			{
-				check(!ChunkStorage->IORequest); // Make sure we do not already have a request
-				check(!ChunkStorage->Data); // Make sure we do not already have data
-				check(Chunk.BulkData.GetBulkDataSize() == ChunkStorage->DataSize); // Make sure that the bulkdata size matches
-
-				FBulkDataIORequestCallBack AsyncFileCallBack =
-					[this, LoadedChunkStorageIndex](bool bWasCancelled, IBulkDataIORequest* Req)
-				{
-					AudioStreamingManager->OnAsyncFileCallback(this, LoadedChunkStorageIndex, Req);
-
-					PendingChunkChangeRequestStatus.Decrement();
-				};
-
-				ChunkStorage->IORequest = Chunk.BulkData.CreateStreamingRequest(AsyncIOPriority, &AsyncFileCallBack, nullptr);
-
-				if (!ChunkStorage->IORequest)
-				{
-					UE_LOG(LogAudio, Error, TEXT("Audio streaming read request failed."));
-
-					// we failed for some reason; file not found I guess.
-					PendingChunkChangeRequestStatus.Decrement();
-				}
-			}
-		}
-
-		// Decrement the state to AudioState_InProgress_Loading + NumChunksCurrentLoading - 1.
-		PendingChunkChangeRequestStatus.Decrement();
-	}
-	else
-	{
-		// Skip straight to finalisation
-		PendingChunkChangeRequestStatus.Set(AudioState_ReadyFor_Finalization);
-	}
-
-}
-
-bool FStreamingWaveData::BlockTillAllRequestsFinished(float TimeLimit)
-{
-	QUICK_SCOPE_CYCLE_COUNTER(FStreamingWaveData_BlockTillAllRequestsFinished);
-	if (TimeLimit == 0.0f)
-	{
-		for (FLoadedAudioChunk& LoadedChunk : LoadedChunks)
-		{
-			if (LoadedChunk.IORequest)
-			{
-				LoadedChunk.IORequest->WaitCompletion();
-
-				delete LoadedChunk.IORequest;
-				LoadedChunk.IORequest = nullptr;
-			}
-		}
-	}
-	else
-	{
-		double EndTime = FPlatformTime::Seconds() + TimeLimit;
-		for (FLoadedAudioChunk& LoadedChunk : LoadedChunks)
-		{
-			if (LoadedChunk.IORequest)
-			{
-				float ThisTimeLimit = EndTime - FPlatformTime::Seconds();
-				if (ThisTimeLimit < .001f || // one ms is the granularity of the platform event system
-					!LoadedChunk.IORequest->WaitCompletion(ThisTimeLimit))
-				{
-					return false;
-				}
-				delete LoadedChunk.IORequest;
-				LoadedChunk.IORequest = nullptr;
-			}
-		}
-	}
-	return true;
-}
-
-#if WITH_EDITORONLY_DATA
-bool FStreamingWaveData::FinishDDCRequests()
-{
-	bool bRequestFailed = false;
-	if (PendingAsyncStreamDerivedChunkTasks.Num())
-	{
-		for (int32 TaskIndex = 0; TaskIndex < PendingAsyncStreamDerivedChunkTasks.Num(); ++TaskIndex)
-		{
-			FAsyncStreamDerivedChunkTask& Task = PendingAsyncStreamDerivedChunkTasks[TaskIndex];
-			Task.EnsureCompletion();
-			bRequestFailed |= Task.GetTask().DidRequestFail();
-		}
-		PendingAsyncStreamDerivedChunkTasks.Empty();
-	}
-	return bRequestFailed;
-}
-#endif //WITH_EDITORONLY_DATA
-
-int32 FStreamingWaveData::AddNewLoadedChunk(int32 ChunkSize, int32 AudioSize)
-{
-	int32 NewIndex = LoadedChunks.Num();
-	LoadedChunks.AddDefaulted();
-
-	LoadedChunks[NewIndex].DataSize = ChunkSize;
-	LoadedChunks[NewIndex].AudioDataSize = AudioSize;
-
-	return NewIndex;
-}
-
-void FStreamingWaveData::FreeLoadedChunk(FLoadedAudioChunk& LoadedChunk)
-{
-	if (LoadedChunk.IORequest)
-	{
-		LoadedChunk.IORequest->Cancel();
-		LoadedChunk.IORequest->WaitCompletion();
-
-		delete LoadedChunk.IORequest;
-		LoadedChunk.IORequest = nullptr;
-
-		// Process pending async requests after iorequest finishes
-		AudioStreamingManager->ProcessPendingAsyncFileResults();
-	}
-
-	if (LoadedChunk.Data != NULL)
-	{
-		FMemory::Free(LoadedChunk.Data);
-
-		DEC_DWORD_STAT_BY(STAT_AudioMemorySize, LoadedChunk.DataSize);
-		DEC_DWORD_STAT_BY(STAT_AudioMemory, LoadedChunk.DataSize);
-	}
-	LoadedChunk.Data = NULL;
-	LoadedChunk.AudioDataSize = 0;
-	LoadedChunk.DataSize = 0;
-	LoadedChunk.Index = 0;
-}
-
-////////////////////////////
-// FAudioStreamingManager //
-////////////////////////////
-
-FLegacyAudioStreamingManager::FLegacyAudioStreamingManager()
-{
-}
-
-FLegacyAudioStreamingManager::~FLegacyAudioStreamingManager()
-{
-}
-
-void FLegacyAudioStreamingManager::OnAsyncFileCallback(FStreamingWaveData* StreamingWaveData, int32 LoadedAudioChunkIndex, IBulkDataIORequest* ReadRequest)
-{
-	// Check to see if we successfully managed to load anything
-	uint8* Mem = ReadRequest->GetReadResults();
-	if (Mem)
-	{
-		// Create a new chunk load result object. Will be deleted on audio thread when TQueue is pumped.
-		FASyncAudioChunkLoadResult* NewAudioChunkLoadResult = new FASyncAudioChunkLoadResult();
-
-		// Copy the ptr which we will use to place the results of the read on the audio thread upon pumping.
-		NewAudioChunkLoadResult->StreamingWaveData = StreamingWaveData;
-
-		// Grab the loaded chunk memory ptr since it will be invalid as soon as this callback finishes
-		NewAudioChunkLoadResult->DataResults = Mem;
-
-		// The chunk index to load the results into
-		NewAudioChunkLoadResult->LoadedAudioChunkIndex = LoadedAudioChunkIndex;
-
-		// Safely add the results of the async file callback into an array to be pumped on audio thread
-		FScopeLock StreamChunkResults(&ChunkResultCriticalSection);
-		AsyncAudioStreamChunkResults.Add(NewAudioChunkLoadResult);
-	}
-}
-
-void FLegacyAudioStreamingManager::ProcessPendingAsyncFileResults()
-{
-	// Pump the results of any async file loads in a protected critical section 
-	FScopeLock StreamChunkResults(&ChunkResultCriticalSection);
-
-	for (FASyncAudioChunkLoadResult* AudioChunkLoadResult : AsyncAudioStreamChunkResults)
-	{
-		// Copy the results to the chunk storage safely
-		const int32 LoadedAudioChunkIndex = AudioChunkLoadResult->LoadedAudioChunkIndex;
-
-		check(AudioChunkLoadResult->StreamingWaveData != nullptr);
-
-		if (LoadedAudioChunkIndex == INDEX_NONE || LoadedAudioChunkIndex >= AudioChunkLoadResult->StreamingWaveData->LoadedChunks.Num())
-		{
-			UE_LOG(LogAudio, Warning, TEXT("Loaded streaming chunk index %d is invalid. Current number of loaded chunks is: %d"), LoadedAudioChunkIndex, AudioChunkLoadResult->StreamingWaveData->LoadedChunks.Num());
-
-			delete AudioChunkLoadResult;
-			AudioChunkLoadResult = nullptr;
-			continue;
-		}
-
-		FLoadedAudioChunk* ChunkStorage = &AudioChunkLoadResult->StreamingWaveData->LoadedChunks[LoadedAudioChunkIndex];
-
-		checkf(!ChunkStorage->Data, TEXT("Chunk storage already has data. (0x%p), datasize: %d"), ChunkStorage->Data, ChunkStorage->DataSize);
-
-		ChunkStorage->Data = AudioChunkLoadResult->DataResults;
-
-		INC_DWORD_STAT_BY(STAT_AudioMemorySize, ChunkStorage->DataSize);
-		INC_DWORD_STAT_BY(STAT_AudioMemory, ChunkStorage->DataSize);
-
-		// Cleanup the chunk load results
-		delete AudioChunkLoadResult;
-		AudioChunkLoadResult = nullptr;
-	}
-
-	AsyncAudioStreamChunkResults.Reset();
-}
-
-void FLegacyAudioStreamingManager::UpdateResourceStreaming(float DeltaTime, bool bProcessEverything /*= false*/)
-{
-	LLM_SCOPE(ELLMTag::Audio);
-
-	FScopeLock Lock(&CriticalSection);
-
-	for (auto& WavePair : StreamingSoundWaves)
-	{
-		WavePair.Value->UpdateStreamingStatus();
-	}
-
-	// Process any async file requests after updating the stream status
-	ProcessPendingAsyncFileResults();
-
-	for (auto Source : StreamingSoundSources)
-	{
-		const FWaveInstance* WaveInstance = Source->GetWaveInstance();
-		USoundWave* Wave = WaveInstance ? WaveInstance->WaveData : nullptr;
-		if (Wave)
-		{
-			FStreamingWaveData** WaveDataPtr = StreamingSoundWaves.Find(Wave);
-
-			if (WaveDataPtr && (*WaveDataPtr)->PendingChunkChangeRequestStatus.GetValue() == AudioState_ReadyFor_Requests)
-			{
-				FStreamingWaveData* WaveData = *WaveDataPtr;
-				// Request the chunk the source is using and the one after that
-				FWaveRequest& WaveRequest = GetWaveRequest(Wave);
-				const FSoundBuffer* SoundBuffer = Source->GetBuffer();
-				if (SoundBuffer)
-				{
-					int32 SourceChunk = SoundBuffer->GetCurrentChunkIndex();
-					if (SourceChunk >= 0 && SourceChunk < (int32)Wave->GetNumChunks())
-					{
-						WaveRequest.RequiredIndices.AddUnique(SourceChunk);
-						WaveRequest.RequiredIndices.AddUnique((SourceChunk + 1) % Wave->GetNumChunks());
-						WaveRequest.bPrioritiseRequest = true;
-					}
-					else
-					{
-						UE_LOG(LogAudio, Log, TEXT("Invalid chunk request curIndex=%d numChunks=%d\n"), SourceChunk, Wave->GetNumChunks());
-					}
-				}
-			}
-		}
-	}
-
-	for (ICompressedAudioInfo* Decoder : CompressedAudioInfos)
-	{
-		const FSoundWaveProxyPtr& SoundWave = Decoder->GetStreamingSoundWave();
-		FStreamingWaveData** WaveDataPtr = StreamingSoundWaves.Find(SoundWave->GetFObjectKey());
-		if (WaveDataPtr && (*WaveDataPtr)->PendingChunkChangeRequestStatus.GetValue() == AudioState_ReadyFor_Requests)
-		{
-			FStreamingWaveData* WaveData = *WaveDataPtr;
-			// Request the chunk the source is using and the one after that
-			FWaveRequest& WaveRequest = GetWaveRequest(SoundWave->GetFObjectKey());
-			int32 SourceChunk = Decoder->GetCurrentChunkIndex();
-
-			// If there's a seek request, use that as our current chunk index.				
-			if (Decoder->IsStreamedCompressedInfo())
-			{
-				// This is awk because the decoder thread catches the seek to
-				// advance, so we have to read this atomically. However, if it changes,
-				// that means that the chunk is already in memory (because StreamCompressedData
-				// only updates it when the chunk is present).
-				// ... and also, GetCurrentChunkIndex has the same problem. So it's layers
-				// on layers.
-				IStreamedCompressedInfo* Streamed = (IStreamedCompressedInfo*)Decoder;
-				int32 StreamSeekToBlockIndex = Streamed->GetStreamSeekBlockIndex();
-				if (StreamSeekToBlockIndex != INDEX_NONE)
-				{
-					UE_LOG(LogAudio, Log, TEXT("Diverting streaming chunk due to seek -- from %d to %d"), SourceChunk, StreamSeekToBlockIndex);
-					SourceChunk = StreamSeekToBlockIndex;
-				}
-			}
-
-			if (SourceChunk >= 0 && SourceChunk < (int32)SoundWave->GetNumChunks())
-			{
-				WaveRequest.RequiredIndices.AddUnique(SourceChunk);
-				WaveRequest.RequiredIndices.AddUnique((SourceChunk + 1) % SoundWave->GetNumChunks());
-				WaveRequest.bPrioritiseRequest = true;
-			}
-			else
-			{
-				UE_LOG(LogAudio, Log, TEXT("Invalid chunk request curIndex=%d numChunks=%d\n"), SourceChunk, SoundWave->GetNumChunks());
-			}
-		}
-	}
-
-	for (auto Iter = WaveRequests.CreateIterator(); Iter; ++Iter)
-	{
-		FStreamingWaveData* WaveData = StreamingSoundWaves.FindRef(Iter.Key());
-
-		if (WaveData && WaveData->PendingChunkChangeRequestStatus.GetValue() == AudioState_ReadyFor_Requests)
-		{
-			WaveData->UpdateChunkRequests(Iter.Value());
-			WaveData->UpdateStreamingStatus();
-			Iter.RemoveCurrent();
-		}
-	}
- 
-	// Process any async file requests after updating the streaming wave data stream statuses
-	ProcessPendingAsyncFileResults();
-}
-
-int32 FLegacyAudioStreamingManager::BlockTillAllRequestsFinished(float TimeLimit, bool)
-{
-	{
-		FScopeLock Lock(&CriticalSection);
-
-		QUICK_SCOPE_CYCLE_COUNTER(FAudioStreamingManager_BlockTillAllRequestsFinished);
-		int32 Result = 0;
-
-		if (TimeLimit == 0.0f)
-		{
-			for (auto& WavePair : StreamingSoundWaves)
-			{
-				WavePair.Value->BlockTillAllRequestsFinished();
-			}
-		}
-		else
-		{
-			double EndTime = FPlatformTime::Seconds() + TimeLimit;
-			for (auto& WavePair : StreamingSoundWaves)
-			{
-				float ThisTimeLimit = EndTime - FPlatformTime::Seconds();
-				if (ThisTimeLimit < .001f || // one ms is the granularity of the platform event system
-					!WavePair.Value->BlockTillAllRequestsFinished(ThisTimeLimit))
-				{
-					Result = 1; // we don't report the actual number, just 1 for any number of outstanding requests
-					break;
-				}
-			}
-		}
-		
-		// After blocking to process all requests, pump the queue
-		ProcessPendingAsyncFileResults();
-
-		return Result;
-	}
-}
-
-void FLegacyAudioStreamingManager::CancelForcedResources()
-{
-}
-
-void FLegacyAudioStreamingManager::NotifyLevelChange()
-{
-}
-
-void FLegacyAudioStreamingManager::SetDisregardWorldResourcesForFrames(int32 NumFrames)
-{
-}
-
-void FLegacyAudioStreamingManager::AddLevel(class ULevel* Level)
-{
-}
-
-void FLegacyAudioStreamingManager::RemoveLevel(class ULevel* Level)
-{
-}
-
-void FLegacyAudioStreamingManager::NotifyLevelOffset(class ULevel* Level, const FVector& Offset)
-{
-}
-
-void FLegacyAudioStreamingManager::AddStreamingSoundWave(const FSoundWaveProxyPtr& SoundWave)
-{
-	if (ensure(SoundWave.IsValid()) && FPlatformProperties::SupportsAudioStreaming() && SoundWave->IsStreaming())
-	{
-		FScopeLock Lock(&CriticalSection);
-
-		if (StreamingSoundWaves.FindRef(SoundWave->GetFObjectKey()) == nullptr)
-		{
-			FStreamingWaveData* NewStreamingWaveData = new FStreamingWaveData;
-			if (NewStreamingWaveData->Initialize(SoundWave, this))
-			{
-				StreamingSoundWaves.Add(SoundWave->GetFObjectKey(), NewStreamingWaveData);
-			}
-			else
-			{
-				// Failed to initialize, don't add to list of streaming sound waves
-				delete NewStreamingWaveData;
-			}
-		}
-	}
-}
-
-void FLegacyAudioStreamingManager::RemoveStreamingSoundWave(const FSoundWaveProxyPtr& SoundWave)
-{
-	if (!ensure(SoundWave.IsValid()))
-	{
-		return;
-	}
-
-	FScopeLock Lock(&CriticalSection);
-	FStreamingWaveData* WaveData = StreamingSoundWaves.FindRef(SoundWave->GetFObjectKey());
-	if (WaveData)
-	{
-		StreamingSoundWaves.Remove(SoundWave->GetFObjectKey());
-
-		// Free the resources of the streaming wave data. This blocks pending IO requests
-		WaveData->FreeResources();
-
-		{
-			// Then we need to remove any results from those pending requests before we delete so that we don't process them
-			FScopeLock StreamChunkResults(&ChunkResultCriticalSection);
-			for (int32 i = AsyncAudioStreamChunkResults.Num() - 1 ; i >= 0; --i)
-			{
-				FASyncAudioChunkLoadResult* LoadResult = AsyncAudioStreamChunkResults[i];
-				FStreamingWaveData* StreamingWaveData = LoadResult->StreamingWaveData;
-				if (StreamingWaveData == WaveData)
-				{
-					delete LoadResult;
-					
-					AsyncAudioStreamChunkResults.RemoveAtSwap(i, EAllowShrinking::No);
-				}
-			}
-		}
-		delete WaveData;
-	}
-	WaveRequests.Remove(SoundWave->GetFObjectKey());
-}
-
-void FLegacyAudioStreamingManager::AddDecoder(ICompressedAudioInfo* InCompressedAudioInfo)
-{
-	FScopeLock Lock(&CriticalSection);
-	CompressedAudioInfos.AddUnique(InCompressedAudioInfo);
-}
-
-void FLegacyAudioStreamingManager::RemoveDecoder(ICompressedAudioInfo* InCompressedAudioInfo)
-{
-	FScopeLock Lock(&CriticalSection);
-	CompressedAudioInfos.Remove(InCompressedAudioInfo);
-}
-
-bool FLegacyAudioStreamingManager::IsManagedStreamingSoundWave(const FSoundWaveProxyPtr& SoundWave) const
-{
-	FScopeLock Lock(&CriticalSection);
-	return ensure(SoundWave.IsValid()) && (StreamingSoundWaves.FindRef(SoundWave->GetFObjectKey()) != NULL);
-}
-
-bool FLegacyAudioStreamingManager::IsStreamingInProgress(const FSoundWaveProxyPtr&  SoundWave)
-{
-	FScopeLock Lock(&CriticalSection);
-	if (ensure(SoundWave.IsValid()))
-	{
-		FStreamingWaveData* WaveData = StreamingSoundWaves.FindRef(SoundWave->GetFObjectKey());
-		if (WaveData)
-		{
-			return WaveData->UpdateStreamingStatus();
-		}
-	}
-
-	return false;
-}
-
-bool FLegacyAudioStreamingManager::CanCreateSoundSource(const FWaveInstance* WaveInstance) const
-{
-	check(WaveInstance);
-	check(WaveInstance->IsStreaming());
-
-	int32 MaxStreams = MaxConcurrentStreamsCvar;
-	if (!MaxStreams)
-	{
-		MaxStreams = GetDefault<UAudioSettings>()->MaximumConcurrentStreams;
-	}
-
-	FScopeLock Lock(&CriticalSection);
-
-	// If the sound wave hasn't been added, or failed when trying to add during sound wave post load, we can't create a streaming sound source with this sound wave
-	if (!WaveInstance->WaveData || !StreamingSoundWaves.Contains(WaveInstance->WaveData))
-	{
-		return false;
-	}
-
-	if ( StreamingSoundSources.Num() < MaxStreams )
-	{
-		return true;
-	}
-
-	for (int32 Index = 0; Index < StreamingSoundSources.Num(); ++Index)
-	{
-		const FSoundSource* ExistingSource = StreamingSoundSources[Index];
-		const FWaveInstance* ExistingWaveInst = ExistingSource->GetWaveInstance();
-		if (!ExistingWaveInst || !ExistingWaveInst->WaveData
-			|| ExistingWaveInst->WaveData->StreamingPriority < WaveInstance->WaveData->StreamingPriority)
-		{
-			return Index < MaxStreams;
-		}
-	}
-
-	return false;
-}
-
-void FLegacyAudioStreamingManager::AddStreamingSoundSource(FSoundSource* SoundSource)
-{
-	const FWaveInstance* WaveInstance = SoundSource->GetWaveInstance();
-	if (WaveInstance && WaveInstance->IsStreaming())
-	{
-		int32 MaxStreams = MaxConcurrentStreamsCvar;
-		if (!MaxStreams)
-		{
-			MaxStreams = GetDefault<UAudioSettings>()->MaximumConcurrentStreams;
-		}
-
-		FScopeLock Lock(&CriticalSection);
-
-		// Add source sorted by priority so we can easily iterate over the amount of streams
-		// that are allowed
-		int32 OrderedIndex = -1;
-		for (int32 Index = 0; Index < StreamingSoundSources.Num() && Index < MaxStreams; ++Index)
-		{
-			const FSoundSource* ExistingSource = StreamingSoundSources[Index];
-			const FWaveInstance* ExistingWaveInst = ExistingSource->GetWaveInstance();
-			if (!ExistingWaveInst || !ExistingWaveInst->WaveData
-				|| ExistingWaveInst->WaveData->StreamingPriority < WaveInstance->WaveData->StreamingPriority)
-			{
-				OrderedIndex = Index;
-				break;
-			}
-		}
-		if (OrderedIndex != -1)
-		{
-			StreamingSoundSources.Insert(SoundSource, OrderedIndex);
-		}
-		else if (StreamingSoundSources.Num() < MaxStreams)
-		{
-			StreamingSoundSources.AddUnique(SoundSource);
-		}
-
-		for (int32 Index = StreamingSoundSources.Num()-1; Index >= MaxStreams; --Index)
-		{
-			StreamingSoundSources[Index]->Stop();
-		}
-	}
-}
-
-void FLegacyAudioStreamingManager::RemoveStreamingSoundSource(FSoundSource* SoundSource)
-{
-	const FWaveInstance* WaveInstance = SoundSource->GetWaveInstance();
-	if (WaveInstance && WaveInstance->WaveData && WaveInstance->WaveData->IsStreaming(nullptr) && !WaveInstance->WaveData->ShouldUseStreamCaching())
-	{
-		FScopeLock Lock(&CriticalSection);
-
-		// Make sure there is a request so that unused chunks
-		// can be cleared if this was the last playing instance
-		GetWaveRequest(FObjectKey(WaveInstance->WaveData));
-		StreamingSoundSources.Remove(SoundSource);
-	}
-}
-
-bool FLegacyAudioStreamingManager::IsManagedStreamingSoundSource(const FSoundSource* SoundSource) const
-{
-	FScopeLock Lock(&CriticalSection);
-	return StreamingSoundSources.FindByKey(SoundSource) != NULL;
-}
-
-bool FLegacyAudioStreamingManager::RequestChunk(const FSoundWaveProxyPtr&, uint32 ChunkIndex, TFunction<void(EAudioChunkLoadResult)> OnLoadCompleted, ENamedThreads::Type ThreadToCallOnLoadCompletedOn, bool bForImmediatePlayback)
-{
-	UE_LOG(LogAudio, Warning, TEXT("RequestChunk is only supported in Stream Caching."));
-	return false;
-}
-
-FAudioChunkHandle FLegacyAudioStreamingManager::GetLoadedChunk(const FSoundWaveProxyPtr&  SoundWave, uint32 ChunkIndex, bool bBlockForLoad, bool bForImmediatePlayback) const
-{
-	// Check for the spoof of failing to load a stream chunk
-	if (ensure(SoundWave.IsValid()) && SpoofFailedStreamChunkLoad > 0)
-	{
-		return FAudioChunkHandle();
-	}
-
-	// If we fail at getting the critical section here, early out. 
-	if (!CriticalSection.TryLock())
-	{
-		return FAudioChunkHandle();
-	}
-
-	const FStreamingWaveData* WaveData = StreamingSoundWaves.FindRef(SoundWave->GetFObjectKey());
-	if (WaveData)
-	{
-		if (WaveData->LoadedChunkIndices.Contains(ChunkIndex))
-		{
-			for (int32 Index = 0; Index < WaveData->LoadedChunks.Num(); ++Index)
-			{
-				if (WaveData->LoadedChunks[Index].Index == ChunkIndex)
-				{
-					CriticalSection.Unlock();
-					return BuildChunkHandle(WaveData->LoadedChunks[Index].Data, WaveData->LoadedChunks[Index].AudioDataSize, SoundWave, SoundWave->GetFName(), ChunkIndex, InvalidAudioStreamCacheLookupID);
-				}
-			}
-		}
-	}
-
-	CriticalSection.Unlock();
-	return FAudioChunkHandle();
-}
-
-FWaveRequest& FLegacyAudioStreamingManager::GetWaveRequest(FObjectKey Key)
-{
-	FWaveRequest* WaveRequest = WaveRequests.Find(Key);
-	if (!WaveRequest)
-	{
-		// Setup the new request so it always asks for chunk 0
-		WaveRequest = &WaveRequests.Add(Key);
-		WaveRequest->RequiredIndices.AddUnique(0);
-		WaveRequest->bPrioritiseRequest = false;
-	}
-	return *WaveRequest;
-}
-
-FString FLegacyAudioStreamingManager::GenerateMemoryReport()
-{
-	return FString(TEXT("Only supported when stream caching is enabled.\n"));
-}
-
-uint64 FLegacyAudioStreamingManager::TrimMemory(uint64 NumBytesToFree)
-{
-	UE_LOG(LogAudio, Warning, TEXT("IAudioStreamingManager::TrimMemory is only supported when Stream Caching is enabled."));
-	return 0;
-}
-
-int32 FLegacyAudioStreamingManager::RenderStatAudioStreaming(UWorld* World, FViewport* Viewport, FCanvas* Canvas, int32 X, int32 Y, const FVector* ViewLocation, const FRotator* ViewRotation)
-{
-	// Only supported by stream caching.
-	return Y;
-}
+#endif // #if WITH_EDITORONLY_DATA
\ No newline at end of file
diff --git a/Engine/Source/Runtime/Engine/Private/AudioStreamingCache.cpp b/Engine/Source/Runtime/Engine/Private/AudioStreamingCache.cpp
--- a/Engine/Source/Runtime/Engine/Private/AudioStreamingCache.cpp
+++ b/Engine/Source/Runtime/Engine/Private/AudioStreamingCache.cpp
@@ -351,7 +351,6 @@ bool FAudioChunkCache::FCacheElement::IsSoundWaveRetainingAudio() const
 FCachedAudioStreamingManager::FCachedAudioStreamingManager(const FCachedAudioStreamingManagerParams& InitParams)
 {
 	LLM_SCOPE(ELLMTag::AudioStreamCache);
-	check(FPlatformCompressionUtilities::IsCurrentPlatformUsingStreamCaching());
 	checkf(InitParams.Caches.Num() > 0, TEXT("FCachedAudioStreamingManager should be initialized with dimensions for at least one cache."));
 
 	// const FAudioStreamCachingSettings& CacheSettings = FPlatformCompressionUtilities::GetStreamCachingSettingsForCurrentPlatform();
@@ -416,16 +415,6 @@ void FCachedAudioStreamingManager::NotifyLevelOffset(class ULevel* Level, const
 	// Unused.
 }
 
-void FCachedAudioStreamingManager::AddStreamingSoundWave(const FSoundWaveProxyPtr& SoundWave)
-{
-	// Unused.
-}
-
-void FCachedAudioStreamingManager::RemoveStreamingSoundWave(const FSoundWaveProxyPtr& SoundWave)
-{
-	// Unused.
-}
-
 void FCachedAudioStreamingManager::AddForceInlineSoundWave(const FSoundWaveProxyPtr& SoundWave)
 {
 	// add the sound wave to the first cache
@@ -462,60 +451,6 @@ void FCachedAudioStreamingManager::RemoveMemoryCountedFeature(const FAudioStream
 	}
 }
 
-void FCachedAudioStreamingManager::AddDecoder(ICompressedAudioInfo* InCompressedAudioInfo)
-{
-	// Unused.
-}
-
-void FCachedAudioStreamingManager::RemoveDecoder(ICompressedAudioInfo* InCompressedAudioInfo)
-{
-	//Unused.
-}
-
-bool FCachedAudioStreamingManager::IsManagedStreamingSoundWave(const FSoundWaveProxyPtr&  SoundWave) const
-{
-	// Unused. The concept of a sound wave being "managed" doesn't apply here.
-	checkf(false, TEXT("Not Implemented!"));
-	return true;
-}
-
-bool FCachedAudioStreamingManager::IsStreamingInProgress(const FSoundWaveProxyPtr&  SoundWave)
-{
-	// This function is used in USoundWave cleanup.
-	// Since this manager owns the binary data we are streaming off of,
-	// It's safe to delete the USoundWave as long as
-	// There are NO sound sources playing with this Sound Wave.
-	//
-	// This is because a playing sound source might kick off a load for a new chunk,
-	// which dereferences the corresponding USoundWave
-	//
-	// As of right now, this is handled by USoundWave::FreeResources(), called
-	// by USoundWave::IsReadyForFinishDestroy.
-	return false;
-}
-
-bool FCachedAudioStreamingManager::CanCreateSoundSource(const FWaveInstance* WaveInstance) const
-{
-	return true;
-}
-
-void FCachedAudioStreamingManager::AddStreamingSoundSource(FSoundSource* SoundSource)
-{
-	// Unused.
-}
-
-void FCachedAudioStreamingManager::RemoveStreamingSoundSource(FSoundSource* SoundSource)
-{
-	// Unused.
-}
-
-bool FCachedAudioStreamingManager::IsManagedStreamingSoundSource(const FSoundSource* SoundSource) const
-{
-	// Unused. The concept of a sound wave being "managed" doesn't apply here.
-	checkf(false, TEXT("Not Implemented!"));
-	return true;
-}
-
 FAudioChunkHandle FCachedAudioStreamingManager::GetLoadedChunk(const FSoundWaveProxyPtr& SoundWave, uint32 ChunkIndex, bool bBlockForLoad, bool bForImmediatePlayback) const
 {
 	LLM_SCOPE(ELLMTag::AudioStreamCache);
diff --git a/Engine/Source/Runtime/Engine/Private/ContentStreaming.cpp b/Engine/Source/Runtime/Engine/Private/ContentStreaming.cpp
--- a/Engine/Source/Runtime/Engine/Private/ContentStreaming.cpp
+++ b/Engine/Source/Runtime/Engine/Private/ContentStreaming.cpp
@@ -790,15 +790,8 @@ FStreamingManagerCollection::FStreamingManagerCollection()
 
 	if (FApp::CanEverRenderAudio())
 	{
-		if (FPlatformCompressionUtilities::IsCurrentPlatformUsingStreamCaching())
-		{
-			FCachedAudioStreamingManagerParams Params = FPlatformCompressionUtilities::BuildCachedStreamingManagerParams();
-			AudioStreamingManager = new FCachedAudioStreamingManager(Params);
-		}
-		else
-		{
-			AudioStreamingManager = new FLegacyAudioStreamingManager();
-		}
+		FCachedAudioStreamingManagerParams Params = FPlatformCompressionUtilities::BuildCachedStreamingManagerParams();
+		AudioStreamingManager = new FCachedAudioStreamingManager(Params);
 	}
 	else
 	{
@@ -1281,17 +1274,8 @@ void FStreamingManagerCollection::OnAudioStreamingParamsChanged()
 	FPlatformCompressionUtilities::RecacheCookOverrides();
 
 	// Finally, reinitialize the streaming manager.
-	if (FPlatformCompressionUtilities::IsCurrentPlatformUsingStreamCaching())
-	{
-		FCachedAudioStreamingManagerParams Params = FPlatformCompressionUtilities::BuildCachedStreamingManagerParams();
-		AudioStreamingManager = new FCachedAudioStreamingManager(Params);
-	}
-	else
-	{
-		AudioStreamingManager = new FLegacyAudioStreamingManager();
-	}
-
-	AddStreamingManager(AudioStreamingManager);
+	FCachedAudioStreamingManagerParams Params = FPlatformCompressionUtilities::BuildCachedStreamingManagerParams();
+	AddStreamingManager(new FCachedAudioStreamingManager(Params));
 }
 #endif
 
@@ -1611,3 +1595,8 @@ FAudioChunkHandle IAudioStreamingManager::BuildChunkHandle(const uint8* InData,
 
 	return {};
 }
+
+void IAudioStreamingManager::LogWarning()
+{
+	UE_LOG(LogContentStreaming, Warning, TEXT("This function is now deprecated on IAudioStreamingManager and no-ops."));
+}
diff --git a/Engine/Source/Runtime/Engine/Private/SoundNode.cpp b/Engine/Source/Runtime/Engine/Private/SoundNode.cpp
--- a/Engine/Source/Runtime/Engine/Private/SoundNode.cpp
+++ b/Engine/Source/Runtime/Engine/Private/SoundNode.cpp
@@ -108,11 +108,6 @@ UPTRINT USoundNode::GetNodeWaveInstanceHash(const UPTRINT ParentWaveInstanceHash
 
 void USoundNode::PrimeChildWavePlayers(bool bRecurse)
 {
-	if (!FPlatformCompressionUtilities::IsCurrentPlatformUsingStreamCaching())
-	{
-		return;
-	}
-
 	// Note: it is not safe to call IAudioStreamingManager::RequestChunk from the game thread
 	// if there is no async loading thread.  This can deadlock:
 
@@ -159,7 +154,7 @@ void USoundNode::RetainChildWavePlayers(bool bRecurse)
 
 void USoundNode::OverrideLoadingBehaviorOnChildWaves(const bool bRecurse, const ESoundWaveLoadingBehavior InLoadingBehavior)
 {
-	if (!BypassRetainInSoundNodesCVar && FPlatformCompressionUtilities::IsCurrentPlatformUsingStreamCaching())
+	if (!BypassRetainInSoundNodesCVar)
 	{
 		// Search child nodes for wave players, then override their waves' loading behavior.
 		for (USoundNode* ChildNode : ChildNodes)
@@ -192,7 +187,7 @@ void USoundNode::OverrideLoadingBehaviorOnChildWaves(const bool bRecurse, const
 
 void USoundNode::ReleaseRetainerOnChildWavePlayers(bool bRecurse)
 {
-	if (bIsRetainingAudio && FPlatformCompressionUtilities::IsCurrentPlatformUsingStreamCaching())
+	if (bIsRetainingAudio)
 	{
 		// Search child nodes for wave players, then release their retainers.
 		for (USoundNode* ChildNode : ChildNodes)
diff --git a/Engine/Source/Runtime/Engine/Private/SoundWave.cpp b/Engine/Source/Runtime/Engine/Private/SoundWave.cpp
--- a/Engine/Source/Runtime/Engine/Private/SoundWave.cpp
+++ b/Engine/Source/Runtime/Engine/Private/SoundWave.cpp
@@ -959,7 +959,6 @@ USoundWave::USoundWave(const FObjectInitializer& ObjectInitializer)
 	CompressionQuality = 80;
 
 #if WITH_EDITOR
-	bWasStreamCachingEnabledOnLastCook = FPlatformCompressionUtilities::IsCurrentPlatformUsingStreamCaching();
 	bLoadedFromCookedData = false;
 
 	SoundWaveDataPtr->ResourceData.Empty();
@@ -2128,11 +2127,6 @@ void USoundWave::PostLoad()
 	// Only add this streaming sound if the platform supports streaming
 	if (FApp::CanEverRenderAudio() && IsStreaming(nullptr) && FPlatformProperties::SupportsAudioStreaming())
 	{
-		if (!ShouldUseStreamCaching())
-		{
-			IStreamingManager::Get().GetAudioStreamingManager().AddStreamingSoundWave(Proxy);
-		}
-
 		// Only request loading the zeroth chunk when streaming is supported and we can render audio.
 		// Especially important to avoid computing unneeded running platform related data during cook.
 		LoadZerothChunk();
@@ -2381,20 +2375,6 @@ void USoundWave::InvalidateSoundWaveIfNeccessary()
 		InvalidateCompressedData(true);
 		SampleRate = SampleRateOverride;
 	}
-	
-	// if stream caching was enabled since the last time we invalidated the compressed audio, force a re-cook.
-	const bool bIsStreamCachingEnabled = FPlatformCompressionUtilities::IsCurrentPlatformUsingStreamCaching();
-	if (bWasStreamCachingEnabledOnLastCook != bIsStreamCachingEnabled)
-	{
-		InvalidateCompressedData(true);
-		bWasStreamCachingEnabledOnLastCook = bIsStreamCachingEnabled;
-
-		// If stream caching is now turned on, recook the streaming audio if necessary.
-		if (bIsStreamCachingEnabled && IsStreaming(nullptr))
-		{
-			LoadZerothChunk();
-		}
-	}
 }
 
 void USoundWave::PreSave(FObjectPreSaveContext InSaveContext)
@@ -3211,22 +3191,10 @@ bool USoundWave::IsReadyForFinishDestroy()
 		}
 	}
 
-	// Only checking to see if it is set to ForceInline. ForceInline is not supported on 
-	// USoundClasses, so it is safe to ignore USoundClasses when calling `GetLoadingBehavior(...)`
-	const ESoundWaveLoadingBehavior CurrentLoadingBehavior = GetLoadingBehavior(false /* bCheckSoundClasses */);
-	bool bIsStreamingInProgress = false;
-	if (CurrentLoadingBehavior != ESoundWaveLoadingBehavior::ForceInline)
-	{
-		if (Proxy.IsValid())
-		{
-			bIsStreamingInProgress = IStreamingManager::Get().GetAudioStreamingManager().IsStreamingInProgress(Proxy);
-		}
-	}
-
 	check(GetPrecacheState() != ESoundWavePrecacheState::InProgress);
 
 	// Wait till streaming and decompression finishes before deleting resource.
-	if (!bIsStreamingInProgress && ResourceState == ESoundWaveResourceState::NeedsFree)
+	if (ResourceState == ESoundWaveResourceState::NeedsFree)
 	{
 		DECLARE_CYCLE_STAT(TEXT("FAudioThreadTask.FreeResources"), STAT_AudioFreeResources, STATGROUP_AudioThreadCommands);
 
@@ -3719,13 +3687,11 @@ bool USoundWave::IsStreaming(const FPlatformAudioCookOverrides& Overrides) const
 bool USoundWave::ShouldUseStreamCaching() const
 {
 	check(SoundWaveDataPtr);
-	const bool bPlatformUsingStreamCaching = FPlatformCompressionUtilities::IsCurrentPlatformUsingStreamCaching();
-	const bool bIsStreaming = IsStreaming(nullptr);
-	const bool Result = bPlatformUsingStreamCaching && bIsStreaming;
 
-	SoundWaveDataPtr->bShouldUseStreamCaching = Result;
+	const bool bIsStreaming = IsStreaming(nullptr);
+	SoundWaveDataPtr->bShouldUseStreamCaching = bIsStreaming;
 
-	return Result;
+	return bIsStreaming;
 }
 
 TArrayView<const uint8> USoundWave::GetZerothChunk(bool bForImmediatePlayback)
@@ -3874,28 +3840,11 @@ void USoundWave::UpdatePlatformData()
 			Proxy = CreateSoundWaveProxy();
 		}
 
-		// Make sure there are no pending requests in flight.
-		while (IStreamingManager::Get().GetAudioStreamingManager().IsStreamingInProgress(Proxy))
-		{
-			// Give up timeslice.
-			FPlatformProcess::Sleep(0);
-		}
-
 #if WITH_EDITORONLY_DATA
-		// Temporarily remove from streaming manager to release currently used data chunks
-		IStreamingManager::Get().GetAudioStreamingManager().RemoveStreamingSoundWave(Proxy);
-
 		// Recache platform data if the source has changed.
 		CachePlatformData(true /* bAsyncCache */);
-
-		// Add back to the streaming manager to reload first chunk
-		IStreamingManager::Get().GetAudioStreamingManager().AddStreamingSoundWave(Proxy);
 #endif
 	}
-	else if (Proxy.IsValid())
-	{
-		IStreamingManager::Get().GetAudioStreamingManager().RemoveStreamingSoundWave(Proxy);
-	}
 }
 
 float USoundWave::GetSampleRateForCurrentPlatform() const
diff --git a/Engine/Source/Runtime/Engine/Public/AudioCompressionSettingsUtils.h b/Engine/Source/Runtime/Engine/Public/AudioCompressionSettingsUtils.h
--- a/Engine/Source/Runtime/Engine/Public/AudioCompressionSettingsUtils.h
+++ b/Engine/Source/Runtime/Engine/Public/AudioCompressionSettingsUtils.h
@@ -24,7 +24,8 @@ public:
 	// null platformname means to use current platform
 	static ENGINE_API const FPlatformAudioCookOverrides* GetCookOverrides(const TCHAR* PlatformName=nullptr, bool bForceRecache = false);
 
-	static ENGINE_API bool IsCurrentPlatformUsingStreamCaching();
+	UE_DEPRECATED(5.5, "this always returns true, no longer needed.")
+	static ENGINE_API bool IsCurrentPlatformUsingStreamCaching() { return true; };
 
 	// null platformname means to use current platform
 	static ENGINE_API const FAudioStreamCachingSettings& GetStreamCachingSettingsForCurrentPlatform();
diff --git a/Engine/Source/Runtime/Engine/Public/AudioStreaming.h b/Engine/Source/Runtime/Engine/Public/AudioStreaming.h
--- a/Engine/Source/Runtime/Engine/Public/AudioStreaming.h
+++ b/Engine/Source/Runtime/Engine/Public/AudioStreaming.h
@@ -6,41 +6,16 @@ AudioStreaming.h: Definitions of classes used for audio streaming.
 
 #pragma once
 
+#include "Async/AsyncWork.h"
 #include "CoreMinimal.h"
-#include "HAL/ThreadSafeCounter.h"
-#include "Containers/IndirectArray.h"
-#include "Containers/Queue.h"
-#include "Stats/Stats.h"
 #include "ContentStreaming.h"
-#include "Async/AsyncWork.h"
-#include "Async/AsyncFileHandle.h"
-#include "HAL/ThreadSafeBool.h"
-#include "UObject/ObjectKey.h"
+#include "HAL/ThreadSafeCounter.h"
 
-class FSoundSource;
 class FSoundWaveProxy;
-class ICompressedAudioInfo;
-struct FLegacyAudioStreamingManager;
-struct FWaveInstance;
-
 using FSoundWaveProxyPtr = TSharedPtr<FSoundWaveProxy, ESPMode::ThreadSafe>;
 
-/** Lists possible states used by Thread-safe counter. */
-enum EAudioStreamingState
-{
-	// There are no pending requests/ all requests have been fulfilled.
-	AudioState_ReadyFor_Requests = 0,
-	// Initial request has completed and finalization needs to be kicked off.
-	AudioState_ReadyFor_Finalization = 1,
-	// We're currently loading in chunk data.
-	AudioState_InProgress_Loading = 2,
-	// ...
-	// States InProgress_Loading+N-1 means we're currently loading in N chunks
-	// ...
-};
-
 /**
- * Async worker to stream audio chunks from the derived data cache.
+ * Async worker to stream audio chunks from the derived data cache. 
  */
 class FAsyncStreamDerivedChunkWorker : public FNonAbandonableTask
 {
@@ -90,143 +65,9 @@ private:
 /** Async task to stream chunks from the derived data cache. */
 typedef FAsyncTask<FAsyncStreamDerivedChunkWorker> FAsyncStreamDerivedChunkTask;
 
-/**
- * Contains a request to load chunks of a sound wave
- */
-struct FWaveRequest
-{
-	TArray<uint32>	RequiredIndices;
-	bool			bPrioritiseRequest;
-};
-
-/**
- * Stores info about an audio chunk once it's been loaded
- */
-struct FLoadedAudioChunk
-{
-	uint8*	Data;
-	class IBulkDataIORequest* IORequest;
-	int32	DataSize;
-	int32	AudioDataSize;
-	uint32	Index;
-
-	FLoadedAudioChunk()
-		: Data(nullptr)
-		, IORequest(nullptr)
-		, DataSize(0)
-		, AudioDataSize(0)
-		, Index(0)
-	{
-	}
-
-	~FLoadedAudioChunk()
-	{
-		checkf(Data == nullptr, TEXT("Audio chunk Data ptr not null (%p), DataSize: %d"), Data, DataSize);
-		checkf(IORequest == nullptr, TEXT("Audio chunk IORequest ptr not null (%p)"), IORequest );
-	}
-};
-
-/**
- * Contains everything that will be needed by a SoundWave that's streaming in data
- */
-struct FStreamingWaveData final
-{
-	FStreamingWaveData();
-	~FStreamingWaveData();
-
-	// Frees streaming wave data resources, blocks pending async IO requests
-	void FreeResources();
-
-	/**
-	 * Sets up the streaming wave data and loads the first chunk of audio for instant play
-	 *
-	 * @param SoundWave	The SoundWave we are managing
-	 */
-	bool Initialize(const FSoundWaveProxyPtr& SoundWave, FLegacyAudioStreamingManager* InStreamingManager);
-
-	/**
-	 * Updates the streaming status of the sound wave and performs finalization when appropriate. The function returns
-	 * true while there are pending requests in flight and updating needs to continue.
-	 *
-	 * @return					true if there are requests in flight, false otherwise
-	 */
-	bool UpdateStreamingStatus();
-
-	/**
-	 * Tells the SoundWave which chunks are currently required so that it can start loading any needed
-	 *
-	 * @param InChunkIndices The Chunk Indices that are currently needed by all sources using this sound
-	 * @param bShouldPrioritizeAsyncIORequest Whether request should have higher priority than usual
-	 */
-	void UpdateChunkRequests(FWaveRequest& InWaveRequest);
-
-		/**
-	 * Checks whether the requested chunk indices differ from those loaded
-	 *
-	 * @param IndicesToLoad		List of chunk indices that should be loaded
-	 * @param IndicesToFree		List of chunk indices that should be freed
-	 * @return Whether any changes to loaded chunks are required
-	 */
-	bool HasPendingRequests(TArray<uint32>& IndicesToLoad, TArray<uint32>& IndicesToFree) const;
-
-	/**
-	 * Kicks off any pending requests
-	 */
-	void BeginPendingRequests(const TArray<uint32>& IndicesToLoad, const TArray<uint32>& IndicesToFree);
-
-	/**
-	* Blocks till all pending requests are fulfilled.
-	*
-	* @param TimeLimit		Optional time limit for processing, in seconds. Specifying 0 means infinite time limit.
-	* @return				Return true if there are no requests left in flight, false if the time limit was reached before they were finished.
-	*/
-	bool BlockTillAllRequestsFinished(float TimeLimit = 0.0f);
-
-#if WITH_EDITORONLY_DATA
-	/**
-	 * Finishes any Derived Data Cache requests that may be in progress
-	 *
-	 * @return Whether any of the requests failed.
-	 */
-	bool FinishDDCRequests();
-#endif //WITH_EDITORONLY_DATA
-
-private:
-	// Don't allow copy construction as it could free shared memory
-	FStreamingWaveData(const FStreamingWaveData& that);
-	FStreamingWaveData& operator=(FStreamingWaveData const&);
-
-	// Creates a new chunk, returns the chunk index
-	int32 AddNewLoadedChunk(int32 ChunkSize, int32 AudioSize);
-	void FreeLoadedChunk(FLoadedAudioChunk& LoadedChunk);
-
-public:
-	/** SoundWave this streaming data is for */
-	FSoundWaveProxyPtr SoundWave;
-
-	/** Thread-safe counter indicating the audio streaming state. */
-	mutable FThreadSafeCounter	PendingChunkChangeRequestStatus;
-
-	/* Contains pointers to Chunks of audio data that have been streamed in */
-	TArray<FLoadedAudioChunk> LoadedChunks;
-
-	/** Indices of chunks that are currently loaded */
-	TArray<uint32>	LoadedChunkIndices;
-
-	/** Indices of chunks we want to have loaded */
-	FWaveRequest	CurrentRequest;
-
-#if WITH_EDITORONLY_DATA
-	/** Pending async derived data streaming tasks */
-	TIndirectArray<FAsyncStreamDerivedChunkTask> PendingAsyncStreamDerivedChunkTasks;
-#endif // #if WITH_EDITORONLY_DATA
-
-	/** Ptr to owning audio streaming manager. */
-	FLegacyAudioStreamingManager* AudioStreamingManager;
-};
-
 /** Struct used to store results of an async file load. */
-struct FASyncAudioChunkLoadResult
+using FStreamingWaveData = void*;
+struct UE_DEPRECATED(5.5, "this struct relies on deprecated / deleted types, and should not be used") FASyncAudioChunkLoadResult
 {
 	// Place to safely copy the ptr of a loaded audio chunk when load result is finished
 	uint8* DataResults;
@@ -243,84 +84,3 @@ struct FASyncAudioChunkLoadResult
 		, LoadedAudioChunkIndex(INDEX_NONE)
 	{}
 };
-
-/**
-* Streaming manager dealing with audio.
-*/
-struct FLegacyAudioStreamingManager : public IAudioStreamingManager
-{
-	/** Constructor, initializing all members */
-	FLegacyAudioStreamingManager();
-
-	virtual ~FLegacyAudioStreamingManager();
-
-	// IStreamingManager interface
-	virtual void UpdateResourceStreaming( float DeltaTime, bool bProcessEverything=false ) override;
-	virtual int32 BlockTillAllRequestsFinished( float TimeLimit = 0.0f, bool bLogResults = false ) override;
-	virtual void CancelForcedResources() override;
-	virtual void NotifyLevelChange() override;
-	virtual void SetDisregardWorldResourcesForFrames( int32 NumFrames ) override;
-	virtual void AddLevel( class ULevel* Level ) override;
-	virtual void RemoveLevel( class ULevel* Level ) override;
-	virtual void NotifyLevelOffset( class ULevel* Level, const FVector& Offset ) override;
-	// End IStreamingManager interface
-
-	// IAudioStreamingManager interface
-	virtual void AddStreamingSoundWave(const FSoundWaveProxyPtr& SoundWave) override;
-	virtual void RemoveStreamingSoundWave(const FSoundWaveProxyPtr& SoundWave) override;
-	virtual void AddDecoder(ICompressedAudioInfo* CompressedAudioInfo) override;
-	virtual void RemoveDecoder(ICompressedAudioInfo* CompressedAudioInfo) override;
-	virtual bool IsManagedStreamingSoundWave(const FSoundWaveProxyPtr&  SoundWave) const override;
-	virtual bool IsStreamingInProgress(const FSoundWaveProxyPtr&  SoundWave) override;
-	virtual bool CanCreateSoundSource(const FWaveInstance* WaveInstance) const override;
-	virtual void AddStreamingSoundSource(FSoundSource* SoundSource) override;
-	virtual void RemoveStreamingSoundSource(FSoundSource* SoundSource) override;
-	virtual bool IsManagedStreamingSoundSource(const FSoundSource* SoundSource) const override;
-	virtual bool RequestChunk(const FSoundWaveProxyPtr& SoundWave, uint32 ChunkIndex, TFunction<void(EAudioChunkLoadResult)> OnLoadCompleted, ENamedThreads::Type ThreadToCallOnLoadCompletedOn, bool bForImmediatePlayback = false) override;
-	virtual FAudioChunkHandle GetLoadedChunk(const FSoundWaveProxyPtr&  SoundWave, uint32 ChunkIndex, bool bBlockForLoad = false, bool bForImmediatePlayback = false) const override;
-	virtual uint64 TrimMemory(uint64 NumBytesToFree) override;
-	virtual int32 RenderStatAudioStreaming(UWorld* World, FViewport* Viewport, FCanvas* Canvas, int32 X, int32 Y, const FVector* ViewLocation, const FRotator* ViewRotation) override;
-	virtual FString GenerateMemoryReport() override;
-	virtual void SetProfilingMode(bool bEnabled) override {};
-	// End IAudioStreamingManager interface
-
-	/** Called when an async callback is made on an async loading audio chunk request. */
-	void OnAsyncFileCallback(FStreamingWaveData* StreamingWaveData, int32 LoadedAudioChunkIndex, IBulkDataIORequest* ReadRequest);
-
-	/** Processes pending async file IO results. */
-	void ProcessPendingAsyncFileResults();
-
-protected:
-
-	/** Unused by the legacy audio streaming manager. */
-	virtual void AddReferenceToChunk(const FAudioChunkHandle& InHandle) override {};
-	virtual void RemoveReferenceToChunk(const FAudioChunkHandle& InHandle) override {};
-
-	/**
-	 * Gets Wave request associated with a specific wave
-	 *
-	 * @param SoundWave		SoundWave we want request for
-	 * @return Existing or new request structure
-	 */
-	FWaveRequest& GetWaveRequest(FObjectKey Key);
-
-	/** Sound Waves being managed. */
-	TMap<FObjectKey, FStreamingWaveData*> StreamingSoundWaves;
-
-	/** Sound Sources being managed. */
-	TArray<FSoundSource*>	StreamingSoundSources;
-
-	/** Map of requests to make next time sound waves are ready */
-	TMap<FObjectKey, FWaveRequest> WaveRequests;
-
-	/** Results of async loading audio chunks. */
-	TArray<FASyncAudioChunkLoadResult*> AsyncAudioStreamChunkResults;
-	mutable FCriticalSection ChunkResultCriticalSection;
-
-	/** Critical section to protect usage of shared gamethread/audiothread members */
-	mutable FCriticalSection CriticalSection;
-
-	/** Compressed audio info objects which are used to avoid deleting chunks with in-fligth decodes. */
-	TArray<ICompressedAudioInfo*> CompressedAudioInfos;
-
-};
diff --git a/Engine/Source/Runtime/Engine/Public/AudioStreamingCache.h b/Engine/Source/Runtime/Engine/Public/AudioStreamingCache.h
--- a/Engine/Source/Runtime/Engine/Public/AudioStreamingCache.h
+++ b/Engine/Source/Runtime/Engine/Public/AudioStreamingCache.h
@@ -533,22 +533,9 @@ public:
 	virtual void NotifyLevelOffset( class ULevel* Level, const FVector& Offset ) override;
 	// End IStreamingManager interface
 
-	// IAudioStreamingManager interface (unused functions)
-	virtual void AddStreamingSoundWave(const FSoundWaveProxyPtr& SoundWave) override;
-	virtual void RemoveStreamingSoundWave(const FSoundWaveProxyPtr& SoundWave) override;
+	// IAudioStreamingManager interface
 	virtual void AddForceInlineSoundWave(const FSoundWaveProxyPtr& SoundWave) override;
 	virtual void RemoveForceInlineSoundWave(const FSoundWaveProxyPtr& SoundWave) override;
-	virtual void AddDecoder(ICompressedAudioInfo* CompressedAudioInfo) override;
-	virtual void RemoveDecoder(ICompressedAudioInfo* CompressedAudioInfo) override;
-	virtual bool IsManagedStreamingSoundWave(const FSoundWaveProxyPtr&  SoundWave) const override;
-	virtual bool IsStreamingInProgress(const FSoundWaveProxyPtr&  SoundWave) override;
-	virtual bool CanCreateSoundSource(const FWaveInstance* WaveInstance) const override;
-	virtual void AddStreamingSoundSource(FSoundSource* SoundSource) override;
-	virtual void RemoveStreamingSoundSource(FSoundSource* SoundSource) override;
-	virtual bool IsManagedStreamingSoundSource(const FSoundSource* SoundSource) const override;
-	// End IAudioStreamingManager interface (unused)
-
-	// IAudioStreamingManager interface (used functions)
 	virtual bool RequestChunk(const FSoundWaveProxyPtr& SoundWave, uint32 ChunkIndex, TFunction<void(EAudioChunkLoadResult)> OnLoadCompleted, ENamedThreads::Type ThreadToCallOnLoadCompletedOn, bool bForImmediatePlayback = false) override;
 	virtual FAudioChunkHandle GetLoadedChunk(const FSoundWaveProxyPtr&  SoundWave, uint32 ChunkIndex, bool bBlockForLoad = false, bool bForImmediatePlayback = false) const override;
 	virtual uint64 TrimMemory(uint64 NumBytesToFree) override;
diff --git a/Engine/Source/Runtime/Engine/Public/ContentStreaming.h b/Engine/Source/Runtime/Engine/Public/ContentStreaming.h
--- a/Engine/Source/Runtime/Engine/Public/ContentStreaming.h
+++ b/Engine/Source/Runtime/Engine/Public/ContentStreaming.h
@@ -511,41 +511,12 @@ enum class EAudioChunkLoadResult : uint8
  */
 struct IAudioStreamingManager : public IStreamingManager
 {
-	/** Adds a new Sound Wave to the streaming manager. */
-	virtual void AddStreamingSoundWave(const FSoundWaveProxyPtr& SoundWave) = 0;
-
-	/** Removes a Sound Wave from the streaming manager. */
-	virtual void RemoveStreamingSoundWave(const FSoundWaveProxyPtr& SoundWave) = 0;
-
 	/** Adds the memory usage of the force inline sound to the streaming cache budget */
 	virtual void AddForceInlineSoundWave(const FSoundWaveProxyPtr& SoundWave) { };
 
 	/** Removes the memory usage of the force inline sound from the streaming cache budget */
 	virtual void RemoveForceInlineSoundWave(const FSoundWaveProxyPtr& SoundWave) { };
-
-	/** Adds the decoder to the streaming manager to prevent stream chunks from getting reaped from underneath it */
-	virtual void AddDecoder(ICompressedAudioInfo* CompressedAudioInfo) = 0;
-
-	/** Removes the decoder from the streaming manager. */
-	virtual void RemoveDecoder(ICompressedAudioInfo* CompressedAudioInfo) = 0;
-
-	/** Returns true if this is a Sound Wave that is managed by the streaming manager. */
-	virtual bool IsManagedStreamingSoundWave(const FSoundWaveProxyPtr&  SoundWave) const = 0;
-
-	/** Returns true if this Sound Wave is currently streaming a chunk. */
-	virtual bool IsStreamingInProgress(const FSoundWaveProxyPtr&  SoundWave) = 0;
-
-	virtual bool CanCreateSoundSource(const FWaveInstance* WaveInstance) const = 0;
-
-	/** Adds a new Sound Source to the streaming manager. */
-	virtual void AddStreamingSoundSource(FSoundSource* SoundSource) = 0;
-
-	/** Removes a Sound Source from the streaming manager. */
-	virtual void RemoveStreamingSoundSource(FSoundSource* SoundSource) = 0;
-
-	/** Returns true if this is a streaming Sound Source that is managed by the streaming manager. */
-	virtual bool IsManagedStreamingSoundSource(const FSoundSource* SoundSource) const = 0;
-
+	
 	/** 
 	 * Manually prepare a chunk to start playing back. This should only be used when the Load On Demand feature is enabled, and returns false on failure. 
 	 * @param SoundWave SoundWave we would like to request a chunk of.
@@ -617,6 +588,50 @@ protected:
 	* This can be used to decrease the memory count for external features. Called by FAudioStreamCacheMemoryHandle.
 	*/
 	virtual void RemoveMemoryCountedFeature(const FAudioStreamCacheMemoryHandle& Feature) { };
+
+// DEPRECATED VIRTUAL INTERFACE
+private:
+	static void LogWarning();
+
+public:
+	/** Adds a new Sound Wave to the streaming manager. */
+	UE_DEPRECATED(5.5, "Use These methods are no longer used (removing the unused legacy streaming manager)")
+	virtual void AddStreamingSoundWave(const FSoundWaveProxyPtr& SoundWave) { LogWarning(); }
+
+	/** Removes a Sound Wave from the streaming manager. */
+	UE_DEPRECATED(5.5, "Use These methods are no longer used (removing the unused legacy streaming manager)")
+	virtual void RemoveStreamingSoundWave(const FSoundWaveProxyPtr& SoundWave) { LogWarning(); }
+
+	/** Adds the decoder to the streaming manager to prevent stream chunks from getting reaped from underneath it */
+	UE_DEPRECATED(5.5, "Use These methods are no longer used (removing the unused legacy streaming manager)")
+	virtual void AddDecoder(ICompressedAudioInfo* CompressedAudioInfo) { LogWarning(); }
+
+	/** Removes the decoder from the streaming manager. */
+	UE_DEPRECATED(5.5, "Use These methods are no longer used (removing the unused legacy streaming manager)")
+	virtual void RemoveDecoder(ICompressedAudioInfo* CompressedAudioInfo) { LogWarning(); }
+
+	/** Returns true if this is a Sound Wave that is managed by the streaming manager. */
+	UE_DEPRECATED(5.5, "Use These methods are no longer used (removing the unused legacy streaming manager)")
+	virtual bool IsManagedStreamingSoundWave(const FSoundWaveProxyPtr&  SoundWave) const { LogWarning(); return false; }
+
+	/** Returns true if this Sound Wave is currently streaming a chunk. */
+	UE_DEPRECATED(5.5, "Use These methods are no longer used (removing the unused legacy streaming manager)")
+	virtual bool IsStreamingInProgress(const FSoundWaveProxyPtr&  SoundWave) { LogWarning(); return false; }
+
+	UE_DEPRECATED(5.5, "Use These methods are no longer used (removing the unused legacy streaming manager)")
+	virtual bool CanCreateSoundSource(const FWaveInstance* WaveInstance) const { LogWarning(); return false; }
+
+	/** Adds a new Sound Source to the streaming manager. */
+	UE_DEPRECATED(5.5, "Use These methods are no longer used (removing the unused legacy streaming manager)")
+	virtual void AddStreamingSoundSource(FSoundSource* SoundSource) { LogWarning(); }
+
+	/** Removes a Sound Source from the streaming manager. */
+	UE_DEPRECATED(5.5, "Use These methods are no longer used (removing the unused legacy streaming manager)")
+	virtual void RemoveStreamingSoundSource(FSoundSource* SoundSource) { LogWarning(); }
+
+	/** Returns true if this is a streaming Sound Source that is managed by the streaming manager. */
+	UE_DEPRECATED(5.5, "Use These methods are no longer used (removing the unused legacy streaming manager)")
+	virtual bool IsManagedStreamingSoundSource(const FSoundSource* SoundSource) const { LogWarning(); return false; }
 };
 
 /**
@@ -633,20 +648,8 @@ struct FDummyAudioStreamingManager final : public IAudioStreamingManager
 	virtual void RemoveLevel(class ULevel* Level) {}
 	virtual void NotifyLevelOffset(class ULevel* Level, const FVector& Offset) {}
 
-	virtual void AddStreamingSoundWave(const FSoundWaveProxyPtr& SoundWave) override {}
-	virtual void RemoveStreamingSoundWave(const FSoundWaveProxyPtr& SoundWave) override {}
-	virtual void AddForceInlineSoundWave(const FSoundWaveProxyPtr& SoundWave) override {}
-	virtual void RemoveForceInlineSoundWave(const FSoundWaveProxyPtr& SoundWave) override {}
 	virtual void AddMemoryCountedFeature(const FAudioStreamCacheMemoryHandle& Feature) override {}
 	virtual void RemoveMemoryCountedFeature(const FAudioStreamCacheMemoryHandle& Feature) override {}
-	virtual void AddDecoder(ICompressedAudioInfo* CompressedAudioInfo) override {}
-	virtual void RemoveDecoder(ICompressedAudioInfo* CompressedAudioInfo) override {}
-	virtual bool IsManagedStreamingSoundWave(const FSoundWaveProxyPtr& SoundWave) const override { return false; }
-	virtual bool IsStreamingInProgress(const FSoundWaveProxyPtr& SoundWave) override { return false; }
-	virtual bool CanCreateSoundSource(const FWaveInstance* WaveInstance) const override { return false; }
-	virtual void AddStreamingSoundSource(FSoundSource* SoundSource) override {}
-	virtual void RemoveStreamingSoundSource(FSoundSource* SoundSource) override {}
-	virtual bool IsManagedStreamingSoundSource(const FSoundSource* SoundSource) const override { return false; }
 	virtual bool RequestChunk(const FSoundWaveProxyPtr& SoundWave, uint32 ChunkIndex, TFunction<void(EAudioChunkLoadResult)> OnLoadCompleted = [](EAudioChunkLoadResult) {}, ENamedThreads::Type ThreadToCallOnLoadCompletedOn = ENamedThreads::AnyThread, bool bForImmediatePlayback = false) override { return false; }
 	virtual FAudioChunkHandle GetLoadedChunk(const FSoundWaveProxyPtr& SoundWave, uint32 ChunkIndex, bool bBlockForLoad = false, bool bForImmediatePlayback = false) const override { return FAudioChunkHandle(); }
 	virtual uint64 TrimMemory(uint64 NumBytesToFree) override { return 0; }
diff --git a/Engine/Source/Runtime/IOS/IOSAudio/Private/IOSAudioSource.cpp b/Engine/Source/Runtime/IOS/IOSAudio/Private/IOSAudioSource.cpp
--- a/Engine/Source/Runtime/IOS/IOSAudio/Private/IOSAudioSource.cpp
+++ b/Engine/Source/Runtime/IOS/IOSAudio/Private/IOSAudioSource.cpp
@@ -318,8 +318,6 @@ void FIOSAudioSoundSource::Stop(void)
 		FPlatformProcess::Sleep(0.0f);
 	}
 
-	IStreamingManager::Get().GetAudioStreamingManager().RemoveStreamingSoundSource(this);
-
 	if (WaveInstance)
 	{
 		Pause();
