Temporarily disable parallel GDME because of a race condition causing random crashes.

This will cause a regression on render thread time, but it's preferrable until we can figure out why the race is happening.

#rnx

diff --git a/Engine/Source/Runtime/Engine/Private/PrimitiveSceneProxy.cpp b/Engine/Source/Runtime/Engine/Private/PrimitiveSceneProxy.cpp
--- a/Engine/Source/Runtime/Engine/Private/PrimitiveSceneProxy.cpp
+++ b/Engine/Source/Runtime/Engine/Private/PrimitiveSceneProxy.cpp
@@ -96,7 +96,8 @@ static TAutoConsoleVariable<int32> CVarApproximateOcclusionQueries(
 	ECVF_RenderThreadSafe
 );
 
-bool GParallelGatherDynamicMeshElements = true;
+// Temporarily off by default as a workaround for UE-220106.
+bool GParallelGatherDynamicMeshElements = false;
 static FAutoConsoleVariableRef CVarParallelGatherDynamicMeshElements(
 	TEXT("r.Visibility.DynamicMeshElements.Parallel"),
 	GParallelGatherDynamicMeshElements,
