Fix crash/memory stomp in the open likely resizing a container

This open is no longer needed

diff --git a/Engine/Source/Runtime/Engine/Private/PrimitiveComponentPhysics.cpp b/Engine/Source/Runtime/Engine/Private/PrimitiveComponentPhysics.cpp
--- a/Engine/Source/Runtime/Engine/Private/PrimitiveComponentPhysics.cpp
+++ b/Engine/Source/Runtime/Engine/Private/PrimitiveComponentPhysics.cpp
@@ -1249,15 +1249,7 @@ void UPrimitiveComponent::SetCollisionEnabled(ECollisionEnabled::Type NewType)
 
 	if (CurrentType != NewType)
 	{
-		UE_AUTORTFM_OPEN2{
-			BodyInstance.SetCollisionEnabled(NewType);
-		};
-
-		// If we fail set the CollisionEnabled back to the CurrentType
-		UE_AUTORTFM_ONABORT2(=, this)
-		{
-			BodyInstance.SetCollisionEnabled(CurrentType);
-		};
+		BodyInstance.SetCollisionEnabled(NewType);
 
 		EnsurePhysicsStateCreated();
 		OnComponentCollisionSettingsChanged();
