Add sockets and correct frame rate to anim compression DDC key

diff --git a/Engine/Source/Runtime/Engine/Private/Animation/AnimCompressionTypes.cpp b/Engine/Source/Runtime/Engine/Private/Animation/AnimCompressionTypes.cpp
--- a/Engine/Source/Runtime/Engine/Private/Animation/AnimCompressionTypes.cpp
+++ b/Engine/Source/Runtime/Engine/Private/Animation/AnimCompressionTypes.cpp
@@ -809,12 +809,7 @@ void FCompressibleAnimData::FetchData(const ITargetPlatform* InPlatform)
 
 	FAnimationUtils::BuildSkeletonMetaData(Skeleton, BoneData);
 
-	const FFrameRate DefaultSamplingFrameRate = AnimSequence->GetSamplingFrameRate();
-	const FFrameRate PlatformSamplingFrameRate = InPlatform ? AnimSequence->GetTargetSamplingFrameRate(InPlatform) : FFrameRate(0,0);
-	
-	const bool bValidTargetSampleRate = PlatformSamplingFrameRate.IsValid() && (PlatformSamplingFrameRate.IsMultipleOf(DefaultSamplingFrameRate) || PlatformSamplingFrameRate.IsFactorOf(DefaultSamplingFrameRate));
-	
-	const FFrameRate& FrameRateToSampleWith = (DefaultSamplingFrameRate != PlatformSamplingFrameRate && bValidTargetSampleRate) ? PlatformSamplingFrameRate : DefaultSamplingFrameRate;	
+	const FFrameRate FrameRateToSampleWith = UE::Anim::Compression::GetCompressionFrameRate(*AnimSequence, InPlatform);
 	
 	const FFrameTime SampleFrameTime = FrameRateToSampleWith.AsFrameTime(SequenceLength);
 	check(FMath::IsNearlyZero(SampleFrameTime.GetSubFrame()));	
@@ -1651,6 +1646,19 @@ UE::Anim::Compression::FAnimDDCKeyArgs::FAnimDDCKeyArgs(const UAnimSequenceBase&
 	, TargetPlatform(TargetPlatform)
 	{
 			}
+
+namespace UE::Anim::Compression
+{
+	FFrameRate GetCompressionFrameRate(const UAnimSequence& AnimSequence, const ITargetPlatform* TargetPlatform)
+	{
+		const FFrameRate DefaultSamplingFrameRate = AnimSequence.GetSamplingFrameRate();
+		const FFrameRate PlatformSamplingFrameRate = TargetPlatform ? AnimSequence.GetTargetSamplingFrameRate(TargetPlatform) : FFrameRate(0, 0);
+
+		const bool bValidTargetSampleRate = PlatformSamplingFrameRate.IsValid() && (PlatformSamplingFrameRate.IsMultipleOf(DefaultSamplingFrameRate) || PlatformSamplingFrameRate.IsFactorOf(DefaultSamplingFrameRate));
+
+		return DefaultSamplingFrameRate != PlatformSamplingFrameRate && bValidTargetSampleRate ? PlatformSamplingFrameRate : DefaultSamplingFrameRate;
+	}
+}
 #endif // WITH_EDITORONLY_DATA
 
 #if WITH_EDITOR
diff --git a/Engine/Source/Runtime/Engine/Private/Animation/AnimSequence.cpp b/Engine/Source/Runtime/Engine/Private/Animation/AnimSequence.cpp
--- a/Engine/Source/Runtime/Engine/Private/Animation/AnimSequence.cpp
+++ b/Engine/Source/Runtime/Engine/Private/Animation/AnimSequence.cpp
@@ -19,6 +19,7 @@
 #include "Animation/MirrorDataTable.h"
 #include "UObject/UObjectIterator.h"
 #include "Engine/SkeletalMesh.h"
+#include "Engine/SkeletalMeshSocket.h"
 #include "EngineUtils.h"
 #include "AnimationUtils.h"
 #include "BonePose.h"
@@ -4622,15 +4623,22 @@ FString UAnimSequence::CreateDerivedDataKeyString(const ITargetPlatform* TargetP
 		}
 	}
 
-	const FFrameRate FrameRate = PlatformTargetFrameRate.GetValueForPlatform(TargetPlatform->GetPlatformInfo().IniPlatformName);
+	// Include sockets since they can impact precision requirements
+	for (const USkeletalMeshSocket* Socket : CurrentSkeleton->Sockets)
+	{
+		FName BoneName = Socket->BoneName;
+		ArcToHexString.Ar << BoneName;
+	}
+
+	const FFrameRate FrameRate = UE::Anim::Compression::GetCompressionFrameRate(*this, TargetPlatform);
 
 	FString Ret = FString::Printf(TEXT("%i_%s%s%s_%c%c%i_%s_%s_%i_%i_%s"),
 		CompressCommandletVersion,
 		PRAGMA_DISABLE_DEPRECATION_WARNINGS
 		*GetDataModel()->GenerateGuid().ToString(),
 		PRAGMA_ENABLE_DEPRECATION_WARNINGS
-		*GetSkeleton()->GetGuid().ToString(),
-		*GetSkeleton()->GetVirtualBoneGuid().ToString(),
+		*CurrentSkeleton->GetGuid().ToString(),
+		*CurrentSkeleton->GetVirtualBoneGuid().ToString(),
 		AdditiveType,
 		RefType,
 		RefFrameIndex,
diff --git a/Engine/Source/Runtime/Engine/Public/Animation/AnimCompressionTypes.h b/Engine/Source/Runtime/Engine/Public/Animation/AnimCompressionTypes.h
--- a/Engine/Source/Runtime/Engine/Public/Animation/AnimCompressionTypes.h
+++ b/Engine/Source/Runtime/Engine/Public/Animation/AnimCompressionTypes.h
@@ -896,7 +896,7 @@ struct FRootMotionReset
 	}
 };
 
-#if WITH_EDITOR
+#if WITH_EDITORONLY_DATA
 namespace UE::Anim::Compression {
 	struct FAnimDDCKeyArgs
 	{
@@ -906,8 +906,10 @@ namespace UE::Anim::Compression {
 		const UAnimSequenceBase& AnimSequence;
 		const ITargetPlatform* TargetPlatform;
 	};
+
+	ENGINE_API FFrameRate GetCompressionFrameRate(const UAnimSequence& AnimSequence, const ITargetPlatform* TargetPlatform);
 }
-#endif // WITH_EDITOR
+#endif // WITH_EDITORONLY_DATA
 
 UE_DEPRECATED(5.1, "Signature of DecompressPose has been deprecated, use UE::Anim::Decompression::DecompressPose instead")
 extern void DecompressPose(FCompactPose& OutPose,
