[NavMesh] Addition of IsGeneratingDiffableObject and using it to prevent tile serialization when diffable object are made

diff --git a/Engine/Source/Runtime/Engine/Private/TransactionCommon.cpp b/Engine/Source/Runtime/Engine/Private/TransactionCommon.cpp
--- a/Engine/Source/Runtime/Engine/Private/TransactionCommon.cpp
+++ b/Engine/Source/Runtime/Engine/Private/TransactionCommon.cpp
@@ -446,6 +446,12 @@ FSerializedTaggedData& FDiffableObjectDataWriter::FCachedTaggedDataEntry::SyncCa
 
 namespace DiffUtil
 {
+static FArchive* GCurrentDiffableObjectDataWriter = nullptr; // Set before GetDiffableObject calls its serialize impl
+	
+bool IsGeneratingDiffableObject(const FArchive& Ar)
+{
+	return GCurrentDiffableObjectDataWriter == &Ar;
+}
 
 FDiffableObject GetDiffableObject(const UObject* Object, const FGetDiffableObjectOptions& Options)
 {
@@ -455,6 +461,8 @@ FDiffableObject GetDiffableObject(const UObject* Object, const FGetDiffableObjec
 	if (Options.bSerializeEvenIfPendingKill || !DiffableObject.ObjectInfo.bIsPendingKill)
 	{
 		FDiffableObjectDataWriter DiffWriter(DiffableObject, Options.PropertiesToSerialize);
+		TGuardValue<FArchive*> GuardGCurrentDiffableObjectDataWriter(GCurrentDiffableObjectDataWriter, &DiffWriter);
+		
 		switch (Options.ObjectSerializationMode)
 		{
 		case EGetDiffableObjectMode::SerializeObject:
diff --git a/Engine/Source/Runtime/Engine/Public/TransactionCommon.h b/Engine/Source/Runtime/Engine/Public/TransactionCommon.h
--- a/Engine/Source/Runtime/Engine/Public/TransactionCommon.h
+++ b/Engine/Source/Runtime/Engine/Public/TransactionCommon.h
@@ -415,6 +415,7 @@ private:
 
 namespace DiffUtil
 {
+ENGINE_API bool IsGeneratingDiffableObject(const FArchive& Ar);
 
 enum class EGetDiffableObjectMode : uint8
 {
diff --git a/Engine/Source/Runtime/NavigationSystem/Private/NavMesh/PImplRecastNavMesh.cpp b/Engine/Source/Runtime/NavigationSystem/Private/NavMesh/PImplRecastNavMesh.cpp
--- a/Engine/Source/Runtime/NavigationSystem/Private/NavMesh/PImplRecastNavMesh.cpp
+++ b/Engine/Source/Runtime/NavigationSystem/Private/NavMesh/PImplRecastNavMesh.cpp
@@ -2,6 +2,7 @@
 
 #include "NavMesh/PImplRecastNavMesh.h"
 #include "NavigationSystem.h"
+#include "TransactionCommon.h"
 
 #if WITH_RECAST
 
@@ -512,7 +513,7 @@ void FPImplRecastNavMesh::Serialize( FArchive& Ar, int32 NavMeshVersion )
 	int32 NumTiles = 0;
 	TArray<FNavTileRef> TilesToSave;
 
-	if (Ar.IsSaving() && !Ar.IsTransacting()) // Do not save tiles during transactions(i.e. undo/redo)
+	if (Ar.IsSaving() && !Ar.IsTransacting() && !UE::Transaction::DiffUtil::IsGeneratingDiffableObject(Ar)) // Do not save tiles during transactions (i.e. undo/redo)
 	{
 		TilesToSave.Reserve(DetourNavMesh->getMaxTiles());
 		
