Perf Metrics - Added histogram for the async loading queue

diff --git a/Engine/Source/Runtime/Engine/Private/ChartCreation.cpp b/Engine/Source/Runtime/Engine/Private/ChartCreation.cpp
--- a/Engine/Source/Runtime/Engine/Private/ChartCreation.cpp
+++ b/Engine/Source/Runtime/Engine/Private/ChartCreation.cpp
@@ -646,6 +646,9 @@ void FPerformanceTrackingChart::Reset(const FDateTime& InStartTime)
 		}
 	}
 
+	// ignore zero since we can infer that from the frame counter
+	AsyncPackageQueueHistogram.InitFromArray({ 1.0, 250.0, 500.0 });
+
 	StartTemperatureLevel = -1.0f;
 	StopTemperatureLevel = -1.0f;
 
@@ -661,6 +664,7 @@ void FPerformanceTrackingChart::AccumulateWith(const FPerformanceTrackingChart&
 	FrametimeHistogram += Chart.FrametimeHistogram;
 	HitchTimeHistogram += Chart.HitchTimeHistogram;
 	DynamicResHistogram += Chart.DynamicResHistogram;
+	AsyncPackageQueueHistogram += Chart.AsyncPackageQueueHistogram;
 	NumFramesBound_GameThread += Chart.NumFramesBound_GameThread;
 	NumFramesBound_RenderThread += Chart.NumFramesBound_RenderThread;
 	NumFramesBound_RHIThread += Chart.NumFramesBound_RHIThread;
@@ -794,6 +798,11 @@ void FPerformanceTrackingChart::ProcessFrame(const FFrameData& FrameData)
 		MaxFlushAsyncLoadingTime = FMath::Max(MaxFlushAsyncLoadingTime, FrameData.FlushAsyncLoadingTime);
 		TotalSyncLoadCount += FrameData.SyncLoadCount;
 
+		if (FrameData.NumAsyncPackages > 0)
+		{
+			AsyncPackageQueueHistogram.AddMeasurement(FrameData.NumAsyncPackages);
+		}
+
 		// Track draw calls
 		// Multi-GPU support : ChartCreation doesn't support MGPU yet
 		MaxDrawCalls = FMath::Max(MaxDrawCalls, GNumDrawCallsRHI[0]);
@@ -848,6 +857,30 @@ void FPerformanceTrackingChart::ProcessFrame(const FFrameData& FrameData)
 	}
 }
 
+double FPerformanceTrackingChart::GetPercentFramesIOBusy(int32 StartingBin) const
+{
+	const int64 NumFrames = GetNumFrames();
+	if (NumFrames > 0)
+	{
+		int32 NumFramesBusy = 0;
+		if (StartingBin == 0)
+		{
+			NumFramesBusy = AsyncPackageQueueHistogram.GetNumMeasurements();
+		}
+		else
+		{
+			for (int32 i = StartingBin; i < AsyncPackageQueueHistogram.GetNumBins(); ++i)
+			{
+				NumFramesBusy += AsyncPackageQueueHistogram.GetBinObservationsCount(i);
+			}
+		}		
+
+		return (NumFramesBusy / (double)NumFrames) * 100.f;
+	}
+
+	return 0.f;
+}
+
 void FPerformanceTrackingChart::DumpFPSChart(const FString& InMapName)
 {
 	TArray<const FPerformanceTrackingChart*> Charts;
@@ -1256,6 +1289,7 @@ IPerformanceDataConsumer::FFrameData FPerformanceTrackingSystem::AnalyzeFrame(fl
 	FrameData.FlushAsyncLoadingTime = GFlushAsyncLoadingTime;
 	FrameData.FlushAsyncLoadingCount = GFlushAsyncLoadingCount;
 	FrameData.SyncLoadCount = GSyncLoadCount;
+	FrameData.NumAsyncPackages = GetNumAsyncPackages();
 
 	// Optionally disregard frames that took too long when accumulating data.
 	FrameData.bBinThisFrame = (DeltaSeconds < GMaximumFrameTimeToConsiderForHitchesAndBinning) || (GMaximumFrameTimeToConsiderForHitchesAndBinning <= 0.0f);
diff --git a/Engine/Source/Runtime/Engine/Public/ChartCreation.h b/Engine/Source/Runtime/Engine/Public/ChartCreation.h
--- a/Engine/Source/Runtime/Engine/Public/ChartCreation.h
+++ b/Engine/Source/Runtime/Engine/Public/ChartCreation.h
@@ -80,6 +80,9 @@ public:
 		/** The number of sync loads performed in this frame. */
 		uint32 SyncLoadCount;
 
+		/** Number of async packages in the loading queue */
+		int32 NumAsyncPackages;
+
 		// Should this frame be considered for histogram generation (controlled by t.FPSChart.MaxFrameDeltaSecsBeforeDiscarding)
 		bool bBinThisFrame;
 
@@ -110,6 +113,7 @@ public:
 			, FlushAsyncLoadingTime(0.0)
 			, FlushAsyncLoadingCount(0)
 			, SyncLoadCount(0)
+			, NumAsyncPackages(0)
 			, bBinThisFrame(false)
 			, bGameThreadBound(false)
 			, bRenderThreadBound(false)
@@ -148,6 +152,9 @@ public:
 	// Hitch time histogram (in seconds)
 	FHistogram DynamicResHistogram;
 
+	// Track size of the async loader package queue (I/O)
+	FHistogram AsyncPackageQueueHistogram;
+
 	/** Number of frames for each time of <boundtype> **/
 	uint32 NumFramesBound_GameThread;
 	uint32 NumFramesBound_RenderThread;
@@ -300,6 +307,8 @@ public:
 		return HitchTimeHistogram.GetAverageOfAllMeasures();
 	}
 
+	ENGINE_API double GetPercentFramesIOBusy(int32 StartingBin=0) const;
+
 	void ChangeLabel(const FString& NewLabel)
 	{
 		ChartLabel = NewLabel;
