[ODSC]
* Fixed linker errors when building game in non-monolithic

diff --git a/Engine/Source/Runtime/Engine/Private/ODSC/ODSCThread.cpp b/Engine/Source/Runtime/Engine/Private/ODSC/ODSCThread.cpp
--- a/Engine/Source/Runtime/Engine/Private/ODSC/ODSCThread.cpp
+++ b/Engine/Source/Runtime/Engine/Private/ODSC/ODSCThread.cpp
@@ -394,7 +394,7 @@ void FODSCThread::AddShaderPipelineRequest(
 
 		if (bShouldAddRequest)
 		{
-#if WITH_ODSC
+#if WITH_ODSC && IS_MONOLITHIC
 			if (PrimitiveSceneInfo)
 			{
 				AActor* OwningActor = PrimitiveSceneInfo->GetComponentForDebugOnly() ? PrimitiveSceneInfo->GetComponentForDebugOnly()->GetOwner() : nullptr;
diff --git a/Engine/Source/Runtime/Engine/Public/ODSC/ODSCManager.h b/Engine/Source/Runtime/Engine/Public/ODSC/ODSCManager.h
--- a/Engine/Source/Runtime/Engine/Public/ODSC/ODSCManager.h
+++ b/Engine/Source/Runtime/Engine/Public/ODSC/ODSCManager.h
@@ -115,8 +115,8 @@ public:
 	static void UnregisterMaterialName(const FMaterial* Material);
 	static void RegisterMaterialShaderMaps(const FString& MaterialName, const TArray<TRefCountPtr<FMaterialShaderMap>>& LoadedShaderMaps);
 	static FMaterialShaderMap* FindMaterialShaderMap(const FString& MaterialName, const FMaterialShaderMapId& ShaderMapId);
-	static void SetCurrentPrimitiveSceneInfo(FPrimitiveSceneInfo* PrimitiveSceneInfo);
-	static void ResetCurrentPrimitiveSceneInfo();
+	ENGINE_API static void SetCurrentPrimitiveSceneInfo(FPrimitiveSceneInfo* PrimitiveSceneInfo);
+	ENGINE_API static void ResetCurrentPrimitiveSceneInfo();
 
 private:
 	friend class FODSCManagerAccess;
