Validate and clamp MinLOD settings for a SkeletalMesh during asset saving and in a cooker

diff --git a/Engine/Source/Runtime/Engine/Classes/Engine/SkeletalMesh.h b/Engine/Source/Runtime/Engine/Classes/Engine/SkeletalMesh.h
--- a/Engine/Source/Runtime/Engine/Classes/Engine/SkeletalMesh.h
+++ b/Engine/Source/Runtime/Engine/Classes/Engine/SkeletalMesh.h
@@ -2402,6 +2402,11 @@ private:
 	 *  are still empty.
 	 */
 	void PostLoadRecoverConvertLODModelsToMeshDescription();
+
+	/** 
+	 *  Check if Quality and Platform MinLOD settings are valid 
+	 */
+	void CheckForValidMinLODs(FPerQualityLevelInt& QualityLocalMinLOD, FPerPlatformInt& LocalMinLOD, int32& OutMinAvailableLOD, TArray<TPair<int32, FName>>& OutInvalidMinLODs) const;
 		
 public:
 	/*
diff --git a/Engine/Source/Runtime/Engine/Private/SkeletalMesh.cpp b/Engine/Source/Runtime/Engine/Private/SkeletalMesh.cpp
--- a/Engine/Source/Runtime/Engine/Private/SkeletalMesh.cpp
+++ b/Engine/Source/Runtime/Engine/Private/SkeletalMesh.cpp
@@ -84,6 +84,8 @@
 #include "Components/BrushComponent.h"
 
 #include "ClothingAssetBase.h"
+#include "Async/Async.h"
+#include "Misc/UObjectToken.h"
 
 #if WITH_EDITOR
 #include "ScopedTransaction.h"
@@ -1467,6 +1469,47 @@ void USkeletalMesh::UpdateGenerateUpToData()
 	}
 }
 
+void USkeletalMesh::CheckForValidMinLODs(FPerQualityLevelInt& QualityLocalMinLOD, FPerPlatformInt& LocalMinLOD, int32& OutMinAvailableLOD, TArray<TPair<int32, FName>>& OutInvalidMinLODs) const
+{
+	const FSkeletalMeshRenderData* LocalRenderData = GetSkeletalMeshRenderData();
+	if (!LocalRenderData)
+	{
+		return;
+	}
+
+	OutMinAvailableLOD = FMath::Max<int32>(LocalRenderData->LODRenderData.Num() - 1, 0);
+
+	auto CheckValidMinLOD = [LocalRenderData, OutMinAvailableLOD, &OutInvalidMinLODs](int32& LODIdx, FName OverrideName)
+	{
+		if (!LocalRenderData->LODRenderData.IsValidIndex(LODIdx))
+		{
+			OutInvalidMinLODs.Emplace(LODIdx, OverrideName);
+			LODIdx = OutMinAvailableLOD;
+		}
+	};
+
+	if (IsMinLodQualityLevelEnable())
+	{
+		QualityLocalMinLOD = GetQualityLevelMinLod();
+		CheckValidMinLOD(QualityLocalMinLOD.Default, NAME_None);
+
+		for (TMap<int32, int32>::TIterator It(QualityLocalMinLOD.PerQuality); It; ++It)
+		{
+			CheckValidMinLOD(It.Value(), QualityLevelProperty::QualityLevelToFName(It.Key()));
+		}
+	}
+	else
+	{
+		LocalMinLOD = GetMinLod();
+		CheckValidMinLOD(LocalMinLOD.Default, NAME_None);
+
+		for (TMap<FName, int32>::TIterator It(LocalMinLOD.PerPlatform); It; ++It)
+		{
+			CheckValidMinLOD(It.Value(), It.Key());
+		}
+	}
+}
+
 EDataValidationResult USkeletalMesh::IsDataValid(FDataValidationContext& Context) const
 {
 	EDataValidationResult ValidationResult = Super::IsDataValid(Context);
@@ -1522,6 +1565,40 @@ EDataValidationResult USkeletalMesh::IsDataValid(FDataValidationContext& Context
 			}
 		}
 	}
+
+	{
+		// check the MinLOD values are all within range
+		FPerQualityLevelInt QualityLocalMinLOD;
+		FPerPlatformInt LocalMinLOD;
+		int32 MinAvailableLOD = INDEX_NONE;
+		TArray<TPair<int32, FName>> InvalidMinLODs;
+		CheckForValidMinLODs(QualityLocalMinLOD, LocalMinLOD, MinAvailableLOD, InvalidMinLODs);
+		if (InvalidMinLODs.Num() > 0)
+		{
+			for (const TPair<int32, FName>& InvalidMinLOD : InvalidMinLODs)
+			{
+				const int32 LODIdx = InvalidMinLOD.Key;
+				const FName OverrideName = InvalidMinLOD.Value;
+
+				FFormatNamedArguments Arguments;
+				Arguments.Add(TEXT("MinLOD"), FText::AsNumber(LODIdx));
+				Arguments.Add(TEXT("MinAvailLOD"), FText::AsNumber(MinAvailableLOD));
+				Arguments.Add(TEXT("OverrideName"), FText::FromName(OverrideName));
+
+				if (OverrideName.IsNone())
+				{
+					Context.AddWarning(FText::Format(LOCTEXT("LoadError_BadMinLOD", "Min LOD value of {MinLOD} is out of range 0..{MinAvailLOD}."), Arguments));
+				}
+				else
+				{
+					Context.AddWarning(FText::Format(LOCTEXT("LoadError_BadMinLODWithOverride", "Min LOD override of {MinLOD} for {OverrideName} is out of range 0..{MinAvailLOD}."), Arguments));
+				}
+			}
+
+			ValidationResult = EDataValidationResult::Invalid;
+		}
+	}
+
 	return ValidationResult;
 }
 
@@ -3599,6 +3676,73 @@ void USkeletalMesh::ExecutePostLoadInternal(FSkinnedAssetPostLoadContext& Contex
 			Context.bHasCachedDerivedData = true;
 		}
 	}
+
+	// check the MinLOD values are all within range
+	FPerQualityLevelInt QualityLocalMinLOD;
+	FPerPlatformInt LocalMinLOD;
+	int32 MinAvailableLOD = INDEX_NONE;
+	TArray<TPair<int32, FName>> InvalidMinLODs;
+	CheckForValidMinLODs(QualityLocalMinLOD, LocalMinLOD, MinAvailableLOD, InvalidMinLODs);
+	if (InvalidMinLODs.Num())
+	{
+		if (IsMinLodQualityLevelEnable())
+		{
+			SetQualityLevelMinLod(QualityLocalMinLOD);
+		}
+		else
+		{
+			SetMinLod(LocalMinLOD);
+		}
+
+		TArray<FText> MinLODErrors;
+		for (const TPair<int32, FName>& InvalidMinLOD : InvalidMinLODs)
+		{
+			const int32 LODIdx = InvalidMinLOD.Key;
+			const FName OverrideName = InvalidMinLOD.Value;
+
+			FFormatNamedArguments Arguments;
+			Arguments.Add(TEXT("MinLOD"), FText::AsNumber(LODIdx));
+			Arguments.Add(TEXT("MinAvailLOD"), FText::AsNumber(MinAvailableLOD));
+			Arguments.Add(TEXT("OverrideName"), FText::FromName(OverrideName));
+			if (OverrideName.IsNone())
+			{
+				MinLODErrors.Add(FText::Format(LOCTEXT("LoadError_BadMinLOD_Fixed", "Min LOD value of {MinLOD} is out of range 0..{MinAvailLOD} and has been adjusted to {MinAvailLOD}. Please verify and resave the asset."), Arguments));
+			}
+			else
+			{
+				MinLODErrors.Add(FText::Format(LOCTEXT("LoadError_BadMinLODWithOverride_Fixed", "Min LOD override of {MinLOD} for {OverrideName} is out of range 0..{MinAvailLOD} and has been adjusted to {MinAvailLOD}. Please verify and resave the asset."), Arguments));
+			}
+		}
+
+		if (IsRunningCommandlet())
+		{
+			for (const FText& MinLODError : MinLODErrors)
+			{
+				UE_ASSET_LOG(LogSkeletalMesh, Warning, this, TEXT("%s"), *MinLODError.ToString());
+			}
+		}
+		else
+		{
+			TSharedRef<FUObjectToken> TokenRef = FUObjectToken::Create(this);
+			Async(EAsyncExecution::TaskGraphMainThread,
+				// No choice to MoveTemp here, the SharedRef is not thread safe so it cannot
+				// be copied to another thread, only moved.
+				[Token = MoveTemp(TokenRef), MinAvailableLOD, MinLODErrors]()
+				{
+					for (const FText& MinLODError : MinLODErrors)
+					{
+						FMessageLog("LoadErrors").Warning()
+							->AddToken(Token)
+							->AddToken(FTextToken::Create(MinLODError));
+					}
+
+					FMessageLog("LoadErrors").Open();
+				}
+			);
+		}
+	}
+
+
 #endif // WITH_EDITOR
 }
 
