Add a small fix to update a TextureSampler Anisotropy at runtime in PIE when r.MaxAnisotropy is not 0

diff --git a/Engine/Source/Runtime/Engine/Private/Rendering/StreamableTextureResource.cpp b/Engine/Source/Runtime/Engine/Private/Rendering/StreamableTextureResource.cpp
--- a/Engine/Source/Runtime/Engine/Private/Rendering/StreamableTextureResource.cpp
+++ b/Engine/Source/Runtime/Engine/Private/Rendering/StreamableTextureResource.cpp
@@ -11,6 +11,7 @@
 #include "ProfilingDebugging/LoadTimeTracker.h"
 #include "ProfilingDebugging/ScopedDebugInfo.h"
 #include "Stats/StatsTrace.h"
+#include "RHIUtilities.h"
 
 #if STATS
 int64 GUITextureMemory = 0;
@@ -250,7 +251,7 @@ void FStreamableTextureResource::RefreshSamplerStates()
 		AddressV,
 		AddressW,
 		MipBias,
-		MaxAniso
+		ComputeAnisotropyRT(MaxAniso) // Need this in case anisotropy is changed at runtime
 	);
 	SamplerStateRHI = GetOrCreateSamplerState(SamplerStateInitializer);
 
