Change behaviour of the MLB function expressions to support legacy Vertex/Pixel attribute pathways + Substrate UseBlendParameters by allowing their settings to be changed via function call inputs (all default to false).

#virtualized

diff --git a/Engine/Source/Editor/MaterialEditor/Private/MaterialEditor.cpp b/Engine/Source/Editor/MaterialEditor/Private/MaterialEditor.cpp
--- a/Engine/Source/Editor/MaterialEditor/Private/MaterialEditor.cpp
+++ b/Engine/Source/Editor/MaterialEditor/Private/MaterialEditor.cpp
@@ -78,6 +78,7 @@
 #include "Materials/MaterialExpressionRuntimeVirtualTextureSampleParameter.h"
 #include "Materials/MaterialExpressionSparseVolumeTextureSample.h"
 #include "Materials/MaterialExpressionScalarParameter.h"
+#include "Materials/MaterialExpressionStaticBool.h"
 #include "Materials/MaterialExpressionStaticComponentMaskParameter.h"
 #include "Materials/MaterialExpressionStaticSwitchParameter.h"
 #include "Materials/MaterialExpressionTextureSampleParameter.h"
@@ -852,15 +853,30 @@ void FMaterialEditor::InitMaterialEditor( const EToolkitMode::Type Mode, const T
 									if(UMaterialExpressionMaterialFunctionCall* BlendFunctionCall = Cast<UMaterialExpressionMaterialFunctionCall>(CreateNewMaterialExpression(UMaterialExpressionMaterialFunctionCall::StaticClass(), FVector2D(-100, 300), false, false)))
 									{
 										BlendFunctionCall->Function = MaterialFunction;
-										InputTop->MaterialExpressionEditorX = -500;
-										InputBottom->MaterialExpressionEditorX = -500;
+										InputTop->MaterialExpressionEditorX = -800;
+										InputBottom->MaterialExpressionEditorX = -800;
+
 										if (BlendFunctionCall->SetMaterialFunction(DefaultBlendFunction))
 										{
-											if (BlendFunctionCall->FunctionInputs.Num() >= 2 && BlendFunctionCall->FunctionOutputs.Num() > 0)
+											if (BlendFunctionCall->FunctionInputs.Num() >= 8 && BlendFunctionCall->FunctionOutputs.Num() > 0)
 											{
 												BlendFunctionCall->FunctionInputs[0].Input.Connect(0, InputBottom);
 												BlendFunctionCall->FunctionInputs[1].Input.Connect(0, InputTop);
 												UMaterialEditingLibrary::ConnectMaterialExpressions(BlendFunctionCall, FString(), Expression, FString());
+
+												int32 StartIndex = 3;
+												int32 EndIndex = 7;
+												int32 BoolYPosition = 500;
+												for (int32 FunctionCallIndex = StartIndex; FunctionCallIndex <= EndIndex; FunctionCallIndex++)
+												{
+													UMaterialExpressionStaticBool* ThisBool = Cast<UMaterialExpressionStaticBool>(CreateNewMaterialExpression(UMaterialExpressionStaticBool::StaticClass(), FVector2D(-450, BoolYPosition), false, false));
+													if (ThisBool)
+													{
+														BlendFunctionCall->FunctionInputs[FunctionCallIndex].Input.Connect(0, ThisBool);
+														BoolYPosition += 80;		
+													}
+												}
+
 												bMaterialDirty = true;
 											}
 										}
diff --git a/Engine/Source/Runtime/Engine/Private/Materials/MaterialExpressions.cpp b/Engine/Source/Runtime/Engine/Private/Materials/MaterialExpressions.cpp
--- a/Engine/Source/Runtime/Engine/Private/Materials/MaterialExpressions.cpp
+++ b/Engine/Source/Runtime/Engine/Private/Materials/MaterialExpressions.cpp
@@ -7464,11 +7464,11 @@ bool UMaterialExpressionGetMaterialAttributes::IsResultSubstrateMaterial(int32 O
 
 void UMaterialExpressionGetMaterialAttributes::GatherSubstrateMaterialInfo(FSubstrateMaterialInfo& SubstrateMaterialInfo, int32 OutputIndex)
 {
+	if (MaterialAttributes.Expression
 #if ENABLE_MATERIAL_LAYER_PROTOTYPE
-	if (IsResultSubstrateMaterial(OutputIndex))
-#else
-	if (MaterialAttributes.Expression)
+	&& IsResultSubstrateMaterial(OutputIndex)
 #endif
+	)
 	{
 		MaterialAttributes.Expression->GatherSubstrateMaterialInfo(SubstrateMaterialInfo, MaterialAttributes.OutputIndex);
 	}
@@ -7476,11 +7476,11 @@ void UMaterialExpressionGetMaterialAttributes::GatherSubstrateMaterialInfo(FSubs
 
 FSubstrateOperator* UMaterialExpressionGetMaterialAttributes::SubstrateGenerateMaterialTopologyTree(class FMaterialCompiler* Compiler, class UMaterialExpression* Parent, int32 OutputIndex)
 {
+	if (MaterialAttributes.Expression
 #if ENABLE_MATERIAL_LAYER_PROTOTYPE
-	if (IsResultSubstrateMaterial(OutputIndex))
-#else
-	if (MaterialAttributes.Expression)
+		&& IsResultSubstrateMaterial(OutputIndex)
 #endif
+	)
 	{
 		return MaterialAttributes.Expression->SubstrateGenerateMaterialTopologyTree(Compiler, Parent, MaterialAttributes.OutputIndex);
 	}
@@ -8154,6 +8154,176 @@ FSubstrateOperator* UMaterialExpressionBlendMaterialAttributes::SubstrateGenerat
 }
 #endif // WITH_EDITOR
 
+UMaterialExpressionLegacyBlendMaterialAttributes::UMaterialExpressionLegacyBlendMaterialAttributes(const FObjectInitializer& ObjectInitializer) 
+	: Super(ObjectInitializer){}
+
+#if ENABLE_MATERIAL_LAYER_PROTOTYPE
+#if WITH_EDITOR
+void UMaterialExpressionLegacyBlendMaterialAttributes::GetCaption(TArray<FString>& OutCaptions) const
+{
+	OutCaptions.Add(TEXT("LegacyBlendMaterialAttributes"));
+}
+bool UMaterialExpressionLegacyBlendMaterialAttributes::CanEditChange(const FProperty* InProperty) const
+{
+	bool bIsEditable = Super::CanEditChange(InProperty);
+	if(InProperty != nullptr)
+	{
+		FName PropertyFName = InProperty->GetFName();
+		if (PropertyFName == GET_MEMBER_NAME_CHECKED(UMaterialExpressionBlendMaterialAttributes, VertexAttributeBlendType))
+		{
+			bIsEditable = !(VertexAttribute_UseA.IsConnected() || VertexAttribute_UseB.IsConnected());
+		}
+		else if (PropertyFName == GET_MEMBER_NAME_CHECKED(UMaterialExpressionBlendMaterialAttributes, PixelAttributeBlendType))
+		{
+			bIsEditable = !(PixelAttribute_UseA.IsConnected() || PixelAttribute_UseB.IsConnected());
+		}
+	}
+	return bIsEditable;
+}
+FExpressionInput* UMaterialExpressionLegacyBlendMaterialAttributes::GetInput(int32 InputIndex)
+{
+	switch(InputIndex)
+	{
+		case 3: return &VertexAttribute_UseA;
+		case 4: return &VertexAttribute_UseB;
+		case 5: return &PixelAttribute_UseA;
+		case 6: return &PixelAttribute_UseB;
+		default: return Super::GetInput(InputIndex);
+	}
+}
+FName UMaterialExpressionLegacyBlendMaterialAttributes::GetInputName(int32 InputIndex) const
+{
+	FName Name;
+
+	switch (InputIndex)
+	{
+		case 3: Name = TEXT("Vertex Attributes Use A"); break;
+		case 4: Name = TEXT("Vertex Attributes Use B"); break;
+		case 5: Name = TEXT("Pixel Attributes Use A"); break;
+		case 6: Name = TEXT("Pixel Attributes Use B"); break;
+		default: Name = Super::GetInputName(InputIndex);
+	}
+
+	return Name;
+}
+bool UMaterialExpressionLegacyBlendMaterialAttributes::IsInputConnectionRequired(int32 InputIndex) const
+{
+	switch (InputIndex)
+	{
+		case 3:
+		case 4:
+		case 5:
+		case 6: return false;
+		default: return Super::IsInputConnectionRequired(InputIndex);
+	}
+}
+uint32 UMaterialExpressionLegacyBlendMaterialAttributes::GetInputType(int32 InputIndex)
+{
+	switch (InputIndex)
+	{
+		case 3:
+		case 4:
+		case 5:
+		case 6: return MCT_StaticBool;
+		default: return Super::GetInputType(InputIndex);
+	}
+}
+
+bool GetStaticBoolExpressionOutput(FMaterialCompiler* Compiler, FExpressionInput& Input)
+{
+	if(Input.GetTracedInput().Expression)
+	{
+		bool bSucceeded = false;
+		bool ResultBool = Compiler->GetStaticBoolValue(Input.Compile(Compiler), bSucceeded);
+		if(bSucceeded)
+		{
+			return ResultBool;
+		}
+	}
+	return false;
+}
+
+int32 UMaterialExpressionLegacyBlendMaterialAttributes::Compile(FMaterialCompiler* Compiler, int32 OutputIndex)
+{
+	VertexAttributeBlendType = EMaterialAttributeBlend::Blend;
+	PixelAttributeBlendType = EMaterialAttributeBlend::Blend;	
+	
+	bool bVertexUseA = GetStaticBoolExpressionOutput(Compiler, VertexAttribute_UseA);
+	bool bVertexUseB = GetStaticBoolExpressionOutput(Compiler, VertexAttribute_UseB);
+	if (bVertexUseA && !bVertexUseB)
+	{
+		VertexAttributeBlendType = EMaterialAttributeBlend::UseA;
+	}
+	if (!bVertexUseA && bVertexUseB)
+	{
+		VertexAttributeBlendType = EMaterialAttributeBlend::UseB;
+	}
+
+	bool bPixelUseA = GetStaticBoolExpressionOutput(Compiler, PixelAttribute_UseA);
+	bool bPixelUseB = GetStaticBoolExpressionOutput(Compiler, PixelAttribute_UseB);
+	if (bPixelUseA && !bPixelUseB)
+	{
+		PixelAttributeBlendType = EMaterialAttributeBlend::UseA;
+	}
+	if (!bPixelUseA && bPixelUseB)
+	{
+		PixelAttributeBlendType = EMaterialAttributeBlend::UseB;
+	}
+
+	const FGuid AttributeID = Compiler->GetMaterialAttribute();
+
+	// Blending is optional, can skip on a per-node basis
+	EMaterialAttributeBlend::Type BlendType;
+	EShaderFrequency AttributeFrequency = FMaterialAttributeDefinitionMap::GetShaderFrequency(AttributeID);
+
+	switch (AttributeFrequency)
+	{
+	case SF_Vertex:	BlendType = VertexAttributeBlendType;	break;
+	case SF_Pixel:	BlendType = PixelAttributeBlendType;	break;
+	default:
+		return Compiler->Errorf(TEXT("Attribute blending for shader frequency %i not implemented."), AttributeFrequency);
+	}
+
+	switch (BlendType)
+	{
+	case EMaterialAttributeBlend::UseA:
+	{
+		Compiler->SubstrateTreeStackPush(this, 0);
+		int32 CodeChunk = A.CompileWithDefault(Compiler, AttributeID);
+		Compiler->SubstrateTreeStackPop();
+		return CodeChunk;
+	}
+	case EMaterialAttributeBlend::UseB:
+	{
+		Compiler->SubstrateTreeStackPush(this, 1);
+		int32 CodeChunk = B.CompileWithDefault(Compiler, AttributeID);
+		Compiler->SubstrateTreeStackPop();
+		return CodeChunk;
+	}
+	default:
+		check(BlendType == EMaterialAttributeBlend::Blend);
+	}
+
+	// Allow custom blends or fallback to standard interpolation
+	Compiler->SubstrateTreeStackPush(this, 0);
+	int32 ResultA = A.CompileWithDefault(Compiler, AttributeID);
+	Compiler->SubstrateTreeStackPop();
+	Compiler->SubstrateTreeStackPush(this, 1);
+	int32 ResultB = B.CompileWithDefault(Compiler, AttributeID);
+	Compiler->SubstrateTreeStackPop();
+	int32 ResultAlpha = Alpha.Compile(Compiler);
+
+	MaterialAttributeBlendFunction BlendFunction = FMaterialAttributeDefinitionMap::GetBlendFunction(AttributeID);
+	if (BlendFunctionType == EMaterialAttributeBlendFunction::Type::Horizontal && BlendFunction)
+	{
+		return BlendFunction(Compiler, ResultA, ResultB, ResultAlpha);
+	}	
+
+	return Compiler->Lerp(ResultA, ResultB, ResultAlpha);
+}
+#endif // WITH_EDITOR
+#endif //ENABLE_MATERIAL_LAYER_PROTOTYPE
+
 //
 //	UMaterialExpressionMaterialAttributeLayers
 //
@@ -16087,18 +16257,21 @@ void UMaterialFunction::ConvertExpressionsBetweenLegacyAndSubstrate()
 	UMaterialFunctionEditorOnlyData* EditorOnly = GetEditorOnlyData();
 	if (!EditorOnly)
 	{
+		UE_LOG(LogMaterial, Warning, TEXT("Material Layer/Blend Function %s editor only data missing."), *GetFullName());
 		return;
 	}
 
 	TArray<TObjectPtr<UMaterialExpression>>& Expressions = EditorOnly->ExpressionCollection.Expressions;
 	if (Expressions.Num() < 1)
 	{
+		UE_LOG(LogMaterial, Warning, TEXT("Material Layer/Blend Function %s contains no expressions."), *GetFullName());
 		return;
 	}
 
 	UMaterialExpressionFunctionOutput* OutputNode = Cast<UMaterialExpressionFunctionOutput>(Expressions[0]);
 	if (!OutputNode || !OutputNode->A.IsConnected() || !OutputNode->A.Expression->IsResultMaterialAttributes(OutputNode->A.OutputIndex))
 	{
+		UE_LOG(LogMaterial, Log, TEXT("Material Layer/Blend Function %s is not valid for Substrate upgrade path."), *GetFullName());
 		return;
 	}
 
@@ -16172,32 +16345,125 @@ void UMaterialFunction::ConvertExpressionsBetweenLegacyAndSubstrate()
 	bool bBlendConverted = false;
 	if (GetMaterialFunctionUsage() == EMaterialFunctionUsage::MaterialLayerBlend)
 	{
+		uint32 InputCount = 0;
 		UMaterialExpressionFunctionInput* BottomInput = nullptr;
 		UMaterialExpressionFunctionInput* TopInput = nullptr;
-
+		UMaterialExpressionBlendMaterialAttributes* BlendNodeForInputs = nullptr;
 		for (UMaterialExpression* Expression : Expressions)
 		{
 			if (UMaterialExpressionFunctionInput* InputNode = Cast<UMaterialExpressionFunctionInput>(Expression))
 			{
-				if (!BottomInput && InputNode->InputName == TEXT("Bottom Layer"))
+				if (!BottomInput)
 				{
 					BottomInput = InputNode;
 				}
-				else if (!TopInput && InputNode->InputName == TEXT("Top Layer"))
+				else if (!TopInput)
 				{
 					TopInput = InputNode;
 				}
+				InputCount++;
+			}
+			else if(!BlendNodeForInputs)
+			{
+				BlendNodeForInputs = Cast<UMaterialExpressionBlendMaterialAttributes>(Expression);
+			}
+
+			if(BottomInput && TopInput && BlendNodeForInputs)
+			{
+				break;
 			}
 		}
-	
-		//The two legacy inputs are required at a minimum.
-		if (!TopInput || !BottomInput)
+
+		//Legacy blend nodes should have exactly 2 inputs, if not, assume invalid for auto-upgrade.
+		if (!TopInput || !BottomInput || InputCount > 2)
 		{
+			UE_LOG(LogMaterial, Warning, TEXT("Material Blend Function %s does not have exactly 2 inputs, so cannot be automatically upgraded to support Substrate."), *GetFullName());
 			return;
 		}
 
-		BottomInput->InputName = TEXT("Background Layer");
-		TopInput->InputName = TEXT("Foreground Layer");
+		bool bInputsPresumedCorrect = false;
+		if(BlendNodeForInputs)
+		{
+			auto RecurseBlendInputs= [&](UMaterialExpression* Expression, auto&& RecurseBlendInputs) -> UMaterialExpressionFunctionInput*
+				{
+					//If null or if the expression has already been recursed, skip.
+					if (!Expression)
+					{
+						return nullptr;
+					}
+
+					if(UMaterialExpressionFunctionInput* InputExpression = Cast<UMaterialExpressionFunctionInput>(Expression))
+					{
+						return InputExpression;
+					}
+
+					//If the expression is valid, iterate the connected expressions for upgrade to Substrate, then recurse via each valid expression
+					for (FExpressionInputIterator It{ Expression }; It; ++It)
+					{
+						if(UMaterialExpressionFunctionInput* InputExpression = RecurseBlendInputs(It.Input->Expression, RecurseBlendInputs))
+						{
+							return InputExpression;
+						}
+					}
+
+					return nullptr;
+				};
+
+			UMaterialExpressionFunctionInput* RecursedBottomInput = RecurseBlendInputs(BlendNodeForInputs->A.Expression, RecurseBlendInputs);
+			UMaterialExpressionFunctionInput* RecursedTopInput = RecurseBlendInputs(BlendNodeForInputs->B.Expression, RecurseBlendInputs);
+
+			if(RecursedBottomInput && RecursedTopInput && RecursedBottomInput != RecursedTopInput)
+			{
+				BottomInput = RecursedBottomInput;
+				TopInput = RecursedTopInput;
+				bInputsPresumedCorrect = true;
+			}
+		}
+
+		auto SwapInputs = [&](UMaterialExpressionFunctionInput* BottomInput, UMaterialExpressionFunctionInput*  TopInput) -> void
+			{
+				UMaterialExpressionFunctionInput* SwapInput = BottomInput;
+				TopInput = BottomInput;
+				BottomInput = SwapInput;
+			};
+
+		if(!bInputsPresumedCorrect)
+		{
+			if (BottomInput->InputName == TEXT("Bottom Layer") || BottomInput->InputName == TEXT("Background Layer"))
+			{
+				SwapInputs(BottomInput, TopInput);
+				if(TopInput->InputName == TEXT("Top Layer") || TopInput->InputName == TEXT("Foreground Layer"))
+				{
+					bInputsPresumedCorrect = true;
+				}
+			}
+		}
+
+		if(!bInputsPresumedCorrect)
+		{
+			if(BottomInput->MaterialExpressionEditorY > TopInput->MaterialExpressionEditorY)
+			{
+				SwapInputs(BottomInput, TopInput);
+				bInputsPresumedCorrect = true;
+			}
+		}
+		
+		if(!bInputsPresumedCorrect)
+		{
+			UE_LOG(LogMaterial, Warning, TEXT("Material Blend Function %s inputs could not be used to upgrade for Substrate."), *GetFullName());
+			return;
+		}
+
+		if(BottomInput->InputName != TEXT("Background Layer"))
+		{
+			BottomInput->InputName = TEXT("Background Layer");
+		}
+
+		if(TopInput->InputName == TEXT("Foreground Layer"))
+		{
+			TopInput->InputName = TEXT("Background Layer");
+		}
+
 		if (BottomInput->PreviewValue == FVector4f(FLinearColor::Black) && TopInput->PreviewValue == FVector4f(FLinearColor::Black))
 		{
 			//Only occurs during the update process and if inputs are set to black default so the blend has distinct preview, further user defined colours will not be altered.
@@ -16232,22 +16498,55 @@ void UMaterialFunction::ConvertExpressionsBetweenLegacyAndSubstrate()
 					if (!BlendFunctionCall)
 					{
 						//Store the mapping of a blend node to it's replacement function call, which means we only create a call once per individual blend node.
-						BlendFunctionCall = Cast<UMaterialExpressionMaterialFunctionCall>(ReplacementNodeMapping.Add(BlendNode, NewObject<UMaterialExpressionMaterialFunctionCall>(this)));
+						BlendFunctionCall = Cast<UMaterialExpressionMaterialFunctionCall>(ReplacementNodeMapping.Add(BlendNode, NewObject<UMaterialExpressionMaterialFunctionCall>(this, NAME_None, RF_Transactional)));
 						BlendFunctionCall->Function = this;
 						BlendFunctionCall->SetMaterialFunction(DefaultBlendFunction);
 
-						if (BlendFunctionCall && BlendFunctionCall->FunctionInputs.Num() >= 3)
+						PlaceNodeInLocation(BlendNode, BlendFunctionCall);
+						InputExpression = BlendFunctionCall;
+						Expressions.Add(BlendFunctionCall);
+						RecursedExpressions.Add(BlendFunctionCall);
+
+						if (BlendFunctionCall && BlendFunctionCall->FunctionInputs.Num() >= 8)
 						{
 							//These should match the sort priority of the Default MLB
 							BlendFunctionCall->FunctionInputs[0].Input.Connect(BlendNode->A.OutputIndex, BlendNode->A.Expression); //Background
 							BlendFunctionCall->FunctionInputs[1].Input.Connect(BlendNode->B.OutputIndex, BlendNode->B.Expression); //Foreground
 							BlendFunctionCall->FunctionInputs[2].Input.Connect(BlendNode->Alpha.OutputIndex, BlendNode->Alpha.Expression);
-						}
 
-						PlaceNodeInLocation(BlendNode, BlendFunctionCall);
-						InputExpression = BlendFunctionCall;
-						Expressions.Add(BlendFunctionCall);
-						RecursedExpressions.Add(BlendFunctionCall);
+							TArray<bool> StaticBoolValues;
+							int32 StartIndex = StaticBoolValues.Add(true); //UseParameterBlend
+							StaticBoolValues.Add(BlendNode->VertexAttributeBlendType == EMaterialAttributeBlend::Type::UseA); //VertexUseA
+							StaticBoolValues.Add(BlendNode->VertexAttributeBlendType == EMaterialAttributeBlend::Type::UseB); //VertexUseB
+							StaticBoolValues.Add(BlendNode->PixelAttributeBlendType == EMaterialAttributeBlend::Type::UseA); //PixelUseA
+							int32 EndIndex = StaticBoolValues.Add(BlendNode->PixelAttributeBlendType == EMaterialAttributeBlend::Type::UseB); //PixelUseB
+
+							TArray<UMaterialExpressionStaticBool*> StaticBoolArray;
+							StartIndex += 3; 
+							EndIndex += 3;
+							for(int32 FunctionCallIndex = StartIndex; FunctionCallIndex <= EndIndex; FunctionCallIndex++)
+							{
+								int32 ArrayIndex = StaticBoolArray.Add(NewObject<UMaterialExpressionStaticBool>(this, NAME_None, RF_Transactional));
+								UMaterialExpressionStaticBool* ThisBool = StaticBoolArray[ArrayIndex];
+								if(ThisBool)
+								{
+									bool ThisValue = StaticBoolValues[ArrayIndex];
+									ThisBool->Value = ThisValue;
+									BlendFunctionCall->FunctionInputs[FunctionCallIndex].Input.Connect(0, ThisBool);
+									Expressions.Add(ThisBool);
+									RecursedExpressions.Add(ThisBool);
+
+									if(FunctionCallIndex == StartIndex)
+									{
+										PlaceBelowNode(BlendFunctionCall, ThisBool, 300);
+									}
+									else
+									{
+										PlaceBelowNode(StaticBoolArray[ArrayIndex-1], ThisBool, 80);
+									}
+								}
+							}
+						}
 					}
 
 					if (BlendFunctionCall)
@@ -16263,13 +16562,13 @@ void UMaterialFunction::ConvertExpressionsBetweenLegacyAndSubstrate()
 				UMaterialExpressionSetMaterialAttributes* SetAttributesNode = ReplacementNodeMapping.Contains(InputExpression) ? Cast<UMaterialExpressionSetMaterialAttributes>(*ReplacementNodeMapping.Find(InputExpression)) : nullptr;
 				if(!SetAttributesNode)
 				{
-					SetAttributesNode = Cast<UMaterialExpressionSetMaterialAttributes>(ReplacementNodeMapping.Add(InputExpression, NewObject<UMaterialExpressionSetMaterialAttributes>(this)));
+					SetAttributesNode = Cast<UMaterialExpressionSetMaterialAttributes>(ReplacementNodeMapping.Add(InputExpression, NewObject<UMaterialExpressionSetMaterialAttributes>(this, NAME_None, RF_Transactional)));
 					SetAttributesNode->Function = this;
 					SetAttributesNode->ConnectInputAttribute(MP_MaterialAttributes, InputExpression);
 					Expressions.Add(SetAttributesNode);
 					RecursedExpressions.Add(SetAttributesNode);
 
-					UMaterialExpressionSubstrateConvertMaterialAttributes* ConvertNode = NewObject<UMaterialExpressionSubstrateConvertMaterialAttributes>(this);
+					UMaterialExpressionSubstrateConvertMaterialAttributes* ConvertNode = NewObject<UMaterialExpressionSubstrateConvertMaterialAttributes>(this, NAME_None, RF_Transactional);
 					ConvertNode->Function = this;
 					ConvertNode->MaterialAttributes.Connect(0, InputExpression);
 					ConvertNode->ShadingModelOverride = MSM_DefaultLit;
@@ -16288,8 +16587,6 @@ void UMaterialFunction::ConvertExpressionsBetweenLegacyAndSubstrate()
 					bBlendConverted = true;
 				}
 			}
-
-			return;			
 		};
 
 		auto RecurseBlendFunction = [&](UMaterialExpression* Expression, auto&& RecurseBlendFunction) -> void
diff --git a/Engine/Source/Runtime/Engine/Public/Materials/MaterialExpressionBlendMaterialAttributes.h b/Engine/Source/Runtime/Engine/Public/Materials/MaterialExpressionBlendMaterialAttributes.h
--- a/Engine/Source/Runtime/Engine/Public/Materials/MaterialExpressionBlendMaterialAttributes.h
+++ b/Engine/Source/Runtime/Engine/Public/Materials/MaterialExpressionBlendMaterialAttributes.h
@@ -20,6 +20,16 @@ namespace EMaterialAttributeBlend
 	};
 }
 
+UENUM()
+namespace EMaterialAttributeBlendFunction
+{
+	enum Type : int
+	{
+		Horizontal,
+		Vertical
+	};
+}
+
 UCLASS(collapsecategories, hidecategories = Object, MinimalAPI)
 class UMaterialExpressionBlendMaterialAttributes : public UMaterialExpression
 {
@@ -58,3 +68,37 @@ class UMaterialExpressionBlendMaterialAttributes : public UMaterialExpression
 #endif
 	//~ End UMaterialExpression Interface
 };
+
+
+UCLASS(collapsecategories, hidecategories = Object, MinimalAPI)
+class UMaterialExpressionLegacyBlendMaterialAttributes : public UMaterialExpressionBlendMaterialAttributes
+{
+	GENERATED_UCLASS_BODY()
+
+	UPROPERTY()
+	FExpressionInput VertexAttribute_UseA;
+
+	UPROPERTY()
+	FExpressionInput VertexAttribute_UseB;
+
+	UPROPERTY()
+	FExpressionInput PixelAttribute_UseA;
+
+	UPROPERTY()
+	FExpressionInput PixelAttribute_UseB;
+
+	UPROPERTY(EditAnywhere, Category = MaterialAttributes)
+	TEnumAsByte<EMaterialAttributeBlendFunction::Type> BlendFunctionType;
+
+#if ENABLE_MATERIAL_LAYER_PROTOTYPE
+#if WITH_EDITOR
+	virtual void GetCaption(TArray<FString>& OutCaptions) const override;
+	virtual bool CanEditChange(const FProperty* InProperty) const override;
+	virtual FExpressionInput* GetInput(int32 InputIndex) override;
+	virtual FName GetInputName(int32 InputIndex) const override;
+	virtual bool IsInputConnectionRequired(int32 InputIndex) const override;
+	virtual uint32 GetInputType(int32 InputIndex) override;
+	virtual int32 Compile(class FMaterialCompiler* Compiler, int32 OutputIndex) override;
+#endif // WITH_EDITOR
+#endif //ENABLE_MATERIAL_LAYER_PROTOTYPE
+};
\ No newline at end of file
