[Deformer Graph] fixed unconnected color output pin lead to incorrect vertex color read in vertex shader. When fallback to static color buffer, there is a chance that it is null. In that case we further fallback to a null vert color buffer of just a few bytes so we also need to update the color index mask to use here to avoid out of bound access.

diff --git a/Engine/Source/Runtime/Engine/Private/GPUSkinVertexFactory.cpp b/Engine/Source/Runtime/Engine/Private/GPUSkinVertexFactory.cpp
--- a/Engine/Source/Runtime/Engine/Private/GPUSkinVertexFactory.cpp
+++ b/Engine/Source/Runtime/Engine/Private/GPUSkinVertexFactory.cpp
@@ -1374,6 +1374,7 @@ void FGPUSkinPassthroughVertexFactory::UpdateUniformBuffer(FRHICommandListBase&
 	{
 		Data.TangentsSRV = SRVs[EShaderResource::Tangent] ? SRVs[EShaderResource::Tangent] : (FRHIShaderResourceView*)InSourceVertexFactory->GetTangentsSRV();
 		Data.ColorComponentsSRV = SRVs[EShaderResource::Color] ? SRVs[EShaderResource::Color] : (FRHIShaderResourceView*)InSourceVertexFactory->GetColorComponentsSRV();
+		Data.ColorIndexMask = SRVs[EShaderResource::Color] ? Data.ColorIndexMask : InSourceVertexFactory->GetColorIndexMask();
 		Data.TextureCoordinatesSRV = SRVs[EShaderResource::TexCoord] ? SRVs[EShaderResource::TexCoord] : (FRHIShaderResourceView*)InSourceVertexFactory->GetTextureCoordinatesSRV();
 
 		const int32 DefaultBaseVertexIndex = 0;
