Strip mesh paint textures from cook if the cook platform doesn't support them.
This will allow us to reduce package size when we use mesh paint textures on some platforms (eg nanite platforms) only.
The MeshPaintTexture reference has been changed to EditorDataOnly and a new MeshPaintTextureCooked reference has been added.
We use MeshPaintTexture for editor data builds and MeshPaintTextureCooked otherwise.
MeshPaintTextureCooked is manually serialized, and the reference is only written for cooked data and if the cook platform supports MeshPaintVirtualTextures.

diff --git a/Engine/Source/Runtime/Core/Public/UObject/FortniteMainBranchObjectVersion.h b/Engine/Source/Runtime/Core/Public/UObject/FortniteMainBranchObjectVersion.h
--- a/Engine/Source/Runtime/Core/Public/UObject/FortniteMainBranchObjectVersion.h
+++ b/Engine/Source/Runtime/Core/Public/UObject/FortniteMainBranchObjectVersion.h
@@ -510,7 +510,10 @@ struct FFortniteMainBranchObjectVersion
 		// Deprecation of the bPlaying flag on FTimeline struct types in favor of a better
 		// PlayingStateTracker type to improve replication reliability 
 		TimelinePlayingStateTrackerDeprecation,
-		
+
+		// Enable SkipOnlyEditorOnly style cooking of UStaticMeshComponent::MeshPaintTexture
+		MeshPaintTextureUsesEditorOnly,
+
 		// -----<new versions can be added above this line>-------------------------------------------------
 		VersionPlusOne,
 		LatestVersion = VersionPlusOne - 1
diff --git a/Engine/Source/Runtime/Engine/Classes/Components/StaticMeshComponent.h b/Engine/Source/Runtime/Engine/Classes/Components/StaticMeshComponent.h
--- a/Engine/Source/Runtime/Engine/Classes/Components/StaticMeshComponent.h
+++ b/Engine/Source/Runtime/Engine/Classes/Components/StaticMeshComponent.h
@@ -313,9 +313,15 @@ public:
 	UPROPERTY(EditAnywhere, Category = "Mesh Painting", meta=(InlineEditConditionToggle))
 	uint8 bOverrideMeshPaintTextureCoordinateIndex : 1;
 
+#if WITH_EDITORONLY_DATA
 	/** Texture containing mesh painting for this mesh component. */
 	UPROPERTY(VisibleAnywhere, Category = "Mesh Painting")
 	TObjectPtr<UTexture> MeshPaintTexture;
+#endif
+
+	/** Cooked pointer to texture containing mesh painting for this mesh component. This will be taken from MeshPaintTexture but can be empty on some platforms if we choose to strip the data. */
+	UPROPERTY(SkipSerialization)
+	TObjectPtr<UTexture> MeshPaintTextureCooked;
 
 	/** Set this to override the locally stored mesh paint texture. */
 	UPROPERTY(Transient)
diff --git a/Engine/Source/Runtime/Engine/Classes/VT/MeshPaintVirtualTexture.h b/Engine/Source/Runtime/Engine/Classes/VT/MeshPaintVirtualTexture.h
--- a/Engine/Source/Runtime/Engine/Classes/VT/MeshPaintVirtualTexture.h
+++ b/Engine/Source/Runtime/Engine/Classes/VT/MeshPaintVirtualTexture.h
@@ -38,6 +38,8 @@ namespace MeshPaintVirtualTexture
 {
 	/** Returns per platform support of mesh paint virtual textures. */
 	ENGINE_API bool IsSupported(EShaderPlatform InShaderPlatform);
+	/** Returns per target platform support of mesh paint virtual textures. (This combines the ShaderPlatform results for a TargetPlatform). */
+	ENGINE_API bool IsSupported(ITargetPlatform const* InTargetPlatform);
 
 	/** Returns the fallback color to use for unmapped virtual textures. Use white (same as default vertex color). */
 	inline uint32 GetDefaultFallbackColor() { return 0xFFFFFFFF; }
diff --git a/Engine/Source/Runtime/Engine/Private/Components/StaticMeshComponent.cpp b/Engine/Source/Runtime/Engine/Private/Components/StaticMeshComponent.cpp
--- a/Engine/Source/Runtime/Engine/Private/Components/StaticMeshComponent.cpp
+++ b/Engine/Source/Runtime/Engine/Private/Components/StaticMeshComponent.cpp
@@ -49,6 +49,7 @@
 #include "Rendering/NaniteResources.h"
 #include "NaniteVertexFactory.h"
 #include "StaticMeshSceneProxyDesc.h"
+#include "VT/MeshPaintVirtualTexture.h"
 #include "WorldPartition/ActorInstanceGuids.h"
 
 #include UE_INLINE_GENERATED_CPP_BY_NAME(StaticMeshComponent)
@@ -419,6 +420,46 @@ void UStaticMeshComponent::Serialize(FArchive& Ar)
 	{
 		bInitialEvaluateWorldPositionOffset = bEvaluateWorldPositionOffset;
 	}
+
+	// Handle stripping the MeshPaintTexture resource from the cook for platforms where MeshPaintVirtualTexture is not supported.
+	const bool bMeshPaintTextureUsesEditorOnly = Ar.CustomVer(FFortniteMainBranchObjectVersion::GUID) >= FFortniteMainBranchObjectVersion::MeshPaintTextureUsesEditorOnly;
+	if (bMeshPaintTextureUsesEditorOnly)
+	{
+		bool bSerializeAsCookedData = Ar.IsCooking();
+		Ar << bSerializeAsCookedData;
+
+		if (bSerializeAsCookedData)
+		{
+			// In cooked data we serialize a UObject* here in the native serializer.
+			if (Ar.IsSaving())
+			{
+#if WITH_EDITORONLY_DATA
+				// If we are choosing to strip from cook because of platform support then serialize nullptr.
+				const bool bCookMeshPaintTexture = MeshPaintVirtualTexture::IsSupported(Ar.CookingTarget());
+				UObject* SavedMeshPaintTexture = bCookMeshPaintTexture ? MeshPaintTexture : nullptr;
+#else
+				// Use MeshPaintTextureCooked.
+				UObject* SavedMeshPaintTexture = MeshPaintTextureCooked;
+#endif
+				Ar << SavedMeshPaintTexture;
+			}
+			else
+			{
+#if WITH_EDITORONLY_DATA
+				// If we are loading a cooked package in an environment that has EditorOnlyData, load into MeshPaintTexture.
+				Ar << MeshPaintTexture;
+#else
+				// Use MeshPaintTextureCooked.
+				Ar << MeshPaintTextureCooked;
+#endif
+			}
+		}
+		else
+		{
+			// In non-cooked data we do not serialize a UObject* here.
+			// Note that the EditorOnlyData MeshPaintTexture is serialized through default property serialization.
+		}
+	}
 }
 
 void UStaticMeshComponent::PostApplyToComponent()
@@ -1136,9 +1177,9 @@ void UStaticMeshComponent::GetStreamingRenderAssetInfo(FStreamingTextureLevelCon
 
 void UStaticMeshComponent::GetUsedTextures(TArray<UTexture*>& OutTextures, EMaterialQualityLevel::Type QualityLevel)
 {
-	if (MeshPaintTexture)
+	if (UTexture* Texture = GetMeshPaintTexture())
 	{
-		OutTextures.AddUnique(MeshPaintTexture);
+		OutTextures.AddUnique(Texture);
 	}
 
 	Super::GetUsedTextures(OutTextures, QualityLevel);
@@ -3103,14 +3144,22 @@ bool UStaticMeshComponent::PrestreamMeshLODs(float Seconds)
 
 UTexture* UStaticMeshComponent::GetMeshPaintTexture() const 
 {
-	return MeshPaintTexture; 
+#if WITH_EDITORONLY_DATA
+	return MeshPaintTexture;
+#else
+	return MeshPaintTextureCooked;
+#endif
 }
 
 void UStaticMeshComponent::SetMeshPaintTexture(UTexture* InTexture)
 {
-	if (MeshPaintTexture != InTexture)
+	if (GetMeshPaintTexture() != InTexture)
 	{
+#if WITH_EDITORONLY_DATA
 		MeshPaintTexture = InTexture;
+#else
+		MeshPaintTextureCooked = InTexture;
+#endif
 		MarkRenderStateDirty();
 	}
 }
diff --git a/Engine/Source/Runtime/Engine/Private/StaticMeshRender.cpp b/Engine/Source/Runtime/Engine/Private/StaticMeshRender.cpp
--- a/Engine/Source/Runtime/Engine/Private/StaticMeshRender.cpp
+++ b/Engine/Source/Runtime/Engine/Private/StaticMeshRender.cpp
@@ -2702,7 +2702,7 @@ void FStaticMeshSceneProxyDesc::InitializeFrom(const UStaticMeshComponent* InCom
 	SetMaterialRelevance(InComponent->GetMaterialRelevance(World->GetFeatureLevel()));
 	SetCollisionResponseToChannels(InComponent->GetCollisionResponseToChannels());
 
-	MeshPaintTexture = InComponent->MeshPaintTextureOverride ? InComponent->MeshPaintTextureOverride : InComponent->MeshPaintTexture;
+	MeshPaintTexture = InComponent->MeshPaintTextureOverride ? InComponent->MeshPaintTextureOverride.Get() : InComponent->GetMeshPaintTexture();
 }
 
 FPrimitiveSceneProxy* UStaticMeshComponent::CreateStaticMeshSceneProxy(Nanite::FMaterialAudit& NaniteMaterials, bool bCreateNanite)
diff --git a/Engine/Source/Runtime/Engine/Private/VT/MeshPaintVirtualTexture.cpp b/Engine/Source/Runtime/Engine/Private/VT/MeshPaintVirtualTexture.cpp
--- a/Engine/Source/Runtime/Engine/Private/VT/MeshPaintVirtualTexture.cpp
+++ b/Engine/Source/Runtime/Engine/Private/VT/MeshPaintVirtualTexture.cpp
@@ -5,6 +5,7 @@
 #include "ComponentRecreateRenderStateContext.h"
 #include "Components/PrimitiveComponent.h"
 #include "EngineModule.h"
+#include "Interfaces/ITargetPlatform.h"
 #include "RendererInterface.h"
 #include "RenderUtils.h"
 #include "ShaderPlatformCachedIniValue.h"
@@ -54,6 +55,24 @@ namespace MeshPaintVirtualTexture
 		return CPlatformVarMeshPaintVirtualTextureSupport.Get(InShaderPlatform) && UseVirtualTexturing(InShaderPlatform);
 	}
 
+	bool IsSupported(ITargetPlatform const* InTargetPlatform)
+	{
+		if (InTargetPlatform != nullptr)
+		{
+			TArray<FName> DesiredShaderFormats;
+			InTargetPlatform->GetAllTargetedShaderFormats(DesiredShaderFormats);
+			for (int32 FormatIndex = 0; FormatIndex < DesiredShaderFormats.Num(); FormatIndex++)
+			{
+				const EShaderPlatform ShaderPlatform = ShaderFormatToLegacyShaderPlatform(DesiredShaderFormats[FormatIndex]);
+				if (IsSupported(ShaderPlatform))
+				{
+					return true;
+				}
+			}
+		}
+		return false;
+	}
+
 	static bool IsEnabled()
 	{
 		return CVarMeshPaintVirtualTextureEnable.GetValueOnAnyThread();
@@ -260,7 +279,7 @@ void UMeshPaintVirtualTexture::BeginCacheForCookedPlatformData(const ITargetPlat
 	// Even though we skip the cook of this object for non VT platforms in URuntimeVirtualTexture::Serialize()
 	// we still load the object at cook time and kick off the DDC build. This will trigger an error in the texture DDC code.
 	// Either we need to make the DDC code more robust for non VT platforms or we can skip the process here...
-	if (!UseVirtualTexturing(GMaxRHIShaderPlatform, TargetPlatform) || !MeshPaintVirtualTexture::IsSupported(GMaxRHIShaderPlatform))
+	if (!UseVirtualTexturing(GMaxRHIShaderPlatform, TargetPlatform))
 	{
 		return;
 	}
@@ -270,7 +289,7 @@ void UMeshPaintVirtualTexture::BeginCacheForCookedPlatformData(const ITargetPlat
 
 bool UMeshPaintVirtualTexture::IsCachedCookedPlatformDataLoaded(const ITargetPlatform* TargetPlatform)
 {
-	if (!UseVirtualTexturing(GMaxRHIShaderPlatform, TargetPlatform) || !MeshPaintVirtualTexture::IsSupported(GMaxRHIShaderPlatform))
+	if (!UseVirtualTexturing(GMaxRHIShaderPlatform, TargetPlatform))
 	{
 		return true;
 	}
@@ -280,7 +299,7 @@ bool UMeshPaintVirtualTexture::IsCachedCookedPlatformDataLoaded(const ITargetPla
 
 void UMeshPaintVirtualTexture::ClearCachedCookedPlatformData(const ITargetPlatform* TargetPlatform)
 {
-	if (!UseVirtualTexturing(GMaxRHIShaderPlatform, TargetPlatform) || !MeshPaintVirtualTexture::IsSupported(GMaxRHIShaderPlatform))
+	if (!UseVirtualTexturing(GMaxRHIShaderPlatform, TargetPlatform))
 	{
 		return;
 	}
