World Partition - RuntimeLevelStreamingCell: Removed unnecessary code introduced while trying to workaround a rare crash
* We now understand the root cause for PLAY-45493

diff --git a/Engine/Source/Runtime/Engine/Private/WorldPartition/WorldPartitionRuntimeLevelStreamingCell.cpp b/Engine/Source/Runtime/Engine/Private/WorldPartition/WorldPartitionRuntimeLevelStreamingCell.cpp
--- a/Engine/Source/Runtime/Engine/Private/WorldPartition/WorldPartitionRuntimeLevelStreamingCell.cpp
+++ b/Engine/Source/Runtime/Engine/Private/WorldPartition/WorldPartitionRuntimeLevelStreamingCell.cpp
@@ -26,17 +26,6 @@ UWorldPartitionRuntimeLevelStreamingCell::UWorldPartitionRuntimeLevelStreamingCe
 	, LevelStreaming(nullptr)
 {}
 
-void UWorldPartitionRuntimeLevelStreamingCell::BeginDestroy()
-{
-	if (LevelStreaming)
-	{
-		LevelStreaming->OnLevelShown.RemoveAll(this);
-		LevelStreaming->OnLevelHidden.RemoveAll(this);
-	}
-
-	Super::BeginDestroy();
-}
-
 EWorldPartitionRuntimeCellState UWorldPartitionRuntimeLevelStreamingCell::GetCurrentState() const
 {
 	if (LevelStreaming)
@@ -117,13 +106,6 @@ TArray<FName> UWorldPartitionRuntimeLevelStreamingCell::GetActors() const
 
 void UWorldPartitionRuntimeLevelStreamingCell::CreateAndSetLevelStreaming(const FString& InPackageName, const FSoftObjectPath& InWorldAsset)
 {
-	// Temporary test while we investigate issue PLAY-45493
-	if (!ensureMsgf(!LevelStreaming, TEXT("StreamingCell already had an assigned LevelStreaming object '%s'"), *LevelStreaming->GetPathName()))
-	{
-		LevelStreaming->OnLevelShown.RemoveAll(this);
-		LevelStreaming->OnLevelHidden.RemoveAll(this);
-	}
-
 	LevelStreaming = CreateLevelStreaming(InPackageName, InWorldAsset);
 }
 
diff --git a/Engine/Source/Runtime/Engine/Public/WorldPartition/WorldPartitionRuntimeLevelStreamingCell.h b/Engine/Source/Runtime/Engine/Public/WorldPartition/WorldPartitionRuntimeLevelStreamingCell.h
--- a/Engine/Source/Runtime/Engine/Public/WorldPartition/WorldPartitionRuntimeLevelStreamingCell.h
+++ b/Engine/Source/Runtime/Engine/Public/WorldPartition/WorldPartitionRuntimeLevelStreamingCell.h
@@ -17,10 +17,6 @@ class UWorldPartitionRuntimeLevelStreamingCell : public UWorldPartitionRuntimeCe
 {
 	GENERATED_UCLASS_BODY()
 
-	//~Begin UObject interface
-	virtual void BeginDestroy() override;
-	//~End UObject interface
-
 	//~Begin UWorldPartitionRuntimeCell interface
 	ENGINE_API virtual void Load() const override;
 	ENGINE_API virtual void Unload() const override;
