Add tick.AllowBatchedTicks cvar (disabled by default) for automatically combining TickFunctions in the same tick group. If enabled, it improves performance by executing multiple compatible tick functions as a single async task but the order within a group may change.
Add bAllowTickBatching to TickFunction to control this behavior per tick and enable by default for actor and component ticks.
Add tick.AllowOptimizedPrerequisites (enabled by default) for improving performance by ignoring any tick dependencies that complete in an earlier group.
Add ForEachNestedTick and ExecuteNestedTick to TickFunction which can be used by tick managers to correctly handle logging and execution of nested ticks.
Clean up handling of tick frame counters to make thread safety and precision explicit.

diff --git a/Engine/Source/Runtime/Engine/Classes/Engine/EngineBaseTypes.h b/Engine/Source/Runtime/Engine/Classes/Engine/EngineBaseTypes.h
--- a/Engine/Source/Runtime/Engine/Classes/Engine/EngineBaseTypes.h
+++ b/Engine/Source/Runtime/Engine/Classes/Engine/EngineBaseTypes.h
@@ -210,6 +210,9 @@ public:
 	UPROPERTY(EditDefaultsOnly, Category="Tick", AdvancedDisplay)
 	uint8 bAllowTickOnDedicatedServer:1;
 
+	/** True if we allow this tick to be combined with other ticks for improved performance */
+	uint8 bAllowTickBatching:1;
+
 	/** Run this tick first within the tick group, presumably to start async tasks that must be completed with this tick group, hiding the latency. */
 	uint8 bHighPriority:1;
 
@@ -242,6 +245,19 @@ private:
 	/** Prerequisites for this tick function **/
 	TArray<struct FTickPrerequisite> Prerequisites;
 
+	/** Defines the internal state of a tick function, set in TickTaskManager */
+	enum class ETickTaskState : uint8
+	{
+		// Has not been queued yet for the current frame, or has already executed
+		NotQueued,
+		// In the process of being queued, but no actual task yet
+		Pending,
+		// Has an actual task and a completion event, may be executing
+		HasTask,
+		// Directly points to a completion event, probably in a batched tick
+		HasCompletionEvent,
+	};
+
 	/** Internal Data structure that contains members only required for a registered tick function **/
 	struct FInternalData
 	{
@@ -253,20 +269,23 @@ private:
 		/** Cache whether this function was rescheduled as an interval function during StartParallel */
 		bool bWasInterval:1;
 
-		/** Internal data that indicates the tick group we actually started in (it may have been delayed due to prerequisites) **/
+		/** Internal state, this determines what TaskPointer points to */
+		ETickTaskState TaskState;
+
+		/** Internal data that indicates the tick group we actually started in (it may have been delayed due to prerequisites) */
 		TEnumAsByte<enum ETickingGroup> ActualStartTickGroup;
 
-		/** Internal data that indicates the tick group we actually started in (it may have been delayed due to prerequisites) **/
+		/** Internal data that indicates the tick group we are guaranteed to end by (it may have been delayed due to prerequisites) */
 		TEnumAsByte<enum ETickingGroup> ActualEndTickGroup;
 		
-		/** Internal data to track if we have started visiting this tick function yet this frame **/
-		int32 TickVisitedGFrameCounter;
-
-		/** Internal data to track if we have finished visiting this tick function yet this frame **/
-		std::atomic<int32> TickQueuedGFrameCounter;
+		/** Internal data to track if we have started visiting this tick function yet this frame, smaller than full frame counter */
+		std::atomic<uint32> TickVisitedGFrameCounter;
 
-		/** Pointer to the task, only used during setup. This is often stale. **/
-		FBaseGraphTask* TaskPointer;
+		/** Internal data to track if we have finished visiting this tick function yet this frame, smaller than full frame counter */
+		std::atomic<uint32> TickQueuedGFrameCounter;
+	
+		/** Pointer to a type determined by TaskState, do not access directly */
+		void* TaskPointer;
 
 		/** The next function in the cooling down list for ticks with an interval*/
 		FTickFunction* Next;
@@ -310,49 +329,52 @@ public:
 	ENGINE_API void SetTickFunctionEnable(bool bInEnabled);
 	/** Returns whether the tick function is currently enabled */
 	bool IsTickFunctionEnabled() const { return TickState != ETickState::Disabled; }
-	/** Returns whether it is valid to access this tick function's completion handle */
-	bool IsCompletionHandleValid() const { return (InternalData && InternalData->TaskPointer); }
+
 	/** Update tick interval in the system and overwrite the current cooldown if any. */
 	ENGINE_API void UpdateTickIntervalAndCoolDown(float NewTickInterval);
 
+	/** Returns true if it is valid to access this tick function's completion handle */
+	ENGINE_API bool IsCompletionHandleValid() const;
+
 	/**
-	* Gets the current completion handle of this tick function, so it can be delayed until a later point when some additional
-	* tasks have been completed.  Only valid after TG_PreAsyncWork has started and then only until the TickFunction finishes
-	* execution
-	**/
+	 * Gets the current completion handle of this tick function, so it can be delayed until a later point when some additional
+	 * tasks have been completed. Only valid after StartFrame has been called and then only until the TickFunction finishes
+	 * execution. This returns a reference so IsCompletionHandleValid must be called first.
+	 */
 	ENGINE_API FGraphEventRef GetCompletionHandle() const;
 
 	/** 
-	* Gets the action tick group that this function will be elligible to start in.
-	* Only valid after TG_PreAsyncWork has started through the end of the frame.
-	**/
+	 * Gets the action tick group that this function will be elligible to start in.
+	 * Only valid after StartFrame has been called through the end of the frame.
+	 */
 	TEnumAsByte<enum ETickingGroup> GetActualTickGroup() const
 	{
 		return (InternalData ? InternalData->ActualStartTickGroup : TickGroup);
 	}
 
 	/** 
-	* Gets the action tick group that this function will be required to end in.
-	* Only valid after TG_PreAsyncWork has started through the end of the frame.
-	**/
+	 * Gets the action tick group that this function will be required to end in.
+ 	 * Only valid after StartFrame has been called through the end of the frame.
+	 */
 	TEnumAsByte<enum ETickingGroup> GetActualEndTickGroup() const
 	{
 		return (InternalData ? InternalData->ActualEndTickGroup : EndTickGroup);
 	}
 
-
 	/** 
 	 * Adds a tick function to the list of prerequisites...in other words, adds the requirement that TargetTickFunction is called before this tick function is 
 	 * @param TargetObject - UObject containing this tick function. Only used to verify that the other pointer is still usable
 	 * @param TargetTickFunction - Actual tick function to use as a prerequisite
 	 **/
 	ENGINE_API void AddPrerequisite(UObject* TargetObject, struct FTickFunction& TargetTickFunction);
+
 	/** 
 	 * Removes a prerequisite that was previously added.
 	 * @param TargetObject - UObject containing this tick function. Only used to verify that the other pointer is still usable
 	 * @param TargetTickFunction - Actual tick function to use as a prerequisite
 	 **/
 	ENGINE_API void RemovePrerequisite(UObject* TargetObject, struct FTickFunction& TargetTickFunction);
+
 	/** 
 	 * Sets this function to hipri and all prerequisites recursively
 	 * @param bInHighPriority - priority to set
@@ -375,50 +397,87 @@ public:
 		return Prerequisites;
 	}
 
+	/**
+	 * For actively registered tick functions with an interval, return the last time it ticked.
+	 * For unregistered or non-interval functions it will return -1;
+	 */
 	float GetLastTickGameTime() const { return (InternalData ? InternalData->LastTickGameTimeSeconds : -1.f); }
+	 
+	/** 
+	 * Correctly call ExecuteTick on a nested function, and will set the completion event properly for the duration of the tick function
+	 * @param DeltaTime - frame time to advance, in seconds
+	 * @param TickType - kind of tick for this frame
+	 * @param CurrentThread - thread we are executing on, useful to pass along as new tasks are created
+	 * @param MyCompletionGraphEvent - completion event for this task. Useful for holding the completetion of this task until certain child tasks are complete.
+	 */
+	ENGINE_API void ExecuteNestedTick(float DeltaTime, ELevelTick TickType, ENamedThreads::Type CurrentThread, const FGraphEventRef& MyCompletionGraphEvent);
 
 private:
+	/** Returns the current internal task only if it matches required state */
+	void* GetTaskPointer(ETickTaskState RequiredState) const
+	{
+		return (InternalData && InternalData->TaskState == RequiredState ? InternalData->TaskPointer : nullptr);
+	}
+
+	/** Directly set the task pointer and state */
+	void SetTaskPointer(ETickTaskState NewState, void* InTaskPointer);
+
 	/**
 	 * Queues a tick function for execution from the game thread
 	 * @param TickContext - context to tick in
 	 */
-	ENGINE_API void QueueTickFunction(class FTickTaskSequencer& TTS, const FTickContext& TickContext);
+	void QueueTickFunction(class FTickTaskSequencer& TTS, const FTickContext& TickContext);
 
 	/**
 	 * Queues a tick function for execution from the game thread
 	 * @param TickContext - context to tick in
 	 * @param StackForCycleDetection - Stack For Cycle Detection
 	 */
-	ENGINE_API void QueueTickFunctionParallel(const FTickContext& TickContext, TArray<FTickFunction*, TInlineAllocator<8> >& StackForCycleDetection);
+	void QueueTickFunctionParallel(const FTickContext& TickContext, TArray<FTickFunction*, TInlineAllocator<8> >& StackForCycleDetection);
 
-	/** Returns the delta time to use when ticking this function given the TickContext */
-	ENGINE_API float CalculateDeltaTime(const FTickContext& TickContext);
+public:
+	// Functions to be used by executing tasks, don't call these directly from user code
 
-	/** 
-	 * Logs the prerequisites
-	 */
+	/** Returns the delta time to use when ticking this function given the delta time and world. This also updates internal tracking data */
+	ENGINE_API float CalculateDeltaTime(float DeltaTime, const class UWorld* TickingWorld);
+
+	/** Logs function info based on flags */
+	ENGINE_API void LogTickFunction(ENamedThreads::Type CurrentThread, bool bLogPrerequisites, int32 Indent = 0);
+
+	/** Logs the prerequisites */
 	ENGINE_API void ShowPrerequistes(int32 Indent = 1);
 
+	/** Clear any current task information such as the completion handle, called after execution */
+	ENGINE_API void ClearTaskInformation();
+
+
 	/** 
-	 * Abstract function actually execute the tick. 
+	 * Abstract function actually execute the tick. Batched tick managers should use ExecuteNestedTick
 	 * @param DeltaTime - frame time to advance, in seconds
 	 * @param TickType - kind of tick for this frame
 	 * @param CurrentThread - thread we are executing on, useful to pass along as new tasks are created
 	 * @param MyCompletionGraphEvent - completion event for this task. Useful for holding the completetion of this task until certain child tasks are complete.
 	 **/
 	ENGINE_API virtual void ExecuteTick(float DeltaTime, ELevelTick TickType, ENamedThreads::Type CurrentThread, const FGraphEventRef& MyCompletionGraphEvent) PURE_VIRTUAL(,);
+	
 	/** Abstract function to describe this tick. Used to print messages about illegal cycles in the dependency graph **/
 	ENGINE_API virtual FString DiagnosticMessage() PURE_VIRTUAL(, return TEXT("DiagnosticMessage() not implemented"););
+	
 	/** Function to give a 'context' for this tick, used for grouped active tick reporting */
 	virtual FName DiagnosticContext(bool bDetailed)
 	{
 		return NAME_None;
 	}
+
+	/**
+	 * To properly expose nested/batched tick functions to the logging and debugging systems of the tick manager, 
+	 * override this to execute InFunc for every nested tick function that would be executed as part of this tick.
+	 */
+	virtual void ForEachNestedTick(TFunctionRef<void(FTickFunction&)> InFunc) const {}
 	
 	friend class FTickTaskSequencer;
 	friend class FTickTaskManager;
 	friend class FTickTaskLevel;
-	friend class FTickFunctionTask;
 
 	// It is unsafe to copy FTickFunctions and any subclasses of FTickFunction should specify the type trait WithCopy = false
 	FTickFunction& operator=(const FTickFunction&) = delete;
diff --git a/Engine/Source/Runtime/Engine/Private/Actor.cpp b/Engine/Source/Runtime/Engine/Private/Actor.cpp
--- a/Engine/Source/Runtime/Engine/Private/Actor.cpp
+++ b/Engine/Source/Runtime/Engine/Private/Actor.cpp
@@ -185,6 +185,7 @@ void AActor::InitializeDefaults()
 	// Default to no tick function, but if we set 'never ticks' to false (so there is a tick function) it is enabled by default
 	PrimaryActorTick.bCanEverTick = false;
 	PrimaryActorTick.bStartWithTickEnabled = true;
+	PrimaryActorTick.bAllowTickBatching = true;
 	PrimaryActorTick.SetTickFunctionEnable(false); 
 	bAsyncPhysicsTickEnabled = false;
 
diff --git a/Engine/Source/Runtime/Engine/Private/Components/ActorComponent.cpp b/Engine/Source/Runtime/Engine/Private/Components/ActorComponent.cpp
--- a/Engine/Source/Runtime/Engine/Private/Components/ActorComponent.cpp
+++ b/Engine/Source/Runtime/Engine/Private/Components/ActorComponent.cpp
@@ -440,6 +440,7 @@ UActorComponent::UActorComponent(const FObjectInitializer& ObjectInitializer /*=
 	PrimaryComponentTick.TickGroup = TG_DuringPhysics;
 	PrimaryComponentTick.bStartWithTickEnabled = true;
 	PrimaryComponentTick.bCanEverTick = false;
+	PrimaryComponentTick.bAllowTickBatching = true;
 	PrimaryComponentTick.SetTickFunctionEnable(false);
 
 	MarkedForEndOfFrameUpdateArrayIndex = INDEX_NONE;
diff --git a/Engine/Source/Runtime/Engine/Private/TickTaskManager.cpp b/Engine/Source/Runtime/Engine/Private/TickTaskManager.cpp
--- a/Engine/Source/Runtime/Engine/Private/TickTaskManager.cpp
+++ b/Engine/Source/Runtime/Engine/Private/TickTaskManager.cpp
@@ -11,6 +11,7 @@
 #include "TickTaskManagerInterface.h"
 #include "Async/ParallelFor.h"
 #include "Misc/TimeGuard.h"
+#include "Templates/UniquePtr.h"
 
 DEFINE_LOG_CATEGORY_STATIC(LogTick, Log, All);
 
@@ -47,7 +48,17 @@ static TAutoConsoleVariable<int32> CVarLogTicksShowPrerequistes(
 static TAutoConsoleVariable<int32> CVarAllowAsyncComponentTicks(
 	TEXT("tick.AllowAsyncComponentTicks"),
 	1,
-	TEXT("Used to control async component ticks."));
+	TEXT("If true, components (and other tick functions) with bRunOnAnyThread set will run in parallel with other ticks."));
+
+static TAutoConsoleVariable<int32> CVarAllowBatchedTicks(
+	TEXT("tick.AllowBatchedTicks"),
+	0,
+	TEXT("If true, tick functions with bAllowTickBatching will be automatically combined into a single tick task."));
+
+static TAutoConsoleVariable<int32> CVarAllowOptimizedPrerequisites(
+	TEXT("tick.AllowOptimizedPrerequisites"),
+	1,
+	TEXT("If true, the code will ignore prerequisites that are not relevant due to guaranteed order of tick groups"));
 
 // This was disabled by default in UE 5.5
 static TAutoConsoleVariable<int32> CVarAllowConcurrentQueue(
@@ -165,83 +176,106 @@ public:
 
 };
 
-
+/** Description of how the tick function should execute */
 struct FTickContext
 {
-	/** Delta time to tick **/
+	/** The world in which the object being ticked is contained */
+	UWorld*					World;
+	/** Delta time to tick */
 	float					DeltaSeconds;
-	/** Tick type **/
-	ELevelTick				TickType;
-	/** Tick type **/
-	ETickingGroup			TickGroup;
-	/** Current or desired thread **/
+	/** Current or desired thread */
 	ENamedThreads::Type		Thread;
-	/** The world in which the object being ticked is contained. **/
-	UWorld*					World;
-
+	/** Tick type such as gameplay */
+	TEnumAsByte<ELevelTick>	TickType;
+	/** Tick group this was started in */
+	TEnumAsByte<ETickingGroup> TickGroup;
+	/** If true, log each tick */
+	bool					bLogTick;
+	/** If true, log prereqs */
+	bool					bLogTicksShowPrerequistes;
+	
+	
 	FTickContext(float InDeltaSeconds = 0.0f, ELevelTick InTickType = LEVELTICK_All, ETickingGroup InTickGroup = TG_PrePhysics, ENamedThreads::Type InThread = ENamedThreads::GameThread)
-		: DeltaSeconds(InDeltaSeconds)
+		: World(nullptr)
+		, DeltaSeconds(InDeltaSeconds)
+		, Thread(InThread)
 		, TickType(InTickType)
 		, TickGroup(InTickGroup)
-		, Thread(InThread)
-		, World(nullptr)
+		, bLogTick(false)
+		, bLogTicksShowPrerequistes(false)
 	{
 	}
 
-	FTickContext(const FTickContext& In)
-		: DeltaSeconds(In.DeltaSeconds)
-		, TickType(In.TickType)
-		, TickGroup(In.TickGroup)
-		, Thread(In.Thread)
-		, World(In.World)
+	FTickContext(const FTickContext& In) = default;
+	FTickContext& operator=(const FTickContext& In) = default;
+};
+
+
+/** This is an integer that represents the conditions for which ticks can be grouped together */
+struct FTickGroupCondition
+{
+	union
+	{
+		uint32 IntVersion;
+		struct {
+			TEnumAsByte<ETickingGroup> StartGroup;
+			TEnumAsByte<ETickingGroup> EndGroup;
+			bool bHighPriority;
+			bool bIsBatch;
+		};
+	};
+	
+	FTickGroupCondition()
+		: IntVersion(0)
 	{
 	}
-	void operator=(const FTickContext& In)
+
+	FTickGroupCondition(const FTickFunction* TickFunction)
+		: StartGroup(TickFunction->GetActualTickGroup())
+		, EndGroup(TickFunction->GetActualEndTickGroup())
+		, bHighPriority(TickFunction->bHighPriority)
+		, bIsBatch(TickFunction->bAllowTickBatching)
 	{
-		DeltaSeconds = In.DeltaSeconds;
-		TickType = In.TickType;
-		TickGroup = In.TickGroup;
-		Thread = In.Thread;
-		World = In.World;
 	}
-};
-
 
+	FORCEINLINE bool operator==(const FTickGroupCondition& Other) const
+	{
+		return IntVersion == Other.IntVersion;
+	}
 
+	FORCEINLINE FTickGroupCondition& operator=(const FTickGroupCondition& Other)
+	{
+		IntVersion = Other.IntVersion;
+		return *this;
+	}
 
+	FORCEINLINE friend uint32 GetTypeHash(const FTickGroupCondition& Condition)
+	{
+		return Condition.IntVersion; 
+	}
+};
 
+static_assert(sizeof(FTickGroupCondition) == 4, "Tick group condition must be an integer");
 
-/**
- * Class that handles the actual tick tasks and starting and completing tick groups
- */
-/** Helper class define the task of ticking a component **/
+/** Task for a single tick function */
 class FTickFunctionTask
 {
-	/** Actor to tick **/
-	FTickFunction*			Target;
-	/** tick context, here thread is desired execution thread **/
-	FTickContext			Context;
-	/** If true, log each tick **/
-	bool					bLogTick;
-	/** If true, log prereqs **/
-	bool					bLogTicksShowPrerequistes;
+	/** Functions to tick */
+	FTickFunction* Target;
+	/** Tick context with the desired execution thread */
+	FTickContext Context;
+
 public:
-	/** Constructor
-		* @param InTarget - Function to tick
-		* @param InContext - context to tick in, here thread is desired execution thread
-	**/
-	FORCEINLINE FTickFunctionTask(FTickFunction* InTarget, const FTickContext* InContext, bool InbLogTick, bool bInLogTicksShowPrerequistes)
+	FORCEINLINE FTickFunctionTask(FTickFunction* InTarget, const FTickContext* InContext)
 		: Target(InTarget)
 		, Context(*InContext)
-		, bLogTick(InbLogTick)
-	, bLogTicksShowPrerequistes(bInLogTicksShowPrerequistes)
 	{
 	}
 	static FORCEINLINE TStatId GetStatId()
 	{
 		RETURN_QUICK_DECLARE_CYCLE_STAT(FTickFunctionTask, STATGROUP_TaskGraphTasks);
 	}
-	/** return the thread for this task **/
+	/** Return the desired execution thread for this task */
 	FORCEINLINE ENamedThreads::Type GetDesiredThread()
 	{
 		return Context.Thread;
@@ -251,20 +285,16 @@ public:
 		return ESubsequentsMode::TrackSubsequents;
 	}
 	/**
-		*	Actually execute the tick.
-		*	@param	CurrentThread; the thread we are running on
-		*	@param	MyCompletionGraphEvent; my completion event. Not always useful since at the end of DoWork, you can assume you are done and hence further tasks do not need you as a prerequisite.
-		*	However, MyCompletionGraphEvent can be useful for passing to other routines or when it is handy to set up subsequents before you actually do work.
-		**/
+	 * Actually execute the tick.
+	 * @param	CurrentThread; the thread we are running on
+	 * @param	MyCompletionGraphEvent; my completion event. Not always useful since at the end of DoWork, you can assume you are done and hence further tasks do not need you as a prerequisite.
+	 * However, MyCompletionGraphEvent can be useful for passing to other routines or when it is handy to set up subsequents before you actually do work.
+	 */
 	void DoTask(ENamedThreads::Type CurrentThread, const FGraphEventRef& MyCompletionGraphEvent)
 	{
-		if (bLogTick)
+		if (Context.bLogTick)
 		{
-			UE_LOG(LogTick, Log, TEXT("tick %s [%1d, %1d] %6llu %2d %s"), Target->bHighPriority ? TEXT("*") : TEXT(" "), (int32)Target->GetActualTickGroup(), (int32)Target->GetActualEndTickGroup(), (uint64)GFrameCounter, (int32)CurrentThread, *Target->DiagnosticMessage());
-			if (bLogTicksShowPrerequistes)
-			{
-				Target->ShowPrerequistes();
-			}
+			Target->LogTickFunction(CurrentThread, Context.bLogTicksShowPrerequistes);
 		}
 		if (Target->IsTickFunctionEnabled())
 		{
@@ -273,10 +303,87 @@ public:
 			SCOPE_TIME_GUARD_DELEGATE_MS(NameFunction, 4);
 #endif
 			LIGHTWEIGHT_TIME_GUARD_BEGIN(FTickFunctionTask, GTimeguardThresholdMS);
-			Target->ExecuteTick(Target->CalculateDeltaTime(Context), Context.TickType, CurrentThread, MyCompletionGraphEvent);
+			Target->ExecuteTick(Target->CalculateDeltaTime(Context.DeltaSeconds, Context.World), Context.TickType, CurrentThread, MyCompletionGraphEvent);
 			LIGHTWEIGHT_TIME_GUARD_END(FTickFunctionTask, Target->DiagnosticMessage());
 		}
-		Target->InternalData->TaskPointer = nullptr;  // This is stale and a good time to clear it for safety
+		Target->ClearTaskInformation();  // This is stale and a good time to clear it for safety
+	}
+};
+
+// Declare the internal async task that is used for scheduling tick tasks with the task graph backend
+// TODO remove this when old frontend is removed. Batching does not work with the old frontend
+#if TASKGRAPH_NEW_FRONTEND
+using FTickGraphTask = FBaseGraphTask;
+#else
+using FTickGraphTask = TGraphTask<FTickFunctionTask>;
+#endif
+
+/** Info used to execute a batch tick */
+struct FTickBatchInfo
+{
+	/** Actual task assigned to this batch */
+	FTickGraphTask* TickTask = nullptr;
+	/** Prerequisites that are important */
+	TArray<FTickFunction*> TickPrerequisites;
+	/** Functions to tick, defaults to a single one */
+	TArray<FTickFunction*> TickFunctions;
+
+	FORCEINLINE void Reset()
+	{
+		TickTask = nullptr;
+		// Maintain sizes because the order will probably be the same next frame
+		TickPrerequisites.Reset();
+		TickFunctions.Reset();
+	}
+
+};
+
+/** Task for executing multiple functions at once */
+class FBatchTickFunctionTask
+{
+	/** Batch to execute */
+	FTickBatchInfo* TickBatch;
+	/** Tick context with the desired execution thread */
+	FTickContext Context;
+
+public:
+	FORCEINLINE FBatchTickFunctionTask(FTickBatchInfo* InTickBatch, const FTickContext* InContext)
+		: TickBatch(InTickBatch)
+		, Context(*InContext)
+	{
+		check(TickBatch);
+	}
+	FORCEINLINE ENamedThreads::Type GetDesiredThread()
+	{
+		return Context.Thread;
+	}
+	/**
+	 * Actually execute the tick.
+	 * @param	CurrentThread; the thread we are running on
+	 * @param	MyCompletionGraphEvent; my completion event. Not always useful since at the end of DoWork, you can assume you are done and hence further tasks do not need you as a prerequisite.
+	 * However, MyCompletionGraphEvent can be useful for passing to other routines or when it is handy to set up subsequents before you actually do work.
+	 */
+	void DoTask(ENamedThreads::Type CurrentThread, const FGraphEventRef& MyCompletionGraphEvent)
+	{
+		check(TickBatch && TickBatch->TickFunctions.Num() > 0);
+		for (FTickFunction* Target : TickBatch->TickFunctions)
+		{
+			if (Context.bLogTick)
+			{
+				Target->LogTickFunction(CurrentThread, Context.bLogTicksShowPrerequistes);
+			}
+			if (Target->IsTickFunctionEnabled())
+			{
+#if DO_TIMEGUARD
+				FTimerNameDelegate NameFunction = FTimerNameDelegate::CreateLambda( [&]{ return FString::Printf(TEXT("Slowtick %s "), *Target->DiagnosticMessage()); } );
+				SCOPE_TIME_GUARD_DELEGATE_MS(NameFunction, 4);
+#endif
+				LIGHTWEIGHT_TIME_GUARD_BEGIN(FBatchTickFunctionTask, GTimeguardThresholdMS);
+				Target->ExecuteTick(Target->CalculateDeltaTime(Context.DeltaSeconds, Context.World), Context.TickType, CurrentThread, MyCompletionGraphEvent);
+				LIGHTWEIGHT_TIME_GUARD_END(FBatchTickFunctionTask, Target->DiagnosticMessage());
+			}
+			Target->ClearTaskInformation();  // This is stale and a good time to clear it for safety
+		}
 	}
 };
 
@@ -359,14 +466,18 @@ class FTickTaskSequencer
 		}
 	};
 
-	/** Completion handles for each phase of ticks */
+	/** List of batched ticks, these stay allocated between frames but are cleared */
+	TArray< TPair<FTickGroupCondition, TUniquePtr<FTickBatchInfo> > > TickBatches;
+	int32 TickBatchesNum = 0;
+
+	/** Completion handles for each phase of ticks, this is compatible with FGraphEventArray */
 	TArrayWithThreadsafeAdd<FGraphEventRef, TInlineAllocator<4> > TickCompletionEvents[TG_MAX];
 
 	/** HiPri Held tasks for each tick group. */
-	TArrayWithThreadsafeAdd<TGraphTask<FTickFunctionTask>*> HiPriTickTasks[TG_MAX][TG_MAX];
+	TArrayWithThreadsafeAdd<FTickGraphTask*> HiPriTickTasks[TG_MAX][TG_MAX];
 
 	/** LowPri Held tasks for each tick group. */
-	TArrayWithThreadsafeAdd<TGraphTask<FTickFunctionTask>*> TickTasks[TG_MAX][TG_MAX];
+	TArrayWithThreadsafeAdd<FTickGraphTask*> TickTasks[TG_MAX][TG_MAX];
 
 	/** These are waited for at the end of the frame; they are not on the critical path, but they have to be done before we leave the frame. */
 	FGraphEventArray CleanupTasks;
@@ -374,54 +485,64 @@ class FTickTaskSequencer
 	/** we keep track of the last TG we have blocked for so when we do block, we know which TG's to wait for . */
 	ETickingGroup WaitForTickGroup;
 
-	/** If true, allow concurrent ticks **/
+	/** If true, ticks can execute on other threads */
 	bool				bAllowConcurrentTicks;
+	/** If true, multiple ticks can be combined into a single task */
+	bool				bAllowBatchedTicksForFrame;
+	/** If true, some prerequisites will be ignored */
+	bool				bAllowOptimizedPrerequisites;
 
-	/** If true, log each tick **/
+	/** If true, log each tick */
 	bool				bLogTicks;
-	/** If true, log each tick **/
+	/** If true, log tick prerequisites when logging ticks */
 	bool				bLogTicksShowPrerequistes;
 
 public:
 
-	/**
-	 * Singleton to retrieve the global tick task sequencer
-	 * @return Reference to the global tick task sequencer
-	**/
+	/** Retrieve the global tick task sequencer shared for all ticks */
 	static FTickTaskSequencer& Get()
 	{
 		static FTickTaskSequencer SingletonInstance;
 		return SingletonInstance;
 	}
-	/**
-	 * Return true if we should be running in single threaded mode, ala dedicated server
-	**/
+
+	/** Return true if we should be running in single threaded mode, ala dedicated server */
 	FORCEINLINE static bool SingleThreadedMode()
 	{
-		if (FForkProcessHelper::IsForkedMultithreadInstance())
+		static TOptional<bool> bSingleThreaded;
+		if (!bSingleThreaded.IsSet())
 		{
-			return false;
+			if (FForkProcessHelper::IsForkedMultithreadInstance())
+			{
+				bSingleThreaded = false;
+			}
+			else if (!FApp::ShouldUseThreadingForPerformance() || IsRunningDedicatedServer() || FPlatformMisc::NumberOfCores() < 3 || !FPlatformProcess::SupportsMultithreading())
+			{
+				bSingleThreaded = true;
+			}
+			else
+			{
+				bSingleThreaded = false;
+			}
 		}
+		return *bSingleThreaded;
+	}
 
-		if (!FApp::ShouldUseThreadingForPerformance() || IsRunningDedicatedServer() || FPlatformMisc::NumberOfCores() < 3 || !FPlatformProcess::SupportsMultithreading())
-		{
-			return true;
-		}
-		return false;
+	/** Accessor that will return a valid task pointer or null if not actually scheduled */
+	FORCEINLINE static FTickGraphTask* GetGraphTask(const FTickFunction* TickFunction)
+	{
+		return (FTickGraphTask*)TickFunction->GetTaskPointer(FTickFunction::ETickTaskState::HasTask);
 	}
-	/**
-	 * Start a component tick task
-	 *
-	 * @param	InPrerequisites - prerequisites that must be completed before this tick can begin
-	 * @param	TickFunction - the tick function to queue
-	 * @param	Context - tick context to tick in. Thread here is the current thread.
-	 */
-	FORCEINLINE void StartTickTask(const FGraphEventArray* Prerequisites, FTickFunction* TickFunction, const FTickContext& TickContext)
+
+	/** Sets up TickContext for a tick function that will possibly tick */
+	FORCEINLINE FTickContext SetupTickContext(FTickFunction* TickFunction, const FTickContext& TickContext)
 	{
 		checkSlow(TickFunction->InternalData);
 		checkSlow(TickFunction->InternalData->ActualStartTickGroup >=0 && TickFunction->InternalData->ActualStartTickGroup < TG_MAX);
 
 		FTickContext UseContext = TickContext;
+		UseContext.bLogTick = bLogTicks;
+		UseContext.bLogTicksShowPrerequistes = bLogTicksShowPrerequistes;
 
 		bool bIsOriginalTickGroup = (TickFunction->InternalData->ActualStartTickGroup == TickFunction->TickGroup);
 
@@ -441,11 +562,23 @@ public:
 			UseContext.Thread = ENamedThreads::SetTaskPriority(ENamedThreads::GameThread, TickFunction->bHighPriority ? ENamedThreads::HighTaskPriority : ENamedThreads::NormalTaskPriority);
 		}
 
-		TickFunction->InternalData->TaskPointer = TGraphTask<FTickFunctionTask>::CreateTask(Prerequisites, TickContext.Thread).ConstructAndHold(TickFunction, &UseContext, bLogTicks, bLogTicksShowPrerequistes);
+		return UseContext;
+	}
+
+	/**
+	 * Start a task for a single function
+	 *
+	 * @param	InPrerequisites - prerequisites that must be completed before this tick can begin
+	 * @param	TickFunction - the tick function to queue
+	 * @param	Context - tick context to tick in, that has been setup by SetupTickContext
+	 */
+	FORCEINLINE void StartTickTask(const FGraphEventArray* Prerequisites, FTickFunction* TickFunction, const FTickContext& TickContext)
+	{
+		
 	}
 
 	/** Add a completion handle to a tick group **/
-	FORCEINLINE void AddTickTaskCompletion(ETickingGroup StartTickGroup, ETickingGroup EndTickGroup, TGraphTask<FTickFunctionTask>* Task, bool bHiPri)
+	FORCEINLINE void AddTickTaskCompletion(ETickingGroup StartTickGroup, ETickingGroup EndTickGroup, FTickGraphTask* Task, bool bHiPri)
 	{
 		checkSlow(StartTickGroup >=0 && StartTickGroup < TG_MAX && EndTickGroup >=0 && EndTickGroup < TG_MAX && StartTickGroup <= EndTickGroup);
 		if (bHiPri)
@@ -459,7 +592,7 @@ public:
 		TickCompletionEvents[EndTickGroup].Add(Task->GetCompletionEvent());
 	}
 	/** Add a completion handle to a tick group, parallel version **/
-	FORCEINLINE void AddTickTaskCompletionParallel(ETickingGroup StartTickGroup, ETickingGroup EndTickGroup, TGraphTask<FTickFunctionTask>* Task, bool bHiPri)
+	FORCEINLINE void AddTickTaskCompletionParallel(ETickingGroup StartTickGroup, ETickingGroup EndTickGroup, FTickGraphTask* Task, bool bHiPri)
 	{
 		check(StartTickGroup >= 0 && StartTickGroup < TG_NewlySpawned && EndTickGroup >= 0 && EndTickGroup < TG_NewlySpawned && StartTickGroup <= EndTickGroup);
 		if (bHiPri)
@@ -472,6 +605,41 @@ public:
 		}
 		TickCompletionEvents[EndTickGroup].AddThreadsafe(Task->GetCompletionEvent());
 	}
+
+	/** Return true if this prerequisite should be tracked in the event graph */
+	FORCEINLINE bool ShouldConsiderPrerequisite(FTickFunction* TickFunction, FTickFunction* Prereq)
+	{
+		// Ignore prereqs that are guaranteed to finish in a previous group
+		// This can be called before TickFunction has it's final group set, but Prereq will always be correct
+		// There is no hard wait for DuringPhysics so always consider those
+		TEnumAsByte<enum ETickingGroup> PrereqEnd = Prereq->GetActualEndTickGroup();
+		return (!bAllowOptimizedPrerequisites || PrereqEnd >= TickFunction->TickGroup || PrereqEnd == TG_DuringPhysics);
+	}
+
+	/** Return true if this tick condition is safe to batch */
+	FORCEINLINE bool CanBatchCondition(FTickGroupCondition Condition)
+	{
+		// Don't batch high priority ticks or ones that last more than a single tick group
+		return Condition.bIsBatch && !Condition.bHighPriority && Condition.StartGroup == Condition.EndGroup;
+	}
+
+	/** Call before doing any batched ticks */
+	void SetupBatchedTicks(int32 NumTicks)
+	{
+		// This is cleared at the end of tick
+		ensure(TickBatchesNum == 0);
+	}
+
+	/** Finished batching ticks for the frame */
+	void FinishBatchedTicks(const FTickContext& InContext)
+	{
+		if (bAllowBatchedTicksForFrame)
+		{
+			// Disable batching for the rest of the frame as we don't want to batch QueueNewlySpawned
+			bAllowBatchedTicksForFrame = false;
+		}
+	}
+
 	/** Set up the lists for AddTickTaskCompletionParallel, since we are using AddThreadsafe, we need to presize the arrays **/
 	void SetupAddTickTaskCompletionParallel(int32 NumTicks)
 	{
@@ -485,8 +653,98 @@ public:
 			TickCompletionEvents[TickGroup].Reserve(NumTicks);
 		}
 	}
+
+	/** This will add to an existing batch, create a new batch, or just spawn a single task and return null */
+	FTickBatchInfo* QueueOrBatchTickTask(TArray<FTickFunction*>& Prerequisites, FTickFunction* TickFunction, const FTickContext& TickContext)
+	{
+		// Batching is not supported with the old frontend
+#if TASKGRAPH_NEW_FRONTEND
+		if (bAllowBatchedTicksForFrame)
+		{
+			FTickGroupCondition Condition = FTickGroupCondition(TickFunction);
+
+			if (CanBatchCondition(Condition))
+			{
+				// Look for an appropriate batch
+				FTickBatchInfo* BatchInfo = nullptr;
+				for (int32 BatchIndex = 0; BatchIndex < TickBatchesNum; BatchIndex++)
+				{
+					if (Condition == TickBatches[BatchIndex].Key)
+					{
+						FTickBatchInfo* PossibleBatch = TickBatches[BatchIndex].Value.Get();
+						bool bPrerequisitesMatch = true;
+
+						for (FTickFunction* Prereq : Prerequisites)
+						{
+							// Ignore prerequisites that are already in this batch
+							if (Prereq->GetTaskPointer(FTickFunction::ETickTaskState::HasTask) != PossibleBatch->TickTask && !PossibleBatch->TickPrerequisites.Contains(Prereq))
+							{
+								bPrerequisitesMatch = false;
+								break;
+							}
+						}
+						if (bPrerequisitesMatch)
+						{
+							BatchInfo = PossibleBatch;
+							break;
+						}
+					}
+				}
+
+				if (!BatchInfo)
+				{
+					// Create a new batch, resizing array if needed
+					check(TickBatchesNum <= TickBatches.Num());
+					if (TickBatchesNum == TickBatches.Num())
+					{
+						TickBatches.Emplace(FTickGroupCondition(), MakeUnique<FTickBatchInfo>());
+						check(TickBatches.IsValidIndex(TickBatchesNum));
+					}
+				
+					TickBatches[TickBatchesNum].Key = Condition;
+					BatchInfo = TickBatches[TickBatchesNum].Value.Get();
+					TickBatchesNum++;
+
+					BatchInfo->TickPrerequisites = Prerequisites;
+					check(BatchInfo->TickTask == nullptr);
+
+					// Create the batched task now
+					FGraphEventArray PrerequisiteEvents;
+					for (FTickFunction* Prereq : Prerequisites)
+					{
+						PrerequisiteEvents.Add(Prereq->GetCompletionHandle());
+					}
+
+					FTickContext UseContext = SetupTickContext(TickFunction, TickContext);
+					BatchInfo->TickTask = TGraphTask<FBatchTickFunctionTask>::CreateTask(&PrerequisiteEvents, ENamedThreads::GameThread).ConstructAndHold(BatchInfo, &UseContext);
+
+					AddTickTaskCompletion(Condition.StartGroup, Condition.EndGroup, BatchInfo->TickTask, Condition.bHighPriority);
+				}
+
+				// Add this tick function to batch, which could be the first one
+				BatchInfo->TickFunctions.Add(TickFunction);
+				TickFunction->SetTaskPointer(FTickFunction::ETickTaskState::HasTask, BatchInfo->TickTask);
+
+				return BatchInfo;
+			}
+		}
+#endif //TASKGRAPH_NEW_FRONTEND
+
+		// No batching, create a single task
+		// FGraphEventArray array has some inline members so it is faster to not explicitly reserve space
+		FGraphEventArray PrerequisiteEvents;
+		for (FTickFunction* Prereq : Prerequisites)
+		{
+			PrerequisiteEvents.Add(Prereq->GetCompletionHandle());
+		}
+
+		QueueTickTask(&PrerequisiteEvents, TickFunction, TickContext);
+
+		return nullptr;
+	}
+
 	/**
-	 * Start a component tick task and add the completion handle
+	 * Start a tick task and add the completion handle
 	 *
 	 * @param	InPrerequisites - prerequisites that must be completed before this tick can begin
 	 * @param	TickFunction - the tick function to queue
@@ -494,15 +752,15 @@ public:
 	 */
 	FORCEINLINE void QueueTickTask(const FGraphEventArray* Prerequisites, FTickFunction* TickFunction, const FTickContext& TickContext)
 	{
-		checkSlow(TickFunction->InternalData);
-		checkSlow(TickContext.Thread == ENamedThreads::GameThread);
-		StartTickTask(Prerequisites, TickFunction, TickContext);
-		TGraphTask<FTickFunctionTask>* Task = (TGraphTask<FTickFunctionTask>*)TickFunction->InternalData->TaskPointer;
+		FTickContext UseContext = SetupTickContext(TickFunction, TickContext);
+		FTickGraphTask* Task = TGraphTask<FTickFunctionTask>::CreateTask(Prerequisites, ENamedThreads::GameThread).ConstructAndHold(TickFunction, &UseContext);
+		TickFunction->SetTaskPointer(FTickFunction::ETickTaskState::HasTask, Task);
+
 		AddTickTaskCompletion(TickFunction->InternalData->ActualStartTickGroup, TickFunction->InternalData->ActualEndTickGroup, Task, TickFunction->bHighPriority);
 	}
 
 	/**
-	 * Start a component tick task and add the completion handle
+	 * Start a tick task and add the completion handle, for parallel queuing
 	 *
 	 * @param	InPrerequisites - prerequisites that must be completed before this tick can begin
 	 * @param	TickFunction - the tick function to queue
@@ -510,10 +768,10 @@ public:
 	 */
 	FORCEINLINE void QueueTickTaskParallel(const FGraphEventArray* Prerequisites, FTickFunction* TickFunction, const FTickContext& TickContext)
 	{
-		checkSlow(TickFunction->InternalData);
-		checkSlow(TickContext.Thread == ENamedThreads::GameThread);
-		StartTickTask(Prerequisites, TickFunction, TickContext);
-		TGraphTask<FTickFunctionTask>* Task = (TGraphTask<FTickFunctionTask>*)TickFunction->InternalData->TaskPointer;
+		FTickContext UseContext = SetupTickContext(TickFunction, TickContext);
+		FTickGraphTask* Task = TGraphTask<FTickFunctionTask>::CreateTask(Prerequisites, ENamedThreads::GameThread).ConstructAndHold(TickFunction, &UseContext);
+		TickFunction->SetTaskPointer(FTickFunction::ETickTaskState::HasTask, Task);
+
 		AddTickTaskCompletionParallel(TickFunction->InternalData->ActualStartTickGroup, TickFunction->InternalData->ActualEndTickGroup, Task, TickFunction->bHighPriority);
 	}
 
@@ -597,6 +855,9 @@ public:
 			bAllowConcurrentTicks = !!CVarAllowAsyncComponentTicks.GetValueOnGameThread();
 		}
 
+		bAllowBatchedTicksForFrame = !!CVarAllowBatchedTicks.GetValueOnGameThread();
+		bAllowOptimizedPrerequisites = !!CVarAllowOptimizedPrerequisites.GetValueOnGameThread();
+
 		WaitForCleanup();
 
 		for (int32 Index = 0; Index < TG_MAX; Index++)
@@ -621,11 +882,22 @@ public:
 		{
 			UE_LOG(LogTick, Log, TEXT("tick %6llu ---------------------------------------- End Frame"),(uint64)GFrameCounter);
 		}
+
+		// Clear out batched tick data but leave allocated for next frame
+		for (TPair<FTickGroupCondition, TUniquePtr<FTickBatchInfo> >& Pair : TickBatches)
+		{
+			Pair.Key.IntVersion = 0;
+			Pair.Value->Reset();
+		}
+		TickBatchesNum = 0;
+
 	}
 private:
 
 	FTickTaskSequencer()
 		: bAllowConcurrentTicks(false)
+		, bAllowBatchedTicksForFrame(false)
+		, bAllowOptimizedPrerequisites(false)
 		, bLogTicks(false)
 		, bLogTicksShowPrerequistes(false)
 	{
@@ -661,7 +933,7 @@ private:
 		QUICK_SCOPE_CYCLE_COUNTER(STAT_DispatchTickGroup);
 		for (int32 IndexInner = 0; IndexInner < TG_MAX; IndexInner++)
 		{
-			TArray<TGraphTask<FTickFunctionTask>*>& TickArray = HiPriTickTasks[WorldTickGroup][IndexInner]; //-V781
+			TArray<FTickGraphTask*>& TickArray = HiPriTickTasks[WorldTickGroup][IndexInner]; //-V781
 			if (IndexInner < WorldTickGroup)
 			{
 				check(TickArray.Num() == 0); // makes no sense to have and end TG before the start TG
@@ -677,7 +949,7 @@ private:
 		}
 		for (int32 IndexInner = 0; IndexInner < TG_MAX; IndexInner++)
 		{
-			TArray<TGraphTask<FTickFunctionTask>*>& TickArray = TickTasks[WorldTickGroup][IndexInner]; //-V781
+			TArray<FTickGraphTask*>& TickArray = TickTasks[WorldTickGroup][IndexInner]; //-V781
 			if (IndexInner < WorldTickGroup)
 			{
 				check(TickArray.Num() == 0); // makes no sense to have and end TG before the start TG
@@ -1061,7 +1333,7 @@ public:
 			if (TickFunction->bTickEvenWhenPaused)
 			{
 				bool bExecuteTick = false;
-				TickFunction->InternalData->TaskPointer = nullptr; // this is stale, clear it out now
+				TickFunction->SetTaskPointer(FTickFunction::ETickTaskState::NotQueued, nullptr); // this is stale, clear it out now
 				if (CumulativeCooldown < InContext.DeltaSeconds)
 				{
 					// Queue up the tick function for later and do the reschedule before as it is in the normal ticking logic
@@ -1097,7 +1369,7 @@ public:
 		for (TSet<FTickFunction*>::TIterator It(AllEnabledTickFunctions); It; ++It)
 		{
 			TickFunction = *It;
-			TickFunction->InternalData->TaskPointer = nullptr; // this is stale, clear it out now
+			TickFunction->SetTaskPointer(FTickFunction::ETickTaskState::NotQueued, nullptr); // this is stale, clear it out now
 			if (TickFunction->bTickEvenWhenPaused && TickFunction->TickState == FTickFunction::ETickState::Enabled)
 			{
 				// Queue up the tick function for later and do the reschedule before as it is in the normal ticking logic
@@ -1115,9 +1387,11 @@ public:
 
 		for (FTickFunction* TickFunctionToExecute : ExecuteTickFunctions)
 		{
-			TickFunctionToExecute->InternalData->TickVisitedGFrameCounter = GFrameCounter;
-			TickFunctionToExecute->InternalData->TickQueuedGFrameCounter = GFrameCounter;
-			TickFunctionToExecute->ExecuteTick(TickFunctionToExecute->CalculateDeltaTime(InContext), InContext.TickType, ENamedThreads::GameThread, FGraphEventRef());
+			// Only use the lower 32 bits of the frame counter
+			uint32 CurrentFrameCounter = (uint32)GFrameCounter;
+			TickFunctionToExecute->InternalData->TickVisitedGFrameCounter.store(CurrentFrameCounter, std::memory_order_relaxed);
+			TickFunctionToExecute->InternalData->TickQueuedGFrameCounter.store(CurrentFrameCounter, std::memory_order_relaxed);
+			TickFunctionToExecute->ExecuteTick(TickFunctionToExecute->CalculateDeltaTime(InContext.DeltaSeconds, InContext.World), InContext.TickType, ENamedThreads::GameThread, FGraphEventRef());
 		}
 	}
 
@@ -1170,7 +1444,7 @@ public:
 	}
 
 	/** Dumps info about a tick function to output device. */
-	FORCEINLINE void DumpTickFunction(FOutputDevice& Ar, FTickFunction* Function, UEnum* TickGroupEnum, const float RemainingCooldown = 0.f)
+	void DumpTickFunction(FOutputDevice& Ar, FTickFunction* Function, UEnum* TickGroupEnum, const float RemainingCooldown = 0.f)
 	{
 		// Info about the function.
 		Ar.Logf(TEXT("%s, %s, ActualStartTickGroup: %s, Prerequesities: %d"),
@@ -1182,7 +1456,7 @@ public:
 		// List all prerequisities
 		for (int32 Index = 0; Index < Function->Prerequisites.Num(); ++Index)
 		{
-			FTickPrerequisite& Prerequisite = Function->Prerequisites[Index];
+			const FTickPrerequisite& Prerequisite = Function->Prerequisites[Index];
 			if (Prerequisite.PrerequisiteObject.IsValid())
 			{
 				Ar.Logf(TEXT("    %s, %s"), *Prerequisite.PrerequisiteObject->GetFullName(), *Prerequisite.PrerequisiteTickFunction->DiagnosticMessage());
@@ -1192,6 +1466,12 @@ public:
 				Ar.Logf(TEXT("    Invalid Prerequisite"));
 			}
 		}
+
+		// Handle any children if they exist
+		Function->ForEachNestedTick([this, &Ar, TickGroupEnum, RemainingCooldown](FTickFunction& NestedTick)
+		{
+			DumpTickFunction(Ar, &NestedTick, TickGroupEnum, RemainingCooldown);
+		});
 	}
 
 	/** Dumps all tick functions to output device. */
@@ -1424,7 +1704,7 @@ private:
 struct FTickGroupCompletionItem
 {
 	/** Task created **/
-	TGraphTask<FTickFunctionTask>* Task;
+	FTickGraphTask* Task;
 	/** Tick group to complete with **/
 	TEnumAsByte<ETickingGroup>	ActualStartTickGroup;
 	/** True if this was a misplaced interval tick that we need to deal with **/
@@ -1509,21 +1789,24 @@ public:
 			}
 			INC_DWORD_STAT_BY(STAT_TicksQueued, TotalTickFunctions);
 			CSV_CUSTOM_STAT(Basic, TicksQueued, TotalTickFunctions, ECsvCustomStatOp::Accumulate);
+			TickTaskSequencer.SetupBatchedTicks(TotalTickFunctions);
 			for( int32 LevelIndex = 0; LevelIndex < LevelList.Num(); LevelIndex++ )
 			{
 				LevelList[LevelIndex]->QueueAllTicks();
 			}
+			TickTaskSequencer.FinishBatchedTicks(Context);
 		}
 		else
 		{
+			ensureMsgf(!CVarAllowBatchedTicks.GetValueOnGameThread(), TEXT("Concurrent queuing is not compatible with batched ticks!"));
+
 			for( int32 LevelIndex = 0; LevelIndex < LevelList.Num(); LevelIndex++ )
 			{
 				LevelList[LevelIndex]->StartFrameParallel(Context, AllTickFunctions);
 			}
 			INC_DWORD_STAT_BY(STAT_TicksQueued, AllTickFunctions.Num());
 			CSV_CUSTOM_STAT(Basic, TicksQueued, AllTickFunctions.Num(), ECsvCustomStatOp::Accumulate);
-			FTickTaskSequencer& TTS = FTickTaskSequencer::Get();
-			TTS.SetupAddTickTaskCompletionParallel(AllTickFunctions.Num());
+			TickTaskSequencer.SetupAddTickTaskCompletionParallel(AllTickFunctions.Num());
 			for( int32 LevelIndex = 0; LevelIndex < LevelList.Num(); LevelIndex++ )
 			{
 				LevelList[LevelIndex]->ReserveTickFunctionCooldowns(AllTickFunctions.Num());
@@ -1799,6 +2082,7 @@ FTickFunction::FTickFunction()
 	, bCanEverTick(false)
 	, bStartWithTickEnabled(false)
 	, bAllowTickOnDedicatedServer(true)
+	, bAllowTickBatching(false)
 	, bHighPriority(false)
 	, bRunOnAnyThread(false)
 	, TickState(ETickState::Enabled)
@@ -1809,6 +2093,7 @@ FTickFunction::FTickFunction()
 FTickFunction::FInternalData::FInternalData()
 	: bRegistered(false)
 	, bWasInterval(false)
+	, TaskState(ETickTaskState::NotQueued)
 	, ActualStartTickGroup(TG_PrePhysics)
 	, ActualEndTickGroup(TG_PrePhysics)
 	, TickVisitedGFrameCounter(0)
@@ -1938,6 +2223,21 @@ void FTickFunction::SetPriorityIncludingPrerequisites(bool bInHighPriority)
 	}
 }
 
+void FTickFunction::LogTickFunction(ENamedThreads::Type CurrentThread, bool bLogPrerequisites, int32 Indent)
+{
+	UE_LOG(LogTick, Log, TEXT("%stick %s [%1d, %1d] %6llu %2d %s"), FCString::Spc(Indent * 2), bHighPriority ? TEXT("*") : TEXT(" "), (int32)GetActualTickGroup(), (int32)GetActualEndTickGroup(), (uint64)GFrameCounter, (int32)CurrentThread, *DiagnosticMessage());
+	if (bLogPrerequisites)
+	{
+		ShowPrerequistes();
+	}
+	
+	// Handle nested ticks
+	ForEachNestedTick([CurrentThread, bLogPrerequisites, Indent](FTickFunction& NestedTick)
+	{
+		NestedTick.LogTickFunction(CurrentThread, bLogPrerequisites, Indent + 1);
+	});
+}
+
 void FTickFunction::ShowPrerequistes(int32 Indent)
 {
 	for (auto& Prereq : Prerequisites)
@@ -1950,27 +2250,75 @@ void FTickFunction::ShowPrerequistes(int32 Indent)
 	}
 }
 
+bool FTickFunction::IsCompletionHandleValid() const
+{
+	return InternalData && (InternalData->TaskState == ETickTaskState::HasTask || InternalData->TaskState == ETickTaskState::HasCompletionEvent);
+}
+
 FGraphEventRef FTickFunction::GetCompletionHandle() const
 {
-	check(InternalData->TaskPointer);
-	TGraphTask<FTickFunctionTask>* Task = (TGraphTask<FTickFunctionTask>*)InternalData->TaskPointer;
+	check(InternalData);
+	if (InternalData->TaskState == ETickTaskState::HasCompletionEvent)
+	{
+		// This will increment the reference count on the event
+		return *((FGraphEventRef*)InternalData->TaskPointer);
+	}
+
+	FTickGraphTask* Task = FTickTaskSequencer::GetGraphTask(this);
+	check(Task);
 	return Task->GetCompletionEvent();
 }
 
+void FTickFunction::ExecuteNestedTick(float DeltaTime, ELevelTick TickType, ENamedThreads::Type CurrentThread, const FGraphEventRef& MyCompletionGraphEvent)
+{
+	// Not valid to call on a real scheduled function
+	check(GetTaskPointer(ETickTaskState::HasTask) == nullptr);
+
+	// This does not increment the reference counter
+	if (MyCompletionGraphEvent.IsValid())
+	{
+		SetTaskPointer(ETickTaskState::HasCompletionEvent, (void*)&MyCompletionGraphEvent);
+	}
+	ExecuteTick(DeltaTime, TickType, CurrentThread, MyCompletionGraphEvent);
+	ClearTaskInformation();
+}
+
+void FTickFunction::ClearTaskInformation()
+{
+	SetTaskPointer(ETickTaskState::NotQueued, nullptr);
+}
+
+void FTickFunction::SetTaskPointer(ETickTaskState NewState, void* InTaskPointer)
+{
+	check(InternalData);
+
+	if (NewState == ETickTaskState::NotQueued)
+	{
+		InternalData->TaskState = ETickTaskState::NotQueued;
+		InternalData->TaskPointer = nullptr;
+	}
+	else
+	{
+		InternalData->TaskState = NewState;
+		InternalData->TaskPointer = InTaskPointer;
+	}
+}
 
 void FTickFunction::QueueTickFunction(FTickTaskSequencer& TTS, const struct FTickContext& TickContext)
 {
 	checkSlow(TickContext.Thread == ENamedThreads::GameThread); // we assume same thread here
 	check(IsTickFunctionRegistered());
 
-	if (InternalData->TickVisitedGFrameCounter != GFrameCounter)
+	// Only compare the 32bit part of the frame counter
+	uint32 CurrentFrameCounter = (uint32)GFrameCounter;
+	if (InternalData->TickVisitedGFrameCounter.load(std::memory_order_relaxed) != CurrentFrameCounter)
 	{
-		InternalData->TickVisitedGFrameCounter = GFrameCounter;
+		InternalData->TickVisitedGFrameCounter.store(CurrentFrameCounter, std::memory_order_relaxed);
 		if (TickState != FTickFunction::ETickState::Disabled)
 		{
 			ETickingGroup MaxPrerequisiteTickGroup =  ETickingGroup(0);
 
-			FGraphEventArray TaskPrerequisites;
+			TArray<FTickFunction*> RawPrerequisites;
 			for (int32 PrereqIndex = 0; PrereqIndex < Prerequisites.Num(); PrereqIndex++)
 			{
 				FTickFunction* Prereq = Prerequisites[PrereqIndex].Get();
@@ -1983,19 +2331,19 @@ void FTickFunction::QueueTickFunction(FTickTaskSequencer& TTS, const struct FTic
 				{
 					// recursive call to make sure my prerequisite is set up so I can use its completion handle
 					Prereq->QueueTickFunction(TTS, TickContext);
-					if (Prereq->InternalData->TickQueuedGFrameCounter != GFrameCounter)
+					if (Prereq->InternalData->TickQueuedGFrameCounter.load(std::memory_order_relaxed) != CurrentFrameCounter)
 					{
 						// this must be up the call stack, therefore this is a cycle
 						UE_LOG(LogTick, Warning, TEXT("While processing prerequisites for %s, could use %s because it would form a cycle."),*DiagnosticMessage(), *Prereq->DiagnosticMessage());
 					}
-					else if (!Prereq->InternalData->TaskPointer)
+					else if (Prereq->InternalData->TaskState == ETickTaskState::NotQueued)
 					{
 						//ok UE_LOG(LogTick, Warning, TEXT("While processing prerequisites for %s, could use %s because it is disabled."),*DiagnosticMessage(), *Prereq->DiagnosticMessage());
 					}
-					else
+					else if (TTS.ShouldConsiderPrerequisite(this, Prereq))
 					{
 						MaxPrerequisiteTickGroup =  FMath::Max<ETickingGroup>(MaxPrerequisiteTickGroup, Prereq->InternalData->ActualStartTickGroup.GetValue());
-						TaskPrerequisites.Add(Prereq->GetCompletionHandle());
+						RawPrerequisites.Add(Prereq);
 					}
 				}
 			}
@@ -2028,25 +2376,25 @@ void FTickFunction::QueueTickFunction(FTickTaskSequencer& TTS, const struct FTic
 
 			if (TickState == FTickFunction::ETickState::Enabled)
 			{
-				TTS.QueueTickTask(&TaskPrerequisites, this, TickContext);
+				TTS.QueueOrBatchTickTask(RawPrerequisites, this, TickContext);
 			}
 		}
-		InternalData->TickQueuedGFrameCounter = GFrameCounter;
+		InternalData->TickQueuedGFrameCounter.store(CurrentFrameCounter, std::memory_order_relaxed);
 	}
 }
 
 void FTickFunction::QueueTickFunctionParallel(const struct FTickContext& TickContext, TArray<FTickFunction*, TInlineAllocator<8> >& StackForCycleDetection)
 {
-	bool bProcessTick;
-
-	int32 OldValue = FPlatformAtomics::AtomicRead_Relaxed(&InternalData->TickVisitedGFrameCounter);
-	if (OldValue != GFrameCounter)
+	// Only compare the 32bit part of the frame counter
+	uint32 CurrentFrameCounter = (uint32)GFrameCounter;
+	uint32 OldValue = InternalData->TickVisitedGFrameCounter.load(std::memory_order_relaxed);
+	if (OldValue != CurrentFrameCounter)
 	{
-		OldValue = FPlatformAtomics::InterlockedCompareExchange(&InternalData->TickVisitedGFrameCounter , GFrameCounter, OldValue);
+		// Modify the visited frame if another thread has not already done so
+		InternalData->TickVisitedGFrameCounter.compare_exchange_strong(OldValue, CurrentFrameCounter);
 	}
-	bProcessTick = OldValue != GFrameCounter;
 
-	if (bProcessTick)
+	if (OldValue != CurrentFrameCounter)
 	{
 		check(IsTickFunctionRegistered());
 		if (TickState != FTickFunction::ETickState::Disabled)
@@ -2074,7 +2422,7 @@ void FTickFunction::QueueTickFunctionParallel(const struct FTickContext& TickCon
 					{
 						// recursive call to make sure my prerequisite is set up so I can use its completion handle
 						Prereq->QueueTickFunctionParallel(TickContext, StackForCycleDetection);
-						if (!Prereq->InternalData->TaskPointer)
+						if (!Prereq->IsCompletionHandleValid())
 						{
 							//ok UE_LOG(LogTick, Warning, TEXT("While processing prerequisites for %s, could use %s because it is disabled."),*DiagnosticMessage(), *Prereq->DiagnosticMessage());
 						}
@@ -2124,16 +2472,16 @@ void FTickFunction::QueueTickFunctionParallel(const struct FTickContext& TickCon
 			}
 		}
 		
-		InternalData->TickQueuedGFrameCounter = GFrameCounter;
+		InternalData->TickQueuedGFrameCounter = CurrentFrameCounter;
 	}
 	else
 	{
 		// if we are not going to process it, we need to at least wait until the other thread finishes it
-		std::atomic<int32>& TickQueuedGFrameCounter = InternalData->TickQueuedGFrameCounter;
-		if (TickQueuedGFrameCounter != GFrameCounter)
+		std::atomic<uint32>& TickQueuedGFrameCounter = InternalData->TickQueuedGFrameCounter;
+		if (TickQueuedGFrameCounter != CurrentFrameCounter)
 		{
 			QUICK_SCOPE_CYCLE_COUNTER(STAT_FTickFunction_QueueTickFunctionParallel_Spin);
-			while (TickQueuedGFrameCounter != GFrameCounter)
+			while (TickQueuedGFrameCounter != CurrentFrameCounter)
 			{
 				FPlatformProcess::YieldThread();
 			}
@@ -2141,10 +2489,8 @@ void FTickFunction::QueueTickFunctionParallel(const struct FTickContext& TickCon
 	}
 }
 
-float FTickFunction::CalculateDeltaTime(const FTickContext& TickContext)
+float FTickFunction::CalculateDeltaTime(float DeltaTime, const UWorld* TickingWorld)
 {
-	float DeltaTimeForFunction = TickContext.DeltaSeconds;
-
 	if (!InternalData->bWasInterval)
 	{
 		// No tick interval. Return the world delta seconds, and make sure to mark that
@@ -2155,15 +2501,15 @@ float FTickFunction::CalculateDeltaTime(const FTickContext& TickContext)
 	{
 		// We've got a tick interval. Mark last-tick-time. If we already had last-tick-time, return
 		// the time since then; otherwise, return the world delta seconds.
-		const float CurrentWorldTime = (bTickEvenWhenPaused ? TickContext.World->GetUnpausedTimeSeconds() : TickContext.World->GetTimeSeconds());
+		const float CurrentWorldTime = (bTickEvenWhenPaused ? TickingWorld->GetUnpausedTimeSeconds() : TickingWorld->GetTimeSeconds());
 		if (InternalData->LastTickGameTimeSeconds >= 0.f)
 		{
-			DeltaTimeForFunction = CurrentWorldTime - InternalData->LastTickGameTimeSeconds;
+			DeltaTime = CurrentWorldTime - InternalData->LastTickGameTimeSeconds;
 		}
 		InternalData->LastTickGameTimeSeconds = CurrentWorldTime;
 	}
 
-	return DeltaTimeForFunction;
+	return DeltaTime;
 }
 
 /**
@@ -2283,4 +2629,3 @@ static FAutoConsoleCommandWithWorldAndArgs AddIndirectTestTickFunctionsCmd(
 	);
 
 
-
