Fixed UTextureRenderTarget2DArray::GetSurfaceHeight returning SizeX instead of SizeY.

diff --git a/Engine/Source/Runtime/Engine/Classes/Engine/TextureRenderTarget2DArray.h b/Engine/Source/Runtime/Engine/Classes/Engine/TextureRenderTarget2DArray.h
--- a/Engine/Source/Runtime/Engine/Classes/Engine/TextureRenderTarget2DArray.h
+++ b/Engine/Source/Runtime/Engine/Classes/Engine/TextureRenderTarget2DArray.h
@@ -82,9 +82,7 @@ class UTextureRenderTarget2DArray : public UTextureRenderTarget
 
 	//~ Begin UTexture Interface.
 	virtual float GetSurfaceWidth() const  override { return static_cast<float>(SizeX); }
-	// PVS-Studio notices that the implementation of GetSurfaceWidth is identical to this one
-	// and warns us. In this case, it is intentional, so we disable the warning:
-	virtual float GetSurfaceHeight()const  override { return static_cast<float>(SizeX); }	 //-V524
+	virtual float GetSurfaceHeight()const  override { return static_cast<float>(SizeY); }
 	virtual float GetSurfaceDepth() const override { return 0.0f; }
 	virtual uint32 GetSurfaceArraySize() const override { return Slices; }
 	virtual FTextureResource* CreateResource() override;
