Fixed missing float identifier for default values (Linux compiler fails without it)

#rnx

diff --git a/Engine/Source/Runtime/Engine/Classes/Kismet/KismetMathLibrary.h b/Engine/Source/Runtime/Engine/Classes/Kismet/KismetMathLibrary.h
--- a/Engine/Source/Runtime/Engine/Classes/Kismet/KismetMathLibrary.h
+++ b/Engine/Source/Runtime/Engine/Classes/Kismet/KismetMathLibrary.h
@@ -1194,7 +1194,7 @@ class UKismetMathLibrary : public UBlueprintFunctionLibrary
 	 * @see NormalSafe2D()
 	 */
 	UFUNCTION(BlueprintCallable, meta = (DisplayName = "Normalize In Place (Vector2D)", Keywords = "Unit Vector", ScriptMethod = "Normalize"), Category = "Math|Vector2D")
-	static ENGINE_API void Normalize2D(UPARAM(ref) FVector2D& A, float Tolerance = 1.e-8);
+	static ENGINE_API void Normalize2D(UPARAM(ref) FVector2D& A, float Tolerance = 1.e-8f);
 
 	/** Converts spherical coordinates on the unit sphere into a Cartesian unit length vector. */
 	UFUNCTION(BlueprintPure, meta = (DisplayName = "Spherical2D To Unit Cartesian", Keywords = "Unit Vector", ScriptMethod = "SphericalToUnitCartesian"), Category = "Math|Vector2D")
@@ -1725,7 +1725,7 @@ class UKismetMathLibrary : public UBlueprintFunctionLibrary
 	 * @param Tolerance Minimum squared length of vector for normalization.
 	 */
 	UFUNCTION(BlueprintCallable, meta = (DisplayName = "Normalize In Place (Vector)", ScriptMethod = "Normalize", Keywords = "Unit Vector"), Category = "Math|Vector")
-	static ENGINE_API void Vector_Normalize(UPARAM(ref) FVector& A, float Tolerance = 1.e-8);
+	static ENGINE_API void Vector_Normalize(UPARAM(ref) FVector& A, float Tolerance = 1.e-8f);
 
 	/** Linearly interpolates between A and B based on Alpha (100% of A when Alpha=0 and 100% of B when Alpha=1) */
 	UFUNCTION(BlueprintPure, meta=(DisplayName = "Lerp (Vector)", ScriptMethod = "LerpTo"), Category="Math|Vector")
@@ -2220,7 +2220,7 @@ class UKismetMathLibrary : public UBlueprintFunctionLibrary
 	 * @param Tolerance Minimum squared length of vector for normalization.
 	 */
 	UFUNCTION(BlueprintCallable, meta = (DisplayName = "Normalize In Place XYZ (Vector4)", ScriptMethod = "Normalize3", Keywords = "Unit Vector"), Category = "Math|Vector4")
-	static ENGINE_API void Vector4_Normalize3(UPARAM(ref) FVector4& A, float Tolerance = 1.e-8);
+	static ENGINE_API void Vector4_Normalize3(UPARAM(ref) FVector4& A, float Tolerance = 1.e-8f);
 
 	/** 
 	 * Given a direction vector and a surface normal, returns the vector reflected across the surface normal.
