Use an atomic to notify if the PSOPrecache are all done compiling. This will remove the memory allocated (64KB) by the TaskGraph when it's was created. In some cases, only PostLoad was called, which ended up not freeing up all the memory allocated by the TGraphTask even if the task was completed. Not clearing the PSOPrecacheCompileEvent when the TGraphTask<FPSOPrecacheFinishedTask> was completed could affect the memory if we had a lot of them in flight.

#tests local replay

diff --git a/Engine/Source/Runtime/Engine/Classes/Components/DecalComponent.h b/Engine/Source/Runtime/Engine/Classes/Components/DecalComponent.h
--- a/Engine/Source/Runtime/Engine/Classes/Components/DecalComponent.h
+++ b/Engine/Source/Runtime/Engine/Classes/Components/DecalComponent.h
@@ -124,13 +124,12 @@ public:
 	UFUNCTION(BlueprintCallable, Category="Rendering|Components|Decal")
 	ENGINE_API virtual class UMaterialInstanceDynamic* CreateDynamicMaterialInstance();
 
+public:
 #if UE_WITH_PSO_PRECACHING
-protected:
-	/** Graph event used to track all the PSO precache events - used for delayed proxy creation */
-	FGraphEventRef PSOPrecacheCompileEvent;
+	/** Atomic bool to track all the PSO precache events - used for delayed proxy creation */
+	std::atomic<bool> bPSOPrecacheFinished = true;
 #endif
 
-public:
 	/** The decal proxy. */
 	FDeferredDecalProxy* SceneProxy;
 
diff --git a/Engine/Source/Runtime/Engine/Classes/Components/PrimitiveComponent.h b/Engine/Source/Runtime/Engine/Classes/Components/PrimitiveComponent.h
--- a/Engine/Source/Runtime/Engine/Classes/Components/PrimitiveComponent.h
+++ b/Engine/Source/Runtime/Engine/Classes/Components/PrimitiveComponent.h
@@ -757,8 +757,8 @@ public:
 	/** Cached array of material PSO requests which can be used to boost the priority */
 	TArray<FMaterialPSOPrecacheRequestID> MaterialPSOPrecacheRequestIDs;
 
-	/** Graph event used to track all the PSO precache events */
-	FGraphEventRef PSOPrecacheCompileEvent;
+	/** Atomic bool used to track all the PSO precache events */
+	std::atomic<bool> bPSOPrecacheFinished = true;
 #endif
 
 	/**
diff --git a/Engine/Source/Runtime/Engine/Private/Components/DecalComponent.cpp b/Engine/Source/Runtime/Engine/Private/Components/DecalComponent.cpp
--- a/Engine/Source/Runtime/Engine/Private/Components/DecalComponent.cpp
+++ b/Engine/Source/Runtime/Engine/Private/Components/DecalComponent.cpp
@@ -331,7 +331,7 @@ void UDecalComponent::PrecachePSOs()
 	}
 
 	// clear the current request data
-	PSOPrecacheCompileEvent = nullptr;
+	bPSOPrecacheFinished = true;
 
 	if (DecalMaterial && !DecalMaterial->HasAnyFlags(RF_NeedPostLoad))
 	{
@@ -346,7 +346,31 @@ void UDecalComponent::PrecachePSOs()
 		// Request recreate of the render state when the PSO compilation is ready (if we want to delay proxy creation)
 		if (GraphEvents.Num() > 0 && GetPSOPrecacheProxyCreationStrategy() != EPSOPrecacheProxyCreationStrategy::AlwaysCreate)
 		{
-			PSOPrecacheCompileEvent = TGraphTask<FMarkActorRenderStateDirtyTask>::CreateTask(&GraphEvents).ConstructAndDispatchWhenReady(this);
+			struct FPSODecalPrecacheFinishedTask
+			{
+				explicit FPSODecalPrecacheFinishedTask(UDecalComponent* InDecalComponent)
+					: WeakDecalComponent(InDecalComponent)
+				{
+				}
+
+				static TStatId GetStatId() { return TStatId(); }
+				static ENamedThreads::Type GetDesiredThread() { return ENamedThreads::GameThread; }
+				static ESubsequentsMode::Type GetSubsequentsMode() { return ESubsequentsMode::TrackSubsequents; }
+
+				void DoTask(ENamedThreads::Type CurrentThread, const FGraphEventRef& MyCompletionGraphEvent)
+				{
+					if (UDecalComponent* DC = WeakDecalComponent.Get())
+					{
+						DC->bPSOPrecacheFinished = true;
+						DC->MarkRenderStateDirty();
+					}
+				}
+
+				TWeakObjectPtr<UDecalComponent> WeakDecalComponent;
+			};
+
+			bPSOPrecacheFinished = false;
+			TGraphTask<FPSODecalPrecacheFinishedTask>::CreateTask(&GraphEvents).ConstructAndDispatchWhenReady(this);
 		}
 	}
 #endif
@@ -393,11 +417,10 @@ FDeferredDecalProxy* UDecalComponent::CreateSceneProxy()
 	LLM_SCOPE(ELLMTag::SceneRender);
 
 #if UE_WITH_PSO_PRECACHING
-	if (PSOPrecacheCompileEvent && !PSOPrecacheCompileEvent->IsComplete() && GetPSOPrecacheProxyCreationStrategy() == EPSOPrecacheProxyCreationStrategy::DelayUntilPSOPrecached)
+	if (!bPSOPrecacheFinished && GetPSOPrecacheProxyCreationStrategy() == EPSOPrecacheProxyCreationStrategy::DelayUntilPSOPrecached)
 	{
 		return nullptr;
 	}
-	PSOPrecacheCompileEvent = nullptr;
 #endif // UE_WITH_PSO_PRECACHING
 
 	return new FDeferredDecalProxy(this);
diff --git a/Engine/Source/Runtime/Engine/Private/Components/PrimitiveComponent.cpp b/Engine/Source/Runtime/Engine/Private/Components/PrimitiveComponent.cpp
--- a/Engine/Source/Runtime/Engine/Private/Components/PrimitiveComponent.cpp
+++ b/Engine/Source/Runtime/Engine/Private/Components/PrimitiveComponent.cpp
@@ -4780,7 +4780,7 @@ void UPrimitiveComponent::PrecachePSOs()
 
 	// clear the current request data
 	MaterialPSOPrecacheRequestIDs.Empty();
-	PSOPrecacheCompileEvent = nullptr;
+	bPSOPrecacheFinished = true;
 	PSOPrecacheRequestPriority = EPSOPrecachePriority::Medium;
 
 	// Collect the data from the derived classes
@@ -4796,19 +4796,31 @@ void UPrimitiveComponent::PrecachePSOs()
 #endif
 }
 
-class FPSOPrecacheFinishedTask : public FMarkActorRenderStateDirtyTask
+#if UE_WITH_PSO_PRECACHING
+struct FPSOPrecacheFinishedTask
 {
-public:
-	explicit FPSOPrecacheFinishedTask(UActorComponent* InActorComponent)
-		: FMarkActorRenderStateDirtyTask(InActorComponent)
-	{ }
+	explicit FPSOPrecacheFinishedTask(UPrimitiveComponent* InPrimitiveComponent)
+		: WeakPrimitiveComponent(InPrimitiveComponent)
+	{
+	}
+
+	static TStatId GetStatId() { return TStatId(); }
+	static ENamedThreads::Type GetDesiredThread() { return ENamedThreads::GameThread; }
+	static ESubsequentsMode::Type GetSubsequentsMode() { return ESubsequentsMode::TrackSubsequents; }
 
 	void DoTask(ENamedThreads::Type CurrentThread, const FGraphEventRef& MyCompletionGraphEvent)
 	{
-		QUICK_SCOPE_CYCLE_COUNTER(STAT_PSOPrecacheFinishedTask);
-		FMarkActorRenderStateDirtyTask::DoTask(CurrentThread, MyCompletionGraphEvent);
+		if (UPrimitiveComponent* PC = WeakPrimitiveComponent.Get())
+		{
+			QUICK_SCOPE_CYCLE_COUNTER(STAT_PSOPrecacheFinishedTask);
+			PC->bPSOPrecacheFinished = true;
+			PC->MarkRenderStateDirty();
+		}
 	}
+
+	TWeakObjectPtr<UPrimitiveComponent> WeakPrimitiveComponent;
 };
+#endif
 
 void UPrimitiveComponent::RequestRecreateRenderStateWhenPSOPrecacheFinished(const FGraphEventArray& PSOPrecacheCompileEvents)
 {
@@ -4817,7 +4829,8 @@ void UPrimitiveComponent::RequestRecreateRenderStateWhenPSOPrecacheFinished(cons
 	// schedule a task to mark the render state dirty when all PSOs are compiled so the proxy gets recreated.
 	if (UsePSOPrecacheRenderProxyDelay() && GetPSOPrecacheProxyCreationStrategy() != EPSOPrecacheProxyCreationStrategy::AlwaysCreate && !PSOPrecacheCompileEvents.IsEmpty())
 	{
-		PSOPrecacheCompileEvent = TGraphTask<FPSOPrecacheFinishedTask>::CreateTask(&PSOPrecacheCompileEvents).ConstructAndDispatchWhenReady(this);
+		bPSOPrecacheFinished = false;
+		TGraphTask<FPSOPrecacheFinishedTask>::CreateTask(&PSOPrecacheCompileEvents).ConstructAndDispatchWhenReady(this);
 	}
 
 	bPSOPrecacheCalled = true;
@@ -4836,7 +4849,7 @@ bool UPrimitiveComponent::UsePSOPrecacheRenderProxyDelay() const
 bool UPrimitiveComponent::IsPSOPrecaching() const
 {
 #if UE_WITH_PSO_PRECACHING
-	return PSOPrecacheCompileEvent && !PSOPrecacheCompileEvent->IsComplete();
+	return !bPSOPrecacheFinished;
 #else
 	return false;
 #endif // UE_WITH_PSO_PRECACHING
@@ -4857,7 +4870,9 @@ bool UPrimitiveComponent::CheckPSOPrecachingAndBoostPriority(EPSOPrecachePriorit
 	ensure(!IsComponentPSOPrecachingEnabled() || bPSOPrecacheCalled);
 	check(NewPSOPrecachePriority == EPSOPrecachePriority::High || NewPSOPrecachePriority == EPSOPrecachePriority::Highest);
 
-	if (PSOPrecacheCompileEvent && !PSOPrecacheCompileEvent->IsComplete())
+	bool bIsPSOPrecaching = IsPSOPrecaching();
+
+	if (bIsPSOPrecaching)
 	{
 		if (PSOPrecacheRequestPriority< NewPSOPrecachePriority)
 		{
@@ -4865,12 +4880,8 @@ bool UPrimitiveComponent::CheckPSOPrecachingAndBoostPriority(EPSOPrecachePriorit
 			PSOPrecacheRequestPriority = NewPSOPrecachePriority;
 		}
 	}
-	else
-	{
-		PSOPrecacheCompileEvent = nullptr;
-	}
 
-	return IsPSOPrecaching();
+	return bIsPSOPrecaching;
 #else
 	return false;
 #endif
@@ -5172,4 +5183,4 @@ HHitProxy* FActorPrimitiveComponentInterface::CreateMeshHitProxy(int32 SectionIn
 }
 #endif
 
-#undef LOCTEXT_NAMESPACE
+#undef LOCTEXT_NAMESPACE
\ No newline at end of file
diff --git a/Engine/Source/Runtime/Engine/Private/Particles/ParticleComponents.cpp b/Engine/Source/Runtime/Engine/Private/Particles/ParticleComponents.cpp
--- a/Engine/Source/Runtime/Engine/Private/Particles/ParticleComponents.cpp
+++ b/Engine/Source/Runtime/Engine/Private/Particles/ParticleComponents.cpp
@@ -3566,7 +3566,7 @@ void UFXSystemComponent::PrecacheAssetPSOs(UFXSystemAsset* FXSystemAsset)
 	check(IsInGameThread() || IsInParallelGameThread());
 
 	MaterialPSOPrecacheRequestIDs.Empty();
-	PSOPrecacheCompileEvent = nullptr;
+	bPSOPrecacheFinished = true;
 	PSOPrecacheRequestPriority = EPSOPrecachePriority::Medium;
 
 	// The asset will keep the Precache events alive, but these might be over. Avoid delaying scene proxy creation if everything is finished
