Validate bone transform before converting to a matrix in skeletal mesh

diff --git a/Engine/Source/Runtime/Engine/Private/SkeletalMesh.cpp b/Engine/Source/Runtime/Engine/Private/SkeletalMesh.cpp
--- a/Engine/Source/Runtime/Engine/Private/SkeletalMesh.cpp
+++ b/Engine/Source/Runtime/Engine/Private/SkeletalMesh.cpp
@@ -1557,8 +1557,8 @@ EDataValidationResult USkeletalMesh::IsDataValid(FDataValidationContext& Context
 			}
 
 			//Validate transform do not contains nan
-			FMatrix RefPoseMatrix = GetRefPoseMatrix(BoneIndex);
-			if (RefPoseMatrix.ContainsNaN())
+			const FTransform& BoneTransform = RawRefBonePose[BoneIndex];
+			if (BoneTransform.ContainsNaN())
 			{
 				Context.AddError(LOCTEXT("SkeletalMeshValidation_PoseMatrixContainNan", "This skeletal mesh asset has NAN (invalid float number) value in the pose matrix. Asset is corrupted and must be re-create"));
 				ValidationResult = EDataValidationResult::Invalid;
