Substrate - Fixed warning message about not found override properties on material instance.

#rnx

diff --git a/Engine/Source/Runtime/Engine/Public/Materials/MaterialInstanceBasePropertyOverrides.h b/Engine/Source/Runtime/Engine/Public/Materials/MaterialInstanceBasePropertyOverrides.h
--- a/Engine/Source/Runtime/Engine/Public/Materials/MaterialInstanceBasePropertyOverrides.h
+++ b/Engine/Source/Runtime/Engine/Public/Materials/MaterialInstanceBasePropertyOverrides.h
@@ -74,7 +74,7 @@ struct FMaterialInstanceBasePropertyOverrides
 	uint8 TwoSided : 1;
 
 	/** Indicates that the material should be rendered as. */
-	UPROPERTY(EditAnywhere, Category = Material, meta = (editcondition = "bOverride_bThinSurface"))
+	UPROPERTY(EditAnywhere, Category = Material, meta = (editcondition = "bOverride_bIsThinSurface"))
 	uint8 bIsThinSurface : 1;
 
 	/** Whether the material should support a dithered LOD transition when used with the foliage system. */
