Substrate - removed not needed interface implementation from Switches nodes since the default implementation already prevent such node to be used with Substrate.

diff --git a/Engine/Source/Runtime/Engine/Private/Materials/MaterialExpressions.cpp b/Engine/Source/Runtime/Engine/Private/Materials/MaterialExpressions.cpp
--- a/Engine/Source/Runtime/Engine/Private/Materials/MaterialExpressions.cpp
+++ b/Engine/Source/Runtime/Engine/Private/Materials/MaterialExpressions.cpp
@@ -2385,7 +2385,7 @@ bool UMaterialExpression::IsUsingNewHLSLGenerator() const
 
 FSubstrateOperator* UMaterialExpression::SubstrateGenerateMaterialTopologyTree(class FMaterialCompiler* Compiler, class UMaterialExpression* Parent, int32 OutputIndex)
 {
-	Compiler->Errorf(TEXT("Missing SubstrateGenerateMaterialTopologyTree implementation for node %s."), *GetClass()->GetName());
+	Compiler->Errorf(TEXT("%s nodes type does not support generating/processing/flowing Substrate data."), *GetClass()->GetName());
 	return nullptr;
 }
 
@@ -21027,22 +21027,6 @@ int32 UMaterialExpressionShaderStageSwitch::Compile(class FMaterialCompiler* Com
 	}
 }
 
-bool UMaterialExpressionShaderStageSwitch::IsResultSubstrateMaterial(int32 OutputIndex)
-{
-	return false; // Not supported
-}
-
-void UMaterialExpressionShaderStageSwitch::GatherSubstrateMaterialInfo(FSubstrateMaterialInfo& SubstrateMaterialInfo, int32 OutputIndex)
-{
-	// Not supported
-}
-
-FSubstrateOperator* UMaterialExpressionShaderStageSwitch::SubstrateGenerateMaterialTopologyTree(class FMaterialCompiler* Compiler, class UMaterialExpression* Parent, int32 OutputIndex)
-{
-	Compiler->Errorf(TEXT("Substrate materials are only supported in pixel shaders: ShaderStageSwitch thus should not be plugged to convey Substrate material informations."));
-	return nullptr;
-}
-
 void UMaterialExpressionShaderStageSwitch::GetCaption(TArray<FString>& OutCaptions) const
 {
 	OutCaptions.Add(TEXT("Shader Stage Switch"));
@@ -21232,22 +21216,6 @@ uint32 UMaterialExpressionRayTracingQualitySwitch::GetInputType(int32 InputIndex
 {
 	return MCT_Unknown;
 }
-
-bool UMaterialExpressionRayTracingQualitySwitch::IsResultSubstrateMaterial(int32 OutputIndex)
-{
-	return false; // Not supported
-}
-
-void UMaterialExpressionRayTracingQualitySwitch::GatherSubstrateMaterialInfo(FSubstrateMaterialInfo& SubstrateMaterialInfo, int32 OutputIndex)
-{
-	// Not supported
-}
-
-FSubstrateOperator* UMaterialExpressionRayTracingQualitySwitch::SubstrateGenerateMaterialTopologyTree(class FMaterialCompiler* Compiler, class UMaterialExpression* Parent, int32 OutputIndex)
-{
-	Compiler->Errorf(TEXT("Substrate material topology must be statically define. We do not support topology update via dynamic evaluation such as `is raytracing or not`. Only input to BSDFs or Operators can be controled this way."));
-	return nullptr;
-}
 #endif // WITH_EDITOR
 
 //
@@ -21324,22 +21292,6 @@ uint32 UMaterialExpressionPathTracingQualitySwitch::GetInputType(int32 InputInde
 {
 	return MCT_Unknown;
 }
-
-bool UMaterialExpressionPathTracingQualitySwitch::IsResultSubstrateMaterial(int32 OutputIndex)
-{
-	return false; // Not supported
-}
-
-void UMaterialExpressionPathTracingQualitySwitch::GatherSubstrateMaterialInfo(FSubstrateMaterialInfo& SubstrateMaterialInfo, int32 OutputIndex)
-{
-	// Not supported
-}
-
-FSubstrateOperator* UMaterialExpressionPathTracingQualitySwitch::SubstrateGenerateMaterialTopologyTree(class FMaterialCompiler* Compiler, class UMaterialExpression* Parent, int32 OutputIndex)
-{
-	Compiler->Errorf(TEXT("Substrate material topology must be statically define. We do not support topology update via dynamic evaluation such as `is pathtracing or not`. Only input to BSDFs or Operators can be controled this way."));
-	return nullptr;
-}
 #endif // WITH_EDITOR
 
 
@@ -26213,21 +26165,6 @@ void UMaterialExpressionReflectionCapturePassSwitch::GetExpressionToolTip(TArray
 	ConvertToMultilineToolTip(TEXT("Allows material to define specialized behavior when being rendered into reflection capture views."), 40, OutToolTip);
 }
 
-bool UMaterialExpressionReflectionCapturePassSwitch::IsResultSubstrateMaterial(int32 OutputIndex)
-{
-	return false; // Not supported
-}
-
-void UMaterialExpressionReflectionCapturePassSwitch::GatherSubstrateMaterialInfo(FSubstrateMaterialInfo& SubstrateMaterialInfo, int32 OutputIndex)
-{
-	// Not supported
-}
-
-FSubstrateOperator* UMaterialExpressionReflectionCapturePassSwitch::SubstrateGenerateMaterialTopologyTree(class FMaterialCompiler* Compiler, class UMaterialExpression* Parent, int32 OutputIndex)
-{
-	Compiler->Errorf(TEXT("Substrate material topology must be statically define. We do not support topology update via dynamic evaluation such as `is reflection or not`. Only input to BSDFs or Operators can be controled this way."));
-	return nullptr;
-}
 #endif // WITH_EDITOR
 
 
diff --git a/Engine/Source/Runtime/Engine/Public/Materials/MaterialExpressionPathTracingQualitySwitch.h b/Engine/Source/Runtime/Engine/Public/Materials/MaterialExpressionPathTracingQualitySwitch.h
--- a/Engine/Source/Runtime/Engine/Public/Materials/MaterialExpressionPathTracingQualitySwitch.h
+++ b/Engine/Source/Runtime/Engine/Public/Materials/MaterialExpressionPathTracingQualitySwitch.h
@@ -31,10 +31,6 @@ class UMaterialExpressionPathTracingQualitySwitch : public UMaterialExpression
 	virtual uint32 GetOutputType(int32 OutputIndex) override { return MCT_Unknown; }
 
 	virtual bool GenerateHLSLExpression(FMaterialHLSLGenerator& Generator, UE::HLSLTree::FScope& Scope, int32 OutputIndex, UE::HLSLTree::FExpression const*& OutExpression) const override;
-
-	virtual bool IsResultSubstrateMaterial(int32 OutputIndex) override;
-	virtual void GatherSubstrateMaterialInfo(FSubstrateMaterialInfo& SubstrateMaterialInfo, int32 OutputIndex) override;
-	virtual FSubstrateOperator* SubstrateGenerateMaterialTopologyTree(class FMaterialCompiler* Compiler, class UMaterialExpression* Parent, int32 OutputIndex) override;
 #endif
 	//~ End UMaterialExpression Interface
 };
diff --git a/Engine/Source/Runtime/Engine/Public/Materials/MaterialExpressionRayTracingQualitySwitch.h b/Engine/Source/Runtime/Engine/Public/Materials/MaterialExpressionRayTracingQualitySwitch.h
--- a/Engine/Source/Runtime/Engine/Public/Materials/MaterialExpressionRayTracingQualitySwitch.h
+++ b/Engine/Source/Runtime/Engine/Public/Materials/MaterialExpressionRayTracingQualitySwitch.h
@@ -30,10 +30,6 @@ class UMaterialExpressionRayTracingQualitySwitch : public UMaterialExpression
 	virtual uint32 GetInputType(int32 InputIndex) override;
 	virtual uint32 GetOutputType(int32 OutputIndex) override { return MCT_Unknown; }
 
-	virtual bool IsResultSubstrateMaterial(int32 OutputIndex) override;
-	virtual void GatherSubstrateMaterialInfo(FSubstrateMaterialInfo& SubstrateMaterialInfo, int32 OutputIndex) override;
-	virtual FSubstrateOperator* SubstrateGenerateMaterialTopologyTree(class FMaterialCompiler* Compiler, class UMaterialExpression* Parent, int32 OutputIndex) override;
-
 	virtual bool GenerateHLSLExpression(FMaterialHLSLGenerator& Generator, UE::HLSLTree::FScope& Scope, int32 OutputIndex, UE::HLSLTree::FExpression const*& OutExpression) const override;
 #endif
 	//~ End UMaterialExpression Interface
diff --git a/Engine/Source/Runtime/Engine/Public/Materials/MaterialExpressionReflectionCapturePassSwitch.h b/Engine/Source/Runtime/Engine/Public/Materials/MaterialExpressionReflectionCapturePassSwitch.h
--- a/Engine/Source/Runtime/Engine/Public/Materials/MaterialExpressionReflectionCapturePassSwitch.h
+++ b/Engine/Source/Runtime/Engine/Public/Materials/MaterialExpressionReflectionCapturePassSwitch.h
@@ -26,10 +26,6 @@ class UMaterialExpressionReflectionCapturePassSwitch : public UMaterialExpressio
 	virtual void GetCaption(TArray<FString>& OutCaptions) const override;
 	virtual void GetExpressionToolTip(TArray<FString>& OutToolTip) override;
 
-	virtual bool IsResultSubstrateMaterial(int32 OutputIndex) override;
-	virtual void GatherSubstrateMaterialInfo(FSubstrateMaterialInfo& SubstrateMaterialInfo, int32 OutputIndex) override;
-	virtual FSubstrateOperator* SubstrateGenerateMaterialTopologyTree(class FMaterialCompiler* Compiler, class UMaterialExpression* Parent, int32 OutputIndex) override;
-
 	virtual bool GenerateHLSLExpression(FMaterialHLSLGenerator& Generator, UE::HLSLTree::FScope& Scope, int32 OutputIndex, UE::HLSLTree::FExpression const*& OutExpression) const override;
 #endif
 	//~ End UMaterialExpression Interface
diff --git a/Engine/Source/Runtime/Engine/Public/Materials/MaterialExpressionShaderStageSwitch.h b/Engine/Source/Runtime/Engine/Public/Materials/MaterialExpressionShaderStageSwitch.h
--- a/Engine/Source/Runtime/Engine/Public/Materials/MaterialExpressionShaderStageSwitch.h
+++ b/Engine/Source/Runtime/Engine/Public/Materials/MaterialExpressionShaderStageSwitch.h
@@ -34,10 +34,6 @@ class UMaterialExpressionShaderStageSwitch : public UMaterialExpression
 	virtual void GetExpressionToolTip(TArray<FString>& OutToolTip) override;
 	virtual FText GetKeywords() const override { return FText::FromString(TEXT("vertex pixel shader")); }
 
-	virtual bool IsResultSubstrateMaterial(int32 OutputIndex) override;
-	virtual void GatherSubstrateMaterialInfo(FSubstrateMaterialInfo& SubstrateMaterialInfo, int32 OutputIndex) override;
-	virtual FSubstrateOperator* SubstrateGenerateMaterialTopologyTree(class FMaterialCompiler* Compiler, class UMaterialExpression* Parent, int32 OutputIndex) override;
-
 	virtual bool GenerateHLSLExpression(FMaterialHLSLGenerator& Generator, UE::HLSLTree::FScope& Scope, int32 OutputIndex, UE::HLSLTree::FExpression const*& OutExpression) const override;
 #endif
 	//~ End UMaterialExpression Interface
