Use move for FStreamingGenerationActorDescViewMap::Emplace instead of copy.

diff --git a/Engine/Source/Runtime/Engine/Private/WorldPartition/WorldPartitionStreamingGeneration.cpp b/Engine/Source/Runtime/Engine/Private/WorldPartition/WorldPartitionStreamingGeneration.cpp
--- a/Engine/Source/Runtime/Engine/Private/WorldPartition/WorldPartitionStreamingGeneration.cpp
+++ b/Engine/Source/Runtime/Engine/Private/WorldPartition/WorldPartitionStreamingGeneration.cpp
@@ -408,9 +408,9 @@ TArray<const FStreamingGenerationActorDescView*> FStreamingGenerationActorDescVi
 	return Result;
 }
 
-FStreamingGenerationActorDescView* FStreamingGenerationActorDescViewMap::Emplace(const FGuid& InGuid, const FStreamingGenerationActorDescView& InActorDescView)
+FStreamingGenerationActorDescView* FStreamingGenerationActorDescViewMap::Emplace(const FGuid& InGuid, FStreamingGenerationActorDescView&& InActorDescView)
 {
-	FStreamingGenerationActorDescView* NewActorDescView = ActorDescViewList.Emplace_GetRef(MakeUnique<FStreamingGenerationActorDescView>(InActorDescView)).Get();
+	FStreamingGenerationActorDescView* NewActorDescView = ActorDescViewList.Emplace_GetRef(MakeUnique<FStreamingGenerationActorDescView>(MoveTemp(InActorDescView))).Get();
 	NewActorDescView->ActorDescViewMap = this;
 
 	const UClass* NativeClass = NewActorDescView->GetActorNativeClass();
diff --git a/Engine/Source/Runtime/Engine/Public/WorldPartition/WorldPartitionStreamingGeneration.h b/Engine/Source/Runtime/Engine/Public/WorldPartition/WorldPartitionStreamingGeneration.h
--- a/Engine/Source/Runtime/Engine/Public/WorldPartition/WorldPartitionStreamingGeneration.h
+++ b/Engine/Source/Runtime/Engine/Public/WorldPartition/WorldPartitionStreamingGeneration.h
@@ -160,7 +160,7 @@ public:
 	FStreamingGenerationActorDescViewMap& operator=(const FStreamingGenerationActorDescViewMap&) = delete;
 	FStreamingGenerationActorDescViewMap& operator=(FStreamingGenerationActorDescViewMap&&) = default;
 
-	ENGINE_API FStreamingGenerationActorDescView* Emplace(const FGuid& InActorGuid, const FStreamingGenerationActorDescView& InActorDescView);
+	ENGINE_API FStreamingGenerationActorDescView* Emplace(const FGuid& InActorGuid, FStreamingGenerationActorDescView&& InActorDescView);
 	ENGINE_API FStreamingGenerationActorDescView* Emplace(const FWorldPartitionActorDescInstance* InActorDescInstance);
 
 	PRAGMA_DISABLE_DEPRECATION_WARNINGS
