Substrate - static switches nodes with dynamic branch cannot compile with Substrate Data.

#rnx

diff --git a/Engine/Source/Runtime/Engine/Private/Materials/MaterialExpressions.cpp b/Engine/Source/Runtime/Engine/Private/Materials/MaterialExpressions.cpp
--- a/Engine/Source/Runtime/Engine/Private/Materials/MaterialExpressions.cpp
+++ b/Engine/Source/Runtime/Engine/Private/Materials/MaterialExpressions.cpp
@@ -10098,6 +10098,12 @@ void UMaterialExpressionStaticSwitchParameter::GatherSubstrateMaterialInfo(FSubs
 
 FSubstrateOperator* UMaterialExpressionStaticSwitchParameter::SubstrateGenerateMaterialTopologyTree(class FMaterialCompiler* Compiler, class UMaterialExpression* Parent, int32 OutputIndex)
 {
+	if (DynamicBranch)
+	{
+		Compiler->Errorf(TEXT("Static Switch nodes processing Substrate data do not support dynamic branching. The compiler must know the topology when translating HLSL (different branches could have different topologies)."));
+		return nullptr;
+	}
+
 	FExpressionInput* EffectiveInput = GetEffectiveInput(Compiler);
 	if (EffectiveInput && EffectiveInput->Expression)
 	{
