Fix incorrect type return by vertex normal code generation.

diff --git a/Engine/Source/Runtime/Engine/Private/Materials/HLSLMaterialTranslator.cpp b/Engine/Source/Runtime/Engine/Private/Materials/HLSLMaterialTranslator.cpp
--- a/Engine/Source/Runtime/Engine/Private/Materials/HLSLMaterialTranslator.cpp
+++ b/Engine/Source/Runtime/Engine/Private/Materials/HLSLMaterialTranslator.cpp
@@ -11275,7 +11275,7 @@ int32 FHLSLMaterialTranslator::VertexNormal()
 	}
 	else
 	{
-		return AddInlinedCodeChunk(MCT_Float4, *FiniteCode);
+		return AddInlinedCodeChunk(MCT_Float3, *FiniteCode);
 	}
 }
 
