Fix shape index resolution in collision event handling. We swap the hit data for many of the parameters in the collision event, but missed the shape index. This makes hits against more complex objects with many shapes potentially fail to identify the shape that was hit.

diff --git a/Engine/Source/Runtime/Engine/Private/PhysicsEngine/Experimental/PhysScene_Chaos.cpp b/Engine/Source/Runtime/Engine/Private/PhysicsEngine/Experimental/PhysScene_Chaos.cpp
--- a/Engine/Source/Runtime/Engine/Private/PhysicsEngine/Experimental/PhysScene_Chaos.cpp
+++ b/Engine/Source/Runtime/Engine/Private/PhysicsEngine/Experimental/PhysScene_Chaos.cpp
@@ -992,8 +992,11 @@ FORCEINLINE void FPhysScene_Chaos::HandleEachCollisionEvent(const TArray<int32>&
 
 			NotifyInfo.SolverTime = CollisionDataItem.SolverTime;
 
-			NotifyInfo.Info0.SetFrom(GetBodyInstanceFromProxyAndShape(PhysicsProxy0, CollisionDataItem.ShapeIndex1), DeltaVelocity1);
-			NotifyInfo.Info1.SetFrom(GetBodyInstanceFromProxyAndShape(PhysicsProxy1, CollisionDataItem.ShapeIndex2), DeltaVelocity2);
+			const int32 ShapeIdx0 = bSwapOrder ? CollisionDataItem.ShapeIndex2 : CollisionDataItem.ShapeIndex1;
+			const int32 ShapeIdx1 = bSwapOrder ? CollisionDataItem.ShapeIndex1 : CollisionDataItem.ShapeIndex2;
+
+			NotifyInfo.Info0.SetFrom(GetBodyInstanceFromProxyAndShape(PhysicsProxy0, ShapeIdx0), DeltaVelocity1);
+			NotifyInfo.Info1.SetFrom(GetBodyInstanceFromProxyAndShape(PhysicsProxy1, ShapeIdx1), DeltaVelocity2);
 
 			// in some case ( like with geometry collections ) we don't have a body instance so the component part will null, we need to handle that 
 			if (NotifyInfo.Info0.Component == nullptr)
