Exposed EDateTimeStyle to BP Date/Time Text functions

diff --git a/Engine/Source/Runtime/Core/Public/Internationalization/Text.h b/Engine/Source/Runtime/Core/Public/Internationalization/Text.h
--- a/Engine/Source/Runtime/Core/Public/Internationalization/Text.h
+++ b/Engine/Source/Runtime/Core/Public/Internationalization/Text.h
@@ -120,9 +120,10 @@ CORE_API bool LexTryParseString(ETextGender& OutValue, const TCHAR* Buffer);
 CORE_API void LexFromString(ETextGender& OutValue, const TCHAR* Buffer);
 CORE_API const TCHAR* LexToString(ETextGender InValue);
 
+/** Redeclared in KismetTextLibrary for meta-data extraction purposes, be sure to update there as well */
 namespace EDateTimeStyle
 {
-	enum Type
+	enum Type : int
 	{
 		Default,
 		Short,
diff --git a/Engine/Source/Runtime/Engine/Classes/Kismet/KismetTextLibrary.h b/Engine/Source/Runtime/Engine/Classes/Kismet/KismetTextLibrary.h
--- a/Engine/Source/Runtime/Engine/Classes/Kismet/KismetTextLibrary.h
+++ b/Engine/Source/Runtime/Engine/Classes/Kismet/KismetTextLibrary.h
@@ -40,6 +40,20 @@ enum class ETextGender : uint8
 	Neuter,
 };
 
+UENUM(BlueprintType)
+namespace EDateTimeStyle
+{
+	enum Type : int
+	{
+		Default,
+		Short,
+		Medium,
+		Long,
+		Full,
+		Custom UMETA(Hidden), // Internal use only
+	};
+}
+
 UENUM(BlueprintType)
 namespace EFormatArgumentType
 {
@@ -248,27 +262,27 @@ class UKismetTextLibrary : public UBlueprintFunctionLibrary
 
 	/** Converts a passed in date & time to a text, formatted as a date using an invariant timezone. This will use the given date & time as-is, so it's assumed to already be in the correct timezone. */
 	UFUNCTION(BlueprintPure, meta=(DisplayName = "As Date", AdvancedDisplay = "1"), Category="Utilities|Text")
-	static ENGINE_API FText AsDate_DateTime(const FDateTime& InDateTime);
+	static ENGINE_API FText AsDate_DateTime(const FDateTime& InDateTime, TEnumAsByte<EDateTimeStyle::Type> InDateStyle = EDateTimeStyle::Default);
 
 	/** Converts a passed in date & time to a text, formatted as a date using the given timezone (default is the local timezone). This will convert the given date & time from UTC to the given timezone (taking into account DST). */
 	UFUNCTION(BlueprintPure, meta=(DisplayName = "As Date (from UTC)", AdvancedDisplay = "1"), Category="Utilities|Text")
-	static ENGINE_API FText AsTimeZoneDate_DateTime(const FDateTime& InDateTime, const FString& InTimeZone = TEXT(""));
+	static ENGINE_API FText AsTimeZoneDate_DateTime(const FDateTime& InDateTime, const FString& InTimeZone = TEXT(""), TEnumAsByte<EDateTimeStyle::Type> InDateStyle = EDateTimeStyle::Default);
 
 	/** Converts a passed in date & time to a text, formatted as a date & time using an invariant timezone. This will use the given date & time as-is, so it's assumed to already be in the correct timezone. */
 	UFUNCTION(BlueprintPure, meta=(DisplayName = "As DateTime", AdvancedDisplay = "1"), Category="Utilities|Text")
-	static ENGINE_API FText AsDateTime_DateTime(const FDateTime& In);
+	static ENGINE_API FText AsDateTime_DateTime(const FDateTime& In, TEnumAsByte<EDateTimeStyle::Type> InDateStyle = EDateTimeStyle::Default, TEnumAsByte<EDateTimeStyle::Type> InTimeStyle = EDateTimeStyle::Default);
 
 	/** Converts a passed in date & time to a text, formatted as a date & time using the given timezone (default is the local timezone). This will convert the given date & time from UTC to the given timezone (taking into account DST). */
 	UFUNCTION(BlueprintPure, meta=(DisplayName = "As DateTime (from UTC)", AdvancedDisplay = "1"), Category="Utilities|Text")
-	static ENGINE_API FText AsTimeZoneDateTime_DateTime(const FDateTime& InDateTime, const FString& InTimeZone = TEXT(""));
+	static ENGINE_API FText AsTimeZoneDateTime_DateTime(const FDateTime& InDateTime, const FString& InTimeZone = TEXT(""), TEnumAsByte<EDateTimeStyle::Type> InDateStyle = EDateTimeStyle::Default, TEnumAsByte<EDateTimeStyle::Type> InTimeStyle = EDateTimeStyle::Default);
 
 	/** Converts a passed in date & time to a text, formatted as a time using an invariant timezone. This will use the given date & time as-is, so it's assumed to already be in the correct timezone. */
 	UFUNCTION(BlueprintPure, meta=(DisplayName = "As Time", AdvancedDisplay = "1"), Category="Utilities|Text")
-	static ENGINE_API FText AsTime_DateTime(const FDateTime& In);
+	static ENGINE_API FText AsTime_DateTime(const FDateTime& In, TEnumAsByte<EDateTimeStyle::Type> InTimeStyle = EDateTimeStyle::Default);
 
 	/** Converts a passed in date & time to a text, formatted as a time using the given timezone (default is the local timezone). This will convert the given date & time from UTC to the given timezone (taking into account DST). */
 	UFUNCTION(BlueprintPure, meta=(DisplayName = "As Time (from UTC)", AdvancedDisplay = "1"), Category="Utilities|Text")
-	static ENGINE_API FText AsTimeZoneTime_DateTime(const FDateTime& InDateTime, const FString& InTimeZone = TEXT(""));
+	static ENGINE_API FText AsTimeZoneTime_DateTime(const FDateTime& InDateTime, const FString& InTimeZone = TEXT(""), TEnumAsByte<EDateTimeStyle::Type> InTimeStyle = EDateTimeStyle::Default);
 
 	/** Converts a passed in time span to a text, formatted as a time span */
 	UFUNCTION(BlueprintPure, meta=(DisplayName = "As Timespan", AdvancedDisplay = "1"), Category="Utilities|Text")
diff --git a/Engine/Source/Runtime/Engine/Private/KismetTextLibrary.cpp b/Engine/Source/Runtime/Engine/Private/KismetTextLibrary.cpp
--- a/Engine/Source/Runtime/Engine/Private/KismetTextLibrary.cpp
+++ b/Engine/Source/Runtime/Engine/Private/KismetTextLibrary.cpp
@@ -272,34 +272,34 @@ FText UKismetTextLibrary::AsPercent_Float(float Value, TEnumAsByte<ERoundingMode
 	return FText::AsPercent(Value, &NumberFormatOptions);
 }
 
-FText UKismetTextLibrary::AsDate_DateTime(const FDateTime& InDateTime)
+FText UKismetTextLibrary::AsDate_DateTime(const FDateTime& InDateTime, TEnumAsByte<EDateTimeStyle::Type> InDateStyle)
 {
-	return FText::AsDate(InDateTime, EDateTimeStyle::Default, FText::GetInvariantTimeZone());
+	return FText::AsDate(InDateTime, InDateStyle, FText::GetInvariantTimeZone());
 }
 
-FText UKismetTextLibrary::AsTimeZoneDate_DateTime(const FDateTime& InDateTime, const FString& InTimeZone)
+FText UKismetTextLibrary::AsTimeZoneDate_DateTime(const FDateTime& InDateTime, const FString& InTimeZone, TEnumAsByte<EDateTimeStyle::Type> InDateStyle)
 {
-	return FText::AsDate(InDateTime, EDateTimeStyle::Default, InTimeZone);
+	return FText::AsDate(InDateTime, InDateStyle, InTimeZone);
 }
 
-FText UKismetTextLibrary::AsDateTime_DateTime(const FDateTime& InDateTime)
+FText UKismetTextLibrary::AsDateTime_DateTime(const FDateTime& InDateTime, TEnumAsByte<EDateTimeStyle::Type> InDateStyle, TEnumAsByte<EDateTimeStyle::Type> InTimeStyle)
 {
-	return FText::AsDateTime(InDateTime, EDateTimeStyle::Default, EDateTimeStyle::Default, FText::GetInvariantTimeZone());
+	return FText::AsDateTime(InDateTime, InDateStyle, InTimeStyle, FText::GetInvariantTimeZone());
 }
 
-FText UKismetTextLibrary::AsTimeZoneDateTime_DateTime(const FDateTime& InDateTime, const FString& InTimeZone)
+FText UKismetTextLibrary::AsTimeZoneDateTime_DateTime(const FDateTime& InDateTime, const FString& InTimeZone, TEnumAsByte<EDateTimeStyle::Type> InDateStyle, TEnumAsByte<EDateTimeStyle::Type> InTimeStyle)
 {
-	return FText::AsDateTime(InDateTime, EDateTimeStyle::Default, EDateTimeStyle::Default, InTimeZone);
+	return FText::AsDateTime(InDateTime, InDateStyle, InTimeStyle, InTimeZone);
 }
 
-FText UKismetTextLibrary::AsTime_DateTime(const FDateTime& InDateTime)
+FText UKismetTextLibrary::AsTime_DateTime(const FDateTime& InDateTime, TEnumAsByte<EDateTimeStyle::Type> InTimeStyle)
 {
-	return FText::AsTime(InDateTime, EDateTimeStyle::Default, FText::GetInvariantTimeZone());
+	return FText::AsTime(InDateTime, InTimeStyle, FText::GetInvariantTimeZone());
 }
 
-FText UKismetTextLibrary::AsTimeZoneTime_DateTime(const FDateTime& InDateTime, const FString& InTimeZone)
+FText UKismetTextLibrary::AsTimeZoneTime_DateTime(const FDateTime& InDateTime, const FString& InTimeZone, TEnumAsByte<EDateTimeStyle::Type> InTimeStyle)
 {
-	return FText::AsTime(InDateTime, EDateTimeStyle::Default, InTimeZone);
+	return FText::AsTime(InDateTime, InTimeStyle, InTimeZone);
 }
 
 FText UKismetTextLibrary::AsTimespan_Timespan(const FTimespan& InTimespan)
