Fixed World Partition Actor Descriptor patching to properly support derived descriptor classes.

diff --git a/Engine/Source/Runtime/Engine/Private/WorldPartition/WorldPartitionActorDesc.cpp b/Engine/Source/Runtime/Engine/Private/WorldPartition/WorldPartitionActorDesc.cpp
--- a/Engine/Source/Runtime/Engine/Private/WorldPartition/WorldPartitionActorDesc.cpp
+++ b/Engine/Source/Runtime/Engine/Private/WorldPartition/WorldPartitionActorDesc.cpp
@@ -209,14 +209,19 @@ void FWorldPartitionActorDesc::Init(const AActor* InActor)
 	Container = nullptr;
 }
 
-void FWorldPartitionActorDesc::Init(const FWorldPartitionActorDescInitData& DescData)
+void FWorldPartitionActorDesc::InitTransientProperties(const FWorldPartitionActorDescInitData& DescData)
 {
 	ActorPackage = DescData.PackageName;
 	ActorPath = DescData.ActorPath;
-	ActorNativeClass = DescData.NativeClass;
-	NativeClass = *DescData.NativeClass->GetPathName();
+	ActorNativeClass = DescData.NativeClass ? DescData.NativeClass : AActor::StaticClass();
+	NativeClass = *ActorNativeClass->GetPathName();
 	ActorName = *FPaths::GetExtension(ActorPath.ToString());
 	ActorNameString = ActorName.ToString();
+}
+
+void FWorldPartitionActorDesc::Init(const FWorldPartitionActorDescInitData& DescData)
+{
+	InitTransientProperties(DescData);
 
 	auto DeprecateClass = [this](FArchive& Archive)
 	{
@@ -303,9 +308,11 @@ void FWorldPartitionActorDesc::Patch(const FWorldPartitionActorDescInitData& Des
 	TArray<uint8> PatchedPayloadData;
 	FMemoryWriter PatchedPayloadAr(PatchedPayloadData, true);
 
-	TUniquePtr<FWorldPartitionActorDesc> ActorDesc(AActor::StaticCreateClassActorDesc(DescData.NativeClass ? DescData.NativeClass : AActor::StaticClass()));
-	FActorDescArchivePatcher ActorDescAr(MetadataAr, ActorDesc.Get(), PatchedPayloadAr, InAssetDataPatcher);	
-	FTopLevelAssetPath ActorClassPath(TEXT("/Script/Engine.Actor"));
+	UClass* NativeClass = DescData.NativeClass ? DescData.NativeClass : AActor::StaticClass();
+	TUniquePtr<FWorldPartitionActorDesc> ActorDesc(AActor::StaticCreateClassActorDesc(NativeClass));
+	ActorDesc->InitTransientProperties(DescData);
+	FActorDescArchivePatcher ActorDescAr(MetadataAr, ActorDesc.Get(), PatchedPayloadAr, InAssetDataPatcher);
+	FTopLevelAssetPath ActorClassPath(NativeClass->GetPathName());
 	ActorDescAr.Init(ActorClassPath);
 
 	ActorDesc->Serialize(ActorDescAr);
diff --git a/Engine/Source/Runtime/Engine/Public/WorldPartition/WorldPartitionActorDesc.h b/Engine/Source/Runtime/Engine/Public/WorldPartition/WorldPartitionActorDesc.h
--- a/Engine/Source/Runtime/Engine/Public/WorldPartition/WorldPartitionActorDesc.h
+++ b/Engine/Source/Runtime/Engine/Public/WorldPartition/WorldPartitionActorDesc.h
@@ -383,6 +383,7 @@ public:
 	ENGINE_API bool IsListedInSceneOutliner() const;
 
 protected:
+	void InitTransientProperties(const FWorldPartitionActorDescInitData& DescData);
 	virtual bool GetChildContainerInstance(const FWorldPartitionActorDescInstance* InActorDescInstance, FContainerInstance& OutContainerInstance) const { return false; }
 	virtual UActorDescContainerInstance* CreateChildContainerInstance(const FWorldPartitionActorDescInstance* InActorDescInstance) const { return nullptr; }
 	virtual UWorldPartition* GetLoadedChildWorldPartition(const FWorldPartitionActorDescInstance* InActorDescInstance) const { return nullptr; }
