[shaders] add a cvar to optionally disambiguate shader debug info folders when overriding the root debug info path.

diff --git a/Engine/Source/Runtime/Engine/Private/ShaderCompiler/ShaderCompiler.cpp b/Engine/Source/Runtime/Engine/Private/ShaderCompiler/ShaderCompiler.cpp
--- a/Engine/Source/Runtime/Engine/Private/ShaderCompiler/ShaderCompiler.cpp
+++ b/Engine/Source/Runtime/Engine/Private/ShaderCompiler/ShaderCompiler.cpp
@@ -516,6 +516,13 @@ static TAutoConsoleVariable<FString> CVarShaderOverrideDebugDir(
 	TEXT("Empty: use default location Saved\\ShaderDebugInfo.\n"),
 	ECVF_ReadOnly);
 
+static TAutoConsoleVariable<bool> CVarDisambiguateShaderDebugDir(
+	TEXT("r.DisambiguateShaderDebugDir"),
+	false,
+	TEXT("If true, appends a folder containing the full project path with directory separators/drive qualifiers replaced with _ to the root debug info folder.\n")
+	TEXT("Intended for use in conjunction with r.OverrideShaderDebugDir to avoid shaderdebuginfo output clashing across workspaces/projects."),
+	ECVF_ReadOnly);
+
 static TAutoConsoleVariable<int32> CVarShadersValidation(
 	TEXT("r.Shaders.Validation"),
 	1,
@@ -810,6 +817,16 @@ FShaderCompilingManager::FShaderCompilingManager() :
 	{
 		AbsoluteDebugInfoDirectory = OverrideShaderDebugDir;
 	}
+
+	if (CVarDisambiguateShaderDebugDir.GetValueOnAnyThread())
+	{
+		FString AppendFolder = FPaths::ConvertRelativePathToFull(FPaths::ProjectDir());
+		FPaths::NormalizeDirectoryName(AppendFolder);
+		AppendFolder.ReplaceInline(TEXT(":/"), TEXT("_"));
+		AppendFolder.ReplaceCharInline('/', '_');
+		AbsoluteDebugInfoDirectory /= AppendFolder;
+	}
+
 	FPaths::NormalizeDirectoryName(AbsoluteDebugInfoDirectory);
 	AbsoluteShaderDebugInfoDirectory = AbsoluteDebugInfoDirectory;
 
