Deprecated ULevelInstanceSettings::DisableLevelInstanceSupport and replaced with SetIsLevelInstanceDisabled to allow toggling the option.

diff --git a/Engine/Source/Runtime/Engine/Public/LevelInstance/LevelInstanceSettings.h b/Engine/Source/Runtime/Engine/Public/LevelInstance/LevelInstanceSettings.h
--- a/Engine/Source/Runtime/Engine/Public/LevelInstance/LevelInstanceSettings.h
+++ b/Engine/Source/Runtime/Engine/Public/LevelInstance/LevelInstanceSettings.h
@@ -20,7 +20,9 @@ public:
 #if WITH_EDITOR
 	bool IsPropertyOverrideEnabled() const;
 
+	UE_DEPRECATED(5.5, "Use SetIsLevelInstanceDisabled instead.")
 	void DisableLevelInstanceSupport() { bIsLevelInstanceDisabled = true; }
+	void SetIsLevelInstanceDisabled(bool bValue) { bIsLevelInstanceDisabled = bValue; }
 	bool IsLevelInstanceDisabled() const { return bIsLevelInstanceDisabled; }
 private:
 	friend class ULevelInstanceSubsystem;
