Deprecate the disallow mesh paint flag and split into two flags for enabling vertex and texture mesh painting.
The flags are editable per instance.
Also added a non editable bSupportMeshPainting flag which is intended for disallowing per component type.
Explicitly disallow painting on ISM using the new flag which has the advantage over the old flag of not being editable.
Now we can't get into a confusing state where user turns on mesh painting on ISM (even though that isn't possible).

diff --git a/Engine/Plugins/MeshPainting/Source/MeshPaintingToolset/Public/MeshPaintStaticMeshAdapter.h b/Engine/Plugins/MeshPainting/Source/MeshPaintingToolset/Public/MeshPaintStaticMeshAdapter.h
--- a/Engine/Plugins/MeshPainting/Source/MeshPaintingToolset/Public/MeshPaintStaticMeshAdapter.h
+++ b/Engine/Plugins/MeshPainting/Source/MeshPaintingToolset/Public/MeshPaintStaticMeshAdapter.h
@@ -28,7 +28,7 @@ public:
 	virtual void OnRemoved() override {}
 	virtual bool IsValid() const override { return StaticMeshComponent.IsValid() && ReferencedStaticMesh && StaticMeshComponent->GetStaticMesh() == ReferencedStaticMesh; }
 	virtual bool SupportsTexturePaint() const override { return true; }
-	virtual bool SupportsVertexPaint() const override { return StaticMeshComponent.IsValid() && !StaticMeshComponent->bDisallowMeshPaintPerInstance; }
+	virtual bool SupportsVertexPaint() const override { return StaticMeshComponent.IsValid() && StaticMeshComponent->CanMeshPaintVertexColors(); }
 	virtual bool LineTraceComponent(struct FHitResult& OutHit, const FVector Start, const FVector End, const struct FCollisionQueryParams& Params) const override;	
 	virtual void QueryPaintableTextures(int32 MaterialIndex, int32& OutDefaultIndex, TArray<struct FPaintableTexture>& InOutTextureList) override;
 	virtual void ApplyOrRemoveTextureOverride(UTexture* SourceTexture, UTexture* OverrideTexture) const override;
diff --git a/Engine/Source/Editor/MeshPaint/Public/MeshPaintStaticMeshAdapter.h b/Engine/Source/Editor/MeshPaint/Public/MeshPaintStaticMeshAdapter.h
--- a/Engine/Source/Editor/MeshPaint/Public/MeshPaintStaticMeshAdapter.h
+++ b/Engine/Source/Editor/MeshPaint/Public/MeshPaintStaticMeshAdapter.h
@@ -31,7 +31,7 @@ public:
 	virtual void OnRemoved() override;
 	virtual bool IsValid() const override { return StaticMeshComponent && StaticMeshComponent->GetStaticMesh() == ReferencedStaticMesh; }
 	virtual bool SupportsTexturePaint() const override { return true; }
-	virtual bool SupportsVertexPaint() const override { return StaticMeshComponent && !StaticMeshComponent->bDisallowMeshPaintPerInstance; }
+	virtual bool SupportsVertexPaint() const override { return StaticMeshComponent && StaticMeshComponent->CanMeshPaintVertexColors(); }
 	virtual bool LineTraceComponent(struct FHitResult& OutHit, const FVector Start, const FVector End, const struct FCollisionQueryParams& Params) const override;	
 	virtual void QueryPaintableTextures(int32 MaterialIndex, int32& OutDefaultIndex, TArray<struct FPaintableTexture>& InOutTextureList) override;
 	virtual void ApplyOrRemoveTextureOverride(UTexture* SourceTexture, UTexture* OverrideTexture) const override;
diff --git a/Engine/Source/Runtime/Engine/Classes/Components/StaticMeshComponent.h b/Engine/Source/Runtime/Engine/Classes/Components/StaticMeshComponent.h
--- a/Engine/Source/Runtime/Engine/Classes/Components/StaticMeshComponent.h
+++ b/Engine/Source/Runtime/Engine/Classes/Components/StaticMeshComponent.h
@@ -186,6 +186,9 @@ protected:
 	/** Whether mip callbacks have been registered and need to be removed on destroy */
 	uint8 bMipLevelCallbackRegistered : 1;
 
+	/** Whether mesh painting is supported. This is can be set by derived component types that don't support mesh painting. */
+	uint8 bSupportMeshPainting : 1;
+
 public:
 
 #if WITH_EDITORONLY_DATA
@@ -232,9 +235,9 @@ public:
 	UPROPERTY(transient)
 	uint8 bForceNavigationObstacle : 1;
 
-	/** If true, vertex mesh painting is disallowed on this instance. Set if vertex colors are overridden in a construction script. */
-	UPROPERTY(EditAnywhere, BlueprintReadWrite, Category = "Mesh Painting")
-	uint8 bDisallowMeshPaintPerInstance : 1;
+	/** Deprecated. Use bEnableVertexColorMeshPainting instead. */
+	UPROPERTY()
+	uint8 bDisallowMeshPaintPerInstance_DEPRECATED : 1;
 
 #if STATICMESH_ENABLE_DEBUG_RENDERING
 	/** Draw mesh collision if used for complex collision */
@@ -308,13 +311,32 @@ public:
 	 */
 	UPROPERTY(EditAnywhere, AdvancedDisplay, BlueprintReadOnly, Category=Lighting)
 	uint8 bReverseCulling : 1;
-	
+
+	/** 
+	 * If false, vertex color mesh painting is disabled on this instance. 
+	 * This may be set to false by blueprint functions that override vertex colors in construction script.
+	 */
+	UPROPERTY(EditAnywhere, AdvancedDisplay, BlueprintReadWrite, Category = "Mesh Painting")
+	uint8 bEnableVertexColorMeshPainting : 1;
+
+	/** If false, texture color mesh painting is disabled on this instance. */
+	UPROPERTY(EditAnywhere, AdvancedDisplay, BlueprintReadWrite, Category = "Mesh Painting")
+	uint8 bEnableTextureColorMeshPainting : 1;
+
 	/** Whether to override the MeshPaintTextureCoordinateIndex set on the static mesh. */
 	UPROPERTY(EditAnywhere, Category = "Mesh Painting", meta=(InlineEditConditionToggle))
 	uint8 bOverrideMeshPaintTextureCoordinateIndex : 1;
 
+	/** The overriden coordinate index to use when texture color painting on this mesh. */
+	UPROPERTY(EditAnywhere, Category = "Mesh Painting", meta=(UIMin = "0", UIMax = "3", editcondition = "bOverrideMeshPaintTextureCoordinateIndex"))
+	int32 OverridenMeshPaintTextureCoordinateIndex;
+
+	/** Light map resolution to use on this component, used if bOverrideLightMapRes is true and there is a valid StaticMesh. */
+	UPROPERTY(EditAnywhere, BlueprintReadOnly, Category=Lighting, meta=(ClampMax = 4096, editcondition="bOverrideLightMapRes") )
+	int32 OverriddenLightMapRes;
+
 #if WITH_EDITORONLY_DATA
-	/** Texture containing mesh painting for this mesh component. */
+	/** Texture containing texture color mesh painting for this mesh component. */
 	UPROPERTY(VisibleAnywhere, Category = "Mesh Painting")
 	TObjectPtr<UTexture> MeshPaintTexture;
 #endif
@@ -327,14 +349,6 @@ public:
 	UPROPERTY(Transient)
 	TObjectPtr<UTexture> MeshPaintTextureOverride;
 
-	/** The overriden coordinate index to use when painting the MeshPaintTexture on this mesh. */
-	UPROPERTY(EditAnywhere, Category = "Mesh Painting", meta=(UIMin = "0", UIMax = "3", editcondition = "bOverrideMeshPaintTextureCoordinateIndex"))
-	int32 OverridenMeshPaintTextureCoordinateIndex;
-
-	/** Light map resolution to use on this component, used if bOverrideLightMapRes is true and there is a valid StaticMesh. */
-	UPROPERTY(EditAnywhere, BlueprintReadOnly, Category=Lighting, meta=(ClampMax = 4096, editcondition="bOverrideLightMapRes") )
-	int32 OverriddenLightMapRes;
-
 	/** 
 	 * Controls how dark the dynamic indirect shadow can be.
 	 */
@@ -739,6 +753,11 @@ public:
 
 	ENGINE_API UMaterialInterface* GetNaniteAuditMaterial(int32 MaterialIndex) const;
 
+	/* Returns true if mesh vertex color painting is supported on this component. */
+	ENGINE_API int32 CanMeshPaintVertexColors() const { return bSupportMeshPainting && bEnableVertexColorMeshPainting; }
+	/* Returns true if mesh texture color painting is supported on this component. */
+	ENGINE_API int32 CanMeshPaintTextureColors() const { return bSupportMeshPainting && bEnableTextureColorMeshPainting; }
+
 private:
 	/** Initializes the resources used by the static mesh component. */
 	ENGINE_API void InitResources();
diff --git a/Engine/Source/Runtime/Engine/Private/Components/StaticMeshComponent.cpp b/Engine/Source/Runtime/Engine/Private/Components/StaticMeshComponent.cpp
--- a/Engine/Source/Runtime/Engine/Private/Components/StaticMeshComponent.cpp
+++ b/Engine/Source/Runtime/Engine/Private/Components/StaticMeshComponent.cpp
@@ -206,7 +206,9 @@ UStaticMeshComponent::UStaticMeshComponent(const FObjectInitializer& ObjectIniti
 	bHasCustomNavigableGeometry = EHasCustomNavigableGeometry::Yes;
 	bOverrideNavigationExport = false;
 	bForceNavigationObstacle = true;
-	bDisallowMeshPaintPerInstance = false;
+	bSupportMeshPainting = true;
+	bEnableVertexColorMeshPainting = true;
+	bEnableTextureColorMeshPainting = true;
 	bForceNaniteForMasked = false;
 	bDisallowNanite = false;
 	bForceDisableNanite = false;
@@ -2079,6 +2081,16 @@ bool UStaticMeshComponent::CanEditChange(const FProperty* InProperty) const
 		{
 			return bOverrideDistanceFieldSelfShadowBias && bAffectDistanceFieldLighting;
 		}
+
+		if (PropertyName == GET_MEMBER_NAME_STRING_CHECKED(UStaticMeshComponent, bEnableVertexColorMeshPainting))
+		{
+			return bSupportMeshPainting;
+		}
+
+		if (PropertyName == GET_MEMBER_NAME_STRING_CHECKED(UStaticMeshComponent, bEnableTextureColorMeshPainting))
+		{
+			return bSupportMeshPainting;
+		}
 	}
 
 	return Super::CanEditChange(InProperty);
@@ -2198,6 +2210,12 @@ void UStaticMeshComponent::PostLoad()
 		}
 	}
 
+	// Handle deprecation of bDisallowMeshPaintPerInstance -> bEnableVertexColorMeshPainting.
+	if (bDisallowMeshPaintPerInstance_DEPRECATED)
+	{
+		bEnableVertexColorMeshPainting = false;
+	}
+
 	// Legacy content may contain a lightmap resolution of 0, which was valid when vertex lightmaps were supported, but not anymore with only texture lightmaps
 	OverriddenLightMapRes = FMath::Max(OverriddenLightMapRes, 4);
 
@@ -2981,7 +2999,7 @@ void UStaticMeshComponent::ApplyComponentInstanceData(FStaticMeshComponentInstan
 		}
 	}
 
-	if (!bDisallowMeshPaintPerInstance)
+	if (CanMeshPaintVertexColors())
 	{
 		FComponentReregisterContext ReregisterStaticMesh(this);
 		StaticMeshInstanceData->ApplyVertexColorData(this);
diff --git a/Engine/Source/Runtime/Engine/Private/InstancedStaticMesh.cpp b/Engine/Source/Runtime/Engine/Private/InstancedStaticMesh.cpp
--- a/Engine/Source/Runtime/Engine/Private/InstancedStaticMesh.cpp
+++ b/Engine/Source/Runtime/Engine/Private/InstancedStaticMesh.cpp
@@ -2038,7 +2038,11 @@ UInstancedStaticMeshComponent::UInstancedStaticMeshComponent(const FObjectInitia
 	Mobility = EComponentMobility::Movable;
 	BodyInstance.bSimulatePhysics = false;
 
-	bDisallowMeshPaintPerInstance = true;
+	// Can't mesh paint on instanced mesh.
+	bSupportMeshPainting = false;
+	bEnableVertexColorMeshPainting = false;
+	bEnableTextureColorMeshPainting = false;
+
 	bMultiBodyOverlap = true;
 
 	bUseGpuLodSelection = true;
diff --git a/Engine/Source/Runtime/Engine/Private/MeshVertexPainter/MeshVertexPainter.cpp b/Engine/Source/Runtime/Engine/Private/MeshVertexPainter/MeshVertexPainter.cpp
--- a/Engine/Source/Runtime/Engine/Private/MeshVertexPainter/MeshVertexPainter.cpp
+++ b/Engine/Source/Runtime/Engine/Private/MeshVertexPainter/MeshVertexPainter.cpp
@@ -42,7 +42,9 @@ void FMeshVertexPainter::PaintVerticesSingleColor(UStaticMeshComponent* StaticMe
 	StaticMeshComponent->CachePaintedDataIfNecessary();
 #endif
 	StaticMeshComponent->MarkRenderStateDirty();
-	StaticMeshComponent->bDisallowMeshPaintPerInstance = true;
+
+	// Explicitly disable the mesh paint tool on the component to prevent stomping the vertex color.
+	StaticMeshComponent->bEnableVertexColorMeshPainting = false;
 }
 
 void FMeshVertexPainter::PaintVerticesLerpAlongAxis(UStaticMeshComponent* StaticMeshComponent, const FLinearColor& StartColor, const FLinearColor& EndColor, EVertexPaintAxis Axis, bool bConvertToSRGB)
@@ -93,7 +95,9 @@ void FMeshVertexPainter::PaintVerticesLerpAlongAxis(UStaticMeshComponent* Static
 	StaticMeshComponent->CachePaintedDataIfNecessary();
 #endif
 	StaticMeshComponent->MarkRenderStateDirty();
-	StaticMeshComponent->bDisallowMeshPaintPerInstance = true;
+
+	// Explicitly disable the mesh paint tool on the component to prevent stomping the vertex color.
+	StaticMeshComponent->bEnableVertexColorMeshPainting = false;
 }
 
 void FMeshVertexPainter::RemovePaintedVertices(UStaticMeshComponent* StaticMeshComponent)
@@ -112,6 +116,6 @@ void FMeshVertexPainter::RemovePaintedVertices(UStaticMeshComponent* StaticMeshC
 	}
 
 	StaticMeshComponent->MarkRenderStateDirty();
-	StaticMeshComponent->bDisallowMeshPaintPerInstance = false;
+	StaticMeshComponent->bEnableVertexColorMeshPainting = true;
 }
 
