Rendering: Fixed off-center projection calculations for Orthographic projections. Using an off-center projection with orthographic mode now uses the same range (-1, 1) as Perspective cameras, and no longer changes how much of an object you see, instead it behaves like Perspective cameras which just shift around your viewing position without changing how much of an object you see.

diff --git a/Engine/Source/Runtime/Engine/Private/Camera/CameraStackTypes.cpp b/Engine/Source/Runtime/Engine/Private/Camera/CameraStackTypes.cpp
--- a/Engine/Source/Runtime/Engine/Private/Camera/CameraStackTypes.cpp
+++ b/Engine/Source/Runtime/Engine/Private/Camera/CameraStackTypes.cpp
@@ -184,7 +184,8 @@ FMatrix FMinimalViewInfo::CalculateProjectionMatrix() const
 {
 	FMatrix ProjectionMatrix;
 
-	if (ProjectionMode == ECameraProjectionMode::Orthographic)
+	const bool bOrthographic = ProjectionMode == ECameraProjectionMode::Orthographic;
+	if (bOrthographic)
 	{
 		const float YScale = 1.0f / AspectRatio;
 
@@ -221,8 +222,19 @@ FMatrix FMinimalViewInfo::CalculateProjectionMatrix() const
 		const float Right = Left + 2.0f;
 		const float Bottom = -1.0f + OffCenterProjectionOffset.Y;
 		const float Top = Bottom + 2.0f;
-		ProjectionMatrix.M[2][0] = (Left + Right) / (Left - Right);
-		ProjectionMatrix.M[2][1] = (Bottom + Top) / (Bottom - Top);
+
+		// Make sure you update CalculateProjectionMatrixGivenViewRectangle(...) as well if you change this, as
+		// it may have already modified some fields in the ProjectionMatrix.
+		if (bOrthographic)
+		{
+			ProjectionMatrix.M[3][0] = (Left + Right) / (Left - Right);
+			ProjectionMatrix.M[3][1] = (Bottom + Top) / (Bottom - Top);
+		}
+		else
+		{
+			ProjectionMatrix.M[2][0] = (Left + Right) / (Left - Right);
+			ProjectionMatrix.M[2][1] = (Bottom + Top) / (Bottom - Top);
+		}
 	}
 
 	return ProjectionMatrix;
@@ -345,8 +357,19 @@ void FMinimalViewInfo::CalculateProjectionMatrixGivenViewRectangle(FMinimalViewI
 		const float Right = Left + 2.0f;
 		const float Bottom = -1.0f + ViewInfo.OffCenterProjectionOffset.Y;
 		const float Top = Bottom + 2.0f;
-		InOutProjectionData.ProjectionMatrix.M[2][0] = (Left + Right) / (Left - Right);
-		InOutProjectionData.ProjectionMatrix.M[2][1] = (Bottom + Top) / (Bottom - Top);
+
+		// Make sure you update CalculateProjectionMatrix() as well if you change this, as
+		// it may have already modified some fields in the ProjectionMatrix.
+		if (bOrthographic)
+		{
+			InOutProjectionData.ProjectionMatrix.M[3][0] = (Left + Right) / (Left - Right);
+			InOutProjectionData.ProjectionMatrix.M[3][1] = (Bottom + Top) / (Bottom - Top);
+		}
+		else
+		{
+			InOutProjectionData.ProjectionMatrix.M[2][0] = (Left + Right) / (Left - Right);
+			InOutProjectionData.ProjectionMatrix.M[2][1] = (Bottom + Top) / (Bottom - Top);
+		}
 	}
 }
 
