Add setting on static mesh for whether texture mesh painting is supported.
The setting has a "default" which allows us to have a project default setting.
This will allow projects to choose to either opt-in or opt-out for texture  mesh painting.

diff --git a/Engine/Source/Runtime/Engine/Classes/Components/StaticMeshComponent.h b/Engine/Source/Runtime/Engine/Classes/Components/StaticMeshComponent.h
--- a/Engine/Source/Runtime/Engine/Classes/Components/StaticMeshComponent.h
+++ b/Engine/Source/Runtime/Engine/Classes/Components/StaticMeshComponent.h
@@ -754,9 +754,9 @@ public:
 	ENGINE_API UMaterialInterface* GetNaniteAuditMaterial(int32 MaterialIndex) const;
 
 	/* Returns true if mesh vertex color painting is supported on this component. */
-	ENGINE_API int32 CanMeshPaintVertexColors() const { return bSupportMeshPainting && bEnableVertexColorMeshPainting; }
+	ENGINE_API bool CanMeshPaintVertexColors() const { return bSupportMeshPainting && bEnableVertexColorMeshPainting; }
 	/* Returns true if mesh texture color painting is supported on this component. */
-	ENGINE_API int32 CanMeshPaintTextureColors() const { return bSupportMeshPainting && bEnableTextureColorMeshPainting; }
+	ENGINE_API bool CanMeshPaintTextureColors() const;
 
 private:
 	/** Initializes the resources used by the static mesh component. */
diff --git a/Engine/Source/Runtime/Engine/Classes/Engine/RendererSettings.h b/Engine/Source/Runtime/Engine/Classes/Engine/RendererSettings.h
--- a/Engine/Source/Runtime/Engine/Classes/Engine/RendererSettings.h
+++ b/Engine/Source/Runtime/Engine/Classes/Engine/RendererSettings.h
@@ -457,6 +457,12 @@ class URendererSettings : public UDeveloperSettings
 		ConfigRestartRequired = true))
 	uint32 MeshPaintVirtualTextureTileBorderSize;
 
+	UPROPERTY(config, EditAnywhere, Category = VirtualTextures, meta = (
+		EditCondition = "bVirtualTextures",
+		ConsoleVariable = "r.StaticMesh.DefaultMeshPaintTextureSupport", DisplayName = "Mesh paint default static mesh support",
+		ToolTip = "Default setting for whether static mesh assets support mesh paint textures."))
+	bool MeshPaintDefaultOnStaticMesh;
+
 	UPROPERTY(config, EditAnywhere, Category = VirtualTextures, meta = (
 		EditCondition = "bVirtualTextures",
 		ConsoleVariable = "r.MeshPaintVirtualTexture.DefaultTexelsPerVertex", DisplayName = "Mesh paint texels per vertex",
diff --git a/Engine/Source/Runtime/Engine/Classes/Engine/StaticMesh.h b/Engine/Source/Runtime/Engine/Classes/Engine/StaticMesh.h
--- a/Engine/Source/Runtime/Engine/Classes/Engine/StaticMesh.h
+++ b/Engine/Source/Runtime/Engine/Classes/Engine/StaticMesh.h
@@ -548,6 +548,17 @@ struct FMaterialRemapIndex
 	TArray<int32> MaterialRemap;
 };
 
+UENUM(BlueprintType)
+enum class EStaticMeshPaintSupport : uint8
+{
+	/** Use the default project setting for whether texture color mesh painting is supported. */
+	Default,
+	/** Enable texture color mesh painting support. */
+	Enabled,
+	/** Disable texture color mesh painting support. */
+	Disabled,
+};
+
 
 /**
  * A StaticMesh is a piece of geometry that consists of a static set of polygons.
@@ -1003,7 +1014,14 @@ public:
 		PRAGMA_ENABLE_DEPRECATION_WARNINGS
 	}
 
-	/** The default coordinate index to use when painting the MeshPaintTexture on this mesh. */
+	/** Whether to support per instance texture color mesh painting on components using this mesh. */
+	UPROPERTY(EditAnywhere, BlueprintReadWrite, Category = StaticMesh, meta = (DisplayName = "Support Texture Color Mesh Painting"))
+	EStaticMeshPaintSupport StaticMeshPaintSupport;
+
+	/** Whether to support per instance texture color mesh painting on components using this mesh. */
+	ENGINE_API bool CanMeshPaintTextureColors() const;
+
+	/** The default coordinate index to use when texture color painting on this mesh. */
 	UPROPERTY(EditAnywhere, AdvancedDisplay, Category = StaticMesh, meta=(UIMin = "0", UIMax = "3"))
 	int32 MeshPaintTextureCoordinateIndex;
 
diff --git a/Engine/Source/Runtime/Engine/Classes/VT/MeshPaintVirtualTexture.h b/Engine/Source/Runtime/Engine/Classes/VT/MeshPaintVirtualTexture.h
--- a/Engine/Source/Runtime/Engine/Classes/VT/MeshPaintVirtualTexture.h
+++ b/Engine/Source/Runtime/Engine/Classes/VT/MeshPaintVirtualTexture.h
@@ -18,7 +18,7 @@ class UMeshPaintVirtualTexture : public UTexture2D
 {
 	GENERATED_UCLASS_BODY()
 
-	/** Weak referernce to the owning primitive component. */
+	/** Weak refererence to the owning primitive component. */
 	UPROPERTY()
 	TWeakObjectPtr<UPrimitiveComponent> OwningComponent;
 
diff --git a/Engine/Source/Runtime/Engine/Private/Components/StaticMeshComponent.cpp b/Engine/Source/Runtime/Engine/Private/Components/StaticMeshComponent.cpp
--- a/Engine/Source/Runtime/Engine/Private/Components/StaticMeshComponent.cpp
+++ b/Engine/Source/Runtime/Engine/Private/Components/StaticMeshComponent.cpp
@@ -3204,6 +3204,19 @@ int32 UStaticMeshComponent::GetMeshPaintTextureCoordinateIndex() const
 	return 0;
 }
 
+bool UStaticMeshComponent::CanMeshPaintTextureColors() const 
+{
+	if (!bSupportMeshPainting || !bEnableTextureColorMeshPainting)
+	{
+		return false;
+	}
+	if (StaticMesh == nullptr || !StaticMesh->CanMeshPaintTextureColors())
+	{
+		return false;
+	}
+	return true;
+}
+
 bool UStaticMeshComponent::IsNavigationRelevant() const
 {
 	if (const UStaticMesh* Mesh = GetStaticMesh())
diff --git a/Engine/Source/Runtime/Engine/Private/StaticMesh.cpp b/Engine/Source/Runtime/Engine/Private/StaticMesh.cpp
--- a/Engine/Source/Runtime/Engine/Private/StaticMesh.cpp
+++ b/Engine/Source/Runtime/Engine/Private/StaticMesh.cpp
@@ -164,6 +164,12 @@ static TAutoConsoleVariable<int32> CVarForceEnableNaniteMeshes(
 	TEXT("Force enables all meshes to also build Nanite data, regardless of the enabled flag on the asset."),
 	ECVF_ReadOnly | ECVF_RenderThreadSafe);
 
+static TAutoConsoleVariable<bool> CVarStaticMeshDefaultMeshPaintTextureSupport(
+	TEXT("r.StaticMesh.DefaultMeshPaintTextureSupport"),
+	true,
+	TEXT("Default setting for whether static mesh assets support mesh paint textures."),
+	ECVF_Default);
+
 #if ENABLE_COOK_STATS
 namespace StaticMeshCookStats
 {
@@ -9201,6 +9207,15 @@ bool UStaticMesh::IsNaniteForceEnabled() const
 
 #endif
 
+bool UStaticMesh::CanMeshPaintTextureColors() const
+{
+	if (StaticMeshPaintSupport == EStaticMeshPaintSupport::Default)
+	{
+		return CVarStaticMeshDefaultMeshPaintTextureSupport.GetValueOnGameThread();
+	}
+	return StaticMeshPaintSupport != EStaticMeshPaintSupport::Disabled;
+}
+
 /*-----------------------------------------------------------------------------
 UStaticMeshSocket
 -----------------------------------------------------------------------------*/
