Substrate - updated message for incompatible nodes to encourage discussion with dev team about adding support.

#rnx

diff --git a/Engine/Source/Runtime/Engine/Private/Materials/MaterialExpressions.cpp b/Engine/Source/Runtime/Engine/Private/Materials/MaterialExpressions.cpp
--- a/Engine/Source/Runtime/Engine/Private/Materials/MaterialExpressions.cpp
+++ b/Engine/Source/Runtime/Engine/Private/Materials/MaterialExpressions.cpp
@@ -2385,7 +2385,7 @@ bool UMaterialExpression::IsUsingNewHLSLGenerator() const
 
 FSubstrateOperator* UMaterialExpression::SubstrateGenerateMaterialTopologyTree(class FMaterialCompiler* Compiler, class UMaterialExpression* Parent, int32 OutputIndex)
 {
-	Compiler->Errorf(TEXT("%s nodes type does not support generating/processing/flowing Substrate data."), *GetClass()->GetName());
+	Compiler->Errorf(TEXT("%s nodes type does not support generating/processing/flowing Substrate data.\nPlease reach out to the development team for feedback and if you want support to be added."), *GetClass()->GetName());
 	return nullptr;
 }
 
