Substrate - Fixed compilation issues with attributes nodes plug in some weird way behind switches nodes. We can now cook and package a full game with Substrate enabled.

diff --git a/Engine/Source/Runtime/Engine/Private/Materials/MaterialExpressions.cpp b/Engine/Source/Runtime/Engine/Private/Materials/MaterialExpressions.cpp
--- a/Engine/Source/Runtime/Engine/Private/Materials/MaterialExpressions.cpp
+++ b/Engine/Source/Runtime/Engine/Private/Materials/MaterialExpressions.cpp
@@ -26635,16 +26635,21 @@ FExpressionInput* UMaterialExpressionSubstrateShadingModels::GetInput(int32 Inpu
 	}
 }
 
+#define LEGACY_DIRECT_ATTRIBUTE_MAPPING(CodeChunkResult, MaterialProperty, Code)	Compiler->PushMaterialAttribute(FMaterialAttributeDefinitionMap::GetID(MaterialProperty));\
+																					int32 CodeChunkResult = Code;\
+																					Compiler->PopMaterialAttribute();
+
 int32 UMaterialExpressionSubstrateShadingModels::Compile(class FMaterialCompiler* Compiler, int32 OutputIndex)
 {
-	int32 RoughnessCodeChunk = CompileWithDefaultFloat1(Compiler, Roughness, 0.5f);
-	int32 AnisotropyCodeChunk = CompileWithDefaultFloat1(Compiler, Anisotropy, 0.0f);
+	LEGACY_DIRECT_ATTRIBUTE_MAPPING(RoughnessCodeChunk, MP_Roughness, CompileWithDefaultFloat1(Compiler, Roughness, 0.5f));
+	LEGACY_DIRECT_ATTRIBUTE_MAPPING(AnisotropyCodeChunk, MP_Anisotropy, CompileWithDefaultFloat1(Compiler, Anisotropy, 0.0f));
+
 	// We also cannot ignore the tangent when using the default Tangent because GetTangentBasis
 	// used in SubstrateGetBSDFSharedBasis cannot be relied on for smooth tangent used for lighting on any mesh.
 	const bool bHasAnisotropy = HasAnisotropy();
 
 	// Regular normal basis
-	int32 NormalCodeChunk = CompileWithDefaultNormalWS(Compiler, Normal);
+	LEGACY_DIRECT_ATTRIBUTE_MAPPING(NormalCodeChunk, MP_Normal, CompileWithDefaultNormalWS(Compiler, Normal));
 
 	// When computing NormalCodeChunk, we invoke TransformNormalFromRequestedBasisToWorld which requires input to be float or float3.
 	// Certain material do not respect this requirement. We handle here a simple recovery when source material doesn't have a valid 
@@ -26652,7 +26657,12 @@ int32 UMaterialExpressionSubstrateShadingModels::Compile(class FMaterialCompiler
 	// to the user, but the compilation will succeed.
 	if (NormalCodeChunk == INDEX_NONE) { NormalCodeChunk = Compiler->VertexNormal(); } 
 
-	int32 TangentCodeChunk = bHasAnisotropy ? CompileWithDefaultTangentWS(Compiler, Tangent) : INDEX_NONE;
+	int32 TangentCodeChunk = INDEX_NONE;
+	if (bHasAnisotropy)
+	{
+		LEGACY_DIRECT_ATTRIBUTE_MAPPING(TangentCodeChunkTmp, MP_Tangent, CompileWithDefaultTangentWS(Compiler, Normal));
+		TangentCodeChunk = TangentCodeChunkTmp;
+	}
 	const FSubstrateRegisteredSharedLocalBasis NewRegisteredSharedLocalBasis = SubstrateCompilationInfoCreateSharedLocalBasis(Compiler, NormalCodeChunk, TangentCodeChunk);
 	const FString BasisIndexMacro = Compiler->GetSubstrateSharedLocalBasisIndexMacro(NewRegisteredSharedLocalBasis);
 
@@ -26705,33 +26715,45 @@ int32 UMaterialExpressionSubstrateShadingModels::Compile(class FMaterialCompiler
 	{
 		// We evaluate opacity only for shading models and blending mode requiring it.
 		// For instance, a translucent shader reading depth for soft fading should no evaluate opacity when an instance forces an opaque mode.
-		OpacityCodeChunk = CompileWithDefaultFloat1(Compiler, Opacity, 1.0f);
+		LEGACY_DIRECT_ATTRIBUTE_MAPPING(OpacityCodeChunkTmp, MP_Opacity, CompileWithDefaultFloat1(Compiler, Opacity, 1.0f));
+		OpacityCodeChunk = OpacityCodeChunkTmp;
 	}
 	else
 	{
 		OpacityCodeChunk = Compiler->Constant(1.0f);
 	}
 
+	LEGACY_DIRECT_ATTRIBUTE_MAPPING(EmissiveCodeChunk,			MP_EmissiveColor,	CompileWithDefaultFloat3(Compiler, EmissiveColor, 0.0f, 0.0f, 0.0f));
+
+	LEGACY_DIRECT_ATTRIBUTE_MAPPING(BaseColorCodeChunk,			MP_BaseColor,		CompileWithDefaultFloat3(Compiler, BaseColor, 0.0f, 0.0f, 0.0f));
+	LEGACY_DIRECT_ATTRIBUTE_MAPPING(SpecularChunk,				MP_Specular,		CompileWithDefaultFloat1(Compiler, Specular, 0.5f));
+	LEGACY_DIRECT_ATTRIBUTE_MAPPING(MetallicCodeChunk,			MP_Metallic,		CompileWithDefaultFloat1(Compiler, Metallic, 0.0f));
+
+	LEGACY_DIRECT_ATTRIBUTE_MAPPING(SubSurfaceColorCodeChunk,	MP_SubsurfaceColor,	CompileWithDefaultFloat3(Compiler, SubSurfaceColor, 1.0f, 1.0f, 1.0f));
+
+	LEGACY_DIRECT_ATTRIBUTE_MAPPING(ClearCoatCodeChunk,			MP_CustomData0,		CompileWithDefaultFloat1(Compiler, ClearCoat, 1.0f));
+	LEGACY_DIRECT_ATTRIBUTE_MAPPING(ClearCoatRoughnessCodeChunk,MP_CustomData1,		CompileWithDefaultFloat1(Compiler, ClearCoatRoughness, 0.1f));
+
 	int32 ShadingModelCodeChunk = ShadingModel.IsConnected() ? CompileWithDefaultFloat1(Compiler, ShadingModel, float(MSM_DefaultLit)) : Compiler->Constant(float(ShadingModelOverride));
 	int32 ShadingModelCount = Compiler->GetMaterialShadingModels().CountShadingModels();
 	const bool bHasDynamicShadingModels = ShadingModelCount > 1;
 	int32 OutputCodeChunk = Compiler->SubstrateConversionFromLegacy(
 		bHasDynamicShadingModels,
 		// Metalness workflow
-		CompileWithDefaultFloat3(Compiler, BaseColor, 0.0f, 0.0f, 0.0f),
-		CompileWithDefaultFloat1(Compiler, Specular, 0.5f),
-		CompileWithDefaultFloat1(Compiler, Metallic,  0.0f),
+		BaseColorCodeChunk,
+		SpecularChunk,
+		MetallicCodeChunk,
 		// Roughness
 		RoughnessCodeChunk,
 		AnisotropyCodeChunk,
 		// SSS
-		CompileWithDefaultFloat3(Compiler, SubSurfaceColor, 1.0f, 1.0f, 1.0f),
+		SubSurfaceColorCodeChunk,
 		SSSProfileCodeChunk != INDEX_NONE ? SSSProfileCodeChunk : Compiler->Constant(0.0f),
 		// Clear Coat / Custom
-		CompileWithDefaultFloat1(Compiler, ClearCoat, 1.0f),
-		CompileWithDefaultFloat1(Compiler, ClearCoatRoughness, 0.1f),
+		ClearCoatCodeChunk,
+		ClearCoatRoughnessCodeChunk,
 		// Misc
-		CompileWithDefaultFloat3(Compiler, EmissiveColor, 0.0f, 0.0f, 0.0f),
+		EmissiveCodeChunk,
 		OpacityCodeChunk,
 		CompileWithDefaultFloat3(Compiler, TransmittanceColor, 0.5f, 0.5f, 0.5f),
 		CompileWithDefaultFloat1(Compiler, ThinTranslucentSurfaceCoverage, 1.0f),
