Gathering dependencies now includes ExternalObjects.
Correction for External(Objects|Actors) path destination path generation.

diff --git a/Engine/Source/Developer/AssetTools/Private/AssetTools.cpp b/Engine/Source/Developer/AssetTools/Private/AssetTools.cpp
--- a/Engine/Source/Developer/AssetTools/Private/AssetTools.cpp
+++ b/Engine/Source/Developer/AssetTools/Private/AssetTools.cpp
@@ -2273,6 +2273,13 @@ TMap<FString, FString> AllSourceAndDestPackages(const TMap<FString, FString>& So
 		FStringView SrcPath = SrcPackagePath.IsEmpty() ? SrcPackageName : SrcPackagePath;
 		FStringView DstPath = DstPackagePath.IsEmpty() ? DstPackageName : DstPackagePath;
 		ExternalMappings.Add(FString(SrcPath), FString(DstPath));
+
+		// if there is a path
+		if (!SrcPackagePath.IsEmpty())
+		{
+			// add the local path/asset for the case of maps (which we cannot tell at this point)
+			ExternalMappings.Add(FString(SrcPath.GetData()), FString(DstPath.GetData()));
+		}
 	}
 
 	TMap<FString, FString> Result;
diff --git a/Engine/Source/Runtime/Engine/Classes/Engine/ExternalObjectAndActorDependencyGatherer.h b/Engine/Source/Runtime/Engine/Classes/Engine/ExternalObjectAndActorDependencyGatherer.h
new file mode 100644
--- /dev/null
+++ b/Engine/Source/Runtime/Engine/Classes/Engine/ExternalObjectAndActorDependencyGatherer.h
@@ -0,0 +1,20 @@
+// Copyright Epic Games, Inc. All Rights Reserved.
+
+#pragma once
+
+#if WITH_EDITOR
+
+#include "AssetRegistry/AssetDependencyGatherer.h"
+#include "AssetRegistry/AssetData.h"
+#include "Engine/World.h"
+
+class ENGINE_API FExternalObjectAndActorDependencyGatherer : public IAssetDependencyGatherer
+{
+public:	
+	FExternalObjectAndActorDependencyGatherer() = default;
+	virtual ~FExternalObjectAndActorDependencyGatherer() = default;
+
+	virtual void GatherDependencies(const FAssetData& AssetData, const FAssetRegistryState& AssetRegistryState, TFunctionRef<FARCompiledFilter(const FARFilter&)> CompileFilterFunc, TArray<IAssetDependencyGatherer::FGathereredDependency>& OutDependencies, TArray<FString>& OutDependencyDirectories) const override;
+};
+
+#endif
\ No newline at end of file
diff --git a/Engine/Source/Runtime/Engine/Private/ExternalActorAssetDependencyGatherer.cpp b/Engine/Source/Runtime/Engine/Private/ExternalActorAssetDependencyGatherer.cpp
--- a/Engine/Source/Runtime/Engine/Private/ExternalActorAssetDependencyGatherer.cpp
+++ b/Engine/Source/Runtime/Engine/Private/ExternalActorAssetDependencyGatherer.cpp
@@ -8,9 +8,6 @@
 #include "AssetRegistry/AssetRegistryState.h"
 #include "Engine/Level.h"
 
-// Register FExternalActorAssetDependencyGatherer for UWorld class
-REGISTER_ASSETDEPENDENCY_GATHERER(FExternalActorAssetDependencyGatherer, UWorld);
-
 void FExternalActorAssetDependencyGatherer::GatherDependencies(const FAssetData& AssetData, const FAssetRegistryState& AssetRegistryState, TFunctionRef<FARCompiledFilter(const FARFilter&)> CompileFilterFunc, TArray<IAssetDependencyGatherer::FGathereredDependency>& OutDependencies, TArray<FString>& OutDependencyDirectories) const
 {
 	if (ULevel::GetIsLevelUsingExternalActorsFromAsset(AssetData))
diff --git a/Engine/Source/Runtime/Engine/Private/ExternalObjectAndActorDependencyGatherer.cpp b/Engine/Source/Runtime/Engine/Private/ExternalObjectAndActorDependencyGatherer.cpp
new file mode 100644
--- /dev/null
+++ b/Engine/Source/Runtime/Engine/Private/ExternalObjectAndActorDependencyGatherer.cpp
@@ -0,0 +1,37 @@
+// Copyright Epic Games, Inc. All Rights Reserved.
+#include "Engine/ExternalObjectAndActorDependencyGatherer.h"
+
+#if WITH_EDITOR
+
+#include "AssetRegistry/ARFilter.h"
+#include "AssetRegistry/AssetRegistryState.h"
+#include "ExternalPackageHelper.h"
+
+void FExternalObjectAndActorDependencyGatherer::GatherDependencies(const FAssetData& AssetData, const FAssetRegistryState& AssetRegistryState, TFunctionRef<FARCompiledFilter(const FARFilter&)> CompileFilterFunc, TArray<IAssetDependencyGatherer::FGathereredDependency>& OutDependencies, TArray<FString>& OutDependencyDirectories) const
+{		
+	const FString ExternalActorsPath = ULevel::GetExternalActorsPath(AssetData.PackageName.ToString());
+	const FString ExternalObjectPath = FExternalPackageHelper::GetExternalObjectsPath(AssetData.PackageName.ToString());
+
+	OutDependencyDirectories.Add(ExternalActorsPath);
+	OutDependencyDirectories.Add(ExternalObjectPath);
+
+	FARFilter Filter;
+	Filter.bRecursivePaths = true;
+	Filter.bIncludeOnlyOnDiskAssets = true;
+
+	Filter.PackagePaths.Add(*ExternalActorsPath);
+	Filter.PackagePaths.Add(*ExternalObjectPath);
+
+	TArray<FAssetData> FilteredAssets;
+	AssetRegistryState.GetAssets(CompileFilterFunc(Filter), {}, FilteredAssets, true);
+
+	for (const FAssetData& FilteredAsset : FilteredAssets)
+	{
+		OutDependencies.Emplace(IAssetDependencyGatherer::FGathereredDependency{ FilteredAsset.PackageName, UE::AssetRegistry::EDependencyProperty::Game | UE::AssetRegistry::EDependencyProperty::Build });
+	}
+}
+
+REGISTER_ASSETDEPENDENCY_GATHERER(FExternalObjectAndActorDependencyGatherer, UWorld);
+
+#endif
+
