[SceneComponent] when attaching the component to another we only update owner navigation data if it's valid (it's possible for a component to be registered without an actor owner)

diff --git a/Engine/Source/Runtime/Engine/Private/Components/SceneComponent.cpp b/Engine/Source/Runtime/Engine/Private/Components/SceneComponent.cpp
--- a/Engine/Source/Runtime/Engine/Private/Components/SceneComponent.cpp
+++ b/Engine/Source/Runtime/Engine/Private/Components/SceneComponent.cpp
@@ -2338,7 +2338,10 @@ bool USceneComponent::AttachToComponent(USceneComponent* Parent, const FAttachme
 			UpdateOverlaps();
 
 			// Update our owner actor in the navigation system since its associated bounds have changed
-			FNavigationSystem::UpdateActorAndComponentData(*GetOwner());
+			if (AActor* Owner = GetOwner())
+			{
+				FNavigationSystem::UpdateActorAndComponentData(*Owner);
+			}
 		}
 
 		return true;
