Saves 16 bytes in FTextureLODGroup (96->80)
Saves 0.35 MB (22500 in game)

diff --git a/Engine/Source/Runtime/Engine/Classes/Engine/TextureLODSettings.h b/Engine/Source/Runtime/Engine/Classes/Engine/TextureLODSettings.h
--- a/Engine/Source/Runtime/Engine/Classes/Engine/TextureLODSettings.h
+++ b/Engine/Source/Runtime/Engine/Classes/Engine/TextureLODSettings.h
@@ -17,130 +17,106 @@ struct FTextureLODGroup
 	GENERATED_USTRUCT_BODY()
 
 	FTextureLODGroup()
-		: Group(TEXTUREGROUP_World)
-		, MaxLODMipCount(32)
-		, LODBias(0)
-		, LODBias_Smaller(-1)
-		, LODBias_Smallest(-1)
-		, NumStreamedMips(-1)
-		, MipGenSettings(TextureMipGenSettings::TMGS_SimpleAverage)
-		, MinLODSize(1)
-		, MaxLODSize(4096)
-		, MaxLODSize_Smaller(-1)
-		, MaxLODSize_Smallest(-1)
-		, MaxLODSize_VT(0)
-		, OptionalLODBias(0)
-		, OptionalMaxLODSize(4096)
-		, MinMagFilter(NAME_Aniso)
-		, MipFilter(NAME_Point)
-		, MipLoadOptions(ETextureMipLoadOptions::AllMips)
-		, HighPriorityLoad(false)
-		, DuplicateNonOptionalMips(false)
-		, Downscale(1.0)
-		, DownscaleOptions(ETextureDownscaleOptions::Default)
-		, VirtualTextureTileCountBias(0)
-		, VirtualTextureTileSizeBias(0)
-		, LossyCompressionAmount(TLCA_Default)
-		, CookPlatformTilingDisabled(false)
-		, MaxAniso(0)
 	{
 		SetupGroup();
 	}
 
 	/** Group ID.																	*/
 	UPROPERTY()
-	TEnumAsByte<TextureGroup> Group;
+	TEnumAsByte<TextureGroup> Group = TEXTUREGROUP_World;
+	
+	/** Sampler filter state.														*/
+	ETextureSamplerFilter Filter;
+
+	/** Defines how the the mip-map generation works, e.g. sharpening				*/
+	UPROPERTY()
+	TEnumAsByte<TextureMipGenSettings> MipGenSettings = TextureMipGenSettings::TMGS_SimpleAverage;
+	
+	UPROPERTY()
+	ETextureMipLoadOptions MipLoadOptions = ETextureMipLoadOptions::AllMips;
 
 	/** Maximum LOD mip count. Bias will be adjusted so texture won't go above.		*/
-	int32 MaxLODMipCount;
+	int32 MaxLODMipCount = 32;
 	
 	/** Group LOD bias.																*/
 	UPROPERTY()
-	int32 LODBias;
+	int32 LODBias = 0;
 
 	UPROPERTY()
-	int32 LODBias_Smaller;
+	int32 LODBias_Smaller = -1;
 
 	UPROPERTY()
-	int32 LODBias_Smallest;
-
-	/** Sampler filter state.														*/
-	ETextureSamplerFilter Filter;
+	int32 LODBias_Smallest = -1;
 	
 	/** Number of mip-levels that can be streamed. -1 means all mips can stream.	*/
 	UPROPERTY()
-	int32 NumStreamedMips;
-
-	/** Defines how the the mip-map generation works, e.g. sharpening				*/
-	UPROPERTY()
-	TEnumAsByte<TextureMipGenSettings> MipGenSettings;
+	int32 NumStreamedMips = -1;
 
 	/** Prevent LODBias from making the textures smaller than this value. Note that this does _not_ affect the smallest mip level size. */
 	UPROPERTY()
-	int32 MinLODSize;
+	int32 MinLODSize = 1;
 
 	/** Cap the number of mips such that the largest mip is this big. Has no effect for textures with no mip chain. Not used for virtual textures. */
 	UPROPERTY()
-	int32 MaxLODSize;
+	int32 MaxLODSize = 4096;
 
 	/** Cap the number of mips such that the largest mip is this big. Has no effect for textures with no mip chain. Used for platforms with the "Smaller" memory bucket. Not used for virtual textures. */
 	UPROPERTY()
-	int32 MaxLODSize_Smaller;
+	int32 MaxLODSize_Smaller = -1;
 
 	/** Cap the number of mips such that the largest mip is this big. Has no effect for textures with no mip chain. Used for platforms with the "Smallest" memory bucket. Not used for virtual textures. */
 	UPROPERTY()
-	int32 MaxLODSize_Smallest;
+	int32 MaxLODSize_Smallest = -1;
 
 	/** Cap the number of mips such that the largest mip is this big. Has no effect for textures with no mip chain. Used for virtual textures. */
 	UPROPERTY()
-	int32 MaxLODSize_VT;
+	int32 MaxLODSize_VT = 0;
 
-	/** If this is greater then 0 will put that number of mips into an optional bulkdata package */
+	/** If this is greater than 0 will put that number of mips into an optional bulkdata package */
 	UPROPERTY()
-	int32 OptionalLODBias;
+	int32 OptionalLODBias = 0;
 
 	/** Put all the mips which have a width / height larger then OptionalLODSize into an optional bulkdata package */
 	UPROPERTY()
-	int32 OptionalMaxLODSize;
+	int32 OptionalMaxLODSize = 4096;
 
 	UPROPERTY()
-	FName MinMagFilter;
+	FName MinMagFilter = NAME_Aniso;
 
 	UPROPERTY()
-	FName MipFilter;
+	FName MipFilter = NAME_Point;
 
+	/** Whether those assets should be loaded with higher load order and higher IO priority. Allows ProjectXX texture groups to behave as character textures. */
 	UPROPERTY()
-	ETextureMipLoadOptions MipLoadOptions;
+	uint8 HighPriorityLoad : 1 = false;
 
-	/** Wether those assets should be loaded with higher load order and higher IO priority. Allows ProjectXX texture groups to behave as character textures. */
 	UPROPERTY()
-	bool HighPriorityLoad;
-
+	uint8 DuplicateNonOptionalMips : 1 = false;
+	
 	UPROPERTY()
-	bool DuplicateNonOptionalMips;
+	uint8 CookPlatformTilingDisabled : 1 = false;
 
 	UPROPERTY()
-	float Downscale;
-
+	ETextureDownscaleOptions DownscaleOptions = ETextureDownscaleOptions::Default;
+	
 	UPROPERTY()
-	ETextureDownscaleOptions DownscaleOptions;
-
+	TEnumAsByte<ETextureLossyCompressionAmount> LossyCompressionAmount = TLCA_Default;
+	
 	UPROPERTY()
-	int32 VirtualTextureTileCountBias;
+	float Downscale = 1.0;
 
 	UPROPERTY()
-	int32 VirtualTextureTileSizeBias;
-	
+	int32 VirtualTextureTileCountBias = 0;
+
 	UPROPERTY()
-	TEnumAsByte<enum ETextureLossyCompressionAmount> LossyCompressionAmount;
+	int32 VirtualTextureTileSizeBias = 0;
+
 
 	/** If true textures with CookPlatformTilingSettings set to TCPTS_FromTextureGroup will not be tiled during cook. They will be tiled when uploaded to the GPU if necessary */
-	UPROPERTY()
-	bool CookPlatformTilingDisabled;
 
 	/** Allows us to override max anisotropy. If unspecified, uses r.MaxAnisotropy */
 	UPROPERTY()
-	int32 MaxAniso;
+	int32 MaxAniso = 0;
 
 	ENGINE_API void SetupGroup();
 
