Added tooltip on Data Layer Load Filter values for clarity.

diff --git a/Engine/Source/Runtime/Engine/Public/WorldPartition/DataLayer/DataLayerAsset.h b/Engine/Source/Runtime/Engine/Public/WorldPartition/DataLayer/DataLayerAsset.h
--- a/Engine/Source/Runtime/Engine/Public/WorldPartition/DataLayer/DataLayerAsset.h
+++ b/Engine/Source/Runtime/Engine/Public/WorldPartition/DataLayer/DataLayerAsset.h
@@ -17,8 +17,11 @@ class AActor;
 UENUM(BlueprintType)
 enum class EDataLayerLoadFilter : uint8
 {
+	/** Data Layer is considered by the client and the server. Client runtime state is replicated. */
 	None,
+	/** Data Layer is only considered by the client. */
 	ClientOnly,
+	/** Data layer is only considered by the server. */
 	ServerOnly
 };
 
