Added description (tooltip) for each UFUNCTION of DataLayerManager.

diff --git a/Engine/Source/Runtime/Engine/Public/WorldPartition/DataLayer/DataLayerManager.h b/Engine/Source/Runtime/Engine/Public/WorldPartition/DataLayer/DataLayerManager.h
--- a/Engine/Source/Runtime/Engine/Public/WorldPartition/DataLayer/DataLayerManager.h
+++ b/Engine/Source/Runtime/Engine/Public/WorldPartition/DataLayer/DataLayerManager.h
@@ -57,29 +57,55 @@ public:
 	}
 
 	//~ Begin Blueprint interface
+
+	/** Returns all Data Layer instances. */
 	UFUNCTION(BlueprintCallable, Category = DataLayers)
 	ENGINE_API TArray<UDataLayerInstance*> GetDataLayerInstances() const;
 
+	/** Returns the Data Layer instance referencing the provided Data Layer asset (if any). */
 	UFUNCTION(BlueprintCallable, Category = DataLayers)
 	ENGINE_API const UDataLayerInstance* GetDataLayerInstanceFromAsset(const UDataLayerAsset* InDataLayerAsset) const;
 
+	/** Returns the Data Layer instance matching the provided Data Layer instance name (if any). */
 	UFUNCTION(BlueprintCallable, Category = DataLayers)
 	ENGINE_API const UDataLayerInstance* GetDataLayerInstanceFromName(const FName& InDataLayerInstanceName) const;
 
+	/** Changes the Data Layer instance runtime state.
+	  * If recursive is set to true, the runtime state will also be applied to all child Data Layer instances.
+	  * Note:
+	  *  - Changing the runtime state of a Client-Only Data Layer instance must be done on the client side or else it will have no effect.
+	  *  - Changing the runtime state of a Server-Only Data Layer instance can only be done on the server side or else it will have no effect.
+	  *  - Changing the runtime state of a runtime Data Layer instance (with no Load Filter set on the asset) must be done on the server side 
+	  *    or else it will have no effect. The runtime state will then be replicated on the client.
+	  * (see Data Layer asset Load Filter for more details)
+	  */
 	UFUNCTION(BlueprintCallable, Category = DataLayers)
 	ENGINE_API bool SetDataLayerInstanceRuntimeState(const UDataLayerInstance* InDataLayerInstance, EDataLayerRuntimeState InState, bool bInIsRecursive = false);
 
+	/** Finds a matching Data Layer instance referencing the provided Data Layer asset and changes its runtime state (if any).
+	  * If recursive is set to true, the runtime state will also be applied to all child Data Layer instances.
+	  * Note:
+	  *  - Changing the runtime state of a Client-Only Data Layer instance must be done on the client side or else it will have no effect.
+	  *  - Changing the runtime state of a Server-Only Data Layer instance can only be done on the server side or else it will have no effect.
+	  *  - Changing the runtime state of a runtime Data Layer instance (with no Load Filter set on the asset) must be done on the server side
+	  *    or else it will have no effect. The runtime state will then be replicated on the client.
+	  * (see Data Layer asset Load Filter for more details)
+	  */
 	UFUNCTION(BlueprintCallable, Category = DataLayers)
 	ENGINE_API bool SetDataLayerRuntimeState(const UDataLayerAsset* InDataLayerAsset, EDataLayerRuntimeState InState, bool bInIsRecursive = false);
 
+	/** Returns the Data Layer instance runtime state. */
 	UFUNCTION(BlueprintCallable, Category = DataLayers)
 	ENGINE_API EDataLayerRuntimeState GetDataLayerInstanceRuntimeState(const UDataLayerInstance* InDataLayerInstance) const;
 
+	/** Finds a matching Data Layer instance referencing the provided Data Layer asset and returns the Data Layer Instance runtime state. */
 	UFUNCTION(BlueprintCallable, Category = DataLayers)
 	ENGINE_API EDataLayerRuntimeState GetDataLayerInstanceEffectiveRuntimeState(const UDataLayerInstance* InDataLayerInstance) const;
 
+	/** Called when a Data Layer instance runtime state has changed. */
 	UPROPERTY(BlueprintAssignable)
 	FOnDataLayerInstanceRuntimeStateChanged OnDataLayerInstanceRuntimeStateChanged;
+	
 	//~ End Blueprint interface
 
 	template<class T>
