Removing / deprecating stale code paths  (DTYPE_Native & DTYPE_Preview) and clawing back raw ptr-based resource management from FAudioDeviceManager.

#rnx

diff --git a/Engine/Source/Runtime/AudioMixer/Private/AudioMixerBuffer.cpp b/Engine/Source/Runtime/AudioMixer/Private/AudioMixerBuffer.cpp
--- a/Engine/Source/Runtime/AudioMixer/Private/AudioMixerBuffer.cpp
+++ b/Engine/Source/Runtime/AudioMixer/Private/AudioMixerBuffer.cpp
@@ -26,15 +26,6 @@ namespace Audio
 	{
 		// Set the base-class NumChannels to wave's NumChannels
 		NumChannels = InWave->NumChannels;
-
-		if (InBufferType != EBufferType::PCMRealTime && (InWave->DecompressionType == EDecompressionType::DTYPE_Native || InWave->DecompressionType == EDecompressionType::DTYPE_Preview))
-		{
-			check(!InWave->RawPCMData || InWave->RawPCMDataSize);
-			Data = InWave->RawPCMData;
-			DataSize = InWave->RawPCMDataSize;
-			NumFrames = DataSize / (sizeof(int16) * NumChannels);
- 			InWave->RawPCMData = nullptr;
-		}
 	}
 
 	FMixerBuffer::~FMixerBuffer()
@@ -215,33 +206,7 @@ namespace Audio
 				Buffer = Init(InAudioDevice, InWave, bForceRealtime);
 			}
 			break;
-
-			case DTYPE_Preview:
-			{
-				// Find any existing buffers
-				if (InWave->ResourceID)
-				{
-					Buffer = (FMixerBuffer*)AudioDeviceManager->GetSoundBufferForResourceID(InWave->ResourceID);
-				}
-
-				// Override with any new PCM data even if the buffer already exists
-				if (InWave->RawPCMData)
-				{
-					// If we already have a buffer for this wave resource, free it
-					if (Buffer)
-					{
-						AudioDeviceManager->FreeBufferResource(Buffer);
-					}
-
-					// Create a new preview buffer
-					Buffer = FMixerBuffer::CreatePreviewBuffer(InAudioDevice, InWave);
-
-					// Track the new created buffer
-					AudioDeviceManager->TrackResource(InWave, Buffer);
-				}
-			}
-			break;
-
+		
 			case DTYPE_Procedural:
 			{
 				// Always create a new buffer for procedural or bus buffers
@@ -255,25 +220,7 @@ namespace Audio
 				Buffer = FMixerBuffer::CreateRealTimeBuffer(InAudioDevice, InWave);
 			}
 			break;
-
-			case DTYPE_Native:
-			{
-				if (InWave->ResourceID)
-				{
-					Buffer = (FMixerBuffer*)AudioDeviceManager->GetSoundBufferForResourceID(InWave->ResourceID);
-				}
-
-				if (Buffer == nullptr)
-				{
-					Buffer = FMixerBuffer::CreateNativeBuffer(InAudioDevice, InWave);
-
-					// Track the resource with the audio device manager
-					AudioDeviceManager->TrackResource(InWave, Buffer);
-					InWave->RemoveAudioResource();
-				}
-			}
-			break;
-
+			
 			case DTYPE_Streaming:
 			{
 				Buffer = FMixerBuffer::CreateStreamingBuffer(InAudioDevice, InWave);
diff --git a/Engine/Source/Runtime/AudioMixer/Private/AudioMixerSource.cpp b/Engine/Source/Runtime/AudioMixer/Private/AudioMixerSource.cpp
--- a/Engine/Source/Runtime/AudioMixer/Private/AudioMixerSource.cpp
+++ b/Engine/Source/Runtime/AudioMixer/Private/AudioMixerSource.cpp
@@ -524,11 +524,12 @@ namespace Audio
 
 		// We've already been passed the wave instance in PrepareForInitialization, make sure we have the same one
 		AUDIO_MIXER_CHECK(WaveInstance && WaveInstance == InWaveInstance);
-
+		AUDIO_MIXER_CHECK(WaveInstance->WaveData);
 		LLM_SCOPE(ELLMTag::AudioMixer);
-
+		
 		FSoundSource::InitCommon();
 
+		NumChannels = WaveInstance->WaveData->NumChannels;
 		if (!ensure(InWaveInstance))
 		{
 			return false;
@@ -550,7 +551,7 @@ namespace Audio
 		{
 			check(!WaveData->RawPCMData || WaveData->RawPCMDataSize);
 			const int32 NumBytes = WaveData->RawPCMDataSize;
-			if (WaveInstance->WaveData->NumChannels > 0)
+			if (NumChannels > 0)
 			{
 				NumFrames = NumBytes / (WaveData->NumChannels * sizeof(int16));
 			}
@@ -1115,7 +1116,6 @@ namespace Audio
 		check(InWaveInstance->WaveData);
 		USoundWave& SoundWave = *InWaveInstance->WaveData;
 
-		Buffer = MixerBuffer;
 		WaveInstance = InWaveInstance;
 
 		LPFFrequency = MAX_FILTER_FREQUENCY;
@@ -1576,7 +1576,6 @@ namespace Audio
 		}
 
 		MixerSourceBuffer.Reset();
-		Buffer = nullptr;
 		bLoopCallback = false;
 		NumTotalFrames = 0;
 
@@ -2028,11 +2027,10 @@ namespace Audio
 		const FAudioPlatformDeviceInfo& DeviceInfo = MixerDevice->GetPlatformDeviceInfo();
 
 		// Compute a new speaker map for each possible output channel mapping for the source
-		const uint32 NumChannels = Buffer->NumChannels;
 		bool bShouldSetMap = false;
 		{
 			FRWScopeLock Lock(ChannelMapLock, SLT_Write);
-			bShouldSetMap = ComputeChannelMap(Buffer->NumChannels, ChannelMap);
+			bShouldSetMap = ComputeChannelMap(GetNumChannels(), ChannelMap);
 		}
 		if(bShouldSetMap)
 		{			
@@ -2115,7 +2113,7 @@ namespace Audio
 		if (WaveInstance->GetUseSpatialization() && (!FMath::IsNearlyEqual(WaveInstance->AbsoluteAzimuth, PreviousAzimuth, 0.01f) || MixerSourceVoice->NeedsSpeakerMap()))
 		{
 			// Make sure our stereo emitter positions are updated relative to the sound emitter position
-			if (Buffer->NumChannels == 2)
+			if (GetNumChannels() == 2)
 			{
 				UpdateStereoEmitterPositions();
 			}
@@ -2276,7 +2274,7 @@ namespace Audio
 
 	bool FMixerSource::UseObjectBasedSpatialization() const
 	{
-		return (Buffer->NumChannels <= MixerDevice->GetCurrentSpatializationPluginInterfaceInfo().MaxChannelsSupportedBySpatializationPlugin &&
+		return (GetNumChannels() <= MixerDevice->GetCurrentSpatializationPluginInterfaceInfo().MaxChannelsSupportedBySpatializationPlugin &&
 				AudioDevice->IsSpatializationPluginEnabled() &&
 				WaveInstance->SpatializationMethod == ESoundSpatializationAlgorithm::SPATIALIZATION_HRTF);
 	}
@@ -2299,28 +2297,28 @@ namespace Audio
 
 	bool FMixerSource::UseSpatializationPlugin() const
 	{
-		return (Buffer->NumChannels <= MixerDevice->GetCurrentSpatializationPluginInterfaceInfo().MaxChannelsSupportedBySpatializationPlugin) &&
+		return (GetNumChannels() <= MixerDevice->GetCurrentSpatializationPluginInterfaceInfo().MaxChannelsSupportedBySpatializationPlugin) &&  
 			AudioDevice->IsSpatializationPluginEnabled() &&
 			WaveInstance->SpatializationPluginSettings != nullptr;
 	}
 
 	bool FMixerSource::UseOcclusionPlugin() const
 	{
-		return (Buffer->NumChannels == 1 || Buffer->NumChannels == 2) &&
+		return (GetNumChannels() == 1 || GetNumChannels() == 2) && 
 			AudioDevice->IsOcclusionPluginEnabled() &&
 			WaveInstance->OcclusionPluginSettings != nullptr;
 	}
 
 	bool FMixerSource::UseReverbPlugin() const
 	{
-		return (Buffer->NumChannels == 1 || Buffer->NumChannels == 2) &&
+		return (GetNumChannels() == 1 || GetNumChannels() == 2) && 
 			AudioDevice->IsReverbPluginEnabled() &&
 			WaveInstance->ReverbPluginSettings != nullptr;
 	}
 
 	bool FMixerSource::UseSourceDataOverridePlugin() const
 	{
-		return (Buffer->NumChannels == 1 || Buffer->NumChannels == 2) &&
+		return (GetNumChannels() == 1 || GetNumChannels() == 2) && 
 			AudioDevice->IsSourceDataOverridePluginEnabled() &&
 			WaveInstance->SourceDataOverridePluginSettings != nullptr;
 	}
diff --git a/Engine/Source/Runtime/Engine/Classes/Sound/SoundWave.h b/Engine/Source/Runtime/Engine/Classes/Sound/SoundWave.h
--- a/Engine/Source/Runtime/Engine/Classes/Sound/SoundWave.h
+++ b/Engine/Source/Runtime/Engine/Classes/Sound/SoundWave.h
@@ -46,8 +46,6 @@ enum EDecompressionType : int
 {
 	DTYPE_Setup,
 	DTYPE_Invalid,
-	DTYPE_Preview,
-	DTYPE_Native,
 	DTYPE_RealTime,
 	DTYPE_Procedural,
 	DTYPE_Xenon,
diff --git a/Engine/Source/Runtime/Engine/Private/Audio.cpp b/Engine/Source/Runtime/Engine/Private/Audio.cpp
--- a/Engine/Source/Runtime/Engine/Private/Audio.cpp
+++ b/Engine/Source/Runtime/Engine/Private/Audio.cpp
@@ -185,11 +185,6 @@ bool IsSpatializationCVarEnabled()
 
 FSoundBuffer::~FSoundBuffer()
 {
-	// remove ourselves from the set of waves that are tracked by the audio device
-	if (ResourceID && GEngine && GEngine->GetAudioDeviceManager())
-	{
-		GEngine->GetAudioDeviceManager()->RemoveSoundBufferForResourceID(ResourceID);
-	}
 }
 
 /**
@@ -433,7 +428,7 @@ void FSoundSource::UpdateStereoEmitterPositions()
 {
 	// Only call this function if we're told to use spatialization
 	check(WaveInstance->GetUseSpatialization());
-	check(Buffer->NumChannels == 2);
+	check(GetNumChannels() == 2);
 
 	if (!DisableStereoSpreadCvar && WaveInstance->StereoSpread > 0.0f)
 	{
@@ -587,7 +582,7 @@ FSpatializationParams FSoundSource::GetSpatializationParams()
 
 		Params.EmitterPosition = EmitterPosition;
 
-		if (Buffer->NumChannels == 2)
+		if (GetNumChannels() == 2)
 		{
 			Params.LeftChannelPosition = AudioDevice->GetListenerTransformedDirection(LeftChannelSourceLocation, nullptr);
 			Params.RightChannelPosition = AudioDevice->GetListenerTransformedDirection(RightChannelSourceLocation, nullptr);			
diff --git a/Engine/Source/Runtime/Engine/Private/AudioDebug.cpp b/Engine/Source/Runtime/Engine/Private/AudioDebug.cpp
--- a/Engine/Source/Runtime/Engine/Private/AudioDebug.cpp
+++ b/Engine/Source/Runtime/Engine/Private/AudioDebug.cpp
@@ -1038,9 +1038,8 @@ namespace Audio
 		if (DeviceManager && DeviceManager->IsVisualizeDebug3dEnabled())
 		{
 			DECLARE_CYCLE_STAT(TEXT("FAudioThreadTask.DrawSourceDebugInfo"), STAT_AudioDrawSourceDebugInfo, STATGROUP_TaskGraphTasks);
-
-			const FSoundBuffer* Buffer = SoundSource.GetBuffer();
-			const bool bSpatialized = Buffer && Buffer->NumChannels == 2 && WaveInstance->GetUseSpatialization();
+			
+			const bool bSpatialized = SoundSource.GetNumChannels() == 2 && WaveInstance->GetUseSpatialization();
 			if (bSpatialized)
 			{
 				const FRotator Rotator = ActiveSound->Transform.GetRotation().Rotator();
diff --git a/Engine/Source/Runtime/Engine/Private/AudioDecompress.cpp b/Engine/Source/Runtime/Engine/Private/AudioDecompress.cpp
--- a/Engine/Source/Runtime/Engine/Private/AudioDecompress.cpp
+++ b/Engine/Source/Runtime/Engine/Private/AudioDecompress.cpp
@@ -771,7 +771,7 @@ void FAsyncAudioDecompressWorker::DoWork()
 			else
 			{
 				LLM_SCOPE(ELLMTag::AudioFullDecompress);
-				check(Wave->DecompressionType == DTYPE_Native || Wave->DecompressionType == DTYPE_Procedural);
+				check(Wave->DecompressionType == DTYPE_Procedural);
 
 				Wave->RawPCMDataSize = QualityInfo.SampleDataSize;
 				check(Wave->RawPCMData == nullptr);
@@ -795,18 +795,6 @@ void FAsyncAudioDecompressWorker::DoWork()
 			Wave->RemoveAudioResource();
 		}
 
-		if (Wave->DecompressionType == DTYPE_Native)
-		{
-			// todo: this code is stale and needs to be gutted since going all-in on stream caching,
-			// but for now, GetOwnedBulkData() is being removed as a function.
-			// FOwnedBulkDataPtr* BulkDataPtr = Wave->GetOwnedBulkData();
-			FOwnedBulkDataPtr* BulkDataPtr = nullptr;
-			UE_CLOG(BulkDataPtr && BulkDataPtr->GetMappedRegion(), LogAudio, Warning, TEXT("Mapped audio (%s) was discarded after decompression. This is not ideal as it takes more load time and doesn't save memory."), *Wave->GetName());
-
-			// Delete the compressed data
-			Wave->RemoveAudioResource();
-		}
-
 		delete AudioInfo;
 	}
 }
diff --git a/Engine/Source/Runtime/Engine/Private/AudioDevice.cpp b/Engine/Source/Runtime/Engine/Private/AudioDevice.cpp
--- a/Engine/Source/Runtime/Engine/Private/AudioDevice.cpp
+++ b/Engine/Source/Runtime/Engine/Private/AudioDevice.cpp
@@ -1040,7 +1040,6 @@ void FAudioDevice::CountBytes(FArchive& Ar)
 	Sources.CountBytes(Ar);
 	// The buffers are stored on the audio device since they are shared amongst all audio devices
 	// Though we are going to count them when querying an individual audio device object about its bytes
-	GEngine->GetAudioDeviceManager()->Buffers.CountBytes(Ar);
 	FreeSources.CountBytes(Ar);
 	WaveInstanceSourceMap.CountBytes(Ar);
 	Ar.CountBytes(sizeof(FWaveInstance) * WaveInstanceSourceMap.Num(), sizeof(FWaveInstance) * WaveInstanceSourceMap.Num());
@@ -1285,12 +1284,6 @@ void FAudioDevice::GetSoundClassInfo(TMap<FName, FAudioClassInfo>& AudioClassInf
 #else
 		switch(SoundWave->DecompressionType)
 		{
-		case DTYPE_Native:
-		case DTYPE_Preview:
-			AudioClassInfo->SizeResident += SoundWave->RawPCMDataSize;
-			AudioClassInfo->NumResident++;
-			break;
-
 		case DTYPE_RealTime:
 			AudioClassInfo->SizeRealTime += SoundWave->GetCompressedDataSize(SoundWave->GetRuntimeFormat());
 			AudioClassInfo->NumRealTime++;
@@ -4557,9 +4550,6 @@ void FAudioDevice::StartSources(TArray<FWaveInstance*>& WaveInstances, int32 Fir
 					}
 					else
 					{
-						// Make sure init cleaned up the buffer when it failed
-						check(Source->Buffer == nullptr);
-
 						// If were ready to call init but failed, then we need to add the source and stop with notification
 						WaveInstance->StopWithoutNotification();
 						FreeSources.Add(Source);
@@ -6684,8 +6674,8 @@ void FAudioDevice::Precache(USoundWave* SoundWave, bool bSynchronous, bool bTrac
 	}
 	else if (SoundWave->RawPCMData)
 	{
-		// Run time created audio; e.g. editor preview data
-		SoundWave->DecompressionType = DTYPE_Preview;
+		// DTYPE_Preview has been removed based on the assumption that this path is unreachable
+		checkNoEntry();
 	}
 	else if (SoundWave->bProcedural)
 	{
@@ -6697,8 +6687,9 @@ void FAudioDevice::Precache(USoundWave* SoundWave, bool bSynchronous, bool bTrac
 		// Buses will initialize as procedural, but not actually become a procedural sound wave
 		SoundWave->DecompressionType = DTYPE_Procedural;
 	}
-	else if (HasCompressedAudioInfoClass(SoundWave))
+	else
 	{
+		check(HasCompressedAudioInfoClass(SoundWave));
 		const FSoundGroup& SoundGroup = GetDefault<USoundGroups>()->GetSoundGroup(SoundWave->SoundGroup);
 
 		if (SoundWave->Duration <= 0.0f)
@@ -6718,30 +6709,21 @@ void FAudioDevice::Precache(USoundWave* SoundWave, bool bSynchronous, bool bTrac
 			SoundWave->DecompressionType = DTYPE_Streaming;
 			SoundWave->bCanProcessAsync = false;
 		}
-		else if (ShouldUseRealtimeDecompression(bForceFullDecompression, SoundGroup, SoundWave, CompressedDurationThreshold))
+		else
 		{
+			check(ShouldUseRealtimeDecompression(bForceFullDecompression, SoundGroup, SoundWave, CompressedDurationThreshold));
+			
 			// Store as compressed data and decompress in realtime
 			SoundWave->DecompressionType = DTYPE_RealTime;
 #if !(UE_BUILD_SHIPPING || UE_BUILD_TEST)
 			++PrecachedRealtime;
 #endif // !(UE_BUILD_SHIPPING || UE_BUILD_TEST)
 		}
-		else
-		{
-			// Fully expand loaded audio data into PCM
-			SoundWave->DecompressionType = DTYPE_Native;
-#if !(UE_BUILD_SHIPPING || UE_BUILD_TEST)
-			++AudioDeviceUtils::PrecachedNative;
-			AverageNativeLength = (AverageNativeLength * (PrecachedNative - 1) + SoundWave->Duration) / PrecachedNative;
-			NativeSampleRateCount.FindOrAdd(SoundWave->GetSampleRateForCurrentPlatform())++;
-			NativeChannelCount.FindOrAdd(SoundWave->NumChannels)++;
-#endif // !(UE_BUILD_SHIPPING || UE_BUILD_TEST)
-		}
 
 		// Grab the compressed audio data
 		SoundWave->InitAudioResource(SoundWave->GetRuntimeFormat());
 
-		if (SoundWave->AudioDecompressor == nullptr && (SoundWave->DecompressionType == DTYPE_Native || SoundWave->DecompressionType == DTYPE_RealTime))
+		if (SoundWave->AudioDecompressor == nullptr && SoundWave->DecompressionType == DTYPE_RealTime)
 		{
 			// Create a worker to decompress the audio data
 			if (bSynchronous)
@@ -6758,19 +6740,8 @@ void FAudioDevice::Precache(USoundWave* SoundWave, bool bSynchronous, bool bTrac
 				SoundWave->AudioDecompressor->StartBackgroundTask();
 				PrecachingSoundWaves.Add(SoundWave);
 			}
-
-			// the audio decompressor will track memory
-			if (SoundWave->DecompressionType == DTYPE_Native)
-			{
-				bTrackMemory = false;
-			}
 		}
 	}
-	else
-	{
-		// Preserve old behavior if there is no compressed audio info class for this audio format
-		SoundWave->DecompressionType = DTYPE_Native;
-	}
 
 	// If we don't have an audio decompressor task, then we're fully precached
 	if (!SoundWave->AudioDecompressor)
@@ -6817,25 +6788,26 @@ bool FAudioDevice::ShouldUseRealtimeDecompression(bool bForceFullDecompression,
 			(ForceRealtimeDecompressionCvar || SoundWave->Duration > CompressedDurationThreshold || (RealtimeDecompressZeroDurationSoundsCvar && SoundWave->Duration <= 0.0f))));
 }
 
+// deprecated
 void FAudioDevice::StopSourcesUsingBuffer(FSoundBuffer* SoundBuffer)
 {
 	TRACE_CPUPROFILER_EVENT_SCOPE(FAudioDevice_StopSourcesUsingBuffer);
 
 	check(IsInAudioThread());
 
-	if (SoundBuffer)
-	{
-		for (int32 SrcIndex = 0; SrcIndex < Sources.Num(); SrcIndex++)
-		{
-			FSoundSource* Src = Sources[SrcIndex];
-			if (Src && Src->Buffer == SoundBuffer)
-			{
-				// Make sure the buffer is no longer referenced by anything
-				Src->StopNow();
-				break;
-			}
-		}
-	}
+	// if (SoundBuffer)
+	// {
+	// 	for (int32 SrcIndex = 0; SrcIndex < Sources.Num(); SrcIndex++)
+	// 	{
+	// 		FSoundSource* Src = Sources[SrcIndex];
+	// 		if (Src && Src->Buffer == SoundBuffer) // deprecated: Src->Buffer no longer exists
+	// 		{
+	// 			// Make sure the buffer is no longer referenced by anything
+	// 			Src->StopNow();
+	// 			break;
+	// 		}
+	// 	}
+	// }
 }
 
 void FAudioDevice::RegisterSoundClass(USoundClass* InSoundClass)
diff --git a/Engine/Source/Runtime/Engine/Private/AudioDeviceManager.cpp b/Engine/Source/Runtime/Engine/Private/AudioDeviceManager.cpp
--- a/Engine/Source/Runtime/Engine/Private/AudioDeviceManager.cpp
+++ b/Engine/Source/Runtime/Engine/Private/AudioDeviceManager.cpp
@@ -189,6 +189,7 @@ FAudioDeviceManager::FAudioDeviceManager()
 #endif //ENABLE_AUDIO_DEBUG
 }
 
+PRAGMA_DISABLE_DEPRECATION_WARNINGS
 FAudioDeviceManager::~FAudioDeviceManager()
 {
 	UE_LOG(LogAudio, Display, TEXT("Beginning Audio Device Manager Shutdown (Module: %s)..."), *AudioMixerModuleName);
@@ -227,13 +228,8 @@ FAudioDeviceManager::~FAudioDeviceManager()
 	MainAudioDeviceHandle.Reset();
 
 	FCoreDelegates::ApplicationWillEnterBackgroundDelegate.RemoveAll(this);
-
-	// Release any loaded buffers - this calls stop on any sources that need it
-	for (int32 Index = Buffers.Num() - 1; Index >= 0; Index--)
-	{
-		FreeBufferResource(Buffers[Index]);
-	}
 }
+PRAGMA_ENABLE_DEPRECATION_WARNINGS
 
 FAudioDevice* FAudioDeviceManager::GetAudioDeviceFromWorldContext(const UObject* WorldContextObject)
 {
@@ -1023,14 +1019,12 @@ uint32 FAudioDeviceManager::GetNewDeviceID()
 	return ++DeviceIDCounter;
 }
 
-void FAudioDeviceManager::StopSourcesUsingBuffer(FSoundBuffer* SoundBuffer)
+// deprecated
+void FAudioDeviceManager::StopSourcesUsingBuffer(FSoundBuffer*)
 {
-	IterateOverAllDevices([SoundBuffer](Audio::FDeviceId Id, FAudioDevice* Device)
-	{
-		Device->StopSourcesUsingBuffer(SoundBuffer);
-	});
 }
 
+// deprecated
 void FAudioDeviceManager::TrackResource(USoundWave* SoundWave, FSoundBuffer* Buffer)
 {
 	// Allocate new resource ID and assign to USoundWave. A value of 0 (default) means not yet registered.
@@ -1038,9 +1032,6 @@ void FAudioDeviceManager::TrackResource(USoundWave* SoundWave, FSoundBuffer* Buf
 	Buffer->ResourceID = ResourceID;
 	SoundWave->ResourceID = ResourceID;
 
-	Buffers.Add(Buffer);
-	WaveBufferMap.Add(ResourceID, Buffer);
-
 #if !(UE_BUILD_SHIPPING || UE_BUILD_TEST)
 	// Keep track of associated resource name.
 	Buffer->ResourceName = SoundWave->GetPathName();
@@ -1051,9 +1042,6 @@ void FAudioDeviceManager::FreeResource(USoundWave* SoundWave)
 {
 	if (SoundWave->ResourceID)
 	{
-		FSoundBuffer* SoundBuffer = WaveBufferMap.FindRef(SoundWave->ResourceID);
-		FreeBufferResource(SoundBuffer);
-
 		// Flag that the sound wave needs to do a full decompress again
 		SoundWave->DecompressionType = DTYPE_Setup;
 		SoundWave->SetPrecacheState(ESoundWavePrecacheState::NotStarted);
@@ -1062,6 +1050,7 @@ void FAudioDeviceManager::FreeResource(USoundWave* SoundWave)
 	}
 }
 
+// deprecated
 void FAudioDeviceManager::FreeBufferResource(FSoundBuffer* SoundBuffer)
 {
 	if (SoundBuffer)
@@ -1069,24 +1058,20 @@ void FAudioDeviceManager::FreeBufferResource(FSoundBuffer* SoundBuffer)
 		// Make sure any realtime tasks are finished that are using this buffer
 		SoundBuffer->EnsureRealtimeTaskCompletion();
 
-		Buffers.Remove(SoundBuffer);
-
-		// Stop any sound sources on any audio device currently using this buffer before deleting
-		StopSourcesUsingBuffer(SoundBuffer);
-
 		delete SoundBuffer;
 		SoundBuffer = nullptr;
 	}
 }
 
+// deprecated
 FSoundBuffer* FAudioDeviceManager::GetSoundBufferForResourceID(uint32 ResourceID)
 {
-	return WaveBufferMap.FindRef(ResourceID);
+	return {};
 }
 
-void FAudioDeviceManager::RemoveSoundBufferForResourceID(uint32 ResourceID)
+// deprecated
+void FAudioDeviceManager::RemoveSoundBufferForResourceID(uint32)
 {
-	WaveBufferMap.Remove(ResourceID);
 }
 
 void FAudioDeviceManager::RemoveSoundMix(USoundMix* SoundMix)
diff --git a/Engine/Source/Runtime/Engine/Private/SoundWave.cpp b/Engine/Source/Runtime/Engine/Private/SoundWave.cpp
--- a/Engine/Source/Runtime/Engine/Private/SoundWave.cpp
+++ b/Engine/Source/Runtime/Engine/Private/SoundWave.cpp
@@ -1000,22 +1000,14 @@ void USoundWave::GetResourceSizeEx(FResourceSizeEx& CumulativeResourceSize)
 	// Finally, report the actual audio memory being used, if this asset isn't using the stream cache.
 	if (FAudioDevice* LocalAudioDevice = GEngine->GetMainAudioDeviceRaw())
 	{
-		if (LocalAudioDevice->HasCompressedAudioInfoClass(this) && DecompressionType == DTYPE_Native)
+		if (DecompressionType == DTYPE_RealTime && CachedRealtimeFirstBuffer)
 		{
-			check(!RawPCMData || RawPCMDataSize);
-			CumulativeResourceSize.AddDedicatedSystemMemoryBytes(RawPCMDataSize);
+			CumulativeResourceSize.AddDedicatedSystemMemoryBytes(MONO_PCM_BUFFER_SIZE * NumChannels);
 		}
-		else
-		{
-			if (DecompressionType == DTYPE_RealTime && CachedRealtimeFirstBuffer)
-			{
-				CumulativeResourceSize.AddDedicatedSystemMemoryBytes(MONO_PCM_BUFFER_SIZE * NumChannels);
-			}
 
-			if (!FPlatformProperties::SupportsAudioStreaming() || !IsStreaming(nullptr))
-			{
-				CumulativeResourceSize.AddDedicatedSystemMemoryBytes(GetCompressedDataSize(GetRuntimeFormat()));
-			}
+		if (!FPlatformProperties::SupportsAudioStreaming() || !IsStreaming(nullptr))
+		{
+			CumulativeResourceSize.AddDedicatedSystemMemoryBytes(GetCompressedDataSize(GetRuntimeFormat()));
 		}
 	}
 }
diff --git a/Engine/Source/Runtime/Engine/Public/Audio.h b/Engine/Source/Runtime/Engine/Public/Audio.h
--- a/Engine/Source/Runtime/Engine/Public/Audio.h
+++ b/Engine/Source/Runtime/Engine/Public/Audio.h
@@ -610,7 +610,6 @@ public:
 	FSoundSource(FAudioDevice* InAudioDevice)
 		: AudioDevice(InAudioDevice)
 		, WaveInstance(nullptr)
-		, Buffer(nullptr)
 		, LFEBleed(0.5f)
 		, LPFFrequency(MAX_FILTER_FREQUENCY)
 		, HPFFrequency(MIN_FILTER_FREQUENCY)
@@ -715,7 +714,10 @@ public:
 	ENGINE_API FSpatializationParams GetSpatializationParams();
 
 	/** Returns the contained sound buffer object. */
-	virtual const FSoundBuffer* GetBuffer() const { return Buffer; }
+	UE_DEPRECATED(5.5, "The Buffer member no longer exists.  Use GetNumChannels() to qeury channel count.")
+	virtual const FSoundBuffer* GetBuffer() const { return nullptr; }
+
+	int32 GetNumChannels() const { return NumChannels; }
 
 	/** Initializes any source effects for this sound source. */
 	virtual void InitializeSourceEffects(uint32 InEffectVoiceId)
@@ -777,6 +779,7 @@ protected:
 	FWaveInstance* WaveInstance;
 
 	/** Cached sound buffer associated with currently bound wave instance. */
+	UE_DEPRECATED(5.5, "This is no longer used and should not be accessed directly.")
 	FSoundBuffer* Buffer;
 
 	/** The amount of a sound to bleed to the LFE speaker */
@@ -824,6 +827,9 @@ protected:
 	/** The frame we started on. */
 	int32 StartFrame;
 
+	/** the number of channels */
+	int32 NumChannels = 0;
+
 	/** Effect ID of this sound source in the audio device sound source array. */
 	uint32 VoiceId;
 
diff --git a/Engine/Source/Runtime/Engine/Public/AudioDevice.h b/Engine/Source/Runtime/Engine/Public/AudioDevice.h
--- a/Engine/Source/Runtime/Engine/Public/AudioDevice.h
+++ b/Engine/Source/Runtime/Engine/Public/AudioDevice.h
@@ -681,6 +681,7 @@ public:
 	*
 	* @param	FSoundBuffer	Buffer to check against
 	*/
+	UE_DEPRECATED(5.5, "No longer needed for resource managment, existing calls may be deleted")
 	ENGINE_API void StopSourcesUsingBuffer(FSoundBuffer * SoundBuffer);
 
 	/**
diff --git a/Engine/Source/Runtime/Engine/Public/AudioDeviceManager.h b/Engine/Source/Runtime/Engine/Public/AudioDeviceManager.h
--- a/Engine/Source/Runtime/Engine/Public/AudioDeviceManager.h
+++ b/Engine/Source/Runtime/Engine/Public/AudioDeviceManager.h
@@ -221,21 +221,26 @@ public:
 	ENGINE_API void SetSoloDevice(Audio::FDeviceId InAudioDeviceHandle);
 
 	/** Links up the resource data indices for looking up and cleaning up. */
+	UE_DEPRECATED(5.5, "This is no longer used.  Existing calls may be deleted")
 	ENGINE_API void TrackResource(USoundWave* SoundWave, FSoundBuffer* Buffer);
 
 	/** Frees the given sound wave resource from the device manager */
 	ENGINE_API void FreeResource(USoundWave* SoundWave);
 
 	/** Frees the sound buffer from the device manager. */
+	UE_DEPRECATED(5.5, "This is no longer used.  Existing calls may be deleted")
 	ENGINE_API void FreeBufferResource(FSoundBuffer* SoundBuffer);
 
 	/** Stops using the given sound buffer. Called before freeing the buffer */
+	UE_DEPRECATED(5.5, "This is no longer used.  Existing calls may be deleted")
 	ENGINE_API void StopSourcesUsingBuffer(FSoundBuffer* Buffer);
 
 	/** Retrieves the sound buffer for the given resource id */
+	UE_DEPRECATED(5.5, "This is no longer used.  Existing calls may be deleted")
 	ENGINE_API FSoundBuffer* GetSoundBufferForResourceID(uint32 ResourceID);
 
 	/** Removes the sound buffer for the given resource id */
+	UE_DEPRECATED(5.5, "This is no longer used.  Existing calls may be deleted")
 	ENGINE_API void RemoveSoundBufferForResourceID(uint32 ResourceID);
 
 	/** Removes sound mix from all audio devices */
@@ -273,9 +278,11 @@ public:
 public:
 
 	/** Array of all created buffers */
+	UE_DEPRECATED(5.5, "This is no longer used and should not be accessed directly.")
 	TArray<FSoundBuffer*>			Buffers;
-
+	
 	/** Look up associating a USoundWave's resource ID with sound buffers	*/
+	UE_DEPRECATED(5.5, "This is no longer used and should not be accessed directly.")
 	TMap<int32, FSoundBuffer*>	WaveBufferMap;
 
 	/** Returns all the audio devices managed by device manager. */
diff --git a/Engine/Source/Runtime/IOS/IOSAudio/Private/IOSAudioBuffer.cpp b/Engine/Source/Runtime/IOS/IOSAudio/Private/IOSAudioBuffer.cpp
--- a/Engine/Source/Runtime/IOS/IOSAudio/Private/IOSAudioBuffer.cpp
+++ b/Engine/Source/Runtime/IOS/IOSAudio/Private/IOSAudioBuffer.cpp
@@ -182,9 +182,7 @@ FIOSAudioSoundBuffer* FIOSAudioSoundBuffer::Init(FIOSAudioDevice* IOSAudioDevice
 			Buffer = new FIOSAudioSoundBuffer(IOSAudioDevice, InWave, bIsStreaming, true);
 			break;
 
-		case DTYPE_Native:
 		case DTYPE_Invalid:
-		case DTYPE_Preview:
 		default:
 			// Invalid will be set if the wave cannot be played
 			UE_LOG( LogIOSAudio, Warning, TEXT("Init Buffer on unsupported sound type name = %s type = %d"), *InWave->GetName(), int32(InWave->DecompressionType));
diff --git a/Engine/Source/Runtime/IOS/IOSAudio/Private/IOSAudioDevice.cpp b/Engine/Source/Runtime/IOS/IOSAudio/Private/IOSAudioDevice.cpp
--- a/Engine/Source/Runtime/IOS/IOSAudio/Private/IOSAudioDevice.cpp
+++ b/Engine/Source/Runtime/IOS/IOSAudio/Private/IOSAudioDevice.cpp
@@ -104,7 +104,7 @@ FIOSAudioDevice::FIOSAudioDevice() :
 	MixerUnit(NULL),
 	NextBusNumber(0)
 {
-	bDisableAudioCaching = true;	// Do not allow DTYPE_Native buffers, only DTYPE_RealTime or DTYPE_Streaming since on the fly decompression is so cheap, it saves memory, and requires fewer code paths
+	bDisableAudioCaching = true;	// Only use DTYPE_RealTime or DTYPE_Streaming since on the fly decompression is so cheap, it saves memory, and requires fewer code paths
 }
 
 bool FIOSAudioDevice::Exec(UWorld* InWorld, const TCHAR* Cmd, FOutputDevice& Ar)
