Fix for saving textures without loading them when they need PreSave work done that changes the DDC key. Resetting the linker would call update resource which could potentially call FinishCache if the async task had completed fast enough, however the texture had been registered with the texture compiler during the immediately preceding PostLoad ultimately leading in a classic AddTextures crash.

diff --git a/Engine/Source/Runtime/Engine/Private/Texture2D.cpp b/Engine/Source/Runtime/Engine/Private/Texture2D.cpp
--- a/Engine/Source/Runtime/Engine/Private/Texture2D.cpp
+++ b/Engine/Source/Runtime/Engine/Private/Texture2D.cpp
@@ -795,7 +795,8 @@ void UTexture2D::PostLinkerChange()
 	// Changing the linker requires re-creating the resource to make sure streaming behavior is right.
 	if( !HasAnyFlags( RF_BeginDestroyed | RF_NeedLoad | RF_NeedPostLoad ) && !IsUnreachable() )
 	{
-		// Update the resource.
+		// Update the resource. We have to complete any builds beforehand.
+		BlockOnAnyAsyncBuild();
 		UpdateResource();
 	}
 }
