Material Layer Prototype : FIxing material graph building taking into account activations

diff --git a/Engine/Source/Runtime/Engine/Private/Materials/MaterialExpressions.cpp b/Engine/Source/Runtime/Engine/Private/Materials/MaterialExpressions.cpp
--- a/Engine/Source/Runtime/Engine/Private/Materials/MaterialExpressions.cpp
+++ b/Engine/Source/Runtime/Engine/Private/Materials/MaterialExpressions.cpp
@@ -8434,15 +8434,6 @@ void UMaterialExpressionMaterialAttributeLayers::RebuildLayerGraph(bool bReportE
 					LayerCallers[LayerIndex]->Outputs.Add(FunctionOutput.Output);
 				}
 
-				// Optional: Single material attributes input, the base input to the stack
-				if (LayerCallers[LayerIndex]->FunctionInputs.Num() > 0)
-				{
-					if (Input.GetTracedInput().Expression)
-					{
-						LayerCallers[LayerIndex]->FunctionInputs[0].Input = Input;
-					}
-				}
-
 				// Recursively run through internal functions to allow connection of inputs/outputs
 				LayerCallers[LayerIndex]->UpdateFromFunctionResource();
 
@@ -8499,33 +8490,42 @@ void UMaterialExpressionMaterialAttributeLayers::RebuildLayerGraph(bool bReportE
 
 			UMaterialExpressionMaterialFunctionCall* AssembleBranchGraph(int32 InSourceId)
 			{
+				// Fetch the Layer function of this node as the base background layer
+				// Always null for the root node.	
 				UMaterialExpressionMaterialFunctionCall* PreviousLayerInput = nullptr;
-
 				if (InSourceId != -1)
 				{
 					int32 BackgroundLayerFuncIdx = Tree->GetLayerFuncIndex(InSourceId);
 					PreviousLayerInput = (This->LayerCallers.IsValidIndex(BackgroundLayerFuncIdx) ? This->LayerCallers[BackgroundLayerFuncIdx] : nullptr);
+					PreviousLayerInput = (PreviousLayerInput && PreviousLayerInput->MaterialFunction ? PreviousLayerInput : nullptr);
 				}
 
+				// For child of this node, try to assemble the blend/layer tree 
 				auto NodeChildrenIds = Tree->GetNodeChildren(InSourceId);
 				for (auto NodeId : NodeChildrenIds)
 				{
+					// Assemble the branch of layer / blend under this sub node by calling recursively  into that sub node
 					UMaterialExpressionMaterialFunctionCall* CurrentLayerInput = AssembleBranchGraph(NodeId);
 
+					// Fetch the blend function of the node
 					int32 BlendFuncIdx = Tree->GetBlendFuncIndex(NodeId);
 					UMaterialExpressionMaterialFunctionCall* CurrentBlendCaller = (This->BlendCallers.IsValidIndex(BlendFuncIdx) ? This->BlendCallers[BlendFuncIdx] : nullptr);
+					CurrentBlendCaller = (CurrentBlendCaller && CurrentBlendCaller->MaterialFunction ? CurrentBlendCaller : nullptr);
 
-					if (CurrentBlendCaller && CurrentBlendCaller->MaterialFunction)
+					if (CurrentBlendCaller) // Blend is the new current layer if valid and connect the previous layer and current layer is valid
 					{
-						CurrentBlendCaller->FunctionInputs[0].Input.Connect(0, PreviousLayerInput);
-						CurrentBlendCaller->FunctionInputs[1].Input.Connect(0, CurrentLayerInput);
+						if (PreviousLayerInput)
+							CurrentBlendCaller->FunctionInputs[0].Input.Connect(0, PreviousLayerInput);
+						if (CurrentLayerInput)
+							CurrentBlendCaller->FunctionInputs[1].Input.Connect(0, CurrentLayerInput);
 
 						PreviousLayerInput = CurrentBlendCaller;
 					}
-					else
+					else if (CurrentLayerInput) // No Blend then the new current layer if valid
 					{
 						PreviousLayerInput = ToRawPtr(CurrentLayerInput);
 					}
+					//else this node's layer and and blend functions aren't valid, ignore and move on to the next sub node
 				}
 
 				return PreviousLayerInput;
@@ -8965,27 +8965,7 @@ FSubstrateOperator* UMaterialExpressionMaterialAttributeLayers::SubstrateGenerat
 		{
 			OutOperator = OutputCaller->SubstrateGenerateMaterialTopologyTree(Compiler, Parent, 0);
 		}
-/*
-		if (NumActiveBlendCallers > 0 && BlendCallers[NumActiveBlendCallers - 1]->MaterialFunction)
-		{
-			// Multiple blended layers
-			OutOperator = BlendCallers[NumActiveBlendCallers - 1]->SubstrateGenerateMaterialTopologyTree(Compiler, Parent, 0);
-		}
-		else if (NumActiveLayerCallers > 0 && LayerCallers[NumActiveLayerCallers - 1]->MaterialFunction)
-		{
-			// Single layer
-			OutOperator = LayerCallers[NumActiveLayerCallers - 1]->SubstrateGenerateMaterialTopologyTree(Compiler, Parent, 0);
-		} */
-		else if (NumActiveLayerCallers == 0)
-		{
-			// Pass-through
-			const FGuid AttributeID = Compiler->GetMaterialAttribute();
-			if (Input.GetTracedInput().Expression)
-			{
-				OutOperator = Input.Expression->SubstrateGenerateMaterialTopologyTree(Compiler, Parent, 0);
-			}
-		}
-
+		
 		if (!OutOperator)
 		{
 			//If this is reached, compile a default operator to avoid a crash
