Allow skeletal mesh component's animation to be completely disabled

Introduce a flag to control enable/disabled state of built-in animation in skeletal mesh components. This also opts-in to threaded tick.
Fix motion blur and refpose initialization for skeletal meshes that are controlled by AnimNext.
Note: still subject to component registration ordering issues - components can be registered in any order and depending on which is done first, either the reference pose or the AnimNext pose will be displayed (in editor - in game the component tick will correctly pose the mesh)

#rnx

diff --git a/Engine/Plugins/Experimental/AnimNext/Source/AnimNext/Private/Component/SkinnedMeshComponentExtensions.cpp b/Engine/Plugins/Experimental/AnimNext/Source/AnimNext/Private/Component/SkinnedMeshComponentExtensions.cpp
--- a/Engine/Plugins/Experimental/AnimNext/Source/AnimNext/Private/Component/SkinnedMeshComponentExtensions.cpp
+++ b/Engine/Plugins/Experimental/AnimNext/Source/AnimNext/Private/Component/SkinnedMeshComponentExtensions.cpp
@@ -17,7 +17,7 @@ void FSkinnedMeshComponentExtensions::CompleteAndDispatch(USkinnedMeshComponent*
 
 		// Flag buffer for flip
 		InComponent->bNeedToFlipSpaceBaseBuffers = true;
-		InComponent->bHasValidBoneTransform = false;
+
 		InComponent->FlipEditableSpaceBases();
 		InComponent->bHasValidBoneTransform = true;
 
diff --git a/Engine/Source/Editor/DetailCustomizations/Private/SkeletalMeshComponentDetails.cpp b/Engine/Source/Editor/DetailCustomizations/Private/SkeletalMeshComponentDetails.cpp
--- a/Engine/Source/Editor/DetailCustomizations/Private/SkeletalMeshComponentDetails.cpp
+++ b/Engine/Source/Editor/DetailCustomizations/Private/SkeletalMeshComponentDetails.cpp
@@ -142,10 +142,9 @@ void FSkeletalMeshComponentDetails::UpdateAnimationCategory(IDetailLayoutBuilder
 	// Place the blueprint property next (which may be hidden, depending on the mode)
 	TAttribute<EVisibility> BlueprintVisibility( this, &FSkeletalMeshComponentDetails::VisibilityForBlueprintMode );
 
-	DetailBuilder.HideProperty(AnimationBlueprintHandle);
-	AnimationCategory.AddCustomRow(AnimationBlueprintHandle->GetPropertyDisplayName())
-		.RowTag(AnimationBlueprintName)
+	AnimationCategory.AddProperty(AnimationBlueprintHandle)
 		.Visibility(BlueprintVisibility)
+		.CustomWidget()
 		.NameContent()
 		[
 			AnimationBlueprintHandle->CreatePropertyNameWidget()
@@ -169,22 +168,6 @@ void FSkeletalMeshComponentDetails::UpdateAnimationCategory(IDetailLayoutBuilder
 					.MinDesiredWidth(200.f)
 				]
 			]
-			+ SHorizontalBox::Slot()
-			.AutoWidth()
-			.HAlign(HAlign_Center)
-			.VAlign(VAlign_Center)
-			.Padding(2.0f, 1.0f)
-			[
-				PropertyCustomizationHelpers::MakeUseSelectedButton(FSimpleDelegate::CreateSP(this, &FSkeletalMeshComponentDetails::UseSelectedAnimBlueprint))
-			]
-			+ SHorizontalBox::Slot()
-			.AutoWidth()
-			.HAlign(HAlign_Center)
-			.VAlign(VAlign_Center)
-			.Padding(2.0f, 1.0f)
-			[
-				PropertyCustomizationHelpers::MakeBrowseButton(FSimpleDelegate::CreateSP(this, &FSkeletalMeshComponentDetails::OnBrowseToAnimBlueprint))
-			]
 		];
 
 	// Hide the parent AnimationData property, and inline the children with custom visibility delegates
diff --git a/Engine/Source/Runtime/Engine/Classes/Components/SkeletalMeshComponent.h b/Engine/Source/Runtime/Engine/Classes/Components/SkeletalMeshComponent.h
--- a/Engine/Source/Runtime/Engine/Classes/Components/SkeletalMeshComponent.h
+++ b/Engine/Source/Runtime/Engine/Classes/Components/SkeletalMeshComponent.h
@@ -344,7 +344,7 @@ public:
 	TObjectPtr<class UAnimBlueprintGeneratedClass> AnimBlueprintGeneratedClass;
 
 	/* The AnimBlueprint class to use. Use 'SetAnimInstanceClass' to change at runtime. */
-	UPROPERTY(EditAnywhere, BlueprintReadOnly, Category = Animation)
+	UPROPERTY(EditAnywhere, BlueprintReadOnly, Category = Animation, meta=(EditCondition = bEnableAnimation))
 	class TSubclassOf<UAnimInstance> AnimClass;
 
 	/** The active animation graph program instance. */
@@ -399,7 +399,7 @@ public:
 	UFUNCTION(BlueprintSetter)
 	ENGINE_API void SetDisablePostProcessBlueprint(bool bInDisablePostProcess);
 
-	UPROPERTY(EditAnywhere, BlueprintReadWrite, Category=Animation, meta=(ShowOnlyInnerProperties))
+	UPROPERTY(EditAnywhere, BlueprintReadWrite, Category=Animation, meta=(ShowOnlyInnerProperties, EditCondition = bEnableAnimation))
 	struct FSingleAnimationPlayData AnimationData;
 
 	// this is explicit copy because this buffer is reused during evaluation
@@ -562,7 +562,7 @@ protected:
 
 public:
 	/** Used to scale speed of all animations on this skeletal mesh. */
-	UPROPERTY(EditAnywhere, AdvancedDisplay, BlueprintReadWrite, Category=Animation)
+	UPROPERTY(EditAnywhere, AdvancedDisplay, BlueprintReadWrite, Category=Animation, meta=(EditCondition = bEnableAnimation))
 	float GlobalAnimRateScale;
 	
 	/** If we are running physics, should we update non-simulated bones based on the animation bone positions. */
@@ -578,7 +578,7 @@ public:
 
 protected:
 	/** Whether to use Animation Blueprint or play Single Animation Asset. */
-	UPROPERTY(EditAnywhere, BlueprintReadOnly, Category=Animation)
+	UPROPERTY(EditAnywhere, BlueprintReadOnly, Category=Animation, meta=(EditCondition = bEnableAnimation))
 	TEnumAsByte<EAnimationMode::Type>	AnimationMode;
 public:
 	// helper function to get the member name and verify it exists, without making it public
@@ -596,7 +596,7 @@ private:
 	/** Controls whether or not this component will evaluate its post process instance. The post-process
 	 *  Instance is dictated by the skeletal mesh so this is used for per-instance control.
 	 */
-	UPROPERTY(EditAnywhere, BlueprintGetter=GetDisablePostProcessBlueprint, BlueprintSetter=SetDisablePostProcessBlueprint, Category = Animation)
+	UPROPERTY(EditAnywhere, BlueprintGetter=GetDisablePostProcessBlueprint, BlueprintSetter=SetDisablePostProcessBlueprint, Category = Animation, meta=(EditCondition = bEnableAnimation))
 	uint8 bDisablePostProcessBlueprint:1;
 
 public:
@@ -725,11 +725,19 @@ public:
 	uint8 bNoSkeletonUpdate:1;
 
 	/** pauses this component's animations (doesn't tick them, but still refreshes bones) */
-	UPROPERTY(EditAnywhere, AdvancedDisplay, BlueprintReadWrite, Category=Animation)
+	UPROPERTY(EditAnywhere, AdvancedDisplay, BlueprintReadWrite, Category=Animation, meta=(EditCondition = bEnableAnimation))
 	uint8 bPauseAnims:1;
 
+	/**
+	 * Whether the built-in animation of this component should run when the component ticks.
+	 * It is assumed that if this is false then some external system will be animating this mesh.
+	 * Note that disabling animation will also cause cloth simulation not to run and the component's tick to run on any thread. 
+	 */
+	UPROPERTY(EditAnywhere, AdvancedDisplay, BlueprintReadOnly, Category=Animation)
+	uint8 bEnableAnimation : 1;
+
 	/** On InitAnim should we set to ref pose (if false use first tick of animation data). If enabled, takes precedence over UAnimationSettings::bTickAnimationOnSkeletalMeshInit*/
-	UPROPERTY(EditAnywhere, AdvancedDisplay, Category = Animation)
+	UPROPERTY(EditAnywhere, AdvancedDisplay, Category = Animation, meta=(EditCondition = bEnableAnimation))
 	uint8 bUseRefPoseOnInitAnim:1;
 
 	/**
@@ -973,6 +981,9 @@ ENGINE_API PRAGMA_ENABLE_DEPRECATION_WARNINGS
 	/** Clear the linked anim instances and mark them pending kill */
 	void ResetLinkedAnimInstances();
 
+	/** Internal helper - copies the mesh's reference pose to the local space transforms and regenerates component space transforms accordingly */ 
+	void ResetToRefPose();
+
 public:
 	UE_DEPRECATED(4.23, "This function is deprecated. Please use GetLinkedAnimGraphInstanceByTag")
 	UAnimInstance* GetSubInstanceByName(FName InTag) const { return GetLinkedAnimGraphInstanceByTag(InTag); }
diff --git a/Engine/Source/Runtime/Engine/Private/Components/SkeletalMeshComponent.cpp b/Engine/Source/Runtime/Engine/Private/Components/SkeletalMeshComponent.cpp
--- a/Engine/Source/Runtime/Engine/Private/Components/SkeletalMeshComponent.cpp
+++ b/Engine/Source/Runtime/Engine/Private/Components/SkeletalMeshComponent.cpp
@@ -338,6 +338,7 @@ USkeletalMeshComponent::USkeletalMeshComponent(const FObjectInitializer& ObjectI
 	RagdollAggregateThreshold = UPhysicsSettings::Get()->RagdollAggregateThreshold;
 
 	bUpdateMeshWhenKinematic = false;
+	bEnableAnimation = true;
 
 	LastPoseTickFrame = 0u;
 
@@ -495,6 +496,8 @@ void USkeletalMeshComponent::PostInitProperties()
 			AnimationMode = EAnimationMode::AnimationSingleNode;
 		}
 #endif
+
+		PrimaryComponentTick.bRunOnAnyThread = !bEnableAnimation;
 	}
 }
 
@@ -738,12 +741,19 @@ void USkeletalMeshComponent::OnRegister()
 	// to correctly populate that list.
 	ResetLinkedAnimInstances();
 
-	// We force an initialization here because we're in one of two cases.
-	// 1) First register, no spawned instance, need to initialize
-	// 2) We're being re-registered, in which case we've went through
-	// OnUnregister and unconditionally uninitialized our anim instances
-	// so we need to force initialize them before we begin to tick.
-	InitAnim(true);
+	if(bEnableAnimation)
+	{
+		// We force an initialization here because we're in one of two cases.
+		// 1) First register, no spawned instance, need to initialize
+		// 2) We're being re-registered, in which case we've went through
+		// OnUnregister and unconditionally uninitialized our anim instances
+		// so we need to force initialize them before we begin to tick.
+		InitAnim(true);
+	}
+	else
+	{
+		ResetToRefPose();
+	}
 
 	if (bRenderStatic || (VisibilityBasedAnimTickOption == EVisibilityBasedAnimTickOption::OnlyTickPoseWhenRendered && !FApp::CanEverRender()))
 	{
@@ -906,6 +916,11 @@ void USkeletalMeshComponent::InitAnim(bool bForceReinit)
 	QUICK_SCOPE_CYCLE_COUNTER(STAT_SkelMeshComp_InitAnim);
 	LLM_SCOPE(ELLMTag::Animation);
 
+	if(!bEnableAnimation)
+	{
+		return;
+	}
+	
 	// a lot of places just call InitAnim without checking Mesh, so 
 	// I'm moving the check here
 	if ( GetSkeletalMeshAsset() != nullptr && IsRegistered() )
@@ -964,13 +979,7 @@ void USkeletalMeshComponent::InitAnim(bool bForceReinit)
 				}
 				else
 				{
-					PRAGMA_DISABLE_DEPRECATION_WARNINGS
-					BoneSpaceTransforms = GetSkeletalMeshAsset()->GetRefSkeleton().GetRefBonePose();
-					//Mini RefreshBoneTransforms (the bit we actually care about)
-					GetSkeletalMeshAsset()->FillComponentSpaceTransforms(BoneSpaceTransforms, FillComponentSpaceTransformsRequiredBones, GetEditableComponentSpaceTransforms());
-					PRAGMA_ENABLE_DEPRECATION_WARNINGS
-					bNeedToFlipSpaceBaseBuffers = true; // Have updated space bases so need to flip
-					FlipEditableSpaceBases();
+					ResetToRefPose();
 				}
 
 				if (bInitializedAnimInstance)
@@ -1408,6 +1417,11 @@ bool USkeletalMeshComponent::ShouldOnlyTickMontages(const float DeltaTime) const
 
 void USkeletalMeshComponent::TickAnimation(float DeltaTime, bool bNeedsValidRootMotion)
 {
+	if(!bEnableAnimation)
+	{
+		return;
+	}
+
 	SCOPED_NAMED_EVENT(USkeletalMeshComponent_TickAnimation, FColor::Yellow);
 	SCOPE_CYCLE_COUNTER(STAT_AnimGameThreadTime);
 	SCOPE_CYCLE_COUNTER(STAT_AnimTickTime);
@@ -1456,6 +1470,11 @@ void USkeletalMeshComponent::SetPredictedLODLevel(int32 InPredictedLODLevel)
 
 void USkeletalMeshComponent::TickAnimInstances(float DeltaTime, bool bNeedsValidRootMotion)
 {
+	if(!bEnableAnimation)
+	{
+		return;
+	}
+
 	// Allow animation instance to do some processing before the linked instances update
 	if (AnimScriptInstance != nullptr)
 	{
@@ -2423,6 +2442,11 @@ void USkeletalMeshComponent::DoInstancePostEvaluation()
 
 void USkeletalMeshComponent::RefreshBoneTransforms(FActorComponentTickFunction* TickFunction)
 {
+	if(!bEnableAnimation)
+	{
+		return;
+	}
+
 	SCOPE_CYCLE_COUNTER(STAT_AnimGameThreadTime);
 	SCOPE_CYCLE_COUNTER(STAT_RefreshBoneTransforms);
 
@@ -3244,6 +3268,17 @@ void USkeletalMeshComponent::ResetLinkedAnimInstances()
 	LinkedInstances.Reset();
 }
 
+void USkeletalMeshComponent::ResetToRefPose()
+{
+	PRAGMA_DISABLE_DEPRECATION_WARNINGS
+	BoneSpaceTransforms = GetSkeletalMeshAsset()->GetRefSkeleton().GetRefBonePose();
+	//Mini RefreshBoneTransforms (the bit we actually care about)
+	GetSkeletalMeshAsset()->FillComponentSpaceTransforms(BoneSpaceTransforms, FillComponentSpaceTransformsRequiredBones, GetEditableComponentSpaceTransforms());
+	PRAGMA_ENABLE_DEPRECATION_WARNINGS
+	bNeedToFlipSpaceBaseBuffers = true; // Have updated space bases so need to flip
+	FlipEditableSpaceBases();
+}
+
 void USkeletalMeshComponent::AllowQueuedAnimEventsNextDispatch()
 {
 	bNeedsQueuedAnimEventsDispatched = true;
@@ -3251,7 +3286,7 @@ void USkeletalMeshComponent::AllowQueuedAnimEventsNextDispatch()
 
 UAnimInstance* USkeletalMeshComponent::GetLinkedAnimGraphInstanceByTag(FName InName) const
 {
-	if(AnimScriptInstance)
+	if(bEnableAnimation && AnimScriptInstance)
 	{
 		return AnimScriptInstance->GetLinkedAnimGraphInstanceByTag(InName);
 	}
@@ -3260,7 +3295,7 @@ UAnimInstance* USkeletalMeshComponent::GetLinkedAnimGraphInstanceByTag(FName InN
 
 void USkeletalMeshComponent::GetLinkedAnimGraphInstancesByTag(FName InTag, TArray<UAnimInstance*>& OutLinkedInstances) const
 {
-	if(AnimScriptInstance)
+	if(bEnableAnimation && AnimScriptInstance)
 	{
 		PRAGMA_DISABLE_DEPRECATION_WARNINGS
 		AnimScriptInstance->GetLinkedAnimGraphInstancesByTag(InTag, OutLinkedInstances);
@@ -3270,7 +3305,7 @@ void USkeletalMeshComponent::GetLinkedAnimGraphInstancesByTag(FName InTag, TArra
 
 void USkeletalMeshComponent::LinkAnimGraphByTag(FName InTag, TSubclassOf<UAnimInstance> InClass)
 {
-	if(AnimScriptInstance)
+	if(bEnableAnimation && AnimScriptInstance)
 	{
 		AnimScriptInstance->LinkAnimGraphByTag(InTag, InClass);
 	}
@@ -3278,7 +3313,7 @@ void USkeletalMeshComponent::LinkAnimGraphByTag(FName InTag, TSubclassOf<UAnimIn
 
 void USkeletalMeshComponent::LinkAnimClassLayers(TSubclassOf<UAnimInstance> InClass)
 {
-	if(AnimScriptInstance)
+	if(bEnableAnimation && AnimScriptInstance)
 	{
 		AnimScriptInstance->LinkAnimClassLayers(InClass);
 	}
@@ -3286,7 +3321,7 @@ void USkeletalMeshComponent::LinkAnimClassLayers(TSubclassOf<UAnimInstance> InCl
 
 void USkeletalMeshComponent::UnlinkAnimClassLayers(TSubclassOf<UAnimInstance> InClass)
 {
-	if(AnimScriptInstance)
+	if(bEnableAnimation && AnimScriptInstance)
 	{
 		AnimScriptInstance->UnlinkAnimClassLayers(InClass);
 	}
@@ -3294,7 +3329,7 @@ void USkeletalMeshComponent::UnlinkAnimClassLayers(TSubclassOf<UAnimInstance> In
 
 UAnimInstance* USkeletalMeshComponent::GetLinkedAnimLayerInstanceByGroup(FName InGroup) const
 {
-	if(AnimScriptInstance)
+	if(bEnableAnimation && AnimScriptInstance)
 	{
 		return AnimScriptInstance->GetLinkedAnimLayerInstanceByGroup(InGroup);
 	}
@@ -3303,7 +3338,7 @@ UAnimInstance* USkeletalMeshComponent::GetLinkedAnimLayerInstanceByGroup(FName I
 
 UAnimInstance* USkeletalMeshComponent::GetLinkedAnimLayerInstanceByClass(TSubclassOf<UAnimInstance> InClass) const
 {
-	if(AnimScriptInstance)
+	if(bEnableAnimation && AnimScriptInstance)
 	{
 		return AnimScriptInstance->GetLinkedAnimLayerInstanceByClass(InClass);
 	}
@@ -3529,6 +3564,12 @@ void USkeletalMeshComponent::GetResourceSizeEx(FResourceSizeEx& CumulativeResour
 
 void USkeletalMeshComponent::SetAnimationMode(EAnimationMode::Type InAnimationMode, bool bForceInitAnimScriptInstance)
 {
+	if(!bEnableAnimation)
+	{
+		UE_LOG(LogAnimation, Warning, TEXT("SetAnimationMode: Animation is currently disabled"));
+		return;
+	}
+
 	const bool bNeedChange = AnimationMode != InAnimationMode;
 	if (bNeedChange)
 	{
@@ -3555,6 +3596,12 @@ EAnimationMode::Type USkeletalMeshComponent::GetAnimationMode() const
 
 void USkeletalMeshComponent::PlayAnimation(class UAnimationAsset* NewAnimToPlay, bool bLooping)
 {
+	if(!bEnableAnimation)
+	{
+		UE_LOG(LogAnimation, Warning, TEXT("PlayAnimation: Animation is currently disabled"));
+		return;
+	}
+
 	SetAnimationMode(EAnimationMode::AnimationSingleNode);
 	SetAnimation(NewAnimToPlay);
 	Play(bLooping);
@@ -3562,6 +3609,12 @@ void USkeletalMeshComponent::PlayAnimation(class UAnimationAsset* NewAnimToPlay,
 
 void USkeletalMeshComponent::SetAnimation(UAnimationAsset* NewAnimToPlay)
 {
+	if(!bEnableAnimation)
+	{
+		UE_LOG(LogAnimation, Warning, TEXT("SetAnimation: Animation is currently disabled"));
+		return;
+	}
+
 	UAnimSingleNodeInstance* SingleNodeInstance = GetSingleNodeInstance();
 	if (SingleNodeInstance)
 	{
@@ -3576,6 +3629,12 @@ void USkeletalMeshComponent::SetAnimation(UAnimationAsset* NewAnimToPlay)
 
 void USkeletalMeshComponent::Play(bool bLooping)
 {
+	if(!bEnableAnimation)
+	{
+		UE_LOG(LogAnimation, Warning, TEXT("Play: Animation is currently disabled"));
+		return;
+	}
+
 	UAnimSingleNodeInstance* SingleNodeInstance = GetSingleNodeInstance();
 	if (SingleNodeInstance)
 	{
@@ -3590,6 +3649,12 @@ void USkeletalMeshComponent::Play(bool bLooping)
 
 void USkeletalMeshComponent::Stop()
 {
+	if(!bEnableAnimation)
+	{
+		UE_LOG(LogAnimation, Warning, TEXT("Stop: Animation is currently disabled"));
+		return;
+	}
+
 	UAnimSingleNodeInstance* SingleNodeInstance = GetSingleNodeInstance();
 	if (SingleNodeInstance)
 	{
@@ -3618,6 +3683,12 @@ bool USkeletalMeshComponent::IsPlaying() const
 
 void USkeletalMeshComponent::SetPosition(float InPos, bool bFireNotifies)
 {
+	if(!bEnableAnimation)
+	{
+		UE_LOG(LogAnimation, Warning, TEXT("SetPosition: Animation is currently disabled"));
+		return;
+	}
+
 	UAnimSingleNodeInstance* SingleNodeInstance = GetSingleNodeInstance();
 	if (SingleNodeInstance)
 	{
@@ -3646,6 +3717,12 @@ float USkeletalMeshComponent::GetPosition() const
 
 void USkeletalMeshComponent::SetPlayRate(float Rate)
 {
+	if(!bEnableAnimation)
+	{
+		UE_LOG(LogAnimation, Warning, TEXT("SetPlayRate: Animation is currently disabled"));
+		return;
+	}
+
 	UAnimSingleNodeInstance* SingleNodeInstance = GetSingleNodeInstance();
 	if (SingleNodeInstance)
 	{
@@ -3674,6 +3751,12 @@ float USkeletalMeshComponent::GetPlayRate() const
 
 void USkeletalMeshComponent::OverrideAnimationData(UAnimationAsset* InAnimToPlay, bool bIsLooping /*= true*/, bool bIsPlaying /*= true*/, float Position /*= 0.f*/, float PlayRate /*= 1.f*/)
 {
+	if(!bEnableAnimation)
+	{
+		UE_LOG(LogAnimation, Warning, TEXT("OverrideAnimationData: Animation is currently disabled"));
+		return;
+	}
+
 	AnimationData.AnimToPlay = InAnimToPlay;
 	AnimationData.bSavedLooping = bIsLooping;
 	AnimationData.bSavedPlaying = bIsPlaying;
