[Material Editor] New Material Translator push.

- Basic texture 2D sampling support
- External inputs support, added texcoords
- Added Graphviz DOT visualization dump
- Many fixes and improvements

diff --git a/Engine/Source/Runtime/Engine/Private/Materials/MaterialExpressionsIR.cpp b/Engine/Source/Runtime/Engine/Private/Materials/MaterialExpressionsIR.cpp
--- a/Engine/Source/Runtime/Engine/Private/Materials/MaterialExpressionsIR.cpp
+++ b/Engine/Source/Runtime/Engine/Private/Materials/MaterialExpressionsIR.cpp
@@ -249,44 +249,56 @@
 
 #include "Materials/MaterialIREmitter.h"
 
-namespace IR = UE::MIR;
+namespace MIR = UE::MIR;
 
 /* Constants */
 
-void UMaterialExpression::Build(IR::FEmitter& Emitter)
+void UMaterialExpression::Build(MIR::FEmitter& Emitter)
 {
 	Emitter.Error(TEXT("Unsupported material expression."));
-}
+} 
 
-void UMaterialExpressionConstant::Build(IR::FEmitter& Emitter)
+void UMaterialExpressionConstant::Build(MIR::FEmitter& Emitter)
 {
-	IR::FValue* Value = Emitter.EmitConstantFloat1(R);
+	MIR::FValue* Value = Emitter.EmitConstantFloat1(R);
 	Emitter.Put(GetOutput(0), Value);
 }
 
-void UMaterialExpressionConstant2Vector::Build(IR::FEmitter& Emitter)
+void UMaterialExpressionConstant2Vector::Build(MIR::FEmitter& Emitter)
 {
-	IR::FValue* Value = Emitter.EmitConstantFloat2({ R, G });
+	MIR::FValue* Value = Emitter.EmitConstantFloat2({ R, G });
 	Emitter.Put(GetOutput(0), Value);
+	for (int i = 0; i < 2; ++i)
+	{
+		Emitter.Put(GetOutput(i + 1), Emitter.TryEmitSubscript(Value, i));
+	}
 }
 
-void UMaterialExpressionConstant3Vector::Build(IR::FEmitter& Emitter)
+void UMaterialExpressionConstant3Vector::Build(MIR::FEmitter& Emitter)
 {
-	IR::FValue* Value = Emitter.EmitConstantFloat3({ Constant.R, Constant.G, Constant.B });
+	MIR::FValue* Value = Emitter.EmitConstantFloat3({ Constant.R, Constant.G, Constant.B });
 	Emitter.Put(GetOutput(0), Value);
+	for (int i = 0; i < 3; ++i)
+	{
+		Emitter.Put(GetOutput(i + 1), Emitter.TryEmitSubscript(Value, i));
+	}
 }
 
-void UMaterialExpressionConstant4Vector::Build(IR::FEmitter& Emitter)
+void UMaterialExpressionConstant4Vector::Build(MIR::FEmitter& Emitter)
 {
-	IR::FValue* Value = Emitter.EmitConstantFloat4(Constant);
+	MIR::FValue* Value = Emitter.EmitConstantFloat4(Constant);
 	Emitter.Put(GetOutput(0), Value);
+	for (int i = 0; i < 4; ++i)
+	{
+		Emitter.Put(GetOutput(i + 1), Emitter.TryEmitSubscript(Value, i));
+	}
 }
 
 /* Mathematical Operations */
 
 static void BuildBinaryArithmeticOperator(
-	IR::FEmitter& Emitter,
-	IR::EBinaryOperator Op,
+	MIR::FEmitter& Emitter,
+	MIR::EBinaryOperator Op,
 	FExpressionInput* LhsInput,
 	float LhsConst,
 	FExpressionInput* RhsInput,
@@ -294,8 +306,8 @@ static void BuildBinaryArithmeticOperator(
 	FExpressionOutput* Output)
 {
 	// Default inputs to their relative constants if disconnected, then get each input after checking it has arithmetic type.
-	IR::FValue* LhsValue = Emitter.DefaultTo(LhsInput, LhsConst).TryGetArithmetic(LhsInput);
-	IR::FValue* RhsValue = Emitter.DefaultTo(RhsInput, RhsConst).TryGetArithmetic(RhsInput);
+	MIR::FValue* LhsValue = Emitter.DefaultTo(LhsInput, LhsConst).TryGetArithmetic(LhsInput);
+	MIR::FValue* RhsValue = Emitter.DefaultTo(RhsInput, RhsConst).TryGetArithmetic(RhsInput);
 
 	if (Emitter.IsInvalid())
 	{
@@ -303,9 +315,9 @@ static void BuildBinaryArithmeticOperator(
 	}
 
 	// Determine operation input/output type by looking at the first connected input and picking float1 otherwise.
-	IR::FTypePtr ResultType = LhsValue ? LhsValue->Type
+	MIR::FTypePtr ResultType = LhsValue ? LhsValue->Type
 		: RhsValue ? RhsValue->Type
-		: IR::FArithmeticType::GetScalar(IR::SK_Float);
+		: MIR::FArithmeticType::GetScalar(MIR::SK_Float);
 
 	// Convert operand values to determined result type. 
 	LhsValue = Emitter.TryEmitConstruct(ResultType, LhsValue);
@@ -317,33 +329,33 @@ static void BuildBinaryArithmeticOperator(
 	}
 
 	// Finally emit the binary operator.
-	IR::FValue* Value = Emitter.EmitBinaryOperator(Op, LhsValue, RhsValue);
+	MIR::FValue* Value = Emitter.EmitBinaryOperator(Op, LhsValue, RhsValue);
 
 	// And flow it out of the expression's only output.
 	Emitter.Put(Output, Value);
 }
 
-void UMaterialExpressionAdd::Build(IR::FEmitter& Emitter)
+void UMaterialExpressionAdd::Build(MIR::FEmitter& Emitter)
 { 
-	BuildBinaryArithmeticOperator(Emitter, IR::BO_Add, &A, ConstA, &B, ConstB, GetOutput(0));
+	BuildBinaryArithmeticOperator(Emitter, MIR::BO_Add, &A, ConstA, &B, ConstB, GetOutput(0));
 }
 
-void UMaterialExpressionSubtract::Build(IR::FEmitter& Emitter)
+void UMaterialExpressionSubtract::Build(MIR::FEmitter& Emitter)
 { 
-	BuildBinaryArithmeticOperator(Emitter, IR::BO_Subtract, &A, ConstA, &B, ConstB, GetOutput(0));
+	BuildBinaryArithmeticOperator(Emitter, MIR::BO_Subtract, &A, ConstA, &B, ConstB, GetOutput(0));
 }
 
-void UMaterialExpressionMultiply::Build(IR::FEmitter& Emitter)
+void UMaterialExpressionMultiply::Build(MIR::FEmitter& Emitter)
 { 
-	BuildBinaryArithmeticOperator(Emitter, IR::BO_Multiply, &A, ConstA, &B, ConstB, GetOutput(0));
+	BuildBinaryArithmeticOperator(Emitter, MIR::BO_Multiply, &A, ConstA, &B, ConstB, GetOutput(0));
 }
 
-void UMaterialExpressionDivide::Build(IR::FEmitter& Emitter)
+void UMaterialExpressionDivide::Build(MIR::FEmitter& Emitter)
 { 
-	BuildBinaryArithmeticOperator(Emitter, IR::BO_Divide, &A, ConstA, &B, ConstB, GetOutput(0));
+	BuildBinaryArithmeticOperator(Emitter, MIR::BO_Divide, &A, ConstA, &B, ConstB, GetOutput(0));
 }
 
-void UMaterialExpressionIf::Build(IR::FEmitter& Emitter)
+void UMaterialExpressionIf::Build(MIR::FEmitter& Emitter)
 {
 	// Create default values flowing into disconnected inputs
 	Emitter.DefaultToFloatZero(&A);
@@ -353,11 +365,11 @@ void UMaterialExpressionIf::Build(IR::FEmitter& Emitter)
 	Emitter.DefaultToFloatZero(&ALessThanB);
 
 	// Get input values and check their types are what we expect.
-	IR::FValue* AValue = Emitter.TryGetScalar(&A);
-	IR::FValue* BValue = Emitter.TryGetScalar(&B);
-	IR::FValue* AGreaterThanBValue = Emitter.TryGetArithmetic(&AGreaterThanB);
-	IR::FValue* AEqualsBValue = Emitter.TryGetArithmetic(&AEqualsB);
-	IR::FValue* ALessThanBValue = Emitter.TryGetArithmetic(&ALessThanB);
+	MIR::FValue* AValue = Emitter.TryGetScalar(&A);
+	MIR::FValue* BValue = Emitter.TryGetScalar(&B);
+	MIR::FValue* AGreaterThanBValue = Emitter.TryGetArithmetic(&AGreaterThanB);
+	MIR::FValue* AEqualsBValue = Emitter.TryGetArithmetic(&AEqualsB);
+	MIR::FValue* ALessThanBValue = Emitter.TryGetArithmetic(&ALessThanB);
 
 	if (Emitter.IsInvalid())
 	{
@@ -365,7 +377,7 @@ void UMaterialExpressionIf::Build(IR::FEmitter& Emitter)
 	}
 
 	// Get the arithmetic common type between the conditional arguments (e.g. if inputs are int and float, it will return float).
-	IR::FArithmeticTypePtr ConditionArgsType = Emitter.TryGetCommonArithmeticType(AValue->Type->AsArithmetic(), BValue->Type->AsArithmetic());
+	MIR::FArithmeticTypePtr ConditionArgsType = Emitter.TryGetCommonArithmeticType(AValue->Type->AsArithmetic(), BValue->Type->AsArithmetic());
 
 	if (Emitter.IsInvalid())
 	{
@@ -377,7 +389,7 @@ void UMaterialExpressionIf::Build(IR::FEmitter& Emitter)
 	BValue = Emitter.TryEmitConstruct(ConditionArgsType, BValue);
 
 	// Now determine the output type by taking the common arithmetic type between result values.
-	IR::FArithmeticTypePtr OutputType = Emitter.TryGetCommonArithmeticType(AGreaterThanBValue->Type->AsArithmetic(), AEqualsBValue->Type->AsArithmetic());
+	MIR::FArithmeticTypePtr OutputType = Emitter.TryGetCommonArithmeticType(AGreaterThanBValue->Type->AsArithmetic(), AEqualsBValue->Type->AsArithmetic());
 	OutputType = Emitter.TryGetCommonArithmeticType(OutputType, ALessThanBValue->Type->AsArithmetic());
 
 	if (Emitter.IsInvalid())
@@ -396,14 +408,63 @@ void UMaterialExpressionIf::Build(IR::FEmitter& Emitter)
 	}
 
 	// Emit the comparison expressions.
-	IR::FValue* ALessThanBConditionValue = Emitter.EmitBinaryOperator(IR::BO_Lower, AValue, BValue);
-	IR::FValue* AEqualsBConditionValue = Emitter.EmitBinaryOperator(IR::BO_Equals, AValue, BValue);
+	MIR::FValue* ALessThanBConditionValue = Emitter.EmitBinaryOperator(MIR::BO_LowerThan, AValue, BValue);
+	MIR::FValue* AEqualsBConditionValue = Emitter.EmitBinaryOperator(MIR::BO_Equals, AValue, BValue);
 
 	// And finally emit the full conditional expression.
-	IR::FValue* OutputValue = Emitter.EmitBranch(AEqualsBConditionValue, AEqualsBValue, AGreaterThanBValue);
+	MIR::FValue* OutputValue = Emitter.EmitBranch(AEqualsBConditionValue, AEqualsBValue, AGreaterThanBValue);
 	OutputValue = Emitter.EmitBranch(ALessThanBConditionValue, ALessThanBValue, OutputValue);
 
 	Emitter.Put(GetOutput(0), OutputValue);
 }
 
+
+void UMaterialExpressionTextureSample::Build(UE::MIR::FEmitter& Emitter)
+{
+	UTexture* InputTexture = Texture;
+	if (MIR::FValue* InputTextureValue = Emitter.Get(&TextureObject))
+	{
+		InputTexture = InputTextureValue->GetTexture();
+		if (!InputTexture)
+		{
+			Emitter.Error(TEXT("Value flowing into texture input is not a texture."));
+			return;
+		}
+	}
+	else if (!Texture)
+	{
+		Emitter.Error("Unspecified texture.");
+		return;
+	}
+
+	MIR::FValue* TexCoordsValue = Emitter.Get(&Coordinates);
+	if (!TexCoordsValue)
+	{
+		TexCoordsValue = Emitter.GetExternalInput(MIR::TexCoordIndexToExternalInput(ConstCoordinate));
+	}
+
+	if (Emitter.IsInvalid())
+	{
+		return;
+	}
+
+	MIR::FValue* OutputValue = Emitter.TryEmitTextureSample(InputTexture, TexCoordsValue, SamplerSource, MipValueMode, SamplerType);
+	
+	Emitter.Put(GetOutput(0), Emitter.TryEmitSwizzle(OutputValue, MIR::FSwizzleMask::XYZ()));
+	Emitter.Put(GetOutput(1), Emitter.TryEmitSubscript(OutputValue, 0)); 
+	Emitter.Put(GetOutput(2), Emitter.TryEmitSubscript(OutputValue, 1)); 
+	Emitter.Put(GetOutput(3), Emitter.TryEmitSubscript(OutputValue, 2));
+	Emitter.Put(GetOutput(4), Emitter.TryEmitSubscript(OutputValue, 3));
+	Emitter.Put(GetOutput(5), OutputValue); 
+}
+
+void UMaterialExpressionTextureCoordinate::Build(UE::MIR::FEmitter& Emitter)
+{
+	MIR::FValue* OutputValue = Emitter.GetExternalInput(MIR::TexCoordIndexToExternalInput(CoordinateIndex));
+
+	// todo: add tiling support
+	
+	Emitter.Put(GetOutput(0), OutputValue);
+}
+
 #endif // WITH_EDITOR
diff --git a/Engine/Source/Runtime/Engine/Private/Materials/MaterialIR.cpp b/Engine/Source/Runtime/Engine/Private/Materials/MaterialIR.cpp
--- a/Engine/Source/Runtime/Engine/Private/Materials/MaterialIR.cpp
+++ b/Engine/Source/Runtime/Engine/Private/Materials/MaterialIR.cpp
@@ -2,15 +2,34 @@
 
 #include "Materials/MaterialIR.h"
 #include "Materials/MaterialIRTypes.h"
+#include "Engine/Texture.h"
 
 #if WITH_EDITOR
 
 namespace UE::MIR
 {
 
+const TCHAR* ValueKindToString(EValueKind Kind)
+{
+	switch (Kind)
+	{
+		case VK_Constant: return TEXT("Constant");
+		case VK_ExternalInput: return TEXT("ExternalInput");
+		case VK_MaterialParameter: return TEXT("MaterialParameter");
+		case VK_Dimensional: return TEXT("Dimensional");
+		case VK_SetMaterialOutput: return TEXT("SetMaterialOutput");
+		case VK_BinaryOperator: return TEXT("BinaryOperator");
+		case VK_Branch: return TEXT("Branch");
+		case VK_Subscript: return TEXT("Subscript");
+		case VK_Cast: return TEXT("Cast");
+		case VK_TextureSample: return TEXT("TextureSample");
+		default: UE_MIR_UNREACHABLE();
+	};
+}
+
 FInstruction* FValue::AsInstruction()
 {
-	return this && (Kind >= VK_InstructionBegin && Kind < VK_InstructionEnd) ? static_cast<FInstruction*>(this) : nullptr;
+	return this && (Kind > VK_InstructionBegin && Kind < VK_InstructionEnd) ? static_cast<FInstruction*>(this) : nullptr;
 }
 
 const FInstruction* FValue::AsInstruction() const
@@ -28,7 +47,10 @@ bool FValue::Equals(const FValue* Other) const
 
 	// Get the size of this value in bytes. It should match that of Other, since the value kinds are the same.
 	uint32 SizeInBytes = GetSizeInBytes();
-	check(SizeInBytes == Other->GetSizeInBytes());
+	if (SizeInBytes != Other->GetSizeInBytes())
+	{
+		return false;
+	}
 
 	// Values are PODs by design, therefore simply comparing bytes is sufficient.
 	return FMemory::Memcmp(this, Other, SizeInBytes) == 0;
@@ -39,21 +61,40 @@ uint32 FValue::GetSizeInBytes() const
 	switch (Kind)
 	{
 		case VK_Constant: return sizeof(FConstant);
+		case VK_ExternalInput: return sizeof(FExternalInput);
+		case VK_MaterialParameter: return sizeof(FMaterialParameter);
 		case VK_SetMaterialOutput: return sizeof(FSetMaterialOutput);
 		case VK_BinaryOperator: return sizeof(FBinaryOperator);
 		case VK_Dimensional: return sizeof(FDimensional) + sizeof(FValue*) * static_cast<const FDimensional*>(this)->GetComponents().Num();
 		case VK_Branch: return sizeof(FBranch);
+		case VK_Subscript: return sizeof(FSubscript);
+		case VK_Cast: return sizeof(FCast);
+		case VK_TextureSample: return sizeof(FTextureSample);
 		default: UE_MIR_UNREACHABLE();
 	}
 }
 
+TArrayView<const FValue*> FValue::GetUses() const
+{
+	TArrayView<FValue*> Uses = const_cast<FValue*>(this)->GetUses();
+	return { Uses.Num() ? (const FValue**)&Uses[0] : nullptr , Uses.Num() };
+}
+
 TArrayView<FValue*> FValue::GetUses()
 {
+	// Values have no uses by definition.
+	if (Kind < VK_InstructionBegin)
+	{
+		return {};
+	}
+
 	switch (Kind)
 	{
-		case VK_Constant:
 		case VK_Dimensional:
-			return {};
+		{
+			auto This = static_cast<FDimensional*>(this);
+			return This->GetComponents();
+		}
 
 		case VK_SetMaterialOutput:
 		{
@@ -73,10 +114,129 @@ TArrayView<FValue*> FValue::GetUses()
 			return { &This->ConditionArg, 3 };
 		}
 
+		case VK_Subscript:
+		{
+			auto This = static_cast<FSubscript*>(this);
+			return { &This->Arg, 1 };
+		}
+
+		case VK_TextureSample:
+		{
+			auto This = static_cast<FTextureSample*>(this);
+			return { &This->TexCoordArg, 3 };
+		}
+
 		default: UE_MIR_UNREACHABLE();
 	}
 }
 
+UTexture* FValue::GetTexture()
+{
+	if (auto Parameter = As<FMaterialParameter>())
+	{
+		UObject* TextureObject = Parameter->Metadata.Value.AsTextureObject();
+		if (TextureObject)
+		{
+			return static_cast<UTexture*>(TextureObject);
+		}
+	}
+
+	return nullptr;
+}
+
+const TCHAR* ExternalInputToString(EExternalInput Input)
+{
+	switch (Input)
+	{
+		case EExternalInput::TexCoord0: return TEXT("TexCoord0");
+		case EExternalInput::TexCoord1: return TEXT("TexCoord1");
+		case EExternalInput::TexCoord2: return TEXT("TexCoord2");
+		case EExternalInput::TexCoord3: return TEXT("TexCoord3");
+		case EExternalInput::TexCoord4: return TEXT("TexCoord4");
+		case EExternalInput::TexCoord5: return TEXT("TexCoord5");
+		case EExternalInput::TexCoord6: return TEXT("TexCoord6");
+		case EExternalInput::TexCoord7: return TEXT("TexCoord7");
+		case EExternalInput::TexCoord0_Ddx: return TEXT("TexCoord0_Ddx");
+		case EExternalInput::TexCoord1_Ddx: return TEXT("TexCoord1_Ddx");
+		case EExternalInput::TexCoord2_Ddx: return TEXT("TexCoord2_Ddx");
+		case EExternalInput::TexCoord3_Ddx: return TEXT("TexCoord3_Ddx");
+		case EExternalInput::TexCoord4_Ddx: return TEXT("TexCoord4_Ddx");
+		case EExternalInput::TexCoord5_Ddx: return TEXT("TexCoord5_Ddx");
+		case EExternalInput::TexCoord6_Ddx: return TEXT("TexCoord6_Ddx");
+		case EExternalInput::TexCoord7_Ddx: return TEXT("TexCoord7_Ddx");
+		case EExternalInput::TexCoord0_Ddy: return TEXT("TexCoord0_Ddy");
+		case EExternalInput::TexCoord1_Ddy: return TEXT("TexCoord1_Ddy");
+		case EExternalInput::TexCoord2_Ddy: return TEXT("TexCoord2_Ddy");
+		case EExternalInput::TexCoord3_Ddy: return TEXT("TexCoord3_Ddy");
+		case EExternalInput::TexCoord4_Ddy: return TEXT("TexCoord4_Ddy");
+		case EExternalInput::TexCoord5_Ddy: return TEXT("TexCoord5_Ddy");
+		case EExternalInput::TexCoord6_Ddy: return TEXT("TexCoord6_Ddy");
+		case EExternalInput::TexCoord7_Ddy: return TEXT("TexCoord7_Ddy");
+		default: UE_MIR_UNREACHABLE();
+	}
+}
+
+MIR::EExternalInput TexCoordIndexToExternalInput(int TexCoordIndex)
+{
+	check(TexCoordIndex < TexCoordMaxNum);
+	return (MIR::EExternalInput)((int)MIR::EExternalInput::TexCoord0 + TexCoordIndex);
+}
+
+FTypePtr GetExternalInputType(EExternalInput Id)
+{
+	int IdIndex = (int)Id;
+	if (IdIndex >= (int)EExternalInput::TexCoord0 && IdIndex <= (int)EExternalInput::TexCoord7_Ddy)
+	{
+		return FArithmeticType::GetFloat2();
+	}
+
+	UE_MIR_UNREACHABLE();
+}
+
+bool IsExternalInputTexCoord(EExternalInput Id)
+{
+	return Id >= EExternalInput::TexCoord0 && Id <= EExternalInput::TexCoord7;
+}
+
+bool IsExternalInputTexCoordDdx(EExternalInput Id)
+{
+	return Id >= EExternalInput::TexCoord0_Ddx && Id <= EExternalInput::TexCoord7_Ddx;
+}
+
+bool IsExternalInputTexCoordDdy(EExternalInput Id)
+{
+	return Id >= EExternalInput::TexCoord0_Ddy && Id <= EExternalInput::TexCoord7_Ddy;
+}
+
+FBlock* FBlock::FindCommonParentWith(MIR::FBlock* Other)
+{
+	FBlock* A = this;
+	FBlock* B = Other;
+
+	if (A == B)
+	{
+		return A;
+	}
+
+	while (A->Level > B->Level)
+	{
+		A = A->Parent;
+	}
+
+	while (B->Level > A->Level)
+	{
+		B = B->Parent;
+	}
+
+	while (A != B)
+	{
+		A = A->Parent;
+		B = B->Parent;
+	}
+
+	return A;
+}
+
 TArrayView<FValue* const> FDimensional::GetComponents() const
 {
 	TArrayView<FValue*> Components = const_cast<FDimensional*>(this)->GetComponents();
@@ -104,42 +264,21 @@ bool FDimensional::AreComponentsConstant() const
 	return true;
 }
 
-bool FInstruction::GetInnerBlock(int32 Index, FValue*& OutArg, FBlock*& OutBlock)
+FBlock* FInstruction::GetDesiredBlockForUse(int32 UseIndex)
 {
-	switch (Kind)
+	if (auto Branch = As<FBranch>())
 	{
-		case VK_Branch:
-		{
-			auto This = static_cast<FBranch*>(this);
-			if (Index == 0)
-			{
-				OutArg = This->TrueArg;
-				OutBlock = &This->TrueBlock;
-				return true;
-			}
-			else if (Index == 1)
-			{
-				OutArg = This->FalseArg;
-				OutBlock = &This->FalseBlock;
-				return true;
-			}
-			break;
-		}
-
-		default:
+		switch (UseIndex)
 		{
-			TArrayView<FValue*> Uses = GetUses();
-			if (Uses.IsValidIndex(Index))
-			{
-				OutArg = Uses[Index];
-				OutBlock = Block;
-				return true;
-			}
-			break;
+			case 0: return Block; // ConditionArg goes into the same block as this instruction's
+			case 1: return &Branch->TrueBlock; // TrueArg
+			case 2: return &Branch->FalseBlock; // FalseArg
+			default: UE_MIR_UNREACHABLE();
 		}
 	}
 
-	return false;
+	// By default, dependencies can go in the same block as this instruction
+	return Block;
 }
 
 bool IsArithmeticOperator(EBinaryOperator Op)
@@ -149,7 +288,25 @@ bool IsArithmeticOperator(EBinaryOperator Op)
 
 bool IsComparisonOperator(EBinaryOperator Op)
 {
-	return Op >= BO_Greater && Op <= BO_NotEquals;
+	return Op >= BO_GreaterThan && Op <= BO_NotEquals;
+}
+
+const TCHAR* BinaryOperatorToString(EBinaryOperator Op)
+{
+	switch (Op)
+	{
+		case BO_Add: return TEXT("Add");
+		case BO_Subtract: return TEXT("Subtract");
+		case BO_Multiply: return TEXT("Multiply");
+		case BO_Divide: return TEXT("Divide");
+		case BO_GreaterThan: return TEXT("GreaterThan");
+		case BO_GreaterThanOrEquals: return TEXT("GreaterThanOrEquals");
+		case BO_LowerThan: return TEXT("LowerThan");
+		case BO_LowerThanOrEquals: return TEXT("LowerThanOrEquals");
+		case BO_Equals: return TEXT("Equals");
+		case BO_NotEquals: return TEXT("NotEquals");
+		default: UE_MIR_UNREACHABLE();
+	};
 }
 
 } // namespace UE::MIR
diff --git a/Engine/Source/Runtime/Engine/Private/Materials/MaterialIREmitter.cpp b/Engine/Source/Runtime/Engine/Private/Materials/MaterialIREmitter.cpp
--- a/Engine/Source/Runtime/Engine/Private/Materials/MaterialIREmitter.cpp
+++ b/Engine/Source/Runtime/Engine/Private/Materials/MaterialIREmitter.cpp
@@ -25,14 +25,14 @@ FSwizzleMask::FSwizzleMask(EVectorComponent X)
 }
 
 FSwizzleMask::FSwizzleMask(EVectorComponent X, EVectorComponent Y)
-: NumComponents{ 4 }
+: NumComponents{ 2 }
 {
 	Components[0] = X;
 	Components[1] = Y;
 }
 
 FSwizzleMask::FSwizzleMask(EVectorComponent X, EVectorComponent Y, EVectorComponent Z)
-: NumComponents{ 4 }
+: NumComponents{ 3 }
 {
 	Components[0] = X;
 	Components[1] = Y;
@@ -49,22 +49,25 @@ FSwizzleMask::FSwizzleMask(EVectorComponent X, EVectorComponent Y, EVectorCompon
 	Components[3] = W;
 }
 
+FSwizzleMask FSwizzleMask::XYZ()
+{
+	return { EVectorComponent::X, EVectorComponent::Y, EVectorComponent::Z };
+}
+
 void FSwizzleMask::Append(EVectorComponent Component)
 {
 	check(NumComponents < 4);
 	Components[NumComponents++] = Component;
 }
 
-struct FEmitter::FPrivate : FEmitter
+struct FEmitter::FPrivate
 {
 	// Looks for an existing value in the module that matches `Prototype` and returns it if found.
 	static FValue* FindValue(FEmitter* Emitter, const FValue* Prototype)
 	{
 		/* todo: improve this with hashmap */
-
 		for (FValue* CurrValue : Emitter->Module->Values)
 		{
-	
 			if (CurrValue->Equals(Prototype))
 			{
 				return CurrValue;
@@ -74,40 +77,59 @@ struct FEmitter::FPrivate : FEmitter
 		return nullptr;
 	}
 
+	// Allocates a block of memory using the emitter's allocator.
+	static uint8* Allocate(FEmitter* Emitter, int Size, int Alignment)
+	{
+		uint8* Bytes = (uint8*) FMemory::Malloc(Size, Alignment); // Emitter->Module->Allocator.PushBytes(Size, Alignment);
+		FMemory::Memzero(Bytes, Size);
+		return Bytes;
+	}
+
+	// Pushes a new value to the list of values.
 	static void PushNewValue(FEmitter* Emitter, FValue* Value)
 	{
 		Emitter->Module->Values.Add(Value);
 	}
 };
 
-template <typename T>
-static T InitValue(FTypePtr InType)
+// Creates a new `FDimensional` value of specified `Type` and returns it.
+static FDimensional* NewDimensionalValue(FEmitter* Emitter, FArithmeticTypePtr Type)
 {
-	T Value{};
-	Value.Kind = T::TypeKind;
-	Value.Type = InType;
+	check(!Type->IsScalar());
+
+	int Dimensions = Type->NumRows * Type->NumColumns;
+	int SizeInBytes = sizeof(FDimensional) + sizeof(FValue*) * Dimensions;
+
+	uint8* Bytes = FEmitter::FPrivate::Allocate(Emitter, SizeInBytes, alignof(FDimensional));
+
+	FDimensional* Value = new (Bytes) FDimensional{};
+	Value->Kind = VK_Dimensional;
+	Value->Type = Type;
+
 	return Value;
 }
 
-template <typename TValueType>
-static TValueType* NewValue(FEmitter* Emitter, FTypePtr Type)
+// Emits specified newly created `Value`. If the exact value already exists,
+// specified one is *destroyed* and existing one is returned instead.
+static FValue* EmitNew(FEmitter* Emitter, FValue* Value)
 {
-	TValueType* Value = (TValueType*)new TValueType{};
-	Value->Kind = TValueType::TypeKind;
-	Value->Type = Type;
-	FEmitter::FPrivate::PushNewValue(Value);
+	if (FValue* Existing = FEmitter::FPrivate::FindValue(Emitter, Value))
+	{
+		delete Value;
+		return Existing;
+	}
+
+	FEmitter::FPrivate::PushNewValue(Emitter, Value);
 	return Value;
 }
 
-// Creates a new `FDimensional` value of specified `Type` and returns it.
-static FDimensional* NewDimensionalValue(FEmitter* Emitter, FArithmeticTypePtr Type)
+template <typename T>
+static T MakePrototype(FTypePtr InType)
 {
-	check(!Type->IsScalar());
-	int Dimensions = Type->NumRows * Type->NumColumns;
-	void* Alloc = FMemory::Malloc(sizeof(FDimensional) + sizeof(FValue*) * Dimensions);
-	FDimensional* Value = new (Alloc) FDimensional{};
-	Value->Kind = VK_Dimensional;
-	Value->Type = Type;
+	T Value;
+	FMemory::Memzero(Value);
+	Value.Kind = T::TypeKind;
+	Value.Type = InType;
 	return Value;
 }
 
@@ -122,45 +144,32 @@ static FValue* EmitPrototype(FEmitter* Emitter, const TValueType& Prototype)
 		return Existing;
 	}
 
-	TValueType* Value = new TValueType{ Prototype };
+	uint8* Bytes = FEmitter::FPrivate::Allocate(Emitter, sizeof(TValueType), alignof(TValueType));
+	TValueType* Value = new (Bytes) TValueType{ Prototype };
+
 	FEmitter::FPrivate::PushNewValue(Emitter, Value);
+
 	return Value;
 }
 
-// Emits specified newly created `Value`. If the exact value already exists,
-// specified one is *destroyed* and existing one is returned instead.
-static FValue* EmitNew(FEmitter* Emitter, FValue* Value)
-{
-	if (FValue* Existing = FEmitter::FPrivate::FindValue(Emitter, Value))
-	{
-		delete Value;
-		return Existing;
-	}
 
-	FEmitter::FPrivate::PushNewValue(Emitter, Value);
-	return Value;
-}
+// FEmitter API
 
 FEmitter::FEmitter(FMaterialIRModuleBuilder* InBuilder, UMaterial* InMaterial, FMaterialIRModule* InModule)
 {
 	Builder = InBuilder;
 	Material = InMaterial;
 	Module = InModule;
-
-	// Create and reference the true/false constants.
-	FConstant Temp = InitValue<FConstant>(FArithmeticType::GetBool1());
-
-	Temp.Boolean = true;
-	ConstantTrue = EmitPrototype(this, Temp);
-
-	Temp.Boolean = false;
-	ConstantFalse = EmitPrototype(this, Temp);
 }
 
 FValue* FEmitter::Get(const FExpressionInput* Input)
 {
 	FValue** Value = Builder->InputValues.Find(Input);
-	return Value ? *Value : nullptr;
+	if (!Value || !*Value)
+	{
+		return nullptr;
+	}
+	return *Value;
 }
 	
 void FEmitter::Put(const FExpressionOutput* Output, FValue* Value)
@@ -193,9 +202,14 @@ FValue* FEmitter::TryGetFloat(const FExpressionInput* Input)
 
 FValue* FEmitter::TryGetScalar(const FExpressionInput* Input)
 {
-	FValue* Value = Get(Input);
-	CheckInputIsScalar(Input, Value);
-	return Value;
+	FValue* Value = TryGetArithmetic(Input);
+	if (!Value)
+	{
+		return nullptr;
+	}
+
+	FArithmeticTypePtr ScalarType = Value->Type->AsArithmetic()->ToScalar();
+	return TryEmitConstruct(ScalarType, Value);
 }
 
 FValue* FEmitter::TryGetArithmetic(const FExpressionInput* Input)
@@ -217,6 +231,16 @@ FValue* FEmitter::TryGetOfType(const FExpressionInput* Input, ETypeKind Kind)
 	return Value;
 }
 
+bool FEmitter::CheckValueValid(const FValue* Value)
+{
+	if (!Value)
+	{
+		Error(TEXT("Unspecified value"));
+		return false;
+	}
+	return true;
+}
+
 void FEmitter::CheckInputIsScalar(const FExpressionInput* Input, FValue* InputValue)
 {
 	FArithmeticTypePtr ArithmeticType = InputValue->Type->AsArithmetic();
@@ -237,7 +261,7 @@ void FEmitter::CheckInputIsScalar(const FExpressionInput* Input, FValue* InputVa
 
 void FEmitter::CheckInputTypeIs(const FExpressionInput* Input, FValue* InputValue, ETypeKind Kind)
 {
-	if (InputValue->Type->Kind != Kind)
+	if (InputValue && InputValue->Type->Kind != Kind)
 	{
 		Errorf(TEXT("Input '%s' expected to be have type %s. It is %s instead."), *Input->InputName.ToString(), TypeKindToString(Kind), InputValue->Type->GetSpelling().GetData());
 	}
@@ -291,7 +315,7 @@ FValue* FEmitter::EmitConstantBool1(bool InX)
 
 FValue* FEmitter::EmitConstantFloat1(TFloat InX)
 {
-	FConstant Scalar = InitValue<FConstant>(FArithmeticType::GetScalar(SK_Float));
+	FConstant Scalar = MakePrototype<FConstant>(FArithmeticType::GetScalar(SK_Float));
 	Scalar.Float = InX;
 	return EmitPrototype(this, Scalar);
 }
@@ -322,7 +346,7 @@ FValue* FEmitter::EmitConstantFloat4(const FVector4f& InValue)
 
 FValue* FEmitter::EmitConstantInt1(TInteger InX)
 {
-	FConstant Scalar = InitValue<FConstant>(FArithmeticType::GetScalar(SK_Int));
+	FConstant Scalar = MakePrototype<FConstant>(FArithmeticType::GetScalar(SK_Int));
 	Scalar.Integer = InX;
 	return EmitPrototype(this, Scalar);
 }
@@ -356,7 +380,7 @@ FValue* FEmitter::EmitVector2(FValue* InX, FValue* InY)
 	check(InX->Type->AsScalar());
 	check(InX->Type == InY->Type);
 
-	TDimensional<2> Vector = InitValue<TDimensional<2>>(FArithmeticType::GetVector(InX->Type->AsArithmetic()->ScalarKind, 2));
+	TDimensional<2> Vector = MakePrototype<TDimensional<2>>(FArithmeticType::GetVector(InX->Type->AsArithmetic()->ScalarKind, 2));
 	TArrayView<FValue*> Components = Vector.GetComponents();
 	Components[0] = InX;
 	Components[1] = InY;
@@ -370,7 +394,7 @@ FValue* FEmitter::EmitVector3(FValue* InX, FValue* InY, FValue* InZ)
 	check(InX->Type == InY->Type);
 	check(InY->Type == InZ->Type);
 
-	TDimensional<3> Vector = InitValue<TDimensional<3>>(FArithmeticType::GetVector(InX->Type->AsArithmetic()->ScalarKind, 3));
+	TDimensional<3> Vector = MakePrototype<TDimensional<3>>(FArithmeticType::GetVector(InX->Type->AsArithmetic()->ScalarKind, 3));
 	TArrayView<FValue*> Components = Vector.GetComponents();
 	Components[0] = InX;
 	Components[1] = InY;
@@ -386,7 +410,7 @@ FValue* FEmitter::EmitVector4(FValue* InX, FValue* InY, FValue* InZ, FValue* InW
 	check(InY->Type == InZ->Type);
 	check(InZ->Type == InW->Type);
 
-	TDimensional<4> Vector = InitValue<TDimensional<4>>(FArithmeticType::GetVector(InX->Type->AsArithmetic()->ScalarKind, 4));
+	TDimensional<4> Vector = MakePrototype<TDimensional<4>>(FArithmeticType::GetVector(InX->Type->AsArithmetic()->ScalarKind, 4));
 	TArrayView<FValue*> Components = Vector.GetComponents();
 	Components[0] = InX;
 	Components[1] = InY;
@@ -396,10 +420,21 @@ FValue* FEmitter::EmitVector4(FValue* InX, FValue* InY, FValue* InZ, FValue* InW
 	return EmitPrototype(this, Vector);
 }
 
-FValue* FEmitter::EmitSubscript(FValue* Value, int Index)
+FValue* FEmitter::GetExternalInput(EExternalInput Id)
+{
+	MIR::FExternalInput Prototype = MakePrototype<MIR::FExternalInput>(GetExternalInputType(Id));
+	Prototype.Id = Id;
+	return EmitPrototype(this, Prototype);
+}
+
+FValue* FEmitter::TryEmitSubscript(FValue* Value, int Index)
 {
 	FArithmeticTypePtr ArithmeticType = Value->Type->AsArithmetic();
-	check(ArithmeticType);
+	if (!ArithmeticType)
+	{
+		Errorf(TEXT("Value of type `%s` cannot be subscripted."), Value->Type->GetSpelling().GetData());
+		return nullptr;
+	}
 
 	// Getting first component and Value is already a scalar, just return itself.
 	if (Index == 0 && Value->Type->AsScalar())
@@ -414,8 +449,14 @@ FValue* FEmitter::EmitSubscript(FValue* Value, int Index)
 		return DimensionalValue->GetComponents()[Index];
 	}
 	
+	// Avoid subscripting a subscript (e.g. no value.xy.x)
+	if (FSubscript* Subscript = Value->As<FSubscript>())
+	{
+		Value = Subscript->Arg;
+	}
+
 	// We can't resolve it at compile time: emit subscript value.
-	FSubscript Prototype = InitValue<FSubscript>(ArithmeticType->ToScalar());
+	FSubscript Prototype = MakePrototype<FSubscript>(ArithmeticType->ToScalar());
 	Prototype.Arg = Value;
 	Prototype.Index = Index;
 
@@ -448,7 +489,7 @@ FValue* FEmitter::TryEmitSwizzle(FValue* Value, FSwizzleMask Mask)
 	// If only one component is requested, we can use EmitSubscript() to return the single component.
 	if (Mask.NumComponents == 1)
 	{
-		return EmitSubscript(Value, (int)Mask.Components[0]);
+		return TryEmitSubscript(Value, (int)Mask.Components[0]);
 	}
 
 	// If the requested number of components is the same as Value and the order in which the components
@@ -477,12 +518,23 @@ FValue* FEmitter::TryEmitSwizzle(FValue* Value, FSwizzleMask Mask)
 
 	for (int i = 0; i < Mask.NumComponents; ++i)
 	{
-		Result->GetComponents()[i] = EmitSubscript(Value, (int)Mask.Components[i]);
+		Result->GetComponents()[i] = TryEmitSubscript(Value, (int)Mask.Components[i]);
 	}
 
 	return Result;
 }
 
+FValue* FEmitter::GetParameter(FName Name, const FMaterialParameterMetadata& Metadata)
+{
+	FMaterialParameterInfo Info{ MoveTemp(Name), EMaterialParameterAssociation::GlobalParameter, INDEX_NONE };
+	
+	MIR::FMaterialParameter Proto = MakePrototype<MIR::FMaterialParameter>(MIR::FTextureType::Get());
+	Proto.Info = Info;
+	Proto.Metadata = Metadata;
+
+	return EmitPrototype(this, Proto);
+}
+
 FSetMaterialOutput* FEmitter::EmitSetMaterialOutput(EMaterialProperty InProperty, FValue* InArgValue)
 {
 	FSetMaterialOutput* Instr = new FSetMaterialOutput;
@@ -514,10 +566,10 @@ static bool FoldComparisonOperatorScalar(EBinaryOperator Operator, T Lhs, T Rhs)
 {
 	switch (Operator)
 	{
-		case BO_Greater: return Lhs > Rhs;
-		case BO_GreaterOrEquals: return Lhs >= Rhs;
-		case BO_Lower: return Lhs < Rhs;
-		case BO_LowerOrEquals: return Lhs <= Rhs;
+		case BO_GreaterThan: return Lhs > Rhs;
+		case BO_GreaterThanOrEquals: return Lhs >= Rhs;
+		case BO_LowerThan: return Lhs < Rhs;
+		case BO_LowerThanOrEquals: return Lhs <= Rhs;
 		case BO_Equals: return Lhs == Rhs;
 		case BO_NotEquals: return Lhs != Rhs;
 		default: UE_MIR_UNREACHABLE();
@@ -622,7 +674,7 @@ FValue* FEmitter::EmitBinaryOperator(EBinaryOperator Operator, FValue* Lhs, FVal
 	}
 
 	// One (or both) of the operands is runtime only. Emit the runtime binary operation instruction.
-	FBinaryOperator Proto = InitValue<FBinaryOperator>(ResultType);
+	FBinaryOperator Proto = MakePrototype<FBinaryOperator>(ResultType);
 	Proto.Operator = Operator;
 	Proto.LhsArg= Lhs;
 	Proto.RhsArg = Rhs;
@@ -646,7 +698,7 @@ FValue* FEmitter::EmitBranch(FValue* Condition, FValue* True, FValue* False)
 	check(True->Type == False->Type);
 
 	// Create the branch instruction.
-	FBranch Proto = InitValue<FBranch>(True->Type);
+	FBranch Proto = MakePrototype<FBranch>(True->Type);
 	Proto.ConditionArg = Condition;
 	Proto.TrueArg = True;
 	Proto.FalseArg = False;
@@ -701,8 +753,18 @@ static FValue* ConstructArithmeticValue(FEmitter* Emitter, FArithmeticTypePtr Ta
 	}
 
 	// Construct a scalar from another scalar.
-	if (TargetArithmeticType->IsScalar() && InitializerArithmeticType->IsScalar())
+	if (TargetArithmeticType->IsScalar())
 	{
+		//
+		Initializer = Emitter->TryEmitSubscript(Initializer, 0);
+		InitializerArithmeticType = Initializer->Type->AsArithmetic();
+		
+		//
+		if (InitializerArithmeticType == TargetArithmeticType)
+		{
+			return Initializer;
+		}
+
 		// Construct the scalar from a constant.
 		if (FConstant* ConstantInitializer = Initializer->As<FConstant>())
 		{
@@ -733,7 +795,7 @@ static FValue* ConstructArithmeticValue(FEmitter* Emitter, FArithmeticTypePtr Ta
 		
 		return EmitNew(Emitter, Result);
 	}
-	
+
 	// Construct a vector from another vector. If constructed vector is larger, initialize
 	// remaining components to zero. If it's smaller, truncate initializer vector and only use
 	// the necessary components.
@@ -742,9 +804,6 @@ static FValue* ConstructArithmeticValue(FEmitter* Emitter, FArithmeticTypePtr Ta
 		int TargetNumComponents = TargetArithmeticType->GetNumComponents();
 		int InitializerNumComponents = InitializerArithmeticType->GetNumComponents();
 
-		int MinNumComponents = FMath::Min(TargetNumComponents, InitializerNumComponents);
-		int MaxNumComponents = FMath::Max(TargetNumComponents, InitializerNumComponents);
-
 		// Create the result dimensional value.
 		FDimensional* Result = NewDimensionalValue(Emitter, TargetArithmeticType);
 
@@ -755,13 +814,14 @@ static FValue* ConstructArithmeticValue(FEmitter* Emitter, FArithmeticTypePtr Ta
 		int Index = 0;
 
 		// Convert components from the initializer vector.
+		int MinNumComponents = FMath::Min(TargetNumComponents, InitializerNumComponents);
 		for (; Index < MinNumComponents; ++Index)
 		{
-			Result->GetComponents()[Index] = Emitter->TryEmitConstruct(ResultComponentType, Emitter->EmitSubscript(Initializer, Index));
+			Result->GetComponents()[Index] = Emitter->TryEmitConstruct(ResultComponentType, Emitter->TryEmitSubscript(Initializer, Index));
 		}
 
 		// Initialize remaining result dimensional components to zero.
-		for (; Index < MaxNumComponents; ++Index)
+		for (; Index < TargetNumComponents; ++Index)
 		{
 			Result->GetComponents()[Index] = Emitter->EmitConstantScalarZero(ResultComponentType->ScalarKind);
 		}
@@ -795,12 +855,12 @@ static FValue* ConstructArithmeticValue(FEmitter* Emitter, FArithmeticTypePtr Ta
 		else
 		{
 			// Initializer is an unknown value, construct target value casting initializer.
-			FCast Prototype = InitValue<FCast>(TargetArithmeticType);
+			FCast Prototype = MakePrototype<FCast>(TargetArithmeticType);
 			Prototype.Arg = Initializer;
 			return EmitPrototype(Emitter, Prototype);
 		}
 	}
-	
+
 	// Initializer value cannot be used to construct this arithmetic type.
 	return nullptr;
 }
@@ -815,7 +875,6 @@ FValue* FEmitter::TryEmitConstruct(FTypePtr Type, FValue* Initializer)
 	}
 	
 	FValue* Result{};
-
 	if (FArithmeticTypePtr ArithmeticType = Type->AsArithmetic())
 	{
 		Result = ConstructArithmeticValue(this, ArithmeticType, Initializer);
@@ -830,6 +889,20 @@ FValue* FEmitter::TryEmitConstruct(FTypePtr Type, FValue* Initializer)
 	return Result;
 }
 
+FValue* FEmitter::TryEmitTextureSample(UTexture* Texture, FValue* TexCoord, ESamplerSourceMode SamplerSourceMode, ETextureMipValueMode MipValueMode, EMaterialSamplerType SamplerType)
+{
+	MIR::FTextureSample Prototype = MakePrototype<MIR::FTextureSample>(FArithmeticType::GetFloat4());
+	Prototype.Texture = Texture;
+	Prototype.TexCoordArg = TexCoord;
+	Prototype.MipValueArg = nullptr;
+	Prototype.AutomaticMipBiasArg = nullptr;
+	Prototype.SamplerSourceMode = SamplerSourceMode;
+	Prototype.MipValueMode = MipValueMode;
+	Prototype.SamplerType = SamplerType;
+
+	return EmitPrototype(this, Prototype);
+}
+
 FArithmeticTypePtr FEmitter::TryGetCommonArithmeticType(FArithmeticTypePtr A, FArithmeticTypePtr B)
 {
 	// Trivial case: types are equal
@@ -861,6 +934,18 @@ void FEmitter::Error(FString Message)
 	bHasExprBuildError = true;
 }
 
+void FEmitter::Initialize()
+{
+	// Create and reference the true/false constants.
+	FConstant Temp = MakePrototype<FConstant>(FArithmeticType::GetBool1());
+
+	Temp.Boolean = true;
+	ConstantTrue = EmitPrototype(this, Temp);
+
+	Temp.Boolean = false;
+	ConstantFalse = EmitPrototype(this, Temp);
+}
+
 } // namespace UE::MIR
 
 #endif // #if WITH_EDITOR
diff --git a/Engine/Source/Runtime/Engine/Private/Materials/MaterialIRModule.cpp b/Engine/Source/Runtime/Engine/Private/Materials/MaterialIRModule.cpp
--- a/Engine/Source/Runtime/Engine/Private/Materials/MaterialIRModule.cpp
+++ b/Engine/Source/Runtime/Engine/Private/Materials/MaterialIRModule.cpp
@@ -2,20 +2,20 @@
 
 #include "Materials/MaterialIRModule.h"
 #include "Materials/MaterialIR.h"
+#include "Materials/MaterialIRTypes.h"
 
 #if WITH_EDITOR
 
-namespace IR = UE::MIR;
+namespace MIR = UE::MIR;
 
 FMaterialIRModule::FMaterialIRModule()
 {
-	RootBlock = new IR::FBlock;
+	RootBlock = new MIR::FBlock;
 }
 
 FMaterialIRModule::~FMaterialIRModule()
 {
 	Empty();
-
 	delete RootBlock;
 }
 
@@ -23,13 +23,25 @@ void FMaterialIRModule::Empty()
 {
 	RootBlock->Instructions = nullptr;
 
-	for (IR::FValue* Value : Values)
+	for (MIR::FValue* Value : Values)
 	{
-		delete Value;
+		FMemory::Free(Value);
 	}
 
 	Values.Empty();
 	Outputs.Empty();
+
+	// Allocator.~FMemStackBase();
+	// new (&Allocator) FMemStackBase;
+
+	// Reset module statistics.
+	for (int i = 0; i < (int)SF_NumFrequencies; ++i)
+	{
+		Statistics.ExternalInputUsedMask[i].Init(false, (int)MIR::EExternalInput::Count);
+	}
+
+	Statistics.NumVertexTexCoords = 0;
+	Statistics.NumPixelTexCoords = 0;
 }
 
 #endif // #if WITH_EDITOR
diff --git a/Engine/Source/Runtime/Engine/Private/Materials/MaterialIRModuleBuilder.cpp b/Engine/Source/Runtime/Engine/Private/Materials/MaterialIRModuleBuilder.cpp
--- a/Engine/Source/Runtime/Engine/Private/Materials/MaterialIRModuleBuilder.cpp
+++ b/Engine/Source/Runtime/Engine/Private/Materials/MaterialIRModuleBuilder.cpp
@@ -5,6 +5,7 @@
 #if WITH_EDITOR
 
 #include "Materials/MaterialIRModule.h"
+#include "Materials/MaterialIRDebug.h"
 #include "Materials/MaterialIR.h"
 #include "Materials/MaterialIRTypes.h"
 #include "Materials/MaterialIREmitter.h"
@@ -16,17 +17,26 @@
 #include "Materials/MaterialExpression.h"
 #include "Materials/MaterialInsights.h"
 #include "Async/ParallelFor.h"
+#include "Engine/Texture.h"
 
-namespace IR = UE::MIR;
+namespace MIR = UE::MIR;
 
 struct FMaterialIRModuleBuilder::FPrivate
 {
 	FMaterialIRModuleBuilder& Builder;
 	FMaterialIRModuleBuildParams Params;
 	FMaterialIRModule& Module;
-	IR::FEmitter& Emitter;
+	MIR::FEmitter& Emitter;
 	TArray<UMaterialExpression*> ExpressionAnalysisStack;
-	TArray<IR::FInstruction*> InstructionStack;
+	TArray<MIR::FInstruction*> InstructionStack;
+
+	void Build_Initialize()
+	{
+		Module.Empty();
+		Module.ShaderPlatform = Params.ShaderPlatform;
+
+		Emitter.Initialize();
+	}
 
 	void Build_GenerateOutputInstructions()
 	{
@@ -36,7 +46,7 @@ struct FMaterialIRModuleBuilder::FPrivate
 		{
 			EMaterialProperty Property = (EMaterialProperty)Index;
 
-			IR::FSetMaterialOutput* Output = Emitter.EmitSetMaterialOutput(Property, nullptr);
+			MIR::FSetMaterialOutput* Output = Emitter.EmitSetMaterialOutput(Property, nullptr);
 
 			if (Input.bUseConstant)
 			{
@@ -53,7 +63,7 @@ struct FMaterialIRModuleBuilder::FPrivate
 		}
 	}
 
-	void Build_AnalyzeExpressionGraph()
+	void Build_BuildIRGraph()
 	{
 		TSet<UMaterialExpression*> BuiltExpressions;
 
@@ -95,17 +105,18 @@ struct FMaterialIRModuleBuilder::FPrivate
 				if (FExpressionOutput* ConnectedOutput = It->GetConnectedOutput())
 				{
 					// Fetch the value flowing through connected output.
-					IR::FValue** ValuePtr = Builder.OutputValues.Find(ConnectedOutput);
-					check(ValuePtr && TEXT("Output value not found."));
-
-					// Set the value flowing into this input.
-					Builder.InputValues.Add(It.Input, *ValuePtr);
+					if (MIR::FValue** ValuePtr = Builder.OutputValues.Find(ConnectedOutput))
+					{
+						// Set the value flowing into this input.
+						Builder.InputValues.Add(It.Input, *ValuePtr);
+					}
+					else
+					{
+						Builder.InputValues.Remove(It.Input);
+					}
 				}
 			}
 
-			// And clear the expression errors.
-			Emitter.bHasExprBuildError = false;
-
 			// Invoke the expression build function. This will perform semantic analysis, error reporting and
 			// emit IR values for its outputs (which will flow into connected expressions inputs).
 			Emitter.Expression->Build(Emitter);
@@ -118,35 +129,22 @@ struct FMaterialIRModuleBuilder::FPrivate
 		}
 	}
 
-	void AddExpressionConnectionInsights(UMaterialExpression* Expression)
-	{
-		// Update expression inputs insight.
-		for (FExpressionInputIterator It{ Expression}; It; ++It)
-		{
-			if (!It->IsConnected())
-			{
-				continue;
-			}
-
-			IR::FValue** Value = Builder.InputValues.Find(It.Input);
-			PushConnectionInsight(Expression, It.Index, It->Expression, It->OutputIndex, Value ? (*Value)->Type : nullptr);
-		}
-	}
-	
 	void Build_LinkMaterialOutputsToIncomingValues()
 	{
-		for (IR::FSetMaterialOutput* Output : Module.Outputs)
+		for (MIR::FSetMaterialOutput* Output : Module.Outputs)
 		{
 			FMaterialInputDescription Input;
 			ensure(Params.Material->GetExpressionInputDescription(Output->Property, Input));
 
 			if (!Output->Arg)
 			{
-				IR::FValue** ValuePtr = Builder.OutputValues.Find(Input.Input->Expression->GetOutput(0));
+				MIR::FValue** ValuePtr = Builder.OutputValues.Find(Input.Input->GetConnectedOutput());
 				check(ValuePtr);
 
 				Builder.InputValues.Add(Input.Input, *ValuePtr);
-				Output->Arg = *ValuePtr;
+
+				MIR::FTypePtr OutputArgType = MIR::FType::FromShaderType(Input.Type);
+				Output->Arg = Emitter.TryEmitConstruct(OutputArgType, *ValuePtr);
 			}
 
 			if (Params.TargetInsight)
@@ -157,33 +155,27 @@ struct FMaterialIRModuleBuilder::FPrivate
 		}
 	}
 
-	void PushConnectionInsight(const UObject* InputObject, int InputIndex, const UMaterialExpression* OutputExpression, int OutputIndex, IR::FTypePtr Type)
-	{
-		FMaterialInsights::FConnectionInsight Insight;
-		Insight.InputObject = InputObject,
-		Insight.OutputExpression = OutputExpression,
-		Insight.InputIndex = InputIndex,
-		Insight.OutputIndex = OutputIndex,
-		Insight.ValueType = Type ? Type->ToValueType() : UE::Shader::EValueType::Any,
-		
-		Params.TargetInsight->ConnectionInsights.Push(Insight);
-	}
-
-	void Build_FinalizeValueGraph()
+	void Build_AnalyzeIRGraph()
 	{
 		InstructionStack.Reserve(64);
 
-		for (IR::FSetMaterialOutput* Output : Module.Outputs)
+		for (MIR::FSetMaterialOutput* Output : Module.Outputs)
 		{
 			InstructionStack.Push(Output);
 		}
 
 		while (!InstructionStack.IsEmpty())
 		{
-			IR::FValue* Instr = InstructionStack.Pop();
-			for (IR::FValue* UseValue : Instr->GetUses())
+			MIR::FValue* Instr = InstructionStack.Pop();
+			for (MIR::FValue* UseValue : Instr->GetUses())
 			{
-				IR::FInstruction* Use = UseValue->AsInstruction();
+				if (UseValue && !(UseValue->Flags & MIR::VF_ValueAnalyzed))
+				{
+					UseValue->SetFlags(MIR::VF_ValueAnalyzed);
+					Build_AnalyzeValue(UseValue);
+				}
+
+				MIR::FInstruction* Use = UseValue->AsInstruction();
 				if (!Use)
 				{
 					continue;
@@ -191,103 +183,160 @@ struct FMaterialIRModuleBuilder::FPrivate
 
 				Use->NumUsers += 1;
 
-				if (!(Use->Flags & IR::IF_Counted))
+				if (!(Use->Flags & MIR::VF_InstructionAnalyzed))
 				{
-					Use->SetFlags(IR::IF_Counted);
+					Use->SetFlags(MIR::VF_InstructionAnalyzed);
 					InstructionStack.Push(Use);
 				}
 			}
 		}
 	}
 
-	void Build_PopulateBlock()
+	void Build_AnalyzeValue(MIR::FValue* Value)
+	{
+		if (auto ExternalInput = Value->As<MIR::FExternalInput>())
+		{
+			Module.Statistics.ExternalInputUsedMask[SF_Vertex][(int)ExternalInput->Id] = true;
+			Module.Statistics.ExternalInputUsedMask[SF_Pixel][(int)ExternalInput->Id] = true;
+		}
+		else if (auto TextureSample = Value->As<MIR::FTextureSample>())
+		{
+			EMaterialTextureParameterType ParamType = UE::Utility::TextureMaterialValueTypeToParameterType(TextureSample->Texture->GetMaterialType());
+
+			FMaterialTextureParameterInfo ParamInfo{};
+			ParamInfo.ParameterInfo = { "", EMaterialParameterAssociation::GlobalParameter, INDEX_NONE };
+			ParamInfo.SamplerSource = SSM_FromTextureAsset; // TODO - Is this needed?
+
+			ParamInfo.TextureIndex = Params.Material->GetReferencedTextures().Find(TextureSample->Texture);
+			check(ParamInfo.TextureIndex != INDEX_NONE);
+
+			TextureSample->TextureParameterIndex = Module.CompilationOutput.UniformExpressionSet.FindOrAddTextureParameter(ParamType, ParamInfo);
+		}
+	}
+
+	void Build_PopulateBlocks()
 	{
 		// This function walks the instruction graph and puts each instruction into the inner most possible block.
 		InstructionStack.Empty(InstructionStack.Max());
 
-		for (IR::FSetMaterialOutput* Output : Module.Outputs)
+		for (MIR::FSetMaterialOutput* Output : Module.Outputs)
 		{
 			Output->Block = Module.RootBlock;
 			InstructionStack.Add(Output);
 		}
 
-		while (!InstructionStack.IsEmpty()) {
-			IR::FInstruction* Instr = InstructionStack.Pop();
+		while (!InstructionStack.IsEmpty())
+		{
+			MIR::FInstruction* Instr = InstructionStack.Pop();
+			if (MIR::FSetMaterialOutput* Output = Instr->As<MIR::FSetMaterialOutput>())
+			{
+				if (Output->Property == EMaterialProperty::MP_BaseColor)
+				{
+					static int l = 0;
+					++l;
+				}
+			}
 
-			// Push the instruction in its block
+			// Push the instruction to its block in reverse order (push front)
 			Instr->Next = Instr->Block->Instructions;
 			Instr->Block->Instructions = Instr;
 
-			IR::FValue* UseValue;
-			IR::FBlock* InnerBlock;
-			for (int32 i = 0; Instr->GetInnerBlock(i, UseValue, InnerBlock); ++i)
+			TArrayView<MIR::FValue*> Uses = Instr->GetUses();
+			for (int32 UseIndex = 0; UseIndex < Uses.Num(); ++UseIndex)
 			{
-				IR::FInstruction* Use = UseValue->AsInstruction();
-				if (!Use)
+				MIR::FValue* Use = Uses[UseIndex];
+				MIR::FInstruction* UseInstr = Use->AsInstruction();
+				if (!UseInstr)
 				{
 					continue;
 				}
 
+				// Get the block into which the dependency instruction should go.
+				MIR::FBlock* TargetBlock = Instr->GetDesiredBlockForUse(UseIndex);
+
 				// Update dependency's block to be a child of current instruction's block.
-				if (InnerBlock != Instr->Block)
+				if (TargetBlock != Instr->Block)
 				{
-					InnerBlock->Parent = Instr->Block;
-					InnerBlock->Level = Instr->Block->Level + 1;
+					TargetBlock->Parent = Instr->Block;
+					TargetBlock->Level = Instr->Block->Level + 1;
 				}
 
 				// Set the dependency's block to the common block betwen its current block and this one.
-				Use->Block = Use->Block
-					? FindCommonParentBlock(Use->Block, InnerBlock)
-					: InnerBlock;
+				UseInstr->Block = UseInstr->Block
+					? UseInstr->Block->FindCommonParentWith(TargetBlock)
+					: TargetBlock;
 
 				// Increase the number of times this dependency instruction has been considered.
 				// When all of its users have processed, we can carry on visiting this instruction.
-				++Use->NumProcessedUsers;
-				check(Use->NumProcessedUsers <= Use->NumUsers);
+				++UseInstr->NumProcessedUsers;
+				check(UseInstr->NumProcessedUsers <= UseInstr->NumUsers);
 
 				// If all dependants have been processed, we can carry the processing from this dependency.
-				if (Use->NumProcessedUsers == Use->NumUsers)
+				if (UseInstr->NumProcessedUsers == UseInstr->NumUsers)
 				{
-					InstructionStack.Push(Use);
+					InstructionStack.Push(UseInstr);
 				}
 			}
 		}
 	}
 
-    static IR::FBlock* FindCommonParentBlock(IR::FBlock* A, IR::FBlock* B)
-    {
-        if (A == B) {
-            return A;
-        }
-
-        while (A->Level > B->Level) {
-            A = A->Parent;
-        }
+	void Build_Finalize()
+	{
+		/* Produce the module statistics */
+		for (int TexCoordIndex = 0; TexCoordIndex < MIR::TexCoordMaxNum; ++TexCoordIndex)
+		{
+			MIR::EExternalInput TexCoordInput = MIR::TexCoordIndexToExternalInput(TexCoordIndex);
+			if (Module.Statistics.ExternalInputUsedMask[SF_Vertex][(int)TexCoordInput])
+			{
+				Module.Statistics.NumVertexTexCoords = TexCoordIndex + 1;
+			}
+			if (Module.Statistics.ExternalInputUsedMask[SF_Pixel][(int)TexCoordInput])
+			{
+				Module.Statistics.NumPixelTexCoords = TexCoordIndex + 1;
+			}
+		}
 
-        while (B->Level > A->Level) {
-            B = B->Parent;
-        }
+		/* Configure the compilation output */
+		FMaterialCompilationOutput& CompilationOutput = Module.CompilationOutput;
+		CompilationOutput.NumUsedUVScalars = Module.Statistics.NumPixelTexCoords * 2;
+	}
 
-        while (A != B) {
-            A = A->Parent;
-            B = B->Parent;
-        }
+	void AddExpressionConnectionInsights(UMaterialExpression* Expression)
+	{
+		// Update expression inputs insight.
+		for (FExpressionInputIterator It{ Expression}; It; ++It)
+		{
+			if (!It->IsConnected())
+			{
+				continue;
+			}
 
-        return A;
-    }
+			MIR::FValue** Value = Builder.InputValues.Find(It.Input);
+			PushConnectionInsight(Expression, It.Index, It->Expression, It->OutputIndex, Value ? (*Value)->Type : nullptr);
+		}
+	}
+	
+	void PushConnectionInsight(const UObject* InputObject, int InputIndex, const UMaterialExpression* OutputExpression, int OutputIndex, MIR::FTypePtr Type)
+	{
+		FMaterialInsights::FConnectionInsight Insight;
+		Insight.InputObject = InputObject,
+		Insight.OutputExpression = OutputExpression,
+		Insight.InputIndex = InputIndex,
+		Insight.OutputIndex = OutputIndex,
+		Insight.ValueType = Type ? Type->ToValueType() : UE::Shader::EValueType::Any,
+		
+		Params.TargetInsight->ConnectionInsights.Push(Insight);
+	}
 };
 
 bool FMaterialIRModuleBuilder::Build(const FMaterialIRModuleBuildParams& Params, FMaterialIRModule* TargetModule)
 {
-	TargetModule->Empty();
-	TargetModule->ShaderPlatform = Params.ShaderPlatform;
-
-	IR::FEmitter Emitter{ this, Params.Material, TargetModule };
-
+	MIR::FEmitter Emitter{ this, Params.Material, TargetModule };
 	FPrivate Private{ *this, Params, *TargetModule, Emitter };
 
+	Private.Build_Initialize();
 	Private.Build_GenerateOutputInstructions();
-	Private.Build_AnalyzeExpressionGraph();
+	Private.Build_BuildIRGraph();
 
 	if (Private.Emitter.IsInvalid())
 	{
@@ -295,8 +344,11 @@ bool FMaterialIRModuleBuilder::Build(const FMaterialIRModuleBuildParams& Params,
 	}
 
 	Private.Build_LinkMaterialOutputsToIncomingValues();
-	Private.Build_FinalizeValueGraph();
-	Private.Build_PopulateBlock();
+	Private.Build_AnalyzeIRGraph();
+	Private.Build_PopulateBlocks();
+	Private.Build_Finalize();
+
+	UE::MIR::DebugDumpIRUseGraph(*TargetModule);
 
 	return true;
 }
diff --git a/Engine/Source/Runtime/Engine/Private/Materials/MaterialIRToHLSLTranslator.cpp b/Engine/Source/Runtime/Engine/Private/Materials/MaterialIRToHLSLTranslator.cpp
--- a/Engine/Source/Runtime/Engine/Private/Materials/MaterialIRToHLSLTranslator.cpp
+++ b/Engine/Source/Runtime/Engine/Private/Materials/MaterialIRToHLSLTranslator.cpp
@@ -14,12 +14,13 @@
 #include "DataDrivenShaderPlatformInfo.h"
 #include "Materials/MaterialExpressionVolumetricAdvancedMaterialOutput.h"
 #include "RenderUtils.h"
+#include "Engine/Texture.h"
 
 #include <inttypes.h>
 
 #if WITH_EDITOR
 
-namespace IR = UE::MIR;
+namespace MIR = UE::MIR;
 
 enum ENoOp { NoOp };
 enum ENewLine { NewLine };
@@ -27,6 +28,11 @@ enum EEndOfStatement { EndOfStatement };
 enum EOpenBrace { OpenBrace };
 enum ECloseBrace { CloseBrace };
 enum EIndentation { Indentation };
+enum EBeginArgs { BeginArgs };
+enum EEndArgs { EndArgs };
+enum EListSeparator { ListSeparator };
+
+#define TAB "    "
 
 struct FHLSLPrinter
 {
@@ -35,7 +41,7 @@ struct FHLSLPrinter
 	int32 Tabs = 0;
 
 	template <int N, typename... Types>
-	void Printf(const TCHAR (&Format)[N], Types... Args)
+	void Appendf(const TCHAR (&Format)[N], Types... Args)
 	{
 		Buffer.Appendf(Format, Args...);
 	}
@@ -101,6 +107,25 @@ struct FHLSLPrinter
         Buffer.AppendChar('}');
         return *this;
     }
+	
+    FHLSLPrinter& operator<<(EBeginArgs)
+    {
+        Buffer.AppendChar('(');
+		BeginList();
+        return *this;
+    }
+
+    FHLSLPrinter& operator<<(EEndArgs)
+    {
+        Buffer.AppendChar(')');
+        return *this;
+    }
+
+	FHLSLPrinter& operator<<(EListSeparator)
+    {
+		PrintListSeparator();
+        return *this;
+    }
 
 	void BeginList()
 	{
@@ -168,9 +193,9 @@ static const TCHAR* GetShadingModelParameterName(EMaterialShadingModel InModel)
 	}
 }
 
-static bool IsFoldable(const IR::FInstruction* Instr)
+static bool IsFoldable(const MIR::FInstruction* Instr)
 {
-	if (auto Branch = Instr->As<IR::FBranch>())
+	if (auto Branch = Instr->As<MIR::FBranch>())
 	{
 		return !Branch->TrueBlock.Instructions && !Branch->FalseBlock.Instructions;
 	}
@@ -181,7 +206,7 @@ static bool IsFoldable(const IR::FInstruction* Instr)
 struct FTranslator : FMaterialIRToHLSLTranslation
 {
 	int32 NumLocals{};
-	TMap<const IR::FInstruction*, FString> LocalToIdentifier;
+	TMap<const MIR::FInstruction*, FString> LocalIdentifier;
 	FHLSLPrinter Printer;
 	FString PixelAttributesHLSL;
 	FString EvaluateOtherMaterialAttributesHLSL;
@@ -213,14 +238,14 @@ struct FTranslator : FMaterialIRToHLSLTranslation
 			EMaterialValueType Type = (Property == MP_SubsurfaceColor) ? MCT_Float4 : FMaterialAttributeDefinitionMap::GetValueType(Property);
 			check(PropertyName.Len() > 0);
 
-			PixelAttributesHLSL.Appendf(TEXT("\t%s %s;\n"), GetHLSLTypeString(Type), *PropertyName);
+			PixelAttributesHLSL.Appendf(TEXT(TAB "%s %s;\n"), GetHLSLTypeString(Type), *PropertyName);
 		}
 	}
 	
-	ENoOp LowerBlock(const IR::FBlock& Block)
+	ENoOp LowerBlock(const MIR::FBlock& Block)
 	{
 		int OldNumLocals = NumLocals;
-        for (IR::FInstruction* Instr = Block.Instructions; Instr; Instr = Instr->Next)
+        for (MIR::FInstruction* Instr = Block.Instructions; Instr; Instr = Instr->Next)
 		{
             if (Instr->NumUsers == 1 && IsFoldable(Instr))
 			{
@@ -232,9 +257,9 @@ struct FTranslator : FMaterialIRToHLSLTranslation
                 FString LocalStr = FString::Printf(TEXT("l%d"), NumLocals);
                 ++NumLocals;
 
-                Printer << InlineType(Instr->Type) << TEXT(" ") << LocalStr;
+                Printer << LowerType(Instr->Type) << TEXT(" ") << LocalStr;
 
-                LocalToIdentifier.Add(Instr, MoveTemp(LocalStr));
+                LocalIdentifier.Add(Instr, MoveTemp(LocalStr));
                 if (IsFoldable(Instr))
 				{
                     Printer << TEXT(" = ";)
@@ -257,192 +282,295 @@ struct FTranslator : FMaterialIRToHLSLTranslation
 
 		return NoOp;
 	}
+	
+	ENoOp LowerValue(MIR::FValue* InValue)
+	{
+		if (MIR::FInstruction* Instr = InValue->AsInstruction())
+		{
+			if (Instr->NumUsers <= 1 && IsFoldable(Instr))
+			{
+				Printer << LowerInstruction(Instr);
+			}
+			else
+			{
+				Printer << LocalIdentifier[Instr];
+			}
+
+			return NoOp;
+		}
+
+		switch (InValue->Kind)
+		{
+			case MIR::VK_Constant: LowerConstant(static_cast<const MIR::FConstant*>(InValue)); break;
+			case MIR::VK_ExternalInput: LowerExternalInput(static_cast<const MIR::FExternalInput*>(InValue)); break;
+			case MIR::VK_MaterialParameter: LowerMaterialParameter(static_cast<const MIR::FMaterialParameter*>(InValue)); break;
+			default: UE_MIR_UNREACHABLE();
+		}
+
+		return NoOp;
+	}
 
-	ENoOp LowerInstruction(const IR::FInstruction* Instr)
+	ENoOp LowerInstruction(MIR::FInstruction* Instr)
 	{
 		switch (Instr->Kind)
 		{
-			case IR::VK_Dimensional:
-			{
-				const IR::FDimensional* Vector = static_cast<const IR::FDimensional*>(Instr);
+			case MIR::VK_Dimensional: LowerDimensional(static_cast<MIR::FDimensional*>(Instr)); break;
+			case MIR::VK_SetMaterialOutput: LowerSetMaterialOutput(static_cast<MIR::FSetMaterialOutput*>(Instr)); break;
+			case MIR::VK_BinaryOperator: LowerBinaryOperator(static_cast<MIR::FBinaryOperator*>(Instr)); break;
+			case MIR::VK_Branch: LowerBranch(static_cast<MIR::FBranch*>(Instr)); break;
+			case MIR::VK_Subscript: LowerSubscript(static_cast<MIR::FSubscript*>(Instr)); break;
+			case MIR::VK_TextureSample: LowerTextureSample(static_cast<MIR::FTextureSample*>(Instr)); break;
 
-				IR::FArithmeticTypePtr ArithType = Vector->Type->AsArithmetic();
-				check(ArithType && ArithType->IsVector());
+			default:
+				UE_MIR_UNREACHABLE();
+		}
 
-				Printer << ScalarKindToString(ArithType->ScalarKind) << ArithType->NumRows << TEXT("(");
+		return NoOp;
+	}
 
-				Printer.BeginList();
-				for (IR::FValue* Component : Vector->GetComponents())
-				{
-					Printer.PrintListSeparator();
-					LowerValue(Component);
-				}
+	void LowerConstant(const MIR::FConstant* Constant)
+	{
+		MIR::FArithmeticTypePtr ArithType = Constant->Type->AsArithmetic();
+		check(ArithType && ArithType->IsScalar());
 
-				Printer << TEXT(")");
-				break;
-			}
+		switch (ArithType->ScalarKind)
+		{
+			case MIR::SK_Bool:  Printer.Buffer.Append(Constant->Boolean ? TEXT("true") : TEXT("false")); break;
+			case MIR::SK_Int:   Printer.Buffer.Appendf(TEXT("%") PRId64, Constant->Integer); break;
+			case MIR::SK_Float: Printer.Buffer.Appendf(TEXT("%.5ff"), Constant->Float); break;
+		}
+	}
 
-			case IR::VK_SetMaterialOutput:
-			{
-				auto Output = static_cast<const IR::FSetMaterialOutput*>(Instr);
+	void LowerExternalInput(const MIR::FExternalInput* ExternalInput)
+	{
+		int ExternalInputIndex = (int)ExternalInput->Id;
 
-				// Special case MP_SubsurfaceColor as the actual property is a combination of the color and the profile but we don't want to expose the profile
-				const FString& PropertyName = (Output->Property == MP_SubsurfaceColor) ? "Subsurface" : FMaterialAttributeDefinitionMap::GetAttributeName(Output->Property);
+		if (MIR::IsExternalInputTexCoord(ExternalInput->Id))
+		{
+			int Index = ExternalInputIndex - (int)MIR::EExternalInput::TexCoord0;
+			Printer.Appendf(TEXT("Parameters.TexCoords[%d]"), Index);
+		}
+		else if (MIR::IsExternalInputTexCoordDdx(ExternalInput->Id))
+		{
+			int Index = ExternalInputIndex - (int)MIR::EExternalInput::TexCoord0_Ddx;
+			Printer.Appendf(TEXT("Parameters.TexCoords_DDX[%d]"), Index);
+		}
+		else if (MIR::IsExternalInputTexCoordDdy(ExternalInput->Id))
+		{
+			int Index = ExternalInputIndex - (int)MIR::EExternalInput::TexCoord0_Ddy;
+			Printer.Appendf(TEXT("Parameters.TexCoords_DDY[%d]"), Index);
+		}
+		else
+		{
+			UE_MIR_UNREACHABLE();
+		}
+	}
 
-				Printer << TEXT("PixelMaterialInputs.") << PropertyName << TEXT(" = ") << LowerValue(Output->Arg);
+	void LowerMaterialParameter(const MIR::FMaterialParameter* Parameter)
+	{
+		UE_MIR_UNREACHABLE();
+	}
+	
+	void LowerDimensional(const MIR::FDimensional* Dimensional) 
+	{
+		MIR::FArithmeticTypePtr ArithmeticType = Dimensional->Type->AsArithmetic();
+		check(ArithmeticType && ArithmeticType->IsVector());
 
-				break;
-			}
+		Printer << ScalarKindToString(ArithmeticType->ScalarKind) << ArithmeticType->NumRows << BeginArgs;
 
-			case IR::VK_BinaryOperator:
-			{
-				auto BinaryOperator = static_cast<const IR::FBinaryOperator*>(Instr);
+		for (MIR::FValue* Component : Dimensional->GetComponents())
+		{
+			Printer << ListSeparator << LowerValue(Component);
+		}
 
-				LowerValue(BinaryOperator->LhsArg);
+		Printer << EndArgs;
+	}
 
-				const TCHAR* OpString;
-				switch (BinaryOperator->Operator)
-				{
-					case IR::BO_Add: OpString = TEXT(" + "); break;
-					case IR::BO_Subtract: OpString = TEXT(" - "); break;
-					case IR::BO_Multiply: OpString = TEXT(" * "); break;
-					case IR::BO_Divide: OpString = TEXT(" / "); break;
-					case IR::BO_Greater: OpString = TEXT(" > "); break;
-					case IR::BO_Lower: OpString = TEXT(" < "); break;
-					case IR::BO_Equals: OpString = TEXT(" == "); break;
-					default: UE_MIR_UNREACHABLE();
-				}
-				Printer << OpString;
+	void LowerSetMaterialOutput(const MIR::FSetMaterialOutput* Output)
+	{
+		// Special case MP_SubsurfaceColor as the actual property is a combination of the color and the profile but we don't want to expose the profile
+		const FString& PropertyName = (Output->Property == MP_SubsurfaceColor) ? "Subsurface" : FMaterialAttributeDefinitionMap::GetAttributeName(Output->Property);
+		Printer << TEXT("PixelMaterialInputs.") << PropertyName << TEXT(" = ") << LowerValue(Output->Arg);
+	}
 
-				LowerValue(BinaryOperator->RhsArg);
+	void LowerBinaryOperator(const MIR::FBinaryOperator* BinaryOperator)
+	{
+		LowerValue(BinaryOperator->LhsArg);
 
-				break;
-			}
+		const TCHAR* OpString;
+		switch (BinaryOperator->Operator)
+		{
+			case MIR::BO_Add: OpString = TEXT(" + "); break;
+			case MIR::BO_Subtract: OpString = TEXT(" - "); break;
+			case MIR::BO_Multiply: OpString = TEXT(" * "); break;
+			case MIR::BO_Divide: OpString = TEXT(" / "); break;
+			case MIR::BO_GreaterThan: OpString = TEXT(" > "); break;
+			case MIR::BO_LowerThan: OpString = TEXT(" < "); break;
+			case MIR::BO_Equals: OpString = TEXT(" == "); break;
+			default: UE_MIR_UNREACHABLE();
+		}
+		Printer << OpString;
 
-			case IR::VK_Branch:
-			{
-				auto Branch = static_cast<const IR::FBranch*>(Instr);
+		LowerValue(BinaryOperator->RhsArg);
+	}
 
-				if (IsFoldable(Branch))
-				{
-					Printer << LowerValue(Branch->ConditionArg)
-						<< TEXT(" ? ") << LowerValue(Branch->TrueArg)
-						<< TEXT(" : ") << LowerValue(Branch->FalseArg);
-				}
-				else
-				{
-					Printer << EndOfStatement;
-					Printer << TEXT("if (") << LowerValue(Branch->ConditionArg) << TEXT(")") << NewLine << OpenBrace;
-					Printer << LowerBlock(Branch->TrueBlock);
-					Printer << LocalToIdentifier[Instr] << " = " << LowerValue(Branch->TrueArg) << EndOfStatement;
-					Printer << CloseBrace << NewLine;
-					Printer << TEXT("else") << NewLine << OpenBrace;
-					Printer << LowerBlock(Branch->FalseBlock);
-					Printer << LocalToIdentifier[Instr] << " = " << LowerValue(Branch->FalseArg) << EndOfStatement;
-					Printer << CloseBrace;
-				}
-				break;
-			}
+	void LowerBranch(const MIR::FBranch* Branch)
+	{
+		if (IsFoldable(Branch))
+		{
+			Printer << LowerValue(Branch->ConditionArg)
+				<< TEXT(" ? ") << LowerValue(Branch->TrueArg)
+				<< TEXT(" : ") << LowerValue(Branch->FalseArg);
+		}
+		else
+		{
+			Printer << EndOfStatement;
+			Printer << TEXT("if (") << LowerValue(Branch->ConditionArg) << TEXT(")") << NewLine << OpenBrace;
+			Printer << LowerBlock(Branch->TrueBlock);
+			Printer << LocalIdentifier[Branch] << " = " << LowerValue(Branch->TrueArg) << EndOfStatement;
+			Printer << CloseBrace << NewLine;
+			Printer << TEXT("else") << NewLine << OpenBrace;
+			Printer << LowerBlock(Branch->FalseBlock);
+			Printer << LocalIdentifier[Branch] << " = " << LowerValue(Branch->FalseArg) << EndOfStatement;
+			Printer << CloseBrace;
+		}
+	}
 			
-			case IR::VK_Subscript:
-			{
-				const IR::FSubscript* Subscript = static_cast<const IR::FSubscript*>(Instr);
+	void LowerSubscript(const MIR::FSubscript* Subscript)
+	{
+		LowerValue(Subscript->Arg);
 
-				LowerValue(Subscript->Arg);
+		if (MIR::FArithmeticTypePtr ArgArithmeticType = Subscript->Arg->Type->AsVector())
+		{
+			const TCHAR* ComponentsStr[] = { TEXT(".x"), TEXT(".y"), TEXT(".z"), TEXT(".w") };
+			check(Subscript->Index <= ArgArithmeticType->GetNumComponents());
 
-				if (IR::FArithmeticTypePtr ArgArithmeticType = Subscript->Arg->Type->AsVector())
-				{
-					const TCHAR* ComponentsStr[] = { TEXT(".x"), TEXT(".y"), TEXT(".z"), TEXT(".w") };
-					check(Subscript->Index <= ArgArithmeticType->GetNumComponents());
+			Printer << ComponentsStr[Subscript->Index];
+		}
+	}
 
-					Printer << ComponentsStr[Subscript->Index];
-				}
+	void LowerTextureSample(MIR::FTextureSample* TextureSample)
+	{
+		bool bUsesSpecialSampler = LowerSamplerType(TextureSample->SamplerType);
+		if (bUsesSpecialSampler)
+		{
+			Printer << TEXT("(");
+		}
 
+		switch (TextureSample->Texture->GetMaterialType())
+		{
+			case MCT_Texture2D:
+				Printer << TEXT("Texture2DSample");
 				break;
-			}
 
 			default:
 				UE_MIR_UNREACHABLE();
 		}
 
-		return NoOp;
+		Printer << BeginArgs
+			<< ListSeparator << LowerTextureReference(TextureSample->Texture->GetMaterialType(), TextureSample->TextureParameterIndex)
+			<< ListSeparator << LowerTextureSamplerReference(TextureSample->SamplerSourceMode, TextureSample->Texture->GetMaterialType(), TextureSample->TextureParameterIndex)
+			<< ListSeparator << LowerValue(TextureSample->TexCoordArg)
+			<< EndArgs;
+
+		if (bUsesSpecialSampler)
+		{
+			Printer << TEXT(")");
+		}
 	}
 
-	ENoOp LowerValue(const IR::FValue* InValue)
+	bool LowerSamplerType(EMaterialSamplerType SamplerType)
 	{
-		if (const IR::FInstruction* Instr = InValue->AsInstruction())
+		switch (SamplerType)
 		{
-			if (Instr->NumUsers <= 1 && IsFoldable(Instr))
-			{
-				Printer << LowerInstruction(Instr);
-			}
-			else
-			{
-				Printer << LocalToIdentifier[Instr];
-			}
+			case SAMPLERTYPE_External:
+				Printer << TEXT("ProcessMaterialExternalTextureLookup");
+				break;
 
-			return NoOp;
-		}
+			case SAMPLERTYPE_Color:
+				Printer << TEXT("ProcessMaterialColorTextureLookup");
+				break;
+			case SAMPLERTYPE_VirtualColor:
+				// has a mobile specific workaround
+				Printer << TEXT("ProcessMaterialVirtualColorTextureLookup");
+				break;
 
-		switch (InValue->Kind)
-		{
-			case IR::VK_Constant:
-			{
-				const IR::FConstant* Scalar = static_cast<const IR::FConstant*>(InValue);
-				
-				IR::FArithmeticTypePtr ArithType = Scalar->Type->AsArithmetic();
-				check(ArithType && ArithType->IsScalar());
+			case SAMPLERTYPE_LinearColor:
+			case SAMPLERTYPE_VirtualLinearColor:
+				Printer << TEXT("ProcessMaterialLinearColorTextureLookup");
+				break;
 
-				switch (ArithType->ScalarKind)
-				{
-					case IR::SK_Bool:  Printer.Buffer.Append(Scalar->Boolean ? TEXT("true") : TEXT("false")); break;
-					case IR::SK_Int:   Printer.Buffer.Appendf(TEXT("%") PRId64, Scalar->Integer); break;
-					case IR::SK_Float: Printer.Buffer.Appendf(TEXT("%.5ff"), Scalar->Float); break;
-				}
+			case SAMPLERTYPE_Alpha:
+			case SAMPLERTYPE_VirtualAlpha:
+			case SAMPLERTYPE_DistanceFieldFont:
+				Printer << TEXT("ProcessMaterialAlphaTextureLookup");
+				break;
 
+			case SAMPLERTYPE_Grayscale:
+			case SAMPLERTYPE_VirtualGrayscale:
+				Printer << TEXT("ProcessMaterialGreyscaleTextureLookup");
 				break;
-			}
+
+			case SAMPLERTYPE_LinearGrayscale:
+			case SAMPLERTYPE_VirtualLinearGrayscale:
+				Printer <<TEXT("ProcessMaterialLinearGreyscaleTextureLookup");
+				break;
+
+			case SAMPLERTYPE_Normal:
+			case SAMPLERTYPE_VirtualNormal:
+				// Normal maps need to be unpacked in the pixel shader.
+				Printer << TEXT("UnpackNormalMap");
+				break;
+
+			case SAMPLERTYPE_Masks:
+			case SAMPLERTYPE_VirtualMasks:
+			case SAMPLERTYPE_Data:
+				return false;
 
 			default:
-			{
 				UE_MIR_UNREACHABLE();
-			}
 		}
 
-		return NoOp;
+		return true;
 	}
 
-	ENoOp InlineType(const IR::FType* Type)
+	ENoOp LowerTextureSamplerReference(ESamplerSourceMode SamplerSource, EMaterialValueType TextureType, int TextureParameterIndex)
 	{
-		if (auto ArithmeticType = Type->AsArithmetic())
+		switch (SamplerSource)
 		{
-			switch (ArithmeticType->ScalarKind)
-			{
-				case IR::SK_Bool:	Printer << TEXT("bool"); break;
-				case IR::SK_Int: 	Printer << TEXT("int"); break;
-				case IR::SK_Float:	Printer << TEXT("float"); break;
-			}
+		case SSM_FromTextureAsset:
+			LowerTextureReference(TextureType, TextureParameterIndex);
+			Printer << TEXT("Sampler");
+			break;
 
-			if (ArithmeticType->NumRows > 1)
-			{
-				Printer << ArithmeticType->NumRows;
-			}
-			
-			if (ArithmeticType->NumColumns > 1)
-			{
-				Printer << TEXT("x") << ArithmeticType->NumColumns;
-			}
+		default:
+			UE_MIR_UNREACHABLE();
 		}
-		else
+
+		return NoOp;
+	}
+
+	ENoOp LowerTextureReference(EMaterialValueType TextureType, int TextureParameterIndex)
+	{
+		Printer << TEXT("Material.");
+
+		switch (TextureType)
 		{
-			UE_MIR_UNREACHABLE();
+			case MCT_Texture2D: Printer << TEXT("Texture2D_"); break;
+			default: UE_MIR_UNREACHABLE();
 		}
 
+		Printer << TextureParameterIndex;
+
 		return NoOp;
 	}
+	
+	/* Finalization */
 
 	void SetMaterialParameters(TMap<FString, FString>& Params)
 	{
+		const FMaterialIRModule::FStatistics ModuleStatistics = Module->GetStatistics();
+
 		auto SetParamInt = [&] (const TCHAR* InParamName, int InValue)
 		{
 			Params.Add(InParamName, FString::Printf(TEXT("%d"), InValue));
@@ -450,17 +578,39 @@ struct FTranslator : FMaterialIRToHLSLTranslation
 		
 		auto SetParamReturnFloat = [&] (const TCHAR* InParamName, float InValue)
 		{
-			Params.Add(InParamName, FString::Printf(TEXT("\treturn %.5f"), InValue));
+			Params.Add(InParamName, FString::Printf(TEXT(TAB "return %.5f"), InValue));
 		};
 
 		Params.Add(TEXT("pixel_material_inputs"), MoveTemp(PixelAttributesHLSL));
 		Params.Add(TEXT("calc_pixel_material_inputs_initial_calculations"), EvaluateOtherMaterialAttributesHLSL);
 		Params.Add(TEXT("calc_pixel_material_inputs_analytic_derivatives_initial"), MoveTemp(EvaluateOtherMaterialAttributesHLSL));
-		Params.Add(TEXT("material_declarations"), TEXT("struct FMaterialAttributes {};"));
-		SetParamInt(TEXT("num_material_texcoords_vertex"), 0);
-		SetParamInt(TEXT("num_material_texcoords"), 0);
+		
+		// MaterialAttributes
+		TArray<FGuid> OrderedVisibleAttributes = FMaterialAttributeDefinitionMap::GetOrderedVisibleAttributeList();
+		
+		FString MaterialDeclarations;
+		MaterialDeclarations.Appendf(TEXT("struct FMaterialAttributes\n{\n"));
+		for (const FGuid& AttributeID : OrderedVisibleAttributes)
+		{
+			const FString& PropertyName = FMaterialAttributeDefinitionMap::GetAttributeName(AttributeID);
+			const EMaterialValueType PropertyType = FMaterialAttributeDefinitionMap::GetValueType(AttributeID);
+			MaterialDeclarations.Appendf(TEXT(TAB "%s %s;\n"), GetHLSLTypeString(PropertyType), *PropertyName);
+		}
+		MaterialDeclarations.Appendf(TEXT("};"));
+		Params.Add(TEXT("material_declarations"), MoveTemp(MaterialDeclarations));
+		
+		SetParamInt(TEXT("num_material_texcoords_vertex"), ModuleStatistics.NumVertexTexCoords);
+		SetParamInt(TEXT("num_material_texcoords"), ModuleStatistics.NumPixelTexCoords);
 		SetParamInt(TEXT("num_custom_vertex_interpolators"), 0);
-		SetParamInt(TEXT("num_tex_coord_interpolators"), 0);
+		SetParamInt(TEXT("num_tex_coord_interpolators"), ModuleStatistics.NumPixelTexCoords);
+
+		FString GetMaterialCustomizedUVS;
+		for (int CustomUVIndex = 0; CustomUVIndex < ModuleStatistics.NumPixelTexCoords; CustomUVIndex++)
+		{
+			const FString AttributeName = FMaterialAttributeDefinitionMap::GetAttributeName((EMaterialProperty)(MP_CustomizedUVs0 + CustomUVIndex));
+			GetMaterialCustomizedUVS.Appendf(TEXT(TAB "OutTexCoords[%u] = Parameters.MaterialAttributes.%s;\n"), CustomUVIndex, *AttributeName);
+		}
+		Params.Add(TEXT("get_material_customized_u_vs"), MoveTemp(GetMaterialCustomizedUVS));
 
 		SetParamReturnFloat(TEXT("get_material_emissive_for_cs"), 0.f);
 		SetParamReturnFloat(TEXT("get_material_translucency_directional_lighting_intensity"), Material->GetTranslucencyDirectionalLightingIntensity());
@@ -471,22 +621,55 @@ struct FTranslator : FMaterialIRToHLSLTranslation
 		SetParamReturnFloat(TEXT("get_material_translucent_backscattering_exponent"), Material->GetTranslucentBackscatteringExponent());
 
 		FLinearColor Extinction = Material->GetTranslucentMultipleScatteringExtinction();
-		Params.Add(TEXT("get_material_translucent_multiple_scattering_extinction"), FString::Printf(TEXT("\treturn MaterialFloat3(%.5f, %.5f, %.5f)"), Extinction.R, Extinction.G, Extinction.B));
+		Params.Add(TEXT("get_material_translucent_multiple_scattering_extinction"), FString::Printf(TEXT(TAB "return MaterialFloat3(%.5f, %.5f, %.5f)"), Extinction.R, Extinction.G, Extinction.B));
 
 		SetParamReturnFloat(TEXT("get_material_opacity_mask_clip_value"), Material->GetOpacityMaskClipValue());
-		Params.Add(TEXT("get_material_world_position_offset_raw"), TEXT("\treturn 0; // todo"));
-		Params.Add(TEXT("get_material_previous_world_position_offset_raw"), TEXT("\treturn 0; // todo"));
+		Params.Add(TEXT("get_material_world_position_offset_raw"), TEXT(TAB "return 0; // todo"));
+		Params.Add(TEXT("get_material_previous_world_position_offset_raw"), TEXT(TAB "return 0; // todo"));
 	
 		// CustomData0/1 are named ClearCoat/ClearCoatRoughness
-		Params.Add(TEXT("get_material_custom_data0"), TEXT("\treturn 1.0f; // todo"));
-		Params.Add(TEXT("get_material_custom_data1"), TEXT("\treturn 0.1f; // todo"));
+		Params.Add(TEXT("get_material_custom_data0"), TEXT(TAB "return 1.0f; // todo"));
+		Params.Add(TEXT("get_material_custom_data1"), TEXT(TAB "return 0.1f; // todo"));
 
 		FString EvaluateMaterialDeclaration;
 		EvaluateMaterialDeclaration.Append(TEXT("void EvaluateVertexMaterialAttributes(in out FMaterialVertexParameters Parameters)\n{\n"));
+		for (int CustomUVIndex = 0; CustomUVIndex < ModuleStatistics.NumPixelTexCoords; CustomUVIndex++)
+		{
+			EvaluateMaterialDeclaration.Appendf(TEXT(TAB "Parameters.MaterialAttributes.CustomizedUV%d = Parameters.TexCoords[%d].xy;\n"), CustomUVIndex, CustomUVIndex);
+		}
 		EvaluateMaterialDeclaration.Append(TEXT("\n}\n"));
-		Params.Add(TEXT("evaluate_material_attributes"), EvaluateMaterialDeclaration);
+		Params.Add(TEXT("evaluate_material_attributes"), MoveTemp(EvaluateMaterialDeclaration));
 	}
+	
+	ENoOp LowerType(const MIR::FType* Type)
+	{
+		if (auto ArithmeticType = Type->AsArithmetic())
+		{
+			switch (ArithmeticType->ScalarKind)
+			{
+				case MIR::SK_Bool:	Printer << TEXT("bool"); break;
+				case MIR::SK_Int: 	Printer << TEXT("int"); break;
+				case MIR::SK_Float:	Printer << TEXT("float"); break;
+			}
+
+			if (ArithmeticType->NumRows > 1)
+			{
+				Printer << ArithmeticType->NumRows;
+			}
+			
+			if (ArithmeticType->NumColumns > 1)
+			{
+				Printer << TEXT("x") << ArithmeticType->NumColumns;
+			}
+		}
+		else
+		{
+			UE_MIR_UNREACHABLE();
+		}
 
+		return NoOp;
+	}
+	
 	void GetShaderCompilerEnvironment(FShaderCompilerEnvironment& OutEnvironment)
 	{
 		const FMaterialCompilationOutput& CompilationOutput = Module->GetCompilationOutput();
@@ -614,4 +797,5 @@ void FMaterialIRToHLSLTranslation::Run(TMap<FString, FString>& OutParameters, FS
 	Translator.GetShaderCompilerEnvironment(OutEnvironment);
 }
 
-#endif // #if WITH_EDITOR
+#undef TAB
+#endif // #if WITH_EDITOR
\ No newline at end of file
diff --git a/Engine/Source/Runtime/Engine/Private/Materials/MaterialIRTypes.cpp b/Engine/Source/Runtime/Engine/Private/Materials/MaterialIRTypes.cpp
--- a/Engine/Source/Runtime/Engine/Private/Materials/MaterialIRTypes.cpp
+++ b/Engine/Source/Runtime/Engine/Private/Materials/MaterialIRTypes.cpp
@@ -157,6 +157,21 @@ FArithmeticTypePtr FArithmeticType::GetFloat1()
 	return GetScalar(SK_Float);
 }
 
+FArithmeticTypePtr FArithmeticType::GetFloat2()
+{
+	return GetVector(SK_Float, 2);
+}
+
+FArithmeticTypePtr FArithmeticType::GetFloat3()
+{
+	return GetVector(SK_Float, 3);
+}
+
+FArithmeticTypePtr FArithmeticType::GetFloat4()
+{
+	return GetVector(SK_Float, 4);
+}
+
 const FArithmeticType* FArithmeticType::GetScalar(EScalarKind InScalarKind)
 {
 	return Get(InScalarKind, 1, 1);
@@ -242,6 +257,13 @@ FArithmeticTypePtr FArithmeticType::ToScalar() const
 	return FArithmeticType::GetScalar(ScalarKind);
 }
 
+FTypePtr FTextureType::Get()
+{
+	static FTextureType Instance { TK_Texture };
+	return &Instance;
+}
+
+
 } // namespace UE::MIR
 
 #endif // #if WITH_EDITOR
diff --git a/Engine/Source/Runtime/Engine/Private/Materials/MaterialIRUtility.cpp b/Engine/Source/Runtime/Engine/Private/Materials/MaterialIRUtility.cpp
--- a/Engine/Source/Runtime/Engine/Private/Materials/MaterialIRUtility.cpp
+++ b/Engine/Source/Runtime/Engine/Private/Materials/MaterialIRUtility.cpp
@@ -8,7 +8,7 @@
 
 #if WITH_EDITOR
 
-namespace IR = UE::MIR;
+namespace MIR = UE::MIR;
 
 namespace UE::Utility {
 
@@ -57,7 +57,7 @@ bool NextMaterialAttributeInput(UMaterial* BaseMaterial, int32& PropertyIndex, F
 	return false;
 }
 
-IR::FValue* CreateMaterialAttributeDefaultValue(IR::FEmitter& Emitter, const UMaterial* Material, EMaterialProperty Property)
+MIR::FValue* CreateMaterialAttributeDefaultValue(MIR::FEmitter& Emitter, const UMaterial* Material, EMaterialProperty Property)
 {
 	EMaterialValueType Type = FMaterialAttributeDefinitionMap::GetValueType(Property);
 	FVector4f DefaultValue = FMaterialAttributeDefinitionMap::GetDefaultValue(Property);
@@ -80,5 +80,18 @@ IR::FValue* CreateMaterialAttributeDefaultValue(IR::FEmitter& Emitter, const UMa
 	}
 }
 
+EMaterialTextureParameterType TextureMaterialValueTypeToParameterType(EMaterialValueType Type)
+{
+	switch (Type)
+	{
+		case MCT_Texture2D: return EMaterialTextureParameterType::Standard2D;
+		case MCT_Texture2DArray: return EMaterialTextureParameterType::Array2D;
+		case MCT_TextureCube: return EMaterialTextureParameterType::Cube;
+		case MCT_TextureCubeArray: return EMaterialTextureParameterType::ArrayCube;
+		case MCT_VolumeTexture: return EMaterialTextureParameterType::Volume;
+		default: UE_MIR_UNREACHABLE();
+	}
+}
+
 } // namespace Utility
 #endif
diff --git a/Engine/Source/Runtime/Engine/Private/Materials/MaterialIRUtility.h b/Engine/Source/Runtime/Engine/Private/Materials/MaterialIRUtility.h
--- a/Engine/Source/Runtime/Engine/Private/Materials/MaterialIRUtility.h
+++ b/Engine/Source/Runtime/Engine/Private/Materials/MaterialIRUtility.h
@@ -2,6 +2,7 @@
 #pragma once
 
 #include "Materials/MaterialIRCommon.h"
+#include "MaterialShared.h"
 
 #if WITH_EDITOR
 
@@ -16,6 +17,9 @@ bool NextMaterialAttributeInput(UMaterial* BaseMaterial, int32& PropertyIndex, F
 //
 UE::MIR::FValue* CreateMaterialAttributeDefaultValue(UE::MIR::FEmitter& Emitter, const UMaterial* Material, EMaterialProperty Property);
 
+//
+EMaterialTextureParameterType TextureMaterialValueTypeToParameterType(EMaterialValueType Type);
+
 } // namespace Utility
 
 #endif
diff --git a/Engine/Source/Runtime/Engine/Private/Materials/MaterialInstanceConstant.cpp b/Engine/Source/Runtime/Engine/Private/Materials/MaterialInstanceConstant.cpp
--- a/Engine/Source/Runtime/Engine/Private/Materials/MaterialInstanceConstant.cpp
+++ b/Engine/Source/Runtime/Engine/Private/Materials/MaterialInstanceConstant.cpp
@@ -222,7 +222,7 @@ void UMaterialInstanceConstant::UpdateCachedData()
 
 		CachedHLSLTree = MoveTemp(LocalCachedTree);
 
-		if (bUsingNewHLSLGenerator && bHasStaticPermutationResource)
+		if (bUsingNewHLSLGenerator && bHasStaticPermutationResource && CachedHLSLTree.IsValid())
 		{
 			// Find the values of all overridden parameters including those overridden by parent instances. This is needed
 			// for correct HLSL tree preparation where static parameter values can affect the final result. We cannot call
diff --git a/Engine/Source/Runtime/Engine/Public/Materials/MaterialExpressionTextureCoordinate.h b/Engine/Source/Runtime/Engine/Public/Materials/MaterialExpressionTextureCoordinate.h
--- a/Engine/Source/Runtime/Engine/Public/Materials/MaterialExpressionTextureCoordinate.h
+++ b/Engine/Source/Runtime/Engine/Public/Materials/MaterialExpressionTextureCoordinate.h
@@ -36,6 +36,7 @@ class UMaterialExpressionTextureCoordinate : public UMaterialExpression
 
 	//~ Begin UMaterialExpression Interface
 #if WITH_EDITOR
+	virtual void Build(UE::MIR::FEmitter& Emitter) override;
 	virtual int32 Compile(class FMaterialCompiler* Compiler, int32 OutputIndex) override;
 	virtual void GetCaption(TArray<FString>& OutCaptions) const override;
 	virtual FText GetKeywords() const override { return FText::FromString(TEXT("texcoord")); }
diff --git a/Engine/Source/Runtime/Engine/Public/Materials/MaterialExpressionTextureSample.h b/Engine/Source/Runtime/Engine/Public/Materials/MaterialExpressionTextureSample.h
--- a/Engine/Source/Runtime/Engine/Public/Materials/MaterialExpressionTextureSample.h
+++ b/Engine/Source/Runtime/Engine/Public/Materials/MaterialExpressionTextureSample.h
@@ -99,7 +99,8 @@ public:
 	ENGINE_API virtual uint32 GetInputType(int32 InputIndex) override;
 	virtual bool CanIgnoreOutputIndex() { return true; }
 	ENGINE_API virtual bool GenerateHLSLExpression(FMaterialHLSLGenerator& Generator, UE::HLSLTree::FScope& Scope, int32 OutputIndex, UE::HLSLTree::FExpression const*& OutExpression) const override;
-#endif // WITH_EDITOR
+	ENGINE_API virtual void Build(UE::MIR::FEmitter& Emitter) override;
+	#endif // WITH_EDITOR
 	//~ End UMaterialExpression Interface
 
 	ENGINE_API void UpdateTextureResource(class UTexture* InTexture);
diff --git a/Engine/Source/Runtime/Engine/Public/Materials/MaterialIR.h b/Engine/Source/Runtime/Engine/Public/Materials/MaterialIR.h
--- a/Engine/Source/Runtime/Engine/Public/Materials/MaterialIR.h
+++ b/Engine/Source/Runtime/Engine/Public/Materials/MaterialIR.h
@@ -2,8 +2,7 @@
 #pragma once
 
 #include "Materials/MaterialIRCommon.h"
-
-enum EMaterialProperty : int;
+#include "MaterialTypes.h"
 
 #if WITH_EDITOR
 
@@ -14,33 +13,50 @@ enum EValueKind
 	/* Values */
 
 	VK_Constant,
+	VK_ExternalInput,
+	VK_MaterialParameter,
 
 	/* Instructions */
 
+	VK_InstructionBegin,
+
 	VK_Dimensional,
 	VK_SetMaterialOutput,
 	VK_BinaryOperator,
 	VK_Branch,
 	VK_Subscript,
 	VK_Cast,
+	VK_TextureSample,
 
 	VK_InstructionEnd,
-	VK_InstructionBegin = VK_Dimensional,
 };
 
+const TCHAR* ValueKindToString(EValueKind Kind);
+
 /* Values */
 
+enum EValueFlags
+{
+	VF_None = 0,
+	VF_ValueAnalyzed = 1,
+	VF_InstructionAnalyzed = 2,
+};
+
 struct FValue
 {
-	EValueKind Kind{};
+	EValueKind Kind : 8{};
+	EValueFlags Flags : 8{};
 	FTypePtr  Type{};
 
+	void SetFlags(EValueFlags InFlags) { Flags = (EValueFlags)(Flags | InFlags); }
 	bool IsA(EValueKind InKind) const { return Kind == InKind; }
 	FInstruction* AsInstruction();
 	const FInstruction* AsInstruction() const;
 	bool Equals(const FValue* Other) const;
 	uint32 GetSizeInBytes() const;
+	TArrayView<const FValue*> GetUses() const;
 	TArrayView<FValue*> GetUses();
+	UTexture* GetTexture();
 	
 	template <typename T>
 	T* As() { return this && IsA(T::TypeKind) ? static_cast<T*>(this) : nullptr; }
@@ -64,7 +80,7 @@ struct FConstant : TValue<VK_Constant>
 	{
 		bool  		Boolean;
 		TInteger	Integer;
-		TFloat 	Float;
+		TFloat 		Float;
 	};
 
 	template <typename T>
@@ -89,6 +105,60 @@ struct FConstant : TValue<VK_Constant>
 	}
 };
 
+enum class EExternalInput
+{
+	None,
+
+	TexCoord0,
+	TexCoord1,
+	TexCoord2,
+	TexCoord3,
+	TexCoord4,
+	TexCoord5,
+	TexCoord6,
+	TexCoord7,
+
+	TexCoord0_Ddx,
+	TexCoord1_Ddx,
+	TexCoord2_Ddx,
+	TexCoord3_Ddx,
+	TexCoord4_Ddx,
+	TexCoord5_Ddx,
+	TexCoord6_Ddx,
+	TexCoord7_Ddx,
+
+	TexCoord0_Ddy,
+	TexCoord1_Ddy,
+	TexCoord2_Ddy,
+	TexCoord3_Ddy,
+	TexCoord4_Ddy,
+	TexCoord5_Ddy,
+	TexCoord6_Ddy,
+	TexCoord7_Ddy,
+
+	Count,
+};
+
+static constexpr int TexCoordMaxNum = 8;
+
+const TCHAR* ExternalInputToString(EExternalInput Input);
+MIR::EExternalInput TexCoordIndexToExternalInput(int TexCoordIndex);
+FTypePtr GetExternalInputType(EExternalInput Id);
+bool IsExternalInputTexCoord(EExternalInput Id);
+bool IsExternalInputTexCoordDdx(EExternalInput Id);
+bool IsExternalInputTexCoordDdy(EExternalInput Id);
+
+struct FExternalInput : TValue<VK_ExternalInput>
+{
+	EExternalInput Id;
+};
+
+struct FMaterialParameter : TValue<VK_MaterialParameter>
+{	
+	FMaterialParameterInfo Info;
+	FMaterialParameterMetadata Metadata;
+};
+
 /* Instructions */
 
 struct FBlock
@@ -96,24 +166,18 @@ struct FBlock
 	FBlock* Parent{};
 	FInstruction* Instructions{};
 	int32 Level{};
-};
 
-enum EInstructionFlags
-{
-	IF_None = 0,
-	IF_Counted = 1,
+	FBlock* FindCommonParentWith(MIR::FBlock* Other);
 };
 
 struct FInstruction : FValue
 {
-	EInstructionFlags Flags{};
 	FInstruction* Next{};
 	FBlock* Block{};
 	uint32 NumUsers{};
 	uint32 NumProcessedUsers{};
 
-	void SetFlags(EInstructionFlags InFlags) { Flags = (EInstructionFlags)(Flags | InFlags); }
-	bool GetInnerBlock(int32 Index, FValue*& OutArg, FBlock*& OutBlock); 
+	FBlock* GetDesiredBlockForUse(int32 UseIndex); 
 };
 
 template <EValueKind TTypeKind>
@@ -158,16 +222,17 @@ enum EBinaryOperator
 	BO_Divide,
 
 	/* Comparison */
-	BO_Greater,
-	BO_GreaterOrEquals,
-	BO_Lower,
-	BO_LowerOrEquals,
+	BO_GreaterThan,
+	BO_GreaterThanOrEquals,
+	BO_LowerThan,
+	BO_LowerThanOrEquals,
 	BO_Equals,
 	BO_NotEquals,
 };
 
 bool IsArithmeticOperator(EBinaryOperator Op);
 bool IsComparisonOperator(EBinaryOperator Op);
+const TCHAR* BinaryOperatorToString(EBinaryOperator Op);
 
 struct FBinaryOperator : TInstruction<VK_BinaryOperator>
 {
@@ -196,5 +261,22 @@ struct FCast : TInstruction<VK_Cast>
 	FValue* Arg{};
 };
 
+struct FTextureSample : TInstruction<VK_TextureSample>
+{
+	FValue* TexCoordArg;
+	// TexCoordDerivatives;
+	FValue* MipValueArg;
+	FValue* AutomaticMipBiasArg;
+	UTexture* Texture;
+	ESamplerSourceMode SamplerSourceMode;
+	ETextureMipValueMode MipValueMode;
+	EMaterialSamplerType SamplerType;
+	
+	/* Analysis Values */
+	int TextureParameterIndex;
+};
+
+UTexture* GetTextureFromValue(FValue* TextureValue);
+
 } // namespace UE::MIR
 #endif
diff --git a/Engine/Source/Runtime/Engine/Public/Materials/MaterialIRCommon.h b/Engine/Source/Runtime/Engine/Public/Materials/MaterialIRCommon.h
--- a/Engine/Source/Runtime/Engine/Public/Materials/MaterialIRCommon.h
+++ b/Engine/Source/Runtime/Engine/Public/Materials/MaterialIRCommon.h
@@ -10,13 +10,7 @@
 
 /* Forward declarations */
 
-namespace UE::Shader
-{
-	struct FValue;
-}
-
-namespace ERHIFeatureLevel { enum Type : int; }
-
+enum EMaterialProperty : int;
 class FMaterial;
 class FMaterialIRModule;
 class FMaterialIRModule;
@@ -30,6 +24,14 @@ struct FShaderCompilerEnvironment;
 struct FStaticParameterSet;
 struct FMaterialInputDescription;
 struct FMaterialInsights;
+class UTexture;
+
+namespace UE::Shader
+{
+	struct FValue;
+}
+
+namespace ERHIFeatureLevel { enum Type : int; }
 
 namespace UE::MIR
 {
@@ -43,10 +45,12 @@ using FArithmeticTypePtr = const FArithmeticType*;
 
 /* IR */
 struct FValue;
+struct FGlobalInput;
 struct FInstruction;
 struct FSetMaterialOutput;
 using FValuePtr = const FValue*;
 using FInstructionPtr = const FInstruction*;
+enum class EExternalInput;
 
 /* Others */
 class FEmitter;
@@ -55,6 +59,6 @@ struct FBlock;
 }
 
 #define UE_MIR_UNREACHABLE() { check(!"Unreachable"); UE_ASSUME(false); }
-#define UE_MIR_TODO() UE_MIR_UNREACHABLE()
+#define UE_MIR_TODO() checkNoEntry()
 
 #endif // #if WITH_EDITOR
diff --git a/Engine/Source/Runtime/Engine/Public/Materials/MaterialIRDebug.cpp b/Engine/Source/Runtime/Engine/Public/Materials/MaterialIRDebug.cpp
new file mode 100644
--- /dev/null
+++ b/Engine/Source/Runtime/Engine/Public/Materials/MaterialIRDebug.cpp
@@ -0,0 +1,177 @@
+// Copyright Epic Games, Inc. All Rights Reserved.
+
+#include "Materials/MaterialIRDebug.h"
+#include "Materials/MaterialIRModule.h"
+#include "Materials/MaterialIR.h"
+#include "Materials/MaterialIRTypes.h"
+#include "Materials/MaterialAttributeDefinitionMap.h"
+#include "Misc/Paths.h"
+#include "Misc/FileHelper.h"
+
+#if WITH_EDITOR
+
+static TAutoConsoleVariable<bool> CVarDumpMaterialIRUseGraph(
+	TEXT("r.MaterialEditor.MaterialTranslator.DumpUseGraph"),
+	true,
+	TEXT("Whether the material translator should emit the Module IR 'Uses' graph in Graphviz Dot syntax (to 'MaterialIRDumpGraph.dot')."),
+	ECVF_RenderThreadSafe);
+
+
+static TAutoConsoleVariable<bool> CVarDumpMaterialIRUseGraph_EnableNext(
+	TEXT("r.MaterialEditor.MaterialTranslator.DumpUseGraphOpts.EnableSuccessors"),
+	false,
+	TEXT("Whether the Material Module IR 'Uses' graph should also display 'Instruction Next' edges."),
+	ECVF_RenderThreadSafe);
+
+
+namespace UE::MIR
+{
+
+static void DumpValueInfo(const MIR::FValue* Value, FString& Out)
+{
+	if (auto Constant = Value->As<MIR::FConstant>())
+	{
+		switch (Constant->Type->AsArithmetic()->ScalarKind)
+		{
+			case MIR::SK_Bool:	Out.Append(Constant->Boolean ? TEXT("true") : TEXT("false")); break;
+			case MIR::SK_Int:	Out.Appendf(TEXT("%lld"), Constant->Integer); break;
+			case MIR::SK_Float: Out.Appendf(TEXT("%f"), Constant->Float); break;
+			default: UE_MIR_UNREACHABLE();
+		}
+	}
+	else if (auto ExternalInput = Value->As<MIR::FExternalInput>())
+	{
+		Out.Append(MIR::ExternalInputToString(ExternalInput->Id));
+	}
+	else if (auto SetMaterailOutput = Value->As<MIR::FSetMaterialOutput>())
+	{
+		const FString& PropertyName = (SetMaterailOutput->Property == MP_SubsurfaceColor)
+		? TEXT("Subsurface")
+		: FMaterialAttributeDefinitionMap::GetAttributeName(SetMaterailOutput->Property);
+		
+		Out.Append(PropertyName);
+	}
+	else if (auto Subscript = Value->As<MIR::FSubscript>())
+	{
+		Out.Appendf(TEXT("Index: %d"), Subscript->Index);
+	}
+	else if (auto BinaryOperator = Value->As<MIR::FBinaryOperator>())
+	{
+		Out.Append(MIR::BinaryOperatorToString(BinaryOperator->Operator));
+	}
+}				
+
+static void DumpUseInfo(const MIR::FValue* Value, const MIR::FValue* Use, int UseIndex, FString& Out)
+{
+	if (auto Dimensional = Value->As<MIR::FDimensional>())
+	{
+		if (Dimensional->Type->AsArithmetic()->IsVector())
+		{
+			check(UseIndex < 4);
+			Out.AppendChar(TEXT("xyzw")[UseIndex]);
+		}
+		else
+		{
+			Out.AppendInt(UseIndex);
+		}
+	}
+	else if (auto If = Value->As<MIR::FBranch>())
+	{
+		static const TCHAR* Uses[] = { TEXT("condition"), TEXT("true"), TEXT("false") };
+		Out.Append(Uses[UseIndex]);
+	}
+	else if (auto BinaryOperator = Value->As<MIR::FBinaryOperator>())
+	{
+		static const TCHAR* Uses[] = { TEXT("lhs"), TEXT("rhs") };
+		Out.Append(Uses[UseIndex]);
+	}
+}
+
+void DebugDumpIRUseGraph(const FMaterialIRModule& Module)
+{
+	if (!CVarDumpMaterialIRUseGraph.GetValueOnAnyThread())
+	{
+		return;
+	}
+
+	FString Content;
+	TSet<const MIR::FValue*> Crawled;
+	TArray<const MIR::FValue*> ValueStack;
+
+	Content.Appendf(TEXT(
+		"digraph G {\n\n"
+		"rankdir=LR\n"
+		"node [shape=box,fontname=\"Consolas\"]\n"
+		"edge [fontname=\"Consolas\"]\n\n"
+	));
+
+	for (const UE::MIR::FSetMaterialOutput* Output : Module.GetOutputs())
+	{
+		ValueStack.Push(Output);
+	}
+
+	bool DumpInstructionSequence = CVarDumpMaterialIRUseGraph_EnableNext.GetValueOnAnyThread();
+
+	while (!ValueStack.IsEmpty())
+	{
+		const MIR::FValue* Value = ValueStack.Pop();
+
+		// Begin the node declaration
+		Content.Appendf(TEXT("\"%p\" [label=\"%s (%s)\n"),
+						Value,
+						MIR::ValueKindToString(Value->Kind),
+						Value->Type ? Value->Type->GetSpelling().GetData() : TEXT("???"));
+
+		DumpValueInfo(Value, Content);
+
+		// End the node declaration
+		Content.Append(TEXT("\"]\n"));
+
+		const MIR::FInstruction* Instr = Value->AsInstruction();
+		if (DumpInstructionSequence && Instr && Instr->Next)
+		{
+			Content.Appendf(TEXT("\"%p\" -> \"%p\" [color=\"red\"]\n"), Instr, Instr->Next);
+		}
+
+		int UseIndex = -1;
+		for (const MIR::FValue* Use : Value->GetUses())
+		{
+			++UseIndex;
+
+			if (!Use)
+			{
+				continue;
+			}
+			
+			Content.Appendf(TEXT("\"%p\" -> \"%p\" [label=\""), Value, Use);
+
+			DumpUseInfo(Value, Use, UseIndex, Content);
+
+			Content.Appendf(TEXT("\"]\n"));
+
+			if (!Crawled.Contains(Use))
+			{
+				Crawled.Add(Use);
+				ValueStack.Push(Use);
+			}
+
+			if (DumpInstructionSequence && Instr)
+			{
+				const MIR::FInstruction* UseInstr = Use->AsInstruction();
+				if (UseInstr && UseInstr->Block != Instr->Block)
+				{
+					Content.Appendf(TEXT("\"%p\" -> \"%p\" [color=\"red\", style=\"dashed\"]\n"), UseInstr, Instr);
+				}
+			}
+		}
+	}
+
+	Content.Appendf(TEXT("\n}\n"));
+	
+	FString FilePath = FPaths::Combine(FPaths::ProjectLogDir(), TEXT("MaterialIRDumpGraph.dot"));
+	FFileHelper::SaveStringToFile(Content, *FilePath);
+}
+
+} // namespace UE::MIR
+
+#endif // #if WITH_EDITOR
diff --git a/Engine/Source/Runtime/Engine/Public/Materials/MaterialIRDebug.h b/Engine/Source/Runtime/Engine/Public/Materials/MaterialIRDebug.h
new file mode 100644
--- /dev/null
+++ b/Engine/Source/Runtime/Engine/Public/Materials/MaterialIRDebug.h
@@ -0,0 +1,16 @@
+// Copyright Epic Games, Inc. All Rights Reserved.
+#pragma once
+
+#include "Materials/MaterialIRCommon.h"
+
+#if WITH_EDITOR
+
+namespace UE::MIR
+{
+
+//
+void DebugDumpIRUseGraph(const FMaterialIRModule& Module);
+
+} // namespace UE::MIR
+
+#endif // #if WITH_EDITOR
diff --git a/Engine/Source/Runtime/Engine/Public/Materials/MaterialIREmitter.h b/Engine/Source/Runtime/Engine/Public/Materials/MaterialIREmitter.h
--- a/Engine/Source/Runtime/Engine/Public/Materials/MaterialIREmitter.h
+++ b/Engine/Source/Runtime/Engine/Public/Materials/MaterialIREmitter.h
@@ -23,6 +23,8 @@ struct FSwizzleMask
 	EVectorComponent Components[4];
 	int NumComponents{};
 
+	static FSwizzleMask XYZ();
+
 	FSwizzleMask() {}
 	FSwizzleMask(EVectorComponent X);
 	FSwizzleMask(EVectorComponent X, EVectorComponent Y);
@@ -70,7 +72,9 @@ public:
 	// type matches `Kind`.
 	FValue* TryGetOfType(const FExpressionInput* Input, ETypeKind Kind);
 
-	/* Analysis */
+	/* Error Checking */
+
+	bool CheckValueValid(const FValue* Value);
 
 	//
 	void CheckInputIsScalar(const FExpressionInput* Input, FValue* InputValue);
@@ -101,11 +105,13 @@ public:
 	FValue* EmitVector2(FValue* InX, FValue* InY);
 	FValue* EmitVector3(FValue* InX, FValue* InY, FValue* InZ);
 	FValue* EmitVector4(FValue* InX, FValue* InY, FValue* InZ, FValue* InW);
+	FValue* GetExternalInput(EExternalInput Id);
 
 	/* Other Values */
 
-	FValue* EmitSubscript(FValue* Value, int ComponentIndex);
+	FValue* TryEmitSubscript(FValue* Value, int ComponentIndex);
 	FValue* TryEmitSwizzle(FValue* Value, FSwizzleMask Mask);
+	FValue* GetParameter(FName Name, const FMaterialParameterMetadata& Metadata);
 
 	/* Instructions */
 
@@ -113,12 +119,14 @@ public:
 	FValue* EmitBinaryOperator(EBinaryOperator Operator, FValue* Lhs, FValue* Rhs);
 	FValue* EmitBranch(FValue* Condition, FValue* True, FValue* False);
 	FValue* TryEmitConstruct(FTypePtr Type, FValue* Initializer);
+	FValue* TryEmitTextureSample(UTexture* Texture, FValue* TexCoord, ESamplerSourceMode SamplerSourceMode, ETextureMipValueMode MipValueMode, EMaterialSamplerType SamplerType);
 
 	/* Types */
 
 	FArithmeticTypePtr TryGetCommonArithmeticType(FArithmeticTypePtr A, FArithmeticTypePtr B);
 
 	/* Error reporting */
+
 	bool IsInvalid() const { return bHasExprBuildError; }
 
 	template <int TFormatLength, typename... TArgs>
@@ -130,8 +138,11 @@ public:
 	void Error(FString Message);
 
 	struct FPrivate;
+	friend FMaterialIRModuleBuilder;
 
 private:
+	void Initialize();
+
 	UMaterial* Material{};
 	FMaterialIRModule* Module{};
 	UMaterialExpression* Expression{};
@@ -139,8 +150,6 @@ private:
 	bool bHasExprBuildError = false;
 	FValue* ConstantTrue;
 	FValue* ConstantFalse;
-
-	friend FMaterialIRModuleBuilder;
 };
 
 } // namespace UE::MIR
diff --git a/Engine/Source/Runtime/Engine/Public/Materials/MaterialIRModule.h b/Engine/Source/Runtime/Engine/Public/Materials/MaterialIRModule.h
--- a/Engine/Source/Runtime/Engine/Public/Materials/MaterialIRModule.h
+++ b/Engine/Source/Runtime/Engine/Public/Materials/MaterialIRModule.h
@@ -15,6 +15,13 @@ public:
 		FString Message;
 	};
 
+	struct FStatistics
+	{
+		TBitArray<> ExternalInputUsedMask[SF_NumFrequencies];
+		int NumVertexTexCoords;
+		int NumPixelTexCoords;
+	};
+
 public:
 	FMaterialIRModule();
 	~FMaterialIRModule();
@@ -24,14 +31,17 @@ public:
 	TArrayView<const UE::MIR::FSetMaterialOutput* const> GetOutputs() const { return Outputs; }
 	const UE::MIR::FBlock& GetRootBlock() const { return *RootBlock; }
 	TArrayView<const FError> GetErrors() const { return Errors; }
+	const FStatistics& GetStatistics() const { return Statistics; }
 
 private:
 	EShaderPlatform ShaderPlatform;
 	FMaterialCompilationOutput CompilationOutput;
+	FMemStackBase Allocator{};
 	TArray<UE::MIR::FValue*> Values;
 	TArray<UE::MIR::FSetMaterialOutput*> Outputs;
 	TArray<FError> Errors;
 	UE::MIR::FBlock* RootBlock;
+	FStatistics Statistics;
 
 	friend UE::MIR::FEmitter;
 	friend FMaterialIRModuleBuilder;
diff --git a/Engine/Source/Runtime/Engine/Public/Materials/MaterialIRTypes.h b/Engine/Source/Runtime/Engine/Public/Materials/MaterialIRTypes.h
--- a/Engine/Source/Runtime/Engine/Public/Materials/MaterialIRTypes.h
+++ b/Engine/Source/Runtime/Engine/Public/Materials/MaterialIRTypes.h
@@ -12,7 +12,8 @@ namespace UE::MIR
 enum ETypeKind
 {
 	TK_Void,
-	TK_Arithmetic
+	TK_Arithmetic,
+	TK_Texture,
 };
 
 const TCHAR* TypeKindToString(ETypeKind Kind);
@@ -67,6 +68,9 @@ struct FArithmeticType : FType
 	static FArithmeticTypePtr GetBool1();
 	static FArithmeticTypePtr GetInt1();
 	static FArithmeticTypePtr GetFloat1();
+	static FArithmeticTypePtr GetFloat2();
+	static FArithmeticTypePtr GetFloat3();
+	static FArithmeticTypePtr GetFloat4();
 
 	static FArithmeticTypePtr GetScalar(EScalarKind InScalarKind);
 	static FArithmeticTypePtr GetVector(EScalarKind InScalarKind, int NumRows);
@@ -80,6 +84,11 @@ struct FArithmeticType : FType
 	FArithmeticTypePtr ToScalar() const;
 };
 
+struct FTextureType : FType
+{
+	static FTypePtr Get();
+};
+
 } // namespace UE::MIR
 
 #endif // #if WITH_EDITOR
