Add MeshPaintTextureReplace material expression.
This is a material graph node that can select between two inputs depending on whether there is a valid mesh paint texture on the shaded primitive or not.
This can be useful for uber materials that need paths for components that are both mesh texture color painted or not. (We already show a white fallback color when no mesh paint texture is present, but this allows more control).
While adding this, the mesh paint descriptor layout was also modified so that it can carry an optional texture coordinate index. This will be useful for debug views where we need to know the UV channel without having any other material context.

diff --git a/Engine/Shaders/Private/MaterialTemplate.ush b/Engine/Shaders/Private/MaterialTemplate.ush
--- a/Engine/Shaders/Private/MaterialTemplate.ush
+++ b/Engine/Shaders/Private/MaterialTemplate.ush
@@ -30,6 +30,7 @@
 #include "/Engine/Private/DeferredShadingCommon.ush"
 #include "/Engine/Private/SparseVolumeTexture/SparseVolumeTextureCommon.ush"
 #include "/Engine/Shared/EnvironmentComponentsFlags.h"
+#include "/Engine/Private/MeshPaintTextureCommon.ush"
 
 // Reduce latency of Nanite material base pass pixel shaders caused by register pressure due to the calculation of WPO
 #if PIXELSHADER && IS_NANITE_SHADING_PASS && IS_BASE_PASS
diff --git a/Engine/Shaders/Private/MeshPaintTextureCommon.ush b/Engine/Shaders/Private/MeshPaintTextureCommon.ush
new file mode 100644
--- /dev/null
+++ b/Engine/Shaders/Private/MeshPaintTextureCommon.ush
@@ -0,0 +1,20 @@
+// Copyright Epic Games, Inc. All Rights Reserved.
+
+#pragma once
+
+#include "SceneData.ush"
+
+uint2 GetMeshPaintTextureDescriptor(FPrimitiveSceneData PrimitiveData)
+{
+	return PrimitiveData.MeshPaintTextureDescriptor;
+}
+
+uint GetMeshPaintTextureCoordinateIndex(FPrimitiveSceneData PrimitiveData)
+{
+	return GetMeshPaintTextureDescriptor(PrimitiveData).y >> 30;
+}
+
+bool GetMeshPaintTextureDescriptorIsValid(FPrimitiveSceneData PrimitiveData)
+{
+	return GetMeshPaintTextureDescriptor(PrimitiveData).x != 0;
+}
diff --git a/Engine/Source/Runtime/Engine/Classes/VT/MeshPaintVirtualTexture.h b/Engine/Source/Runtime/Engine/Classes/VT/MeshPaintVirtualTexture.h
--- a/Engine/Source/Runtime/Engine/Classes/VT/MeshPaintVirtualTexture.h
+++ b/Engine/Source/Runtime/Engine/Classes/VT/MeshPaintVirtualTexture.h
@@ -51,7 +51,7 @@ namespace MeshPaintVirtualTexture
 	 * Get the 2 dword texture descriptor from texture resource. 
 	 * Will return a null descriptor if the texture resource is not from a UMeshPaintVirtualTexture.
 	 */
-	FUintVector2 GetTextureDescriptor(FTextureResource* InTextureResource);
+	FUintVector2 GetTextureDescriptor(FTextureResource* InTextureResource, uint32 InOptionalCoordinateIndex = 0);
 
 	/** Scene view parameters that describe the virtual texture space shared by all UMeshPaintVirtualTexture objects. */
 	struct FUniformParams
diff --git a/Engine/Source/Runtime/Engine/Private/Materials/HLSLMaterialTranslator.cpp b/Engine/Source/Runtime/Engine/Private/Materials/HLSLMaterialTranslator.cpp
--- a/Engine/Source/Runtime/Engine/Private/Materials/HLSLMaterialTranslator.cpp
+++ b/Engine/Source/Runtime/Engine/Private/Materials/HLSLMaterialTranslator.cpp
@@ -2679,7 +2679,7 @@ void FHLSLMaterialTranslator::GetMaterialEnvironment(EShaderPlatform InPlatform,
 
 			OutEnvironment.SetDefine(
 				*FString::Printf(TEXT("VIRTUALTEXTURE_PAGETABLE_UNIFORM_%d"), i), 
-				TEXT("GetPrimitiveData(Parameters).MeshPaintTextureDescriptor"));
+				TEXT("GetMeshPaintTextureDescriptor(GetPrimitiveData(Parameters))"));
 		}
 		else
 		{
@@ -8793,7 +8793,12 @@ int32 FHLSLMaterialTranslator::VertexColor()
 
 int32 FHLSLMaterialTranslator::MeshPaintTextureDescriptor()
 {
-	return GetPrimitiveProperty(EMaterialValueType(MCT_TextureVirtual | MCT_TextureMeshPaint), TEXT("MeshPaintTextureDescriptor"), TEXT("MeshPaintTextureDescriptor"));
+	return AddInlinedCodeChunkZeroDeriv(EMaterialValueType(MCT_TextureVirtual | MCT_TextureMeshPaint), TEXT("GetMeshPaintTextureDescriptor(GetPrimitiveData(Parameters))"));
+}
+
+int32 FHLSLMaterialTranslator::MeshPaintTextureReplace(int32 Invalid, int32 Valid)
+{
+	return GenericSwitch(TEXT("GetMeshPaintTextureDescriptorIsValid(GetPrimitiveData(Parameters))"), Valid, Invalid);
 }
 
 int32 FHLSLMaterialTranslator::PreSkinnedPosition()
diff --git a/Engine/Source/Runtime/Engine/Private/Materials/HLSLMaterialTranslator.h b/Engine/Source/Runtime/Engine/Private/Materials/HLSLMaterialTranslator.h
--- a/Engine/Source/Runtime/Engine/Private/Materials/HLSLMaterialTranslator.h
+++ b/Engine/Source/Runtime/Engine/Private/Materials/HLSLMaterialTranslator.h
@@ -1119,6 +1119,7 @@ protected:
 
 	virtual int32 VertexColor() override;
 	virtual int32 MeshPaintTextureDescriptor() override;
+	virtual int32 MeshPaintTextureReplace(int32 Invalid, int32 Valid) override;
 
 	virtual int32 PreSkinnedPosition() override;
 	virtual int32 PreSkinnedNormal() override;
diff --git a/Engine/Source/Runtime/Engine/Private/Materials/MaterialExpressions.cpp b/Engine/Source/Runtime/Engine/Private/Materials/MaterialExpressions.cpp
--- a/Engine/Source/Runtime/Engine/Private/Materials/MaterialExpressions.cpp
+++ b/Engine/Source/Runtime/Engine/Private/Materials/MaterialExpressions.cpp
@@ -125,6 +125,7 @@
 #include "Materials/MaterialExpressionInverseLinearInterpolate.h"
 #include "Materials/MaterialExpressionLightmapUVs.h"
 #include "Materials/MaterialExpressionMeshPaintTextureObject.h"
+#include "Materials/MaterialExpressionMeshPaintTextureReplace.h"
 #include "Materials/MaterialExpressionPrecomputedAOMask.h"
 #include "Materials/MaterialExpressionLightmassReplace.h"
 #include "Materials/MaterialExpressionLightVector.h"
@@ -3954,6 +3955,16 @@ UMaterialExpressionMeshPaintTextureObject::UMaterialExpressionMeshPaintTextureOb
 
 #if WITH_EDITOR
 
+void UMaterialExpressionMeshPaintTextureObject::GetCaption(TArray<FString>& OutCaptions) const
+{
+	OutCaptions.Add(TEXT("MeshPaintTextureObject"));
+}
+
+void UMaterialExpressionMeshPaintTextureObject::GetExpressionToolTip(TArray<FString>& OutToolTip)
+{
+	ConvertToMultilineToolTip(TEXT("Get the Mesh Paint Texture object for feeding to a Texture Sample node."), 40, OutToolTip);
+}
+
 uint32 UMaterialExpressionMeshPaintTextureObject::GetOutputType(int32 OutputIndex)
 {
 	return MCT_TextureVirtual | MCT_TextureMeshPaint;
@@ -3966,11 +3977,60 @@ int32 UMaterialExpressionMeshPaintTextureObject::Compile(class FMaterialCompiler
 
 #endif // WITH_EDITOR
 
-void UMaterialExpressionMeshPaintTextureObject::GetCaption(TArray<FString>& OutCaptions) const
+UMaterialExpressionMeshPaintTextureReplace::UMaterialExpressionMeshPaintTextureReplace(const FObjectInitializer& ObjectInitializer)
+	: Super(ObjectInitializer)
 {
-	OutCaptions.Add(TEXT("MeshPaintTextureObject"));
+	// Structure to hold one-time initialization
+	struct FConstructorStatics
+	{
+		FText NAME_MeshPaintTexture;
+		FConstructorStatics()
+			: NAME_MeshPaintTexture(LOCTEXT("MeshPaintTexture", "MeshPaintTexture"))
+		{
+		}
+	};
+	static FConstructorStatics ConstructorStatics;
+
+#if WITH_EDITORONLY_DATA
+	MenuCategories.Add(ConstructorStatics.NAME_MeshPaintTexture);
+#endif
+
+	Outputs.Reset();
+	Outputs.Add(FExpressionOutput(TEXT("")));
+}
+
+#if WITH_EDITOR
+
+void UMaterialExpressionMeshPaintTextureReplace::GetCaption(TArray<FString>& OutCaptions) const
+{
+	OutCaptions.Add(TEXT("MeshPaintTextureReplace"));
 }
 
+void UMaterialExpressionMeshPaintTextureReplace::GetExpressionToolTip(TArray<FString>& OutToolTip)
+{
+	ConvertToMultilineToolTip(TEXT("Switch between inputs according to whether there is a valid Mesh Paint Texture available to sample."), 40, OutToolTip);
+}
+
+int32 UMaterialExpressionMeshPaintTextureReplace::Compile(class FMaterialCompiler* Compiler, int32 OutputIndex)
+{
+	if (!Default.GetTracedInput().Expression)
+	{
+		return Compiler->Errorf(TEXT("Missing input Default"));
+	}
+	else if (!MeshPaintTexture.GetTracedInput().Expression)
+	{
+		return Compiler->Errorf(TEXT("Missing input MeshPaintTexture"));
+	}
+	else
+	{
+		const int32 Arg1 = Default.Compile(Compiler);
+		const int32 Arg2 = MeshPaintTexture.Compile(Compiler);
+		return Compiler->MeshPaintTextureReplace(Arg1, Arg2);
+	}
+}
+
+#endif // WITH_EDITOR
+
 UMaterialExpressionAdd::UMaterialExpressionAdd(const FObjectInitializer& ObjectInitializer)
 	: Super(ObjectInitializer)
 {
diff --git a/Engine/Source/Runtime/Engine/Private/VT/MeshPaintVirtualTexture.cpp b/Engine/Source/Runtime/Engine/Private/VT/MeshPaintVirtualTexture.cpp
--- a/Engine/Source/Runtime/Engine/Private/VT/MeshPaintVirtualTexture.cpp
+++ b/Engine/Source/Runtime/Engine/Private/VT/MeshPaintVirtualTexture.cpp
@@ -193,7 +193,7 @@ namespace MeshPaintVirtualTexture
 		return {};
 	}
 
-	FUintVector2 GetTextureDescriptor(FTextureResource* InTextureResource)
+	FUintVector2 GetTextureDescriptor(FTextureResource* InTextureResource, uint32 InOptionalCoordinateIndex)
 	{
 		FUintVector2 Descriptor(0, 0);
 
@@ -228,8 +228,9 @@ namespace MeshPaintVirtualTexture
 		const uint32 MaxLevel = AllocatedVT->GetMaxLevel();
 
 		// Descriptor layout should match shader unpacking in VTPageTableUniform_Unpack().
+		// We optionally pack a coordinate index which is used for debug views.
 		Descriptor.X = vPageX | (vPageY << 12) | (vPageTableMipBias << 24) | (SpaceID << 28);
-		Descriptor.Y = WidthInPages | (HeightInPages << 12) | (MaxLevel << 24);
+		Descriptor.Y = WidthInPages | (HeightInPages << 12) | (MaxLevel << 24) | (InOptionalCoordinateIndex << 30);
 		
 		return Descriptor;
 	}
diff --git a/Engine/Source/Runtime/Engine/Public/MaterialCompiler.h b/Engine/Source/Runtime/Engine/Public/MaterialCompiler.h
--- a/Engine/Source/Runtime/Engine/Public/MaterialCompiler.h
+++ b/Engine/Source/Runtime/Engine/Public/MaterialCompiler.h
@@ -452,6 +452,7 @@ public:
 
 	virtual int32 VertexColor() = 0;
 	virtual int32 MeshPaintTextureDescriptor() = 0;
+	virtual int32 MeshPaintTextureReplace(int32 Invalid, int32 Valid) = 0;
 
 	virtual int32 PreSkinnedPosition() = 0;
 	virtual int32 PreSkinnedNormal() = 0;
@@ -1002,6 +1003,7 @@ public:
 
 	virtual int32 VertexColor() override { return Compiler->VertexColor(); }
 	virtual int32 MeshPaintTextureDescriptor() override { return Compiler->MeshPaintTextureDescriptor(); }
+	virtual int32 MeshPaintTextureReplace(int32 Invalid, int32 Valid) override { return Compiler->MeshPaintTextureReplace(Invalid, Valid); }
 
 	virtual int32 PreSkinnedPosition() override { return Compiler->PreSkinnedPosition(); }
 	virtual int32 PreSkinnedNormal() override { return Compiler->PreSkinnedNormal(); }
diff --git a/Engine/Source/Runtime/Engine/Public/Materials/MaterialExpressionMeshPaintTextureObject.h b/Engine/Source/Runtime/Engine/Public/Materials/MaterialExpressionMeshPaintTextureObject.h
--- a/Engine/Source/Runtime/Engine/Public/Materials/MaterialExpressionMeshPaintTextureObject.h
+++ b/Engine/Source/Runtime/Engine/Public/Materials/MaterialExpressionMeshPaintTextureObject.h
@@ -11,15 +11,12 @@ class UMaterialExpressionMeshPaintTextureObject : public UMaterialExpression
 {
 	GENERATED_UCLASS_BODY()
 
-protected:
-	/** Initialize the output pins. */
-	ENGINE_API void InitOutputs();
-
 	//~ Begin UMaterialExpression Interface
 #if WITH_EDITOR
-	ENGINE_API virtual uint32 GetOutputType(int32 OutputIndex) override;
-	ENGINE_API virtual int32 Compile(class FMaterialCompiler* Compiler, int32 OutputIndex) override;
+	virtual void GetCaption(TArray<FString>& OutCaptions) const override;
+	virtual void GetExpressionToolTip(TArray<FString>& OutToolTip) override;
+	virtual uint32 GetOutputType(int32 OutputIndex) override;
+	virtual int32 Compile(class FMaterialCompiler* Compiler, int32 OutputIndex) override;
 #endif
-	void GetCaption(TArray<FString>& OutCaptions) const;
 	//~ End UMaterialExpression Interface
 };
diff --git a/Engine/Source/Runtime/Engine/Public/Materials/MaterialExpressionMeshPaintTextureReplace.h b/Engine/Source/Runtime/Engine/Public/Materials/MaterialExpressionMeshPaintTextureReplace.h
new file mode 100644
--- /dev/null
+++ b/Engine/Source/Runtime/Engine/Public/Materials/MaterialExpressionMeshPaintTextureReplace.h
@@ -0,0 +1,27 @@
+// Copyright Epic Games, Inc. All Rights Reserved.
+
+#pragma once
+
+#include "Materials/MaterialExpression.h"
+
+#include "MaterialExpressionMeshPaintTextureReplace.generated.h"
+
+UCLASS(collapsecategories, hidecategories = Object, MinimalAPI)
+class UMaterialExpressionMeshPaintTextureReplace : public UMaterialExpression
+{
+	GENERATED_UCLASS_BODY()
+
+	UPROPERTY(meta = (RequiredInput = "true", ToolTip = "Input will be used when no valid MeshPaintTexture is available"))
+	FExpressionInput Default;
+
+	UPROPERTY(meta = (RequiredInput = "true", ToolTip = "Input will be used when a valid MeshPaintTexture is available"))
+	FExpressionInput MeshPaintTexture;
+
+	//~ Begin UMaterialExpression Interface
+#if WITH_EDITOR
+	virtual void GetCaption(TArray<FString>& OutCaptions) const override;
+	virtual void GetExpressionToolTip(TArray<FString>& OutToolTip) override;
+	virtual int32 Compile(class FMaterialCompiler* Compiler, int32 OutputIndex) override;
+#endif
+	//~ End UMaterialExpression Interface
+};
