Null check deformer graph vertex factory updates to avoid rare crash. Prefer to skip the update instead.

diff --git a/Engine/Source/Runtime/Engine/Private/SkeletalMeshDeformerHelpers.cpp b/Engine/Source/Runtime/Engine/Private/SkeletalMeshDeformerHelpers.cpp
--- a/Engine/Source/Runtime/Engine/Private/SkeletalMeshDeformerHelpers.cpp
+++ b/Engine/Source/Runtime/Engine/Private/SkeletalMeshDeformerHelpers.cpp
@@ -258,7 +258,12 @@ void FSkeletalMeshDeformerHelpers::UpdateVertexFactoryBufferOverrides(FRHIComman
 	{
 		FGPUBaseSkinVertexFactory const* BaseVertexFactory = MeshObjectGPU->GetBaseSkinVertexFactory(InLodIndex, SectionIndex);
 		FGPUSkinPassthroughVertexFactory* TargetVertexFactory = LOD.GPUSkinVertexFactories.PassthroughVertexFactories[SectionIndex].Get();
-		TargetVertexFactory->SetVertexAttributes(RHICmdList, BaseVertexFactory, Desc);
+
+		// The passthrough vertex factory should exist if we got this far, but prefer skipping the update to crashing if that assumption fails.
+		if (ensure(TargetVertexFactory))
+		{
+			TargetVertexFactory->SetVertexAttributes(RHICmdList, BaseVertexFactory, Desc);
+		}
 	}
 }
 
@@ -290,7 +295,12 @@ void FSkeletalMeshDeformerHelpers::ResetVertexFactoryBufferOverrides(FSkeletalMe
 	for (int32 SectionIndex = 0; SectionIndex < NumSections; ++SectionIndex)
 	{
 		FGPUSkinPassthroughVertexFactory* TargetVertexFactory = LOD.GPUSkinVertexFactories.PassthroughVertexFactories[SectionIndex].Get();
-		TargetVertexFactory->ResetVertexAttributes(RHICmdList);
+
+		// The passthrough vertex factory should exist if we got this far, but prefer skipping the update to crashing if that assumption fails.
+		if (ensure(TargetVertexFactory))
+		{
+			TargetVertexFactory->ResetVertexAttributes(RHICmdList);
+		}
 	}
 
 #if RHI_RAYTRACING
