* Deprecate unused ResponseID variable in NetConnection
* Deleted pre-5.5 deprecated variables and function from PlayerState and NetConnection

#rnx

diff --git a/Engine/Source/Runtime/Engine/Classes/Engine/NetConnection.h b/Engine/Source/Runtime/Engine/Classes/Engine/NetConnection.h
--- a/Engine/Source/Runtime/Engine/Classes/Engine/NetConnection.h
+++ b/Engine/Source/Runtime/Engine/Classes/Engine/NetConnection.h
@@ -438,9 +438,10 @@ public:
 	int32			PacketOverhead;			// Bytes overhead per packet sent.
 	FString			Challenge;				// Server-generated challenge.
 	FString			ClientResponse;			// Client-generated response.
+	FString			RequestURL;				// URL requested by client
 
+	UE_DEPRECATED(5.5, "Variable is deprecated and unassigned")
 	int32			ResponseId;				// Id assigned by the server for linking responses to connections upon authentication
-	FString			RequestURL;				// URL requested by client
 
 	// Login state tracking
 	EClientLoginState::Type	ClientLoginState;
@@ -609,12 +610,6 @@ public:
 	UPROPERTY(config)
 	int32 DefaultMaxChannelSize;
 
-	UE_DEPRECATED(5.1, "Deprecated in favor of DefaultMaxChannelSize config property.")
-	static const int32 DEFAULT_MAX_CHANNEL_SIZE;
-
-	UE_DEPRECATED(5.1, "No longer used")
-	int32 MaxChannelSize;
-
 	TArray<TObjectPtr<UChannel>>	Channels;
 	TArray<int32>		OutReliable;
 	TArray<int32>		InReliable;
@@ -622,12 +617,6 @@ public:
 	int32				InitOutReliable;
 	int32				InitInReliable;
 
-	// Network version
-	UE_DEPRECATED(5.2, "Deprecated in favor of NetworkCustomVersions, please use GetNetworkCustomVersion instead")
-	uint32				EngineNetworkProtocolVersion;
-	UE_DEPRECATED(5.2, "Deprecated in favor of NetworkCustomVersions, please use GetNetworkCustomVersion instead")
-	uint32				GameNetworkProtocolVersion;
-
 	uint32 GetNetworkCustomVersion(const FGuid& VersionGuid) const;
 	void SetNetworkCustomVersions(const FCustomVersionContainer& CustomVersions);
 
@@ -825,10 +814,6 @@ public:
 	/** This holds a list of actor channels that want to fully shutdown, but need to continue processing bunches before doing so */
 	TMap<FNetworkGUID, TArray<TObjectPtr<class UActorChannel>>> KeepProcessingActorChannelBunchesMap;
 
-	/** A list of replicators that belong to recently dormant actors/objects */
-	UE_DEPRECATED(5.2, "The DormantReplicatorMap is deprecated in favor of the private DormantReplicatorSet.")
-	TMap<FObjectKey, TSharedRef<FObjectReplicator>> DormantReplicatorMap;
-
 private:
 
 	UE::Net::Private::FDormantReplicatorHolder DormantReplicatorSet;
@@ -1292,10 +1277,6 @@ public:
 	/** Forces properties on this actor to do a compare for one frame (rather than share shadow state) */
 	ENGINE_API void ForcePropertyCompare( AActor* Actor );
 
-	/** Wrapper for validating an objects dormancy state, and to prepare the object for replication again */
-	UE_DEPRECATED(5.2, "FlushDormancyForObject has been replaced with a version that needs to receive the dormant actor.")
-	void FlushDormancyForObject( UObject* Object ) {}
-
 	/**
 	* Validate an objects dormancy state and prepare the object for replication again
 	* 
@@ -1393,17 +1374,9 @@ public:
 	/** Returns the OutgoingBunches array, only to be used by UChannel::SendBunch */
 	TArray<FOutBunch *>& GetOutgoingBunches() { return OutgoingBunches; }
 
-	/** Add a replicator to the dormancy map and release its strong pointer to its object */
-	UE_DEPRECATED(5.2, "AddDormantReplicator has been replaced by StoreDormantReplicator and will be removed soon.")
-	void AddDormantReplicator(UObject* Object, const TSharedRef<FObjectReplicator>& Replicator) {}
-
 	/** Store a replicator to the dormancy map and release its strong pointer to its object */
 	void StoreDormantReplicator(AActor* OwnerActor, UObject* Object, const TSharedRef<FObjectReplicator>& ObjectReplicator);
 	
-	/** Find a dormant replicator for the channel actor or one of its subobjects. Removes it from the map if found. */
-	UE_DEPRECATED(5.2, "FindAndRemoveDormantReplicator is deprecated. Use the new version that needs to receive the owning actor.")
-	TSharedPtr<FObjectReplicator> FindAndRemoveDormantReplicator(UObject* Object) { return {}; }
-
 	/** Find a dormant replicator for the channel actor or one of its subobjects. Removes it from the map if found. */
 	TSharedPtr<FObjectReplicator> FindAndRemoveDormantReplicator(AActor* OwnerActor, UObject* Object);
 
diff --git a/Engine/Source/Runtime/Engine/Classes/GameFramework/PlayerState.h b/Engine/Source/Runtime/Engine/Classes/GameFramework/PlayerState.h
--- a/Engine/Source/Runtime/Engine/Classes/GameFramework/PlayerState.h
+++ b/Engine/Source/Runtime/Engine/Classes/GameFramework/PlayerState.h
@@ -42,20 +42,16 @@ class APlayerState : public AInfo
 {
 	GENERATED_UCLASS_BODY()
 
-	// destructor for handling property deprecation, please remove after all deprecated properties are gone
-	ENGINE_API virtual ~APlayerState();
+private:
 
 	/** Player's current score. */
-	UE_DEPRECATED(4.25, "This member will be made private. Use GetScore or SetScore instead.")
-	UPROPERTY(ReplicatedUsing=OnRep_Score, Category=PlayerState, BlueprintGetter=GetScore)
+	UPROPERTY(ReplicatedUsing = OnRep_Score, Category = PlayerState, BlueprintGetter = GetScore)
 	float Score;
 
 	/** Unique net id number. Actual value varies based on current online subsystem, use it only as a guaranteed unique number per player. */
-	UE_DEPRECATED(4.25, "This member will be made private. Use GetPlayerId or SetPlayerId instead.")
-	UPROPERTY(ReplicatedUsing=OnRep_PlayerId, Category=PlayerState, BlueprintGetter=GetPlayerId)
+	UPROPERTY(ReplicatedUsing = OnRep_PlayerId, Category = PlayerState, BlueprintGetter = GetPlayerId)
 	int32 PlayerId;
 
-private:
 	/** Replicated compressed ping for this player (holds ping in msec divided by 4) */
 	UPROPERTY(Replicated, Category=PlayerState, BlueprintGetter=GetCompressedPing, meta=(AllowPrivateAccess))
 	uint8 CompressedPing;
@@ -72,27 +68,25 @@ private:
 	UPROPERTY(EditDefaultsOnly, Category=PlayerState)
 	uint8 bShouldUpdateReplicatedPing:1;
 
-public:
 	/** Whether this player is currently a spectator */
-	UE_DEPRECATED(4.25, "This member will be made private. Use IsSpectator or SetIsSpectator instead.")
 	UPROPERTY(Replicated, Category=PlayerState, BlueprintGetter=IsSpectator)
 	uint8 bIsSpectator:1;
 
 	/** Whether this player can only ever be a spectator */
-	UE_DEPRECATED(4.25, "This member will be made private. Use IsOnlyASpectator or SetIsOnlyASpectator instead.")
 	UPROPERTY(Replicated)
 	uint8 bOnlySpectator:1;
 
 	/** True if this PlayerState is associated with an AIController */
-	UE_DEPRECATED(4.25, "This member will be made private. Use IsABot or SetIsABot instead.")
 	UPROPERTY(Replicated, Category=PlayerState, BlueprintGetter=IsABot)
 	uint8 bIsABot:1;
 
+public:
 	/** client side flag - whether this player has been welcomed or not (player entered message) */
 	uint8 bHasBeenWelcomed:1;
 
+private:
+
 	/** Means this PlayerState came from the GameMode's InactivePlayerArray */
-	UE_DEPRECATED(4.25, "This member will be made private. Use IsInactive or SetIsInactive instead.")
 	UPROPERTY(ReplicatedUsing=OnRep_bIsInactive)
 	uint8 bIsInactive:1;
 
@@ -100,18 +94,28 @@ public:
 	 * waiting for the player to finish the transition before creating a new one
 	 * this is used to avoid preserving the PlayerState in the InactivePlayerArray if the player leaves
 	 */
-	UE_DEPRECATED(4.25, "This member will be made private. Use IsFromPreviousLevel or SetIsFromPreviousLevel instead.")
 	UPROPERTY(Replicated)
 	uint8 bFromPreviousLevel:1;
 
+public:
+
 	/** if set, GetPlayerName() will call virtual GetPlayerNameCustom() to allow custom access */
 	uint8 bUseCustomPlayerNames : 1;
 
+private:
+
 	/** Elapsed time on server when this PlayerState was first created.  */
-	UE_DEPRECATED(4.25, "This member will be made private. Use GetStartTime or SetStartTime instead.")
 	UPROPERTY(Replicated)
 	int32 StartTime;
 
+	/** The id used by the network to uniquely identify a player.
+	 * NOTE: the internals of this property should *never* be exposed to the player as it's transient
+	 * and opaque in meaning (ie it might mean date/time followed by something else).
+	 * It is OK to use and pass around this property, though. */
+	UPROPERTY(ReplicatedUsing = OnRep_UniqueId)
+	FUniqueNetIdRepl UniqueId;
+
+public:
 	/** This is used for sending game agnostic messages that can be localized */
 	UPROPERTY()
 	TSubclassOf<class ULocalMessage> EngineMessageClass;
@@ -123,14 +127,6 @@ public:
 	UPROPERTY()
 	FString SavedNetworkAddress;
 
-	/** The id used by the network to uniquely identify a player.
-	 * NOTE: the internals of this property should *never* be exposed to the player as it's transient
-	 * and opaque in meaning (ie it might mean date/time followed by something else).
-	 * It is OK to use and pass around this property, though. */
-	UE_DEPRECATED(4.25, "This member will be made private. Use GetUniqueId or SetUniqueId instead.")
-	UPROPERTY(ReplicatedUsing=OnRep_UniqueId)
-	FUniqueNetIdRepl UniqueId; 
-
 	/** The session that the player needs to join/remove from as it is created/leaves */
 	FName SessionName;
 
@@ -265,10 +261,6 @@ public:
 	/** called by seamless travel when initializing a player on the other side - copy properties to the new PlayerState that should persist */
 	ENGINE_API virtual void SeamlessTravelTo(class APlayerState* NewPlayerState);
 
-	/** return true if PlayerState is primary (ie. non-splitscreen) player */
-	UE_DEPRECATED(5.1, "This version of IsPrimaryPlayer has been deprecated, please use the Platform Device Mapper to check the owning PlatformUserId instead.")
-	ENGINE_API virtual bool IsPrimaryPlayer() const;
-
 	ENGINE_API virtual void GetLifetimeReplicatedProps(TArray< FLifetimeProperty > & OutLifetimeProps) const override;
 
 	/** calls OverrideWith and triggers OnOverrideWith for BP extension */
diff --git a/Engine/Source/Runtime/Engine/Private/DataChannel.cpp b/Engine/Source/Runtime/Engine/Private/DataChannel.cpp
--- a/Engine/Source/Runtime/Engine/Private/DataChannel.cpp
+++ b/Engine/Source/Runtime/Engine/Private/DataChannel.cpp
@@ -2357,7 +2357,6 @@ void UActorChannel::CleanupReplicators(const bool bKeepReplicators)
 			//		While the channel is on this list, we also re-open a new channel using the same actor
 			//		KeepProcessingActorChannelBunchesMap will get in here, then when the channel closes a second time, we'll hit this assert
 			//		It should be okay to just set the most recent replicator
-			//check( Connection->DormantReplicatorMap.Find( CompIt.Value()->GetObject() ) == NULL );
 			Connection->StoreDormantReplicator(Actor, ObjectReplicatorRef->GetObject(), ObjectReplicatorRef);
 			ObjectReplicatorRef->StopReplicating(this);		// Stop replicating on this channel
 		}
@@ -5304,7 +5303,7 @@ TSharedRef<FObjectReplicator>& UActorChannel::FindOrCreateReplicator(UObject* Ob
 	TSharedRef<FObjectReplicator>* ReplicatorRefPtr = FindReplicator(Obj);
 
 	// This should only be false if we found the replicator in the ReplicationMap
-	// If we pickup the replicator from the DormantReplicatorMap we treat it as it has been created.
+	// If we pickup the replicator from the DormantReplicatorSet we treat it as it has been created.
 	if (bOutCreated != nullptr)
 	{
 		*bOutCreated = (ReplicatorRefPtr == nullptr);
diff --git a/Engine/Source/Runtime/Engine/Private/NetConnection.cpp b/Engine/Source/Runtime/Engine/Private/NetConnection.cpp
--- a/Engine/Source/Runtime/Engine/Private/NetConnection.cpp
+++ b/Engine/Source/Runtime/Engine/Private/NetConnection.cpp
@@ -289,10 +289,6 @@ namespace UE::Net::Private
 	extern bool bTrackDormantObjectsByLevel;
 }
 
-PRAGMA_DISABLE_DEPRECATION_WARNINGS
-const int32 UNetConnection::DEFAULT_MAX_CHANNEL_SIZE = 32767;
-PRAGMA_ENABLE_DEPRECATION_WARNINGS
-
 /*-----------------------------------------------------------------------------
 	UNetConnection implementation.
 -----------------------------------------------------------------------------*/
@@ -317,7 +313,9 @@ UNetConnection::UNetConnection(const FObjectInitializer& ObjectInitializer)
 ,	Handler()
 ,	StatelessConnectComponent()
 ,	PacketOverhead		( 0 )
+PRAGMA_DISABLE_DEPRECATION_WARNINGS
 ,	ResponseId			( 0 ) // variable is unused and will be deleted
+PRAGMA_ENABLE_DEPRECATION_WARNINGS
 ,	QueuedBits			( 0 )
 ,	TickCount			( 0 )
 ,	LastProcessedFrame	( 0 )
@@ -385,11 +383,6 @@ UNetConnection::UNetConnection(const FObjectInitializer& ObjectInitializer)
 ,	bFlushingPacketOrderCache(false)
 ,	ConnectionId(0)
 {
-	PRAGMA_DISABLE_DEPRECATION_WARNINGS
-	EngineNetworkProtocolVersion = FNetworkVersion::GetEngineNetworkProtocolVersion();
-	GameNetworkProtocolVersion = FNetworkVersion::GetGameNetworkProtocolVersion();
-	PRAGMA_ENABLE_DEPRECATION_WARNINGS
-
 	NetworkCustomVersions = FNetworkVersion::GetNetworkCustomVersions();
 }
 
@@ -430,10 +423,6 @@ void UNetConnection::InitChannelData()
 
 	UE_LOG(LogNet, Log, TEXT("%s setting maximum channels to: %d"), *GetNameSafe(this), ChannelSize);
 
-	PRAGMA_DISABLE_DEPRECATION_WARNINGS
-	MaxChannelSize = ChannelSize;
-	PRAGMA_ENABLE_DEPRECATION_WARNINGS
-
 	Channels.AddDefaulted(ChannelSize);
 	OutReliable.AddDefaulted(ChannelSize);
 	InReliable.AddDefaulted(ChannelSize);
diff --git a/Engine/Source/Runtime/Engine/Private/PlayerState.cpp b/Engine/Source/Runtime/Engine/Private/PlayerState.cpp
--- a/Engine/Source/Runtime/Engine/Private/PlayerState.cpp
+++ b/Engine/Source/Runtime/Engine/Private/PlayerState.cpp
@@ -38,13 +38,6 @@ APlayerState::APlayerState(const FObjectInitializer& ObjectInitializer)
 	bShouldUpdateReplicatedPing = true; // Preserved behavior before bShouldUpdateReplicatedPing was added
 	bUseCustomPlayerNames = false;
 }
-PRAGMA_ENABLE_DEPRECATION_WARNINGS
-
-PRAGMA_DISABLE_DEPRECATION_WARNINGS
-APlayerState::~APlayerState()
-{
-}
-PRAGMA_ENABLE_DEPRECATION_WARNINGS
 
 void APlayerState::UpdatePing(float InPing)
 {
@@ -400,12 +393,6 @@ void APlayerState::SeamlessTravelTo(APlayerState* NewPlayerState)
 	NewPlayerState->SetIsOnlyASpectator(IsOnlyASpectator());
 }
 
-
-bool APlayerState::IsPrimaryPlayer() const
-{
-	return true;
-}
-
 void APlayerState::GetLifetimeReplicatedProps(TArray< FLifetimeProperty > & OutLifetimeProps) const
 {
 	Super::GetLifetimeReplicatedProps(OutLifetimeProps);
@@ -431,7 +418,6 @@ void APlayerState::GetLifetimeReplicatedProps(TArray< FLifetimeProperty > & OutL
 }
 
 
-PRAGMA_DISABLE_DEPRECATION_WARNINGS
 void APlayerState::SetScore(const float NewScore)
 {
 	MARK_PROPERTY_DIRTY_FROM_NAME(APlayerState, Score, this);
@@ -548,4 +534,3 @@ void APlayerState::OnPawnPrivateDestroyed(AActor* InActor)
 	}
 }
 
-PRAGMA_ENABLE_DEPRECATION_WARNINGS
diff --git a/Engine/Source/Runtime/Engine/Private/ReplayHelper.cpp b/Engine/Source/Runtime/Engine/Private/ReplayHelper.cpp
--- a/Engine/Source/Runtime/Engine/Private/ReplayHelper.cpp
+++ b/Engine/Source/Runtime/Engine/Private/ReplayHelper.cpp
@@ -113,10 +113,6 @@ void FReplayHelper::SetPlaybackNetworkVersions(UNetConnection* Connection)
 {
 	if (Connection)
 	{
-		PRAGMA_DISABLE_DEPRECATION_WARNINGS
-		Connection->EngineNetworkProtocolVersion = PlaybackDemoHeader.GetCustomVersion(FEngineNetworkCustomVersion::Guid);
-		Connection->GameNetworkProtocolVersion = PlaybackDemoHeader.GetCustomVersion(FGameNetworkCustomVersion::Guid);
-		PRAGMA_ENABLE_DEPRECATION_WARNINGS
 		Connection->SetNetworkCustomVersions(PlaybackDemoHeader.CustomVersions);
 		Connection->SetPackageVersionUE(PlaybackDemoHeader.PackageVersionUE);
 		Connection->SetPackageVersionLicenseeUE(PlaybackDemoHeader.PackageVersionLicenseeUE);
