Remove the old deprecated InputKey and InputAxis functions

diff --git a/Engine/Source/Runtime/Engine/Classes/GameFramework/PlayerInput.h b/Engine/Source/Runtime/Engine/Classes/GameFramework/PlayerInput.h
--- a/Engine/Source/Runtime/Engine/Classes/GameFramework/PlayerInput.h
+++ b/Engine/Source/Runtime/Engine/Classes/GameFramework/PlayerInput.h
@@ -579,16 +579,8 @@ public:
 	/** Flushes the current key state of the keys associated with the action name passed in */
 	ENGINE_API void FlushPressedActionBindingKeys(FName ActionName);
 
-	/** Handles a key input event.  Returns true if there is an action that handles the specified key. */
-	UE_DEPRECATED(5.0, "This version of InputKey has been deprecated, please use that which takes in FInputKeyParams")
-	ENGINE_API virtual bool InputKey(FKey Key, enum EInputEvent Event, float AmountDepressed, bool bGamepad);
-
 	/** Handles a key input event.  Returns true if there is an action that handles the specified key. */
 	ENGINE_API virtual bool InputKey(const FInputKeyParams& Params);
-	
-	/** Handles an axis input event.  Returns true if a legacy key bind handled the input, otherwise false. */
-	UE_DEPRECATED(5.0, "InputAxis has been deprecated, please use the version of InputKey which takes in FInputKeyParams")
-	ENGINE_API bool InputAxis(FKey Key, float Delta, float DeltaTime, int32 NumSamples, bool bGamepad);
 
 	/** Handles a touch input event.  Returns true. */
 	ENGINE_API bool InputTouch(uint32 Handle, ETouchType::Type Type, const FVector2D& TouchLocation, float Force, FDateTime DeviceTimestamp, uint32 TouchpadIndex);
diff --git a/Engine/Source/Runtime/Engine/Private/UserInterface/PlayerInput.cpp b/Engine/Source/Runtime/Engine/Private/UserInterface/PlayerInput.cpp
--- a/Engine/Source/Runtime/Engine/Private/UserInterface/PlayerInput.cpp
+++ b/Engine/Source/Runtime/Engine/Private/UserInterface/PlayerInput.cpp
@@ -246,17 +246,6 @@ void UPlayerInput::FlushPressedActionBindingKeys(FName ActionName)
 	}
 }
 
-bool UPlayerInput::InputKey(FKey Key, EInputEvent Event, float AmountDepressed, bool bGamepad)
-{
-	FInputKeyParams Params;
-	Params.Key = Key;
-	Params.Event = Event;
-	Params.Delta = FVector((double)AmountDepressed, 0.0, 0.0);
-	Params.bIsGamepadOverride = bGamepad;
-	
-	return InputKey(Params);
-}
-
 bool UPlayerInput::InputKey(const FInputKeyParams& Params)
 {
 	UE_LOG(LogPlayerInput, VeryVerbose, TEXT("[%hs] %s (outer: %s) received input : Key: %s Value: %s Event:  %s"),
@@ -454,17 +443,6 @@ bool UPlayerInput::InputKey(const FInputKeyParams& Params)
 	}
 }
 
-bool UPlayerInput::InputAxis(FKey Key, float Delta, float DeltaTime, int32 NumSamples, bool bGamepad)
-{
-	FInputKeyParams Params;
-	Params.Key = Key;
-	Params.Delta = FVector((double)Delta, 0.0, 0.0);
-	Params.NumSamples = NumSamples;
-	Params.bIsGamepadOverride = bGamepad;
-	
-	return InputKey(Params);
-}
-
 bool UPlayerInput::InputTouch(uint32 Handle, ETouchType::Type Type, const FVector2D& TouchLocation, float Force, FDateTime DeviceTimestamp, uint32 TouchpadIndex)
 {
 	// if touch disabled from the command line, consume the input and do nothing
