New AudioPropertySheet parsing architecture
- Use of class AudioPropertyParser uobject class that handles data parsing and validation]
- AudioPropertiesBindings are no longer necessary
- every parser generates a FAudioPropertiesParsingData, which is then used to inject/validate properties

#rnx

diff --git a/Engine/Source/Runtime/AudioExtensions/Public/AudioPropertiesBindings.h b/Engine/Source/Runtime/AudioExtensions/Public/AudioPropertiesBindings.h
deleted file mode 100644
diff --git a/Engine/Source/Runtime/AudioExtensions/Public/AudioPropertiesSheetAssetBase.h b/Engine/Source/Runtime/AudioExtensions/Public/AudioPropertiesSheetAssetBase.h
--- a/Engine/Source/Runtime/AudioExtensions/Public/AudioPropertiesSheetAssetBase.h
+++ b/Engine/Source/Runtime/AudioExtensions/Public/AudioPropertiesSheetAssetBase.h
@@ -15,5 +15,8 @@ class AUDIOEXTENSIONS_API UAudioPropertiesSheetAssetBase : public UObject
 	GENERATED_BODY()
 
 public: 
-	virtual bool CopyToObjectProperties(TObjectPtr<UObject> TargetObject, const TObjectPtr<const UAudioPropertiesBindings> PropertiesBindings) const { return false; };
+#if WITH_EDITOR
+	UFUNCTION(BlueprintCallable, BlueprintPure=false, Category="AudioProperties")
+	virtual bool CopyToObjectProperties(UObject* TargetObject) const { return false; };
+#endif
 };
\ No newline at end of file
diff --git a/Engine/Source/Runtime/Engine/Classes/Sound/SoundBase.h b/Engine/Source/Runtime/Engine/Classes/Sound/SoundBase.h
--- a/Engine/Source/Runtime/Engine/Classes/Sound/SoundBase.h
+++ b/Engine/Source/Runtime/Engine/Classes/Sound/SoundBase.h
@@ -210,10 +210,6 @@ public:
 	UPROPERTY(EditAnywhere, Category = AudioProperties)
 	TObjectPtr<UAudioPropertiesSheetAssetBase> AudioPropertiesSheet;
 
-	UPROPERTY(EditAnywhere, Category = AudioProperties)
-	TObjectPtr<UAudioPropertiesBindings> AudioPropertiesBindings;
-
-
 private:
 	UPROPERTY()
 	FSoundTimecodeOffset TimecodeOffset;
diff --git a/Engine/Source/Runtime/Engine/Private/SoundBase.cpp b/Engine/Source/Runtime/Engine/Private/SoundBase.cpp
--- a/Engine/Source/Runtime/Engine/Private/SoundBase.cpp
+++ b/Engine/Source/Runtime/Engine/Private/SoundBase.cpp
@@ -195,12 +195,11 @@ void USoundBase::PostEditChangeProperty(struct FPropertyChangedEvent& PropertyCh
 	Super::PostEditChangeProperty(PropertyChangedEvent);
 
 	const FName AudioPropertiesSheetFName = GET_MEMBER_NAME_CHECKED(USoundBase, AudioPropertiesSheet);
-	const FName AudioPropertiesBindingsFName = GET_MEMBER_NAME_CHECKED(USoundBase, AudioPropertiesBindings);
 
 	if (FProperty* PropertyThatChanged = PropertyChangedEvent.Property)
 	{
 		const FName& Name = PropertyThatChanged->GetFName();
-		if (Name == AudioPropertiesSheetFName || Name == AudioPropertiesBindingsFName)
+		if (Name == AudioPropertiesSheetFName)
 		{
 			InjectPropertySheet();
 		}
@@ -352,9 +351,9 @@ TOptional<FSoundTimecodeOffset> USoundBase::GetTimecodeOffset() const
 
 void USoundBase::InjectPropertySheet()
 {
-	if (AudioPropertiesSheet && AudioPropertiesBindings)
+	if (AudioPropertiesSheet)
 	{
-		AudioPropertiesSheet->CopyToObjectProperties(this, AudioPropertiesBindings);
+		AudioPropertiesSheet->CopyToObjectProperties(this);
 	}
 }
 
