Remove useless deprecation warning pragmas from PlayerState

diff --git a/Engine/Source/Runtime/Engine/Classes/GameFramework/PlayerState.h b/Engine/Source/Runtime/Engine/Classes/GameFramework/PlayerState.h
--- a/Engine/Source/Runtime/Engine/Classes/GameFramework/PlayerState.h
+++ b/Engine/Source/Runtime/Engine/Classes/GameFramework/PlayerState.h
@@ -301,7 +301,6 @@ private:
 	// Hidden functions that don't make sense to use on this class.
 	HIDE_ACTOR_TRANSFORM_FUNCTIONS();
 
-PRAGMA_DISABLE_DEPRECATION_WARNINGS
 	//~ Begin Methods for Replicated Members.
 public:
 
@@ -431,7 +430,6 @@ public:
 	ENGINE_API virtual void OnSetUniqueId();
 
 	//~ End Methods for Replicated Members.
-PRAGMA_ENABLE_DEPRECATION_WARNINGS
 };
 
 struct FSetPlayerStatePawn
diff --git a/Engine/Source/Runtime/Engine/Private/PlayerState.cpp b/Engine/Source/Runtime/Engine/Private/PlayerState.cpp
--- a/Engine/Source/Runtime/Engine/Private/PlayerState.cpp
+++ b/Engine/Source/Runtime/Engine/Private/PlayerState.cpp
@@ -15,7 +15,6 @@
 
 #include UE_INLINE_GENERATED_CPP_BY_NAME(PlayerState)
 
-PRAGMA_DISABLE_DEPRECATION_WARNINGS
 APlayerState::APlayerState(const FObjectInitializer& ObjectInitializer)
 	: Super(ObjectInitializer
 		.DoNotCreateDefaultSubobject(TEXT("Sprite")) )
