Animation - fix for UBlendSpace::ResetBlendSamples crash when adding a blend space into a pose search database

diff --git a/Engine/Source/Runtime/Engine/Private/Animation/BlendSpace.cpp b/Engine/Source/Runtime/Engine/Private/Animation/BlendSpace.cpp
--- a/Engine/Source/Runtime/Engine/Private/Animation/BlendSpace.cpp
+++ b/Engine/Source/Runtime/Engine/Private/Animation/BlendSpace.cpp
@@ -304,11 +304,11 @@ void UBlendSpace::ResetBlendSamples(TArray<FBlendSampleData>& InOutSampleDataCac
 	// Ensure we have a valid normalized time.
 	InNormalizedCurrentTime = bLooping ? FMath::Wrap(InNormalizedCurrentTime, 0.0f, 1.0f) : FMath::Clamp(InNormalizedCurrentTime, 0.0f, 1.0f);
 	
-	if (bCanDoMarkerSync)
-	{
-		// Query highest weighted sample with marker information. This will become the leader for all other samples to follow.
-		const int32 HighestMarkerSyncWeightIndex = FBlendSpaceUtilities::GetHighestWeightMarkerSyncSample(InOutSampleDataCache, SampleData);
+	// Query highest weighted sample with marker information. This will become the leader for all other samples to follow.
+	const int32 HighestMarkerSyncWeightIndex = bCanDoMarkerSync ? FBlendSpaceUtilities::GetHighestWeightMarkerSyncSample(InOutSampleDataCache, SampleData) : INDEX_NONE;
 		
+	if (HighestMarkerSyncWeightIndex != INDEX_NONE)
+	{
 		// Query leader sample information.
 		FBlendSampleData& LeaderSampleData = InOutSampleDataCache[HighestMarkerSyncWeightIndex];
 		const FBlendSample& LeaderSample = SampleData[LeaderSampleData.SampleDataIndex];
@@ -578,10 +578,10 @@ void UBlendSpace::TickAssetPlayer(FAnimTickRecord& Instance, struct FAnimNotifyQ
 				Context.SetPreviousAnimationPositionRatio(NormalizedCurrentTime);
 
 				// Get highest weight sample with sync markers. This will become the leader for all other samples to follow.
-				const int32 HighestMarkerSyncWeightIndex = bCanDoMarkerSync ? FBlendSpaceUtilities::GetHighestWeightMarkerSyncSample(SampleDataList, SampleData) : -1;
+				const int32 HighestMarkerSyncWeightIndex = bCanDoMarkerSync ? FBlendSpaceUtilities::GetHighestWeightMarkerSyncSample(SampleDataList, SampleData) : INDEX_NONE;
 
 				// Skip syncing, fallback to normal ticking.
-				if (HighestMarkerSyncWeightIndex == -1)
+				if (HighestMarkerSyncWeightIndex == INDEX_NONE)
 				{
 					bCanDoMarkerSync = false;
 				}
@@ -682,7 +682,7 @@ void UBlendSpace::TickAssetPlayer(FAnimTickRecord& Instance, struct FAnimNotifyQ
 							}
 
 							// Tick all samples as followers
-							TickFollowerSamples(SampleDataList, -1, Context, false, Instance.bLooping, Instance.MirrorDataTable);
+							TickFollowerSamples(SampleDataList, INDEX_NONE, Context, false, Instance.bLooping, Instance.MirrorDataTable);
 						}
 						
 						*Instance.MarkerTickRecord = SampleDataItem.MarkerTickRecord;
@@ -724,7 +724,7 @@ void UBlendSpace::TickAssetPlayer(FAnimTickRecord& Instance, struct FAnimNotifyQ
 
 				// Get the index of the highest weight, assuming that the first is the highest until we find otherwise
 				const bool bTriggerNotifyHighestWeightedAnim = NotifyTriggerMode == ENotifyTriggerMode::HighestWeightedAnimation && SampleDataList.Num() > 0;
-				const int32 HighestWeightIndex = (bGenerateNotifies && bTriggerNotifyHighestWeightedAnim) ? FBlendSpaceUtilities::GetHighestWeightSample(SampleDataList) : -1;
+				const int32 HighestWeightIndex = (bGenerateNotifies && bTriggerNotifyHighestWeightedAnim) ? FBlendSpaceUtilities::GetHighestWeightSample(SampleDataList) : INDEX_NONE;
 
 				for (int32 I = 0; I < SampleDataList.Num(); ++I)
 				{
@@ -1543,7 +1543,7 @@ int32 UBlendSpace::AddSample(const FVector& SampleValue)
 		UpdatePreviewBasePose();
 	}
 
-	return bValidSampleData ? SampleData.Num() - 1 : -1;
+	return bValidSampleData ? SampleData.Num() - 1 : INDEX_NONE;
 }
 
 int32 UBlendSpace::AddSample(UAnimSequence* AnimationSequence, const FVector& SampleValue)
@@ -1559,7 +1559,7 @@ int32 UBlendSpace::AddSample(UAnimSequence* AnimationSequence, const FVector& Sa
 		UpdatePreviewBasePose();
 	}
 
-	return bValidSampleData ? SampleData.Num() - 1 : -1;
+	return bValidSampleData ? SampleData.Num() - 1 : INDEX_NONE;
 }
 
 bool UBlendSpace::EditSampleValue(const int32 BlendSampleIndex, const FVector& NewValue)
@@ -2842,7 +2842,7 @@ void FBlendSpaceData::GetSamples2D(
 		const FBlendSpaceTriangle* Triangle = &Triangles[InOutTriangleIndex];
 		// Look for the edge which has the target point most outside it
 		float LargestDistance = UE_KINDA_SMALL_NUMBER;
-		int32 LargestEdgeIndex = -1;
+		int32 LargestEdgeIndex = INDEX_NONE;
 		for (int32 VertexIndex = 0; VertexIndex != FBlendSpaceTriangle::NUM_VERTICES; ++VertexIndex)
 		{
 			FVector2D Corner = Triangle->Vertices[VertexIndex];
diff --git a/Engine/Source/Runtime/Engine/Private/Animation/BlendSpaceUtilities.cpp b/Engine/Source/Runtime/Engine/Private/Animation/BlendSpaceUtilities.cpp
--- a/Engine/Source/Runtime/Engine/Private/Animation/BlendSpaceUtilities.cpp
+++ b/Engine/Source/Runtime/Engine/Private/Animation/BlendSpaceUtilities.cpp
@@ -68,7 +68,7 @@ int32 FBlendSpaceUtilities::GetHighestWeightSample(const TArray<struct FBlendSam
 
 int32 FBlendSpaceUtilities::GetHighestWeightMarkerSyncSample(const TArray<struct FBlendSampleData> &SampleDataList, const TArray<struct FBlendSample>& BlendSamples)
 {
-	int32 HighestWeightIndex = -1;
+	int32 HighestWeightIndex = INDEX_NONE;
 	float HighestWeight = FLT_MIN;
 
 	for (int32 I = 0; I < SampleDataList.Num(); I++)
