Disable per instance mesh vertex color painting on nanite mesh.
Previously the tool allowed the painting operation to happen even though we can't see vertex color on nanite mesh.
Disabling nanite (completely or just on the component) will allow the tool to paint again.
Note that this only affects the mesh paint tool behavior, we never remove existing painted data.

diff --git a/Engine/Source/Runtime/Engine/Classes/Components/StaticMeshComponent.h b/Engine/Source/Runtime/Engine/Classes/Components/StaticMeshComponent.h
--- a/Engine/Source/Runtime/Engine/Classes/Components/StaticMeshComponent.h
+++ b/Engine/Source/Runtime/Engine/Classes/Components/StaticMeshComponent.h
@@ -754,7 +754,7 @@ public:
 	ENGINE_API UMaterialInterface* GetNaniteAuditMaterial(int32 MaterialIndex) const;
 
 	/* Returns true if mesh vertex color painting is supported on this component. */
-	ENGINE_API bool CanMeshPaintVertexColors() const { return bSupportMeshPainting && bEnableVertexColorMeshPainting; }
+	ENGINE_API bool CanMeshPaintVertexColors() const;
 	/* Returns true if mesh texture color painting is supported on this component. */
 	ENGINE_API bool CanMeshPaintTextureColors() const;
 
diff --git a/Engine/Source/Runtime/Engine/Private/Components/StaticMeshComponent.cpp b/Engine/Source/Runtime/Engine/Private/Components/StaticMeshComponent.cpp
--- a/Engine/Source/Runtime/Engine/Private/Components/StaticMeshComponent.cpp
+++ b/Engine/Source/Runtime/Engine/Private/Components/StaticMeshComponent.cpp
@@ -2999,7 +2999,7 @@ void UStaticMeshComponent::ApplyComponentInstanceData(FStaticMeshComponentInstan
 		}
 	}
 
-	if (CanMeshPaintVertexColors())
+	if (bSupportMeshPainting && bEnableVertexColorMeshPainting)
 	{
 		FComponentReregisterContext ReregisterStaticMesh(this);
 		StaticMeshInstanceData->ApplyVertexColorData(this);
@@ -3204,6 +3204,20 @@ int32 UStaticMeshComponent::GetMeshPaintTextureCoordinateIndex() const
 	return 0;
 }
 
+bool UStaticMeshComponent::CanMeshPaintVertexColors() const
+{
+	if (!bSupportMeshPainting || !bEnableVertexColorMeshPainting)
+	{
+		return false;
+	}
+	if (ShouldCreateNaniteProxy())
+	{
+		// Don't vertex paint on nanite.
+		return false;
+	}
+	return true;
+}
+
 bool UStaticMeshComponent::CanMeshPaintTextureColors() const 
 {
 	if (!bSupportMeshPainting || !bEnableTextureColorMeshPainting)
