Substrate - fixed UI to respect legacy blend mode behavior.

diff --git a/Engine/Shaders/Private/Substrate/SubstrateLegacyConversion.ush b/Engine/Shaders/Private/Substrate/SubstrateLegacyConversion.ush
--- a/Engine/Shaders/Private/Substrate/SubstrateLegacyConversion.ush
+++ b/Engine/Shaders/Private/Substrate/SubstrateLegacyConversion.ush
@@ -586,3 +586,43 @@ FSubstrateData SubstrateConvertLegacyMaterialStatic(
 
 	return Out;
 }
+
+FSubstrateData SubstrateCreateUIMaterial(
+	float3 Emissive,
+	float Opacity)
+{
+	// We create a UI material using a Unlit Substrate BSDF. And we adapt to the blending mode in order to respect legacy behavior.
+	FSubstrateData Out = GetInitialisedSubstrateData();
+
+	const float3 RawNormal = float3(0, 0, 1);	// Unused in this case.
+	const float GreyTransmittance = saturate(1.0f - Opacity);
+
+#if MATERIALBLENDING_SOLID || MATERIALBLENDING_MASKED
+	// Opaque materials only write emissive.
+	Out = GetSubstrateUnlitBSDF(Emissive, 0.0f, RawNormal);
+#elif MATERIALBLENDING_ALPHACOMPOSITE
+	// Uses (Luminance, black transmittance) for premultiplied alpha blending. 
+	// Coverage over the background is specified through the root node input "Coverage Override".
+	Out = GetSubstrateUnlitBSDF(Emissive, 0.0f, RawNormal);
+#elif MATERIALBLENDING_ALPHAHOLDOUT
+	// The MATERIALBLENDING_ALPHAHOLDOUT case needs to be before the MATERIALBLENDING_TRANSLUCENT case, as alphaholdout material have both cases defined.
+	Out = GetSubstrateUnlitBSDF(0.0f, GreyTransmittance, RawNormal);
+#elif MATERIALBLENDING_TRANSLUCENT
+	// Uses (Luminance, black transmittance) for translucency blending. Opacity is controled using coverage.
+	Out = GetSubstrateUnlitBSDF(Emissive, 0.0f, RawNormal);
+#if SUBSTRATE_INLINE_SHADING
+	Out.InlinedBSDF.Coverage = Opacity; // Applied to emissive via the Substrate Tree, (see SubstrateUpdateTreeUnlit).
+#endif
+#elif MATERIALBLENDING_ADDITIVE
+	// Emissive is weighted by opacity in the legacy material.
+	Out = GetSubstrateUnlitBSDF(Emissive * Opacity, 1.0f, RawNormal);
+#elif MATERIALBLENDING_MODULATE
+	// Setting up emissive as the transmittance color. It is not clamped in case brightening is required.
+	Out = GetSubstrateUnlitBSDF(0.0f, Emissive, RawNormal);
+#else
+	// Required default for some materials such as Editor UI.
+	Out = GetSubstrateUnlitBSDF(Emissive, 0.0f, RawNormal);
+#endif
+
+	return Out;
+}
diff --git a/Engine/Shaders/Shared/SubstrateDefinitions.h b/Engine/Shaders/Shared/SubstrateDefinitions.h
--- a/Engine/Shaders/Shared/SubstrateDefinitions.h
+++ b/Engine/Shaders/Shared/SubstrateDefinitions.h
@@ -9,7 +9,7 @@
 
 #ifndef __cplusplus
 // Change this to force recompilation of all Substrate dependent shaders (for instance https://guidgenerator.com/online-guid-generator.aspx)
-#pragma message("UESHADERMETADATA_VERSION B013B52C-494E-4FC3-9EB5-B16B754E5CA4")
+#pragma message("UESHADERMETADATA_VERSION 0A6FB545-2B60-4E5D-B93C-A4CBE10274E5")
 #endif
 
 // Closure offsets are packed into 32bits, each entry using SUBSTRATE_CLOSURE_OFFSET_BIT_COUNT bits
diff --git a/Engine/Source/Runtime/Engine/Private/Materials/HLSLMaterialTranslator.cpp b/Engine/Source/Runtime/Engine/Private/Materials/HLSLMaterialTranslator.cpp
--- a/Engine/Source/Runtime/Engine/Private/Materials/HLSLMaterialTranslator.cpp
+++ b/Engine/Source/Runtime/Engine/Private/Materials/HLSLMaterialTranslator.cpp
@@ -13727,7 +13727,7 @@ int32 FHLSLMaterialTranslator::SubstrateUnlitBSDF(int32 EmissiveColor, int32 Tra
 			MCT_Substrate, TEXT("Parameters.%s.PromoteParameterBlendedBSDFToOperator(GetSubstrateUnlitBSDF(%s, %s, %s), %u, %u, %u, %u)"),
 			*GetParametersSubstrateTreeName(CurrentSubstrateCompilationContext),
 			*SubstrateGetCastParameterCode(EmissiveColor,		MCT_Float3),
-			*SubstrateGetCastParameterCode(TransmittanceColor, MCT_Float3),
+			*SubstrateGetCastParameterCode(TransmittanceColor,	MCT_Float3),
 			*SubstrateGetCastParameterCode(Normal,				MCT_Float3),
 			PromoteToOperator->Index,
 			PromoteToOperator->BSDFIndex,
@@ -13738,21 +13738,43 @@ int32 FHLSLMaterialTranslator::SubstrateUnlitBSDF(int32 EmissiveColor, int32 Tra
 	return AddCodeChunk(
 		MCT_Substrate, TEXT("GetSubstrateUnlitBSDF(%s, %s, %s)"),
 		*SubstrateGetCastParameterCode(EmissiveColor,		MCT_Float3),
-		*SubstrateGetCastParameterCode(TransmittanceColor, MCT_Float3),
+		*SubstrateGetCastParameterCode(TransmittanceColor,	MCT_Float3),
 		*SubstrateGetCastParameterCode(Normal,				MCT_Float3)
 	);
 }
 
+int32 FHLSLMaterialTranslator::SubstrateUIBSDF(int32 EmissiveColor, int32 Opacity, FSubstrateOperator* PromoteToOperator)
+{
+	if (PromoteToOperator)
+	{
+		return AddCodeChunk(
+			MCT_Substrate, TEXT("Parameters.%s.PromoteParameterBlendedBSDFToOperator(SubstrateCreateUIMaterial(%s, %s), %u, %u, %u, %u)"),
+			*GetParametersSubstrateTreeName(CurrentSubstrateCompilationContext),
+			*SubstrateGetCastParameterCode(EmissiveColor, MCT_Float3),
+			*SubstrateGetCastParameterCode(Opacity, MCT_Float1),
+			PromoteToOperator->Index,
+			PromoteToOperator->BSDFIndex,
+			PromoteToOperator->LayerDepth,
+			PromoteToOperator->bIsBottom ? 1 : 0);
+	}
+
+	return AddCodeChunk(
+		MCT_Substrate, TEXT("SubstrateCreateUIMaterial(%s, %s)"),
+		*SubstrateGetCastParameterCode(EmissiveColor,		MCT_Float3),
+		*SubstrateGetCastParameterCode(Opacity,				MCT_Float1)
+	);
+}
+
 int32 FHLSLMaterialTranslator::SubstrateHairBSDF(int32 BaseColor, int32 Scatter, int32 Specular, int32 Roughness, int32 Backlit, int32 EmissiveColor, int32 Tangent, const FString& SharedLocalBasisIndexMacro, FSubstrateOperator* PromoteToOperator)
 {
 	return AddCodeChunk(
 		MCT_Substrate, TEXT("Parameters.%s.PromoteParameterBlendedBSDFToOperator(GetSubstrateHairBSDF(%s, %s, %s, %s, %s, %s, %s), %u, %u, %u, %u) /* Tangent:%s */"),
 		*GetParametersSubstrateTreeName(CurrentSubstrateCompilationContext),
 		*SubstrateGetCastParameterCode(BaseColor,			MCT_Float3),
-		*SubstrateGetCastParameterCode(Scatter,			MCT_Float),
+		*SubstrateGetCastParameterCode(Scatter,				MCT_Float),
 		*SubstrateGetCastParameterCode(Specular,			MCT_Float),
 		*SubstrateGetCastParameterCode(Roughness,			MCT_Float),
-		*SubstrateGetCastParameterCode(Backlit,			MCT_Float),
+		*SubstrateGetCastParameterCode(Backlit,				MCT_Float),
 		*SubstrateGetCastParameterCode(EmissiveColor,		MCT_Float3),
 		*SharedLocalBasisIndexMacro,
 		PromoteToOperator->Index,
diff --git a/Engine/Source/Runtime/Engine/Private/Materials/HLSLMaterialTranslator.h b/Engine/Source/Runtime/Engine/Private/Materials/HLSLMaterialTranslator.h
--- a/Engine/Source/Runtime/Engine/Private/Materials/HLSLMaterialTranslator.h
+++ b/Engine/Source/Runtime/Engine/Private/Materials/HLSLMaterialTranslator.h
@@ -1309,6 +1309,7 @@ protected:
 		FSubstrateOperator* PromoteToOperator) override;
 	virtual int32 SubstrateVolumetricFogCloudBSDF(int32 Albedo, int32 Extinction, int32 EmissiveColor, int32 AmbientOcclusion, bool bEmissiveOnly) override;
 	virtual int32 SubstrateUnlitBSDF(int32 EmissiveColor, int32 TransmittanceColor, int32 Normal, FSubstrateOperator* PromoteToOperator) override;
+	virtual int32 SubstrateUIBSDF(int32 EmissiveColor, int32 Opacity, FSubstrateOperator* PromoteToOperator) override;
 	virtual int32 SubstrateHairBSDF(int32 BaseColor, int32 Scatter, int32 Specular, int32 Roughness, int32 Backlit, int32 EmissiveColor, int32 Tangent, const FString& SharedLocalBasisIndexMacro, FSubstrateOperator* PromoteToOperator) override;
 	virtual int32 SubstrateEyeBSDF(int32 DiffuseColor, int32 Roughness, int32 IrisMask, int32 IrisDistance, int32 EmissiveColor, int32 CorneaNormal, int32 IrisNormal, int32 IrisPlaneNormal, int32 SSSProfileId, const FString& SharedLocalBasisIndexMacro, FSubstrateOperator* PromoteToOperator) override;
 	virtual int32 SubstrateSingleLayerWaterBSDF(
diff --git a/Engine/Source/Runtime/Engine/Private/Materials/MaterialExpressions.cpp b/Engine/Source/Runtime/Engine/Private/Materials/MaterialExpressions.cpp
--- a/Engine/Source/Runtime/Engine/Private/Materials/MaterialExpressions.cpp
+++ b/Engine/Source/Runtime/Engine/Private/Materials/MaterialExpressions.cpp
@@ -28108,13 +28108,11 @@ UMaterialExpressionSubstrateUI::UMaterialExpressionSubstrateUI(const FObjectInit
 int32 UMaterialExpressionSubstrateUI::Compile(class FMaterialCompiler* Compiler, int32 OutputIndex)
 {
 	int OpacityCodeChunk = CompileWithDefaultFloat1(Compiler, Opacity, 0.0f);
-	int TransmittanceCodeChunk = Compiler->Saturate(Compiler->Sub(Compiler->Constant(1.0f), OpacityCodeChunk));
 
 	FSubstrateOperator& SubstrateOperator = Compiler->SubstrateCompilationGetOperator(Compiler->SubstrateTreeStackGetPathUniqueId());
-	int32 OutputCodeChunk = Compiler->SubstrateUnlitBSDF(
+	int32 OutputCodeChunk = Compiler->SubstrateUIBSDF(
 		CompileWithDefaultFloat3(Compiler, Color, 0.0f, 0.0f, 0.0f),
-		TransmittanceCodeChunk,
-		Compiler->Constant3(0.0f, 0.0f, 1.0f),	// place holder normal
+		OpacityCodeChunk,
 		&SubstrateOperator);
 	return OutputCodeChunk;
 }
diff --git a/Engine/Source/Runtime/Engine/Public/MaterialCompiler.h b/Engine/Source/Runtime/Engine/Public/MaterialCompiler.h
--- a/Engine/Source/Runtime/Engine/Public/MaterialCompiler.h
+++ b/Engine/Source/Runtime/Engine/Public/MaterialCompiler.h
@@ -683,6 +683,7 @@ public:
 		FSubstrateOperator* PromoteToOperator) = 0;
 	virtual int32 SubstrateVolumetricFogCloudBSDF(int32 Albedo, int32 Extinction, int32 EmissiveColor, int32 AmbientOcclusion, bool bEmissiveOnly) = 0;
 	virtual int32 SubstrateUnlitBSDF(int32 EmissiveColor, int32 TransmittanceColor, int32 Normal, FSubstrateOperator* PromoteToOperator) = 0;
+	virtual int32 SubstrateUIBSDF(int32 EmissiveColor, int32 Opacity, FSubstrateOperator* PromoteToOperator) = 0;
 	virtual int32 SubstrateHairBSDF(int32 BaseColor, int32 Scatter, int32 Specular, int32 Roughness, int32 Backlit, int32 EmissiveColor, int32 Tangent, const FString& SharedLocalBasisIndexMacro, FSubstrateOperator* PromoteToOperator) = 0;
 	virtual int32 SubstrateEyeBSDF(int32 DiffuseAlbedo, int32 Roughness, int32 IrisMask, int32 IrisDistance, int32 IrisNormal, int32 IrisPlaneNormal, int32 SSSProfileId, int32 EmissiveColor, int32 CorneaNormal, const FString& SharedLocalBasisIndexMacro, FSubstrateOperator* PromoteToOperator) = 0;
 	virtual int32 SubstrateSingleLayerWaterBSDF(
@@ -1399,6 +1400,11 @@ public:
 		return Compiler->SubstrateUnlitBSDF(EmissiveColor, TransmittanceColor, Normal, PromoteToOperator);
 	}
 
+	virtual int32 SubstrateUIBSDF(int32 EmissiveColor, int32 Opacity, FSubstrateOperator* PromoteToOperator) override
+	{
+		return Compiler->SubstrateUIBSDF(EmissiveColor, Opacity, PromoteToOperator);
+	}
+
 	virtual int32 SubstrateHairBSDF(int32 BaseColor, int32 Scatter, int32 Specular, int32 Roughness, int32 Backlit, int32 EmissiveColor, int32 Tangent, const FString& SharedLocalBasisIndexMacro, FSubstrateOperator* PromoteToOperator) override
 	{
 		return Compiler->SubstrateHairBSDF(BaseColor, Scatter, Specular, Roughness, Backlit, EmissiveColor, Tangent, SharedLocalBasisIndexMacro, PromoteToOperator);
