Fix typo in property name OverriddenMeshPaintTextureCoordinateIndex
This is as yet unreleased code, so just changing and not doing deprecation process.

diff --git a/Engine/Source/Runtime/Engine/Classes/Components/StaticMeshComponent.h b/Engine/Source/Runtime/Engine/Classes/Components/StaticMeshComponent.h
--- a/Engine/Source/Runtime/Engine/Classes/Components/StaticMeshComponent.h
+++ b/Engine/Source/Runtime/Engine/Classes/Components/StaticMeshComponent.h
@@ -329,7 +329,7 @@ public:
 
 	/** The overriden coordinate index to use when texture color painting on this mesh. */
 	UPROPERTY(EditAnywhere, Category = "Mesh Painting", meta=(UIMin = "0", UIMax = "3", editcondition = "bOverrideMeshPaintTextureCoordinateIndex"))
-	int32 OverridenMeshPaintTextureCoordinateIndex;
+	int32 OverriddenMeshPaintTextureCoordinateIndex;
 
 	/** Light map resolution to use on this component, used if bOverrideLightMapRes is true and there is a valid StaticMesh. */
 	UPROPERTY(EditAnywhere, BlueprintReadOnly, Category=Lighting, meta=(ClampMax = 4096, editcondition="bOverrideLightMapRes") )
diff --git a/Engine/Source/Runtime/Engine/Private/Components/StaticMeshComponent.cpp b/Engine/Source/Runtime/Engine/Private/Components/StaticMeshComponent.cpp
--- a/Engine/Source/Runtime/Engine/Private/Components/StaticMeshComponent.cpp
+++ b/Engine/Source/Runtime/Engine/Private/Components/StaticMeshComponent.cpp
@@ -3195,7 +3195,7 @@ int32 UStaticMeshComponent::GetMeshPaintTextureCoordinateIndex() const
 {
 	if (bOverrideMeshPaintTextureCoordinateIndex)
 	{
-		return OverridenMeshPaintTextureCoordinateIndex;
+		return OverriddenMeshPaintTextureCoordinateIndex;
 	}
 	if (StaticMesh != nullptr)
 	{
