UOL: Optimized Actor Locator Fragments by removing unnecessary string copy operations and fixup

We no longer do any string copying ourselves when resolving actor locator fragments unless the fragment is unresolved; when it does not resolve, and in-editor, the fragment will go through an additional core redirect fixup (through a call to PreSavePath).

Tested the repro steps for UE-56235 to make sure that didn't regress with this change (it doesn't)

diff --git a/Engine/Source/Runtime/Engine/Private/UniversalObjectLocators/ActorLocatorFragment.cpp b/Engine/Source/Runtime/Engine/Private/UniversalObjectLocators/ActorLocatorFragment.cpp
--- a/Engine/Source/Runtime/Engine/Private/UniversalObjectLocators/ActorLocatorFragment.cpp
+++ b/Engine/Source/Runtime/Engine/Private/UniversalObjectLocators/ActorLocatorFragment.cpp
@@ -109,36 +109,51 @@ UE::UniversalObjectLocator::FResolveResult FActorLocatorFragment::Resolve(const
 		Result = ResolveActorWithinLevel(*this, Level);
 	}
 
+	if (Result)
+	{
+		return FResolveResultData(Result);
+	}
+
+#if WITH_EDITORONLY_DATA
+
+	const UPackage* ContextPackage = Params.Context ? Params.Context->GetOutermost() : nullptr;
+	const int32     PIEInstanceID  = ContextPackage ? ContextPackage->GetPIEInstanceID() : INDEX_NONE;
+
 	// Finally fallback to just trying to resolve the path directly
+	auto ResolvePathWithPIEHandling = [PIEInstanceID](const FSoftObjectPath& PathPtr)
+	{
+		// If we are resolving within a PIE instance, fixup the PIE instance on the path first
+		if (PIEInstanceID != INDEX_NONE)
+		{
+			FSoftObjectPath PIEPath = PathPtr;
+			PIEPath.FixupForPIE(PIEInstanceID);
+			return PIEPath.ResolveObject();
+		}
+		return PathPtr.ResolveObject();
+	};
+
+	Result = ResolvePathWithPIEHandling(Path);
+
 	if (!Result)
 	{
+		// If the path failed to resolve, attempt to fixup redirectors on the path to handle cases
+		//     where this path hasn't been saved yet, but references things that have (and have been redirected)
 		FSoftObjectPath TempPath = Path;
-
-		// Soft Object Paths don't follow asset redirectors when attempting to call ResolveObject or TryLoad.
-		// We want to follow the asset redirector so that maps that have been renamed (from Untitled to their first asset name)
-		// properly resolve. This fixes Possessable bindings losing their references the first time you save a map.
 		TempPath.PreSavePath();
 
-#if WITH_EDITORONLY_DATA
-		UPackage* ContextPackage = Params.Context ? Params.Context->GetOutermost() : nullptr;
-		if (ContextPackage)
+		if (TempPath != Path)
 		{
-			// If we are resolving within a PIE instance, fixp the PIE instance on the path first
-			const int32 PIEInstanceID = ContextPackage->GetPIEInstanceID();
-			if (PIEInstanceID != INDEX_NONE)
-			{
-				FSoftObjectPath PIEPath = TempPath;
-				PIEPath.FixupForPIE(PIEInstanceID);
-
-				Result = PIEPath.ResolveObject();
-				return FResolveResultData(Result);
-			}
+			Result = ResolvePathWithPIEHandling(TempPath);
 		}
-#endif
-
-		Result = TempPath.ResolveObject();
 	}
 
+#else  // WITH_EDITORONLY_DATA
+
+	// By default we just resolve the path directly
+	Result = Path.ResolveObject();
+
+#endif // WITH_EDITORONLY_DATA
+
 	return FResolveResultData(Result);
 }
 
