Fix crash with USkeletalMeshComponent::UpdateFollowerComponent if the leader is not a USkeletalMeshComponent

diff --git a/Engine/Source/Runtime/Engine/Private/Components/SkeletalMeshComponent.cpp b/Engine/Source/Runtime/Engine/Private/Components/SkeletalMeshComponent.cpp
--- a/Engine/Source/Runtime/Engine/Private/Components/SkeletalMeshComponent.cpp
+++ b/Engine/Source/Runtime/Engine/Private/Components/SkeletalMeshComponent.cpp
@@ -2135,28 +2135,30 @@ void USkeletalMeshComponent::UpdateFollowerComponent()
 
 	if(ensure(LeaderPoseComponent.IsValid()))
 	{
-		USkeletalMeshComponent* LeaderSMC = Cast<USkeletalMeshComponent>(LeaderPoseComponent.Get());
-		// first set any animation-driven curves from the leader SMC
-		if (LeaderSMC->AnimScriptInstance)
-		{
-			LeaderSMC->AnimScriptInstance->RefreshCurves(this);
-		}
-
-		// we changed order of morphtarget to be overriden by SetMorphTarget from BP
-		// so this has to go first
-		// now propagate BP-driven curves from the leader SMC...
-		if (GetSkeletalMeshAsset())
+		if (USkeletalMeshComponent* LeaderSMC = Cast<USkeletalMeshComponent>(LeaderPoseComponent.Get()))
 		{
-			check(MorphTargetWeights.Num() == GetSkeletalMeshAsset()->GetMorphTargets().Num());
-			if (LeaderSMC->MorphTargetCurves.Num() > 0)
+			// first set any animation-driven curves from the leader SMC
+			if (LeaderSMC->AnimScriptInstance)
 			{
-				FAnimationRuntime::AppendActiveMorphTargets(GetSkeletalMeshAsset(), LeaderSMC->MorphTargetCurves, ActiveMorphTargets, MorphTargetWeights);
+				LeaderSMC->AnimScriptInstance->RefreshCurves(this);
 			}
 
-			// if follower also has it, add it here. 
-			if (MorphTargetCurves.Num() > 0)
+			// we changed order of morphtarget to be overriden by SetMorphTarget from BP
+			// so this has to go first
+			// now propagate BP-driven curves from the leader SMC...
+			if (GetSkeletalMeshAsset())
 			{
-				FAnimationRuntime::AppendActiveMorphTargets(GetSkeletalMeshAsset(), MorphTargetCurves, ActiveMorphTargets, MorphTargetWeights);
+				check(MorphTargetWeights.Num() == GetSkeletalMeshAsset()->GetMorphTargets().Num());
+				if (LeaderSMC->MorphTargetCurves.Num() > 0)
+				{
+					FAnimationRuntime::AppendActiveMorphTargets(GetSkeletalMeshAsset(), LeaderSMC->MorphTargetCurves, ActiveMorphTargets, MorphTargetWeights);
+				}
+
+				// if follower also has it, add it here. 
+				if (MorphTargetCurves.Num() > 0)
+				{
+					FAnimationRuntime::AppendActiveMorphTargets(GetSkeletalMeshAsset(), MorphTargetCurves, ActiveMorphTargets, MorphTargetWeights);
+				}
 			}
 		}
 	}
