Changed an assert to an ensure when serializing a level with an actor that has its outer mismatching the owning level.

diff --git a/Engine/Source/Runtime/Engine/Private/Level.cpp b/Engine/Source/Runtime/Engine/Private/Level.cpp
--- a/Engine/Source/Runtime/Engine/Private/Level.cpp
+++ b/Engine/Source/Runtime/Engine/Private/Level.cpp
@@ -638,7 +638,10 @@ void ULevel::Serialize( FArchive& Ar )
 				return false;
 			}
 
-			check(Actor->GetLevel() == this);
+			if (!ensureMsgf(Actor->GetLevel() == this, TEXT("Actor %s is in wrong level:") LINE_TERMINATOR TEXT("\tOwner Level=%s") LINE_TERMINATOR TEXT("\tActor Level=%s)"), *GetPathNameSafe(Actor), *GetPathNameSafe(this), *GetPathNameSafe(Actor->GetLevel())))
+			{
+				return false;
+			}
 
 			if (Actor->HasAnyFlags(RF_Transient))
 			{
@@ -4322,7 +4325,7 @@ UPackage* ULevel::CreateActorPackage(UPackage* InLevelPackage, EActorPackagingSc
 	const FString PackageName = ULevel::GetActorPackageName(InLevelPackage, InActorPackagingScheme, InActorPath, InTargetContextObject);
 	UPackage* ActorPackage = CreateActorPackageInternal(PackageName, InActorPath);
 	// Should be prevented upstream but we propagate the flag to prevent issues in asset enumeration
-	if (!ensureMsgf(!(InLevelPackage->GetPackageFlags() & PKG_PlayInEditor), TEXT("Actor packages should not be created on PlayInEditor levels")))
+	
 	{
 		ActorPackage->SetPackageFlags(PKG_PlayInEditor);
 	}
