Alpha: Converts "Enable alpha channel support in post processing" (r.PostProcessing.PropagateAlpha) back to a boolean where "True" propagates alpha through post-processing.

This relaxes the "linear-only" constraint with 1, and simplifies the option for users. Values 1 & 2 are automatically converted to "True" upon config/cvar load, and thus both existing values remain valid. While most backwards compatiblity is preserved, there are two instances where external code will need to be updated:
 * Calls to FindTConsoleVariableDataInt(..) will need to be replaced with regular FindConsoleVariable(..), since the former call would silently fail on a bool cvar.
 * Any literal comparison to "2" not using EAlphaChannelMode::Type, which is now deprecated.
We believe this remains preferable over migration to an entirely new cvar. This change also removes experimental naming.

diff --git a/Engine/Plugins/Experimental/Avalanche/Source/AvalancheMediaEditor/Private/Rundown/Preview/SAvaRundownPagePreview.cpp b/Engine/Plugins/Experimental/Avalanche/Source/AvalancheMediaEditor/Private/Rundown/Preview/SAvaRundownPagePreview.cpp
--- a/Engine/Plugins/Experimental/Avalanche/Source/AvalancheMediaEditor/Private/Rundown/Preview/SAvaRundownPagePreview.cpp
+++ b/Engine/Plugins/Experimental/Avalanche/Source/AvalancheMediaEditor/Private/Rundown/Preview/SAvaRundownPagePreview.cpp
@@ -304,11 +304,11 @@ void SAvaRundownPagePreview::HandleCheckerboardActionExecute() const
 	if (bShowCheckerBoard)
 	{
 		const IConsoleVariable* PropagateAlphaCVar = IConsoleManager::Get().FindConsoleVariable(TEXT("r.PostProcessing.PropagateAlpha"));
-		if (PropagateAlphaCVar && PropagateAlphaCVar->GetInt() != static_cast<int32>(EAlphaChannelMode::AllowThroughTonemapper))
+		if (PropagateAlphaCVar && PropagateAlphaCVar->GetBool() != true)
 		{
 			const FText NotificationText = LOCTEXT("AlphaSupport",
 				"An output requested Alpha Support but the required project setting is not enabled!\n"
-				"Go to Project Settings > Rendering > PostProcessing > 'Enable Alpha Channel Support in Post Processing' and set it to 'Allow through tonemapper'.");
+				"Go to Project Settings > Rendering > PostProcessing > 'Enable Alpha Channel Support in Post Processing' and enable it.");
 
 			FNotificationInfo Info(NotificationText);
 			Info.ExpireDuration = 5.0f;	// The message is long, need more time to read it.
diff --git a/Engine/Plugins/Experimental/Avalanche/Source/AvalancheViewport/Private/Visualizers/AvaViewportPostProcessVisualizer.cpp b/Engine/Plugins/Experimental/Avalanche/Source/AvalancheViewport/Private/Visualizers/AvaViewportPostProcessVisualizer.cpp
--- a/Engine/Plugins/Experimental/Avalanche/Source/AvalancheViewport/Private/Visualizers/AvaViewportPostProcessVisualizer.cpp
+++ b/Engine/Plugins/Experimental/Avalanche/Source/AvalancheViewport/Private/Visualizers/AvaViewportPostProcessVisualizer.cpp
@@ -69,7 +69,7 @@ bool FAvaViewportPostProcessVisualizer::CanActivate(bool bInSilent) const
 	URendererSettings* RendererSettings = GetMutableDefault<URendererSettings>();
 	check(RendererSettings);
 
-	if (RendererSettings->bEnableAlphaChannelInPostProcessing == EAlphaChannelMode::AllowThroughTonemapper)
+	if (RendererSettings->bEnableAlphaChannelInPostProcessing)
 	{
 		return true;
 	}
@@ -84,11 +84,11 @@ bool FAvaViewportPostProcessVisualizer::CanActivate(bool bInSilent) const
 	{
 		case EAppReturnType::Yes:
 		{
-			RendererSettings->bEnableAlphaChannelInPostProcessing = EAlphaChannelMode::AllowThroughTonemapper;
+			RendererSettings->bEnableAlphaChannelInPostProcessing = true;
 
 			if (IConsoleVariable* PropagateAlphaCVar = IConsoleManager::Get().FindConsoleVariable(TEXT("r.PostProcessing.PropagateAlpha")))
 			{
-				PropagateAlphaCVar->Set(EAlphaChannelMode::AllowThroughTonemapper);
+				PropagateAlphaCVar->Set(true);
 			}
 
 			FProperty* Property = RendererSettings->GetClass()->FindPropertyByName(GET_MEMBER_NAME_CHECKED(URendererSettings, bEnableAlphaChannelInPostProcessing));
diff --git a/Engine/Plugins/Experimental/Compositing/HoldoutComposite/Source/HoldoutComposite/Private/HoldoutCompositeSubsystem.cpp b/Engine/Plugins/Experimental/Compositing/HoldoutComposite/Source/HoldoutComposite/Private/HoldoutCompositeSubsystem.cpp
--- a/Engine/Plugins/Experimental/Compositing/HoldoutComposite/Source/HoldoutComposite/Private/HoldoutCompositeSubsystem.cpp
+++ b/Engine/Plugins/Experimental/Compositing/HoldoutComposite/Source/HoldoutComposite/Private/HoldoutCompositeSubsystem.cpp
@@ -38,7 +38,7 @@ void UHoldoutCompositeSubsystem::RegisterPrimitives(TArrayView<TSoftObjectPtr<UP
 {
 	// The compositing relies on alpha preserved through the "tonemapper" post-processing step.
 	static IConsoleVariable* CVarPropagateAlpha = IConsoleManager::Get().FindConsoleVariable(TEXT("r.PostProcessing.PropagateAlpha"));
-	if (CVarPropagateAlpha && CVarPropagateAlpha->GetInt() != static_cast<int32>(EAlphaChannelMode::AllowThroughTonemapper))
+	if (CVarPropagateAlpha && CVarPropagateAlpha->GetBool() != true)
 	{
 		UE_CALL_ONCE([]()
 			{
diff --git a/Engine/Plugins/Experimental/VirtualScouting/Source/VirtualScoutingEditor/Private/VPScoutingMode.cpp b/Engine/Plugins/Experimental/VirtualScouting/Source/VirtualScoutingEditor/Private/VPScoutingMode.cpp
--- a/Engine/Plugins/Experimental/VirtualScouting/Source/VirtualScoutingEditor/Private/VPScoutingMode.cpp
+++ b/Engine/Plugins/Experimental/VirtualScouting/Source/VirtualScoutingEditor/Private/VPScoutingMode.cpp
@@ -106,9 +106,9 @@ bool UVPScoutingMode::ValidateSettings()
 	IConsoleManager& ConsoleMgr = IConsoleManager::Get();
 	if (IConsoleVariable* PropagateAlpha = ConsoleMgr.FindConsoleVariable(TEXT("r.PostProcessing.PropagateAlpha")))
 	{
-		if (PropagateAlpha->GetInt() != 0)
+		if (PropagateAlpha->GetBool())
 		{
-			InvalidSettingNotification(LOCTEXT("InvalidCvarPropagateAlpha", "r.PostProcessing.PropagateAlpha must be set to 0 (and requires an engine restart)"));
+			InvalidSettingNotification(LOCTEXT("InvalidCvarPropagateAlpha", "r.PostProcessing.PropagateAlpha must be disabled (and requires an engine restart)"));
 			bSettingsValid = false;
 		}
 	}
diff --git a/Engine/Plugins/Experimental/XRCreativeFramework/Source/XRCreativeEditor/Private/XRCreativeVREditorMode.cpp b/Engine/Plugins/Experimental/XRCreativeFramework/Source/XRCreativeEditor/Private/XRCreativeVREditorMode.cpp
--- a/Engine/Plugins/Experimental/XRCreativeFramework/Source/XRCreativeEditor/Private/XRCreativeVREditorMode.cpp
+++ b/Engine/Plugins/Experimental/XRCreativeFramework/Source/XRCreativeEditor/Private/XRCreativeVREditorMode.cpp
@@ -286,9 +286,9 @@ bool UXRCreativeVREditorMode::ValidateSettings()
 	IConsoleManager& ConsoleMgr = IConsoleManager::Get();
 	if (IConsoleVariable* PropagateAlpha = ConsoleMgr.FindConsoleVariable(TEXT("r.PostProcessing.PropagateAlpha")))
 	{
-		if (PropagateAlpha->GetInt() != 0)
+		if (PropagateAlpha->GetBool())
 		{
-			InvalidSettingNotification(LOCTEXT("InvalidCvarPropagateAlpha", "r.PostProcessing.PropagateAlpha must be set to 0 (and requires an engine restart)"));
+			InvalidSettingNotification(LOCTEXT("InvalidCvarPropagateAlpha", "r.PostProcessing.PropagateAlpha must be disabled (and requires an engine restart)"));
 			bSettingsValid = false;
 		}
 	}
diff --git a/Engine/Plugins/Media/AppleProResMedia/Source/AppleProResMedia/Private/AppleProResEncoderProtocol.cpp b/Engine/Plugins/Media/AppleProResMedia/Source/AppleProResMedia/Private/AppleProResEncoderProtocol.cpp
--- a/Engine/Plugins/Media/AppleProResMedia/Source/AppleProResMedia/Private/AppleProResEncoderProtocol.cpp
+++ b/Engine/Plugins/Media/AppleProResMedia/Source/AppleProResMedia/Private/AppleProResEncoderProtocol.cpp
@@ -570,9 +570,8 @@ ProResFormatDescriptionFieldDetail UAppleProResEncoderProtocol::GetFieldDetail()
 
 bool UAppleProResEncoderProtocol::HasAlpha() const
 {
-	static const auto CVarPropagateAlpha = IConsoleManager::Get().FindTConsoleVariableDataInt(TEXT("r.PostProcessing.PropagateAlpha"));
-	const EAlphaChannelMode::Type PropagateAlpha = EAlphaChannelMode::FromInt(CVarPropagateAlpha->GetValueOnAnyThread());
-	return EAlphaChannelMode::AllowThroughTonemapper == PropagateAlpha;
+	static const auto CVarPropagateAlpha = IConsoleManager::Get().FindConsoleVariable(TEXT("r.PostProcessing.PropagateAlpha"));
+	return CVarPropagateAlpha->GetBool();
 }
 
 void UAppleProResEncoderProtocol::ConvertFColorToRGBA4444(const TArray<FColor>& InColorbuffer)
diff --git a/Engine/Plugins/Media/MediaIOFramework/Source/MediaIOCore/Private/FileMediaOutput.cpp b/Engine/Plugins/Media/MediaIOFramework/Source/MediaIOCore/Private/FileMediaOutput.cpp
--- a/Engine/Plugins/Media/MediaIOFramework/Source/MediaIOCore/Private/FileMediaOutput.cpp
+++ b/Engine/Plugins/Media/MediaIOFramework/Source/MediaIOCore/Private/FileMediaOutput.cpp
@@ -84,10 +84,10 @@ EMediaCaptureConversionOperation UFileMediaOutput::GetConversionOperation(EMedia
 	// All formats supporting alpha
 	if (WriteOptions.Format == EDesiredImageFormat::EXR || WriteOptions.Format == EDesiredImageFormat::PNG)
 	{
-		// We invert alpha only when alpha channel as valid data when used with "passthrough tone mapper" or using a render target, otherwise we force it to 1.0f.
-		static const auto CVarPropagateAlpha = IConsoleManager::Get().FindTConsoleVariableDataInt(TEXT("r.PostProcessing.PropagateAlpha"));
-		EAlphaChannelMode::Type PropagateAlpha = EAlphaChannelMode::FromInt(CVarPropagateAlpha->GetValueOnAnyThread());
-		if ((PropagateAlpha == EAlphaChannelMode::AllowThroughTonemapper) || (InSourceType == EMediaCaptureSourceType::RENDER_TARGET))
+		// We invert alpha only when alpha channel as valid data when used with propagate alpha or using a render target, otherwise we force it to 1.0f.
+		static const auto CVarPropagateAlpha = IConsoleManager::Get().FindConsoleVariable(TEXT("r.PostProcessing.PropagateAlpha"));
+		const bool bPropagateAlpha = CVarPropagateAlpha->GetBool();
+		if (bPropagateAlpha || (InSourceType == EMediaCaptureSourceType::RENDER_TARGET))
 		{
 			return bInvertAlpha ? EMediaCaptureConversionOperation::INVERT_ALPHA :
 				EMediaCaptureConversionOperation::NONE;
diff --git a/Engine/Plugins/MovieScene/MovieRenderPipeline/Source/MovieRenderPipelineCore/Private/Graph/MovieGraphRenderLayerSubsystem.cpp b/Engine/Plugins/MovieScene/MovieRenderPipeline/Source/MovieRenderPipelineCore/Private/Graph/MovieGraphRenderLayerSubsystem.cpp
--- a/Engine/Plugins/MovieScene/MovieRenderPipeline/Source/MovieRenderPipelineCore/Private/Graph/MovieGraphRenderLayerSubsystem.cpp
+++ b/Engine/Plugins/MovieScene/MovieRenderPipeline/Source/MovieRenderPipelineCore/Private/Graph/MovieGraphRenderLayerSubsystem.cpp
@@ -192,12 +192,11 @@ void UMovieGraphRenderPropertyModifier::ApplyModifier(const UWorld* World)
 
 	// Generate a warning if holdout is being used, but alpha is not enabled in post processing. Without that setting enabled, holdout will not work.
 	const URendererSettings* RendererSettings = GetDefault<URendererSettings>();
-	if (bHoldout && (RendererSettings->bEnableAlphaChannelInPostProcessing == EAlphaChannelMode::Disabled))
+	if (bHoldout && (RendererSettings->bEnableAlphaChannelInPostProcessing == false))
 	{
 		// TODO: Ideally this is called in a general-purpose validation step instead, but that framework does not exist yet.
 		UE_LOG(LogMovieRenderPipeline, Warning, TEXT("A modifier with 'Holdout' is active, but 'Enable alpha channel support in post processing' in "
-													 "the project's Rendering settings is set to Disabled. Holdout will not work properly if this "
-													 "setting is Disabled."));
+													 "the project's Rendering settings is disabled: holdout will not work properly."));
 	}
 	
 	for (const UMovieGraphCollection* Collection : Collections)
diff --git a/Engine/Plugins/MovieScene/MovieRenderPipeline/Source/MovieRenderPipelineCore/Private/MoviePipeline.cpp b/Engine/Plugins/MovieScene/MovieRenderPipeline/Source/MovieRenderPipelineCore/Private/MoviePipeline.cpp
--- a/Engine/Plugins/MovieScene/MovieRenderPipeline/Source/MovieRenderPipelineCore/Private/MoviePipeline.cpp
+++ b/Engine/Plugins/MovieScene/MovieRenderPipeline/Source/MovieRenderPipelineCore/Private/MoviePipeline.cpp
@@ -120,9 +120,9 @@ void UMoviePipeline::ValidateSequenceAndSettings() const
 			bAnyOutputWantsAlpha |= Output->IsAlphaInTonemapperRequired();
 		}
 
-		if (bAnyOutputWantsAlpha && TonemapAlphaCVar->GetInt() == 0)
+		if (bAnyOutputWantsAlpha && TonemapAlphaCVar->GetBool() == false)
 		{
-			UE_LOG(LogMovieRenderPipeline, Warning, TEXT("An output requested Alpha Support but the required project setting is not enabled! Go to Project Settings > Rendering > PostProcessing > 'Enable Alpha Channel Support in Post Processing' and set it to 'Linear Color Space Only'."));
+			UE_LOG(LogMovieRenderPipeline, Warning, TEXT("An output requested Alpha Support but the required project setting is not enabled! Go to Project Settings > Rendering > PostProcessing > 'Enable Alpha Channel Support in Post Processing' and set it to 'True'."));
 		}
 	}
 }
diff --git a/Engine/Plugins/MovieScene/MovieRenderPipeline/Source/MovieRenderPipelineCore/Private/MoviePipelineRenderPass.cpp b/Engine/Plugins/MovieScene/MovieRenderPipeline/Source/MovieRenderPipelineCore/Private/MoviePipelineRenderPass.cpp
--- a/Engine/Plugins/MovieScene/MovieRenderPipeline/Source/MovieRenderPipelineCore/Private/MoviePipelineRenderPass.cpp
+++ b/Engine/Plugins/MovieScene/MovieRenderPipeline/Source/MovieRenderPipelineCore/Private/MoviePipelineRenderPass.cpp
@@ -11,7 +11,7 @@ void UMoviePipelineRenderPass::ValidateStateImpl()
 	if (IsAlphaInTonemapperRequired())
 	{
 		const URendererSettings* RenderSettings = GetDefault<URendererSettings>();
-		if (RenderSettings->bEnableAlphaChannelInPostProcessing == EAlphaChannelMode::Type::Disabled)
+		if (RenderSettings->bEnableAlphaChannelInPostProcessing == false)
 		{
 			ValidationState = EMoviePipelineValidationState::Warnings;
 			ValidationResults.Add(NSLOCTEXT("MovieRenderPipeline", "Outputs_AlphaWithoutProjectSetting", "This option does not work without enabling the Alpha Support in Tonemapper setting via Project Settings > Rendering > Post Processing > Enable Alpha Channel Support."));
diff --git a/Engine/Plugins/MovieScene/MovieRenderPipeline/Source/MovieRenderPipelineRenderPasses/Public/MoviePipelineDeferredPasses.h b/Engine/Plugins/MovieScene/MovieRenderPipeline/Source/MovieRenderPipelineRenderPasses/Public/MoviePipelineDeferredPasses.h
--- a/Engine/Plugins/MovieScene/MovieRenderPipeline/Source/MovieRenderPipelineRenderPasses/Public/MoviePipelineDeferredPasses.h
+++ b/Engine/Plugins/MovieScene/MovieRenderPipeline/Source/MovieRenderPipelineRenderPasses/Public/MoviePipelineDeferredPasses.h
@@ -101,7 +101,7 @@ protected:
 public:
 	/**
 	* Should multiple temporal/spatial samples accumulate the alpha channel? This requires r.PostProcessing.PropagateAlpha
-	* to be set to 1 or 2 (see "Enable Alpha Channel Support in Post Processing" under Project Settings > Rendering). This adds
+	* to be enabled (see "Enable Alpha Channel Support in Post Processing" under Project Settings > Rendering). This adds
 	* ~30% cost to the accumulation so you should not enable it unless necessary. You must delete both the sky and fog to ensure
 	* that they do not make all pixels opaque.
 	*/
diff --git a/Engine/Plugins/Runtime/OpenXR/Source/OpenXRHMD/Private/OpenXRHMD.cpp b/Engine/Plugins/Runtime/OpenXR/Source/OpenXRHMD/Private/OpenXRHMD.cpp
--- a/Engine/Plugins/Runtime/OpenXR/Source/OpenXRHMD/Private/OpenXRHMD.cpp
+++ b/Engine/Plugins/Runtime/OpenXR/Source/OpenXRHMD/Private/OpenXRHMD.cpp
@@ -1692,8 +1692,8 @@ bool FOpenXRHMD::ReconfigureForShaderPlatform(EShaderPlatform NewShaderPlatform)
 	UE::StereoRenderUtils::FStereoShaderAspects Aspects(NewShaderPlatform);
 	bIsMobileMultiViewEnabled = Aspects.IsMobileMultiViewEnabled();
 
-	static const auto CVarPropagateAlpha = IConsoleManager::Get().FindTConsoleVariableDataInt(TEXT("r.PostProcessing.PropagateAlpha"));
-	bProjectionLayerAlphaEnabled = !IsMobilePlatform(NewShaderPlatform) && CVarPropagateAlpha->GetValueOnAnyThread() != 0;
+	static const auto CVarPropagateAlpha = IConsoleManager::Get().FindConsoleVariable(TEXT("r.PostProcessing.PropagateAlpha"));
+	bProjectionLayerAlphaEnabled = !IsMobilePlatform(NewShaderPlatform) && CVarPropagateAlpha->GetBool();
 
 	ConfiguredShaderPlatform = NewShaderPlatform;
 
diff --git a/Engine/Plugins/Runtime/nDisplay/Source/DisplayCluster/Private/Render/Viewport/Configuration/DisplayClusterViewportConfigurationHelpers_RenderFrameSettings.cpp b/Engine/Plugins/Runtime/nDisplay/Source/DisplayCluster/Private/Render/Viewport/Configuration/DisplayClusterViewportConfigurationHelpers_RenderFrameSettings.cpp
--- a/Engine/Plugins/Runtime/nDisplay/Source/DisplayCluster/Private/Render/Viewport/Configuration/DisplayClusterViewportConfigurationHelpers_RenderFrameSettings.cpp
+++ b/Engine/Plugins/Runtime/nDisplay/Source/DisplayCluster/Private/Render/Viewport/Configuration/DisplayClusterViewportConfigurationHelpers_RenderFrameSettings.cpp
@@ -86,15 +86,14 @@ EDisplayClusterRenderFrameAlphaChannelCaptureMode FDisplayClusterViewportConfigu
 {
 	ECVarDisplayClusterAlphaChannelCaptureMode AlphaChannelCaptureMode = (ECVarDisplayClusterAlphaChannelCaptureMode)FMath::Clamp(GDisplayClusterAlphaChannelCaptureMode, 0, (int32)ECVarDisplayClusterAlphaChannelCaptureMode::COUNT - 1);
 
-	static const auto CVarPropagateAlpha = IConsoleManager::Get().FindTConsoleVariableDataInt(TEXT("r.PostProcessing.PropagateAlpha"));
-	const EAlphaChannelMode::Type PropagateAlpha = EAlphaChannelMode::FromInt(CVarPropagateAlpha->GetValueOnGameThread());
-	const bool bAllowThroughTonemapper = PropagateAlpha == EAlphaChannelMode::AllowThroughTonemapper;
+	static const auto CVarPropagateAlpha = IConsoleManager::Get().FindConsoleVariable(TEXT("r.PostProcessing.PropagateAlpha"));
+	const bool bPropagateAlpha = CVarPropagateAlpha->GetBool();
 
 	switch (AlphaChannelCaptureMode)
 	{
 	case ECVarDisplayClusterAlphaChannelCaptureMode::ThroughTonemapper:
 		// Disable alpha capture if PropagateAlpha not valid
-		return bAllowThroughTonemapper ? EDisplayClusterRenderFrameAlphaChannelCaptureMode::ThroughTonemapper : EDisplayClusterRenderFrameAlphaChannelCaptureMode::None;
+		return bPropagateAlpha ? EDisplayClusterRenderFrameAlphaChannelCaptureMode::ThroughTonemapper : EDisplayClusterRenderFrameAlphaChannelCaptureMode::None;
 
 	case ECVarDisplayClusterAlphaChannelCaptureMode::FXAA:
 		return EDisplayClusterRenderFrameAlphaChannelCaptureMode::FXAA;
diff --git a/Engine/Shaders/Private/PostProcessTonemap.usf b/Engine/Shaders/Private/PostProcessTonemap.usf
--- a/Engine/Shaders/Private/PostProcessTonemap.usf
+++ b/Engine/Shaders/Private/PostProcessTonemap.usf
@@ -615,13 +615,7 @@ float4 TonemapCommonPS(
 	#endif
 
 	#if POST_PROCESS_ALPHA
-	#if POST_PROCESS_ALPHA == 1
-		BRANCH
-		if(GetOutputDevice() == TONEMAPPER_OUTPUT_NoToneCurve || GetOutputDevice() == TONEMAPPER_OUTPUT_WithToneCurve)
-	#endif
-		{
-			OutColor.a = SceneColor.a;
-		}
+		OutColor.a = SceneColor.a;
 	#endif
 
 	return OutColor;
diff --git a/Engine/Shaders/Private/PostProcessUpscale.usf b/Engine/Shaders/Private/PostProcessUpscale.usf
--- a/Engine/Shaders/Private/PostProcessUpscale.usf
+++ b/Engine/Shaders/Private/PostProcessUpscale.usf
@@ -289,8 +289,9 @@ void MainPS(noperspective float4 UVAndScreenPos : TEXCOORD0, float4 SvPosition :
 		OutColor.a = 0.0;
 	}
 
-#if POST_PROCESS_ALPHA == 2 && SHADING_PATH_MOBILE
-	OutColor.a = 1.0f - OutColor.a;	// invert alpha when propagate alpha is used
+#if SHADING_PATH_MOBILE
+	// TODO: Remove inversion & replace with more robust solution for HMI: pass dependent & can potentially conflict with OpenXR backbuffer alpha inversion.
+	OutColor.a = 1.0f - OutColor.a;
 #endif
 	
 #else
diff --git a/Engine/Source/Editor/UnrealEd/Private/PlayLevel.cpp b/Engine/Source/Editor/UnrealEd/Private/PlayLevel.cpp
--- a/Engine/Source/Editor/UnrealEd/Private/PlayLevel.cpp
+++ b/Engine/Source/Editor/UnrealEd/Private/PlayLevel.cpp
@@ -3342,9 +3342,8 @@ TSharedRef<SPIEViewport> UEditorEngine::GeneratePIEViewportWindow(const FRequest
 	bool bRenderDirectlyToWindow = bVRPreview;
 	bool bEnableStereoRendering = bVRPreview;
 
-	static const auto CVarPropagateAlpha = IConsoleManager::Get().FindTConsoleVariableDataInt(TEXT("r.PostProcessing.PropagateAlpha"));
-	const EAlphaChannelMode::Type PropagateAlpha = EAlphaChannelMode::FromInt(CVarPropagateAlpha->GetValueOnGameThread());
-	const bool bIgnoreTextureAlpha = (PropagateAlpha != EAlphaChannelMode::AllowThroughTonemapper);
+	static const auto CVarPropagateAlpha = IConsoleManager::Get().FindConsoleVariable(TEXT("r.PostProcessing.PropagateAlpha"));
+	const bool bIgnoreTextureAlpha = (CVarPropagateAlpha->GetBool() != true);
 
 	TSharedRef<SPIEViewport> PieViewportWidget =
 		SNew(SPIEViewport)
diff --git a/Engine/Source/Runtime/Engine/Classes/Engine/RendererSettings.h b/Engine/Source/Runtime/Engine/Classes/Engine/RendererSettings.h
--- a/Engine/Source/Runtime/Engine/Classes/Engine/RendererSettings.h
+++ b/Engine/Source/Runtime/Engine/Classes/Engine/RendererSettings.h
@@ -114,29 +114,26 @@ namespace EVertexDeformationOutputsVelocity
 UENUM()
 namespace EAlphaChannelMode
 {
-	enum Type : int
+	enum UE_DEPRECATED(5.5, "Propagate alpha was converted back to a boolean.") Type : int
 	{
 		/** Disabled, reducing GPU cost to the minimum. (default). */
-		Disabled = 0 UMETA(DisplayName = "Disabled"),
+		Disabled = 0,
+
+		/** Propagate alpha channel through post-processing. */
+		Enabled = 1,
 
-		/** Maintain alpha channel only within linear color space. Tonemapper won't output alpha channel. */
-		LinearColorSpaceOnly = 1 UMETA(DisplayName="Linear color space only"),
-
-		/** Maintain alpha channel within linear color space, but also pass it through the tonemapper.
-		 *
-		 * CAUTION: Passing the alpha channel through the tonemapper can unevitably lead to pretty poor compositing quality as
-		 * opposed to linear color space compositing, especially on purely additive pixels bloom can generate. This settings is
-		 * exclusively targeting broadcast industry in case of hardware unable to do linear color space compositing and
-		 * tonemapping.
-		 */
-		AllowThroughTonemapper = 2 UMETA(DisplayName="Allow through tonemapper"),
+		// Deprecated
+		LinearColorSpaceOnly = Enabled UMETA(Hidden),
+		AllowThroughTonemapper = Enabled UMETA(Hidden)
 	};
 }
 
+PRAGMA_DISABLE_DEPRECATION_WARNINGS
 namespace EAlphaChannelMode
 {
 	ENGINE_API EAlphaChannelMode::Type FromInt(int32 InAlphaChannelMode);
 }
+PRAGMA_ENABLE_DEPRECATION_WARNINGS
 
 /** used by FPostProcessSettings AutoExposure*/
 UENUM()
@@ -729,10 +726,10 @@ class URendererSettings : public UDeveloperSettings
 	uint32 bCustomDepthTaaJitter : 1;
 
 	UPROPERTY(config, EditAnywhere, Category = Postprocessing, meta = (
-		ConsoleVariable = "r.PostProcessing.PropagateAlpha", DisplayName = "Enable alpha channel support in post processing (experimental).",
-		ToolTip = "Configures alpha channel support in renderer's post processing chain. Still experimental: works only with Temporal AA, Motion Blur, Circle Depth Of Field. This option also force disable the separate translucency.",
+		ConsoleVariable = "r.PostProcessing.PropagateAlpha", DisplayName = "Enable alpha channel support in post processing",
+		ToolTip = "Propagates the alpha channel through the renderer's post-processing chain.",
 		ConfigRestartRequired = true))
-	TEnumAsByte<EAlphaChannelMode::Type> bEnableAlphaChannelInPostProcessing;
+	bool bEnableAlphaChannelInPostProcessing;
 
 	UPROPERTY(config, EditAnywhere, Category = DefaultSettings, meta = (
 		ConsoleVariable = "r.DefaultFeature.Bloom", DisplayName = "Bloom",
diff --git a/Engine/Source/Runtime/Engine/Private/RendererSettings.cpp b/Engine/Source/Runtime/Engine/Private/RendererSettings.cpp
--- a/Engine/Source/Runtime/Engine/Private/RendererSettings.cpp
+++ b/Engine/Source/Runtime/Engine/Private/RendererSettings.cpp
@@ -32,6 +32,7 @@ extern UNREALED_API class UEditorEngine* GEditor;
 
 #define LOCTEXT_NAMESPACE "RendererSettings"
 
+PRAGMA_DISABLE_DEPRECATION_WARNINGS
 namespace EAlphaChannelMode
 {
 	EAlphaChannelMode::Type FromInt(int32 InAlphaChannelMode)
@@ -39,6 +40,7 @@ namespace EAlphaChannelMode
 		return static_cast<EAlphaChannelMode::Type>(FMath::Clamp(InAlphaChannelMode, (int32)Disabled, (int32)AllowThroughTonemapper));
 	}
 }
+PRAGMA_ENABLE_DEPRECATION_WARNINGS
 
 namespace EDefaultBackBufferPixelFormat
 {
diff --git a/Engine/Source/Runtime/Engine/Private/ShaderCompiler/ShaderCompiler.cpp b/Engine/Source/Runtime/Engine/Private/ShaderCompiler/ShaderCompiler.cpp
--- a/Engine/Source/Runtime/Engine/Private/ShaderCompiler/ShaderCompiler.cpp
+++ b/Engine/Source/Runtime/Engine/Private/ShaderCompiler/ShaderCompiler.cpp
@@ -3821,20 +3821,15 @@ void GlobalBeginCompileShader(
 
 	{
 		static IConsoleVariable* CVar = IConsoleManager::Get().FindConsoleVariable(TEXT("r.PostProcessing.PropagateAlpha"));
-		int32 PropagateAlpha = CVar->GetInt();
+		bool bPropagateAlpha = CVar->GetBool();
 
 		if (bIsMobilePlatform)
 		{
 			static FShaderPlatformCachedIniValue<int32> MobilePropagateAlphaIniValue(TEXT("r.Mobile.PropagateAlpha"));
-			int MobilePropagateAlphaIniValueInt = MobilePropagateAlphaIniValue.Get((EShaderPlatform)ShaderPlatform);
-			PropagateAlpha = MobilePropagateAlphaIniValueInt > 0 ? 2 : 0;
+			bPropagateAlpha = MobilePropagateAlphaIniValue.Get((EShaderPlatform)ShaderPlatform) > 0;
 		}
 
-		if (PropagateAlpha < 0 || PropagateAlpha > 2)
-		{
-			PropagateAlpha = 0;
-		}
-		SET_SHADER_DEFINE(Input.Environment, POST_PROCESS_ALPHA, PropagateAlpha);
+		SET_SHADER_DEFINE(Input.Environment, POST_PROCESS_ALPHA, bPropagateAlpha ? 1 : 0);
 	}
 
 	if (TargetPlatform && 
diff --git a/Engine/Source/Runtime/Engine/Private/ShaderCompiler/ShaderGenerationUtil.cpp b/Engine/Source/Runtime/Engine/Private/ShaderCompiler/ShaderGenerationUtil.cpp
--- a/Engine/Source/Runtime/Engine/Private/ShaderCompiler/ShaderGenerationUtil.cpp
+++ b/Engine/Source/Runtime/Engine/Private/ShaderCompiler/ShaderGenerationUtil.cpp
@@ -435,12 +435,7 @@ static FShaderGlobalDefines FetchShaderGlobalDefines(EShaderPlatform TargetPlatf
 
 	{
 		static IConsoleVariable* CVar = IConsoleManager::Get().FindConsoleVariable(TEXT("r.PostProcessing.PropagateAlpha"));
-		int32 PropagateAlpha = CVar->GetInt();
-		if (PropagateAlpha < 0 || PropagateAlpha > 2)
-		{
-			PropagateAlpha = 0;
-		}
-		Ret.POST_PROCESS_ALPHA = PropagateAlpha != 0;
+		Ret.POST_PROCESS_ALPHA = CVar->GetBool();
 	}
 
 	Ret.PLATFORM_SUPPORTS_SHADER_ROOT_CONSTANTS = RHISupportsShaderRootConstants(EShaderPlatform(TargetPlatform)) ? 1 : 0;
diff --git a/Engine/Source/Runtime/RenderCore/Private/Shader.cpp b/Engine/Source/Runtime/RenderCore/Private/Shader.cpp
--- a/Engine/Source/Runtime/RenderCore/Private/Shader.cpp
+++ b/Engine/Source/Runtime/RenderCore/Private/Shader.cpp
@@ -2192,34 +2192,25 @@ void ShaderMapAppendKey(EShaderPlatform Platform, FShaderKeyGenerator& KeyGen)
 	}
 
 	{
-		int PropagateAlphaType = 0;
+		bool bPropagateAlpha = false;
 		if (IsMobilePlatform(Platform))
 		{
 			static FShaderPlatformCachedIniValue<int32> MobilePropagateAlphaIniValue(TEXT("r.Mobile.PropagateAlpha"));
-			int MobilePropagateAlphaIniValueInt = MobilePropagateAlphaIniValue.Get((EShaderPlatform)Platform);
-			PropagateAlphaType = MobilePropagateAlphaIniValueInt > 0 ? 2 : 0;
+			bPropagateAlpha = MobilePropagateAlphaIniValue.Get((EShaderPlatform)Platform) > 0;
 		}
 		else
 		{
-			static const auto CVar = IConsoleManager::Get().FindTConsoleVariableDataInt(TEXT("r.PostProcessing.PropagateAlpha"));
-			if (CVar)
+			static const auto CVarPropagateAlpha = IConsoleManager::Get().FindConsoleVariable(TEXT("r.PostProcessing.PropagateAlpha"));
+			if (CVarPropagateAlpha)
 			{
-				PropagateAlphaType = CVar->GetValueOnAnyThread();
+				bPropagateAlpha = CVarPropagateAlpha->GetBool();
 			}
 		}
 
-		if (PropagateAlphaType > 0)
+		if (bPropagateAlpha)
 		{
-			if (PropagateAlphaType == 2)
-			{
-				KeyGen.AppendSeparator();
-				KeyGen.Append(TEXT("SA2"));
-			}
-			else
-			{
-				KeyGen.AppendSeparator();
-				KeyGen.Append(TEXT("SA"));
-			}
+			KeyGen.AppendSeparator();
+			KeyGen.Append(TEXT("PA"));
 		}
 	}
 
diff --git a/Engine/Source/Runtime/Renderer/Private/HairStrands/HairStrandsComposition.cpp b/Engine/Source/Runtime/Renderer/Private/HairStrands/HairStrandsComposition.cpp
--- a/Engine/Source/Runtime/Renderer/Private/HairStrands/HairStrandsComposition.cpp
+++ b/Engine/Source/Runtime/Renderer/Private/HairStrands/HairStrandsComposition.cpp
@@ -100,10 +100,9 @@ void InternalCommonDrawPass(
 			{
 				// Alpha usage/output is controlled with r.PostProcessing.PropagateAlpha. The value are:
 				// 0: disabled(default);
-				// 1: enabled in linear color space;
-				// 2: same as 1, but also enable it through the tonemapper.
+				// 1: enabled
 				//
-				// When enable (PorpagateAlpha is set to 1 or 2), the alpha value means:
+				// When enabled, the alpha value means:
 				// 0: valid pixel
 				// 1: invalid pixel (background)
 				GraphicsPSOInit.BlendState = TStaticBlendState<CW_RGBA, BO_Add, BF_One, BF_InverseSourceAlpha, BO_Add, BF_Zero, BF_InverseSourceAlpha>::GetRHI();
diff --git a/Engine/Source/Runtime/Renderer/Private/PostProcess/PostProcessing.cpp b/Engine/Source/Runtime/Renderer/Private/PostProcess/PostProcessing.cpp
--- a/Engine/Source/Runtime/Renderer/Private/PostProcess/PostProcessing.cpp
+++ b/Engine/Source/Runtime/Renderer/Private/PostProcess/PostProcessing.cpp
@@ -105,14 +105,20 @@ TAutoConsoleVariable<bool> CVarBloomApplyLocalExposure(
 	TEXT("Whether to apply local exposure when calculating bloom, default: true"),
 	ECVF_Scalability | ECVF_RenderThreadSafe);
 
-TAutoConsoleVariable<int32> CVarPostProcessingPropagateAlpha(
+static bool GPostProcessingPropagateAlpha = false;
+/**
+* NOTE (5.5):
+* r.PostProcessing.PropagateAlpha has been converted back to a boolean. In order to prevent silent failures
+* with IConsoleManager::Get().FindTConsoleVariableDataInt returning 0 with a boolean cvar set to True, we now use
+* a FAutoConsoleVariableRef which will warn licensees with typed-access at runtime, see IConsoleObject::AsVariableBool()
+* or IConsoleObject::AsVariableInt(). However both CVar->GetBool() & CVar->GetInt() will continue to work, assuming > 0
+* or EAlphaChannelMode::Type comparisons were used.
+*/
+FAutoConsoleVariableRef CVarPostProcessingPropagateAlpha(
 	TEXT("r.PostProcessing.PropagateAlpha"),
-	0,
-	TEXT("0 to disable scene alpha channel support in the post processing.\n")
-	TEXT(" 0: disabled (default);\n")
-	TEXT(" 1: enabled in linear color space;\n")
-	TEXT(" 2: same as 1, but also enable it through the tonemapper. Compositing after the tonemapper is incorrect, as their is no meaning to tonemap the alpha channel. This is only meant to be use exclusively for broadcasting hardware that does not support linear color space compositing and tonemapping."),
-	ECVF_ReadOnly);
+	GPostProcessingPropagateAlpha,
+	TEXT("Whether to propagate alpha through post-processing, default: false"),
+	ECVF_ReadOnly | ECVF_RenderThreadSafe);
 
 TAutoConsoleVariable<int32> CVarPostProcessingPreferCompute(
 	TEXT("r.PostProcessing.PreferCompute"),
@@ -212,7 +218,7 @@ bool IsPostProcessingEnabled(const FViewInfo& View)
 
 bool IsPostProcessingWithAlphaChannelSupported()
 {
-	return CVarPostProcessingPropagateAlpha.GetValueOnAnyThread() != 0;
+	return CVarPostProcessingPropagateAlpha->GetBool();
 }
 
 #if DEBUG_POST_PROCESS_VOLUME_ENABLE
diff --git a/Engine/Source/Runtime/Renderer/Private/PostProcess/TemporalSuperResolution.cpp b/Engine/Source/Runtime/Renderer/Private/PostProcess/TemporalSuperResolution.cpp
--- a/Engine/Source/Runtime/Renderer/Private/PostProcess/TemporalSuperResolution.cpp
+++ b/Engine/Source/Runtime/Renderer/Private/PostProcess/TemporalSuperResolution.cpp
@@ -62,7 +62,7 @@ TAutoConsoleVariable<int32> CVarTSRR11G11B10History(
 	TEXT("Select the bitdepth of the history. r.TSR.History.R11G11B10=1 Saves memory bandwidth that is of particular interest of the TSR's ")
 	TEXT("UpdateHistory's runtime performance by saving memory both at previous frame's history reprojection and write out of the output and ")
 	TEXT("new history.\n")
-	TEXT("This optimisation is unsupported with r.PostProcessing.PropagateAlpha=1.\n")
+	TEXT("This optimisation is unsupported with r.PostProcessing.PropagateAlpha=True.\n")
 	TEXT("\n")
 	TEXT("Please also not that increasing r.TSR.History.ScreenPercentage=200 adds 2 additional implicit encoding bits in the history compared to the TSR.Output's bitdepth thanks to the downscaling pass from TSR history resolution to TSR output resolution."),
 	ECVF_Scalability | ECVF_RenderThreadSafe);
