[World Partition] Support IsMainWorldOnly actor flag in classic levels when used as embedded level instances.

diff --git a/Engine/Source/Runtime/Engine/Private/ActorEditor.cpp b/Engine/Source/Runtime/Engine/Private/ActorEditor.cpp
--- a/Engine/Source/Runtime/Engine/Private/ActorEditor.cpp
+++ b/Engine/Source/Runtime/Engine/Private/ActorEditor.cpp
@@ -101,9 +101,8 @@ bool AActor::CanEditChange(const FProperty* PropertyThatWillChange) const
 	const bool bIsRuntimeGridProperty = PropertyThatWillChange->GetFName() == GET_MEMBER_NAME_CHECKED(AActor, RuntimeGrid);
 	const bool bIsDataLayersProperty = PropertyThatWillChange->GetFName() == GET_MEMBER_NAME_CHECKED(AActor, DataLayerAssets);
 	const bool bIsHLODLayerProperty = PropertyThatWillChange->GetFName() == GET_MEMBER_NAME_CHECKED(AActor, HLODLayer);
-	const bool bIsMainWorldOnlyProperty = PropertyThatWillChange->GetFName() == GET_MEMBER_NAME_CHECKED(AActor, bIsMainWorldOnly);
 
-	if (bIsSpatiallyLoadedProperty || bIsRuntimeGridProperty || bIsDataLayersProperty || bIsHLODLayerProperty || bIsMainWorldOnlyProperty)
+	if (bIsSpatiallyLoadedProperty || bIsRuntimeGridProperty || bIsDataLayersProperty || bIsHLODLayerProperty)
 	{
 		if (!IsTemplate())
 		{
