FSkeletalMeshProxy: set view relevance dynamic if any materials are hidden.

It appears that hidden materials are not supported for static view relevance. Switch to dynamic if any there are any hidden materials.

diff --git a/Engine/Source/Runtime/Engine/Private/SkeletalMesh.cpp b/Engine/Source/Runtime/Engine/Private/SkeletalMesh.cpp
--- a/Engine/Source/Runtime/Engine/Private/SkeletalMesh.cpp
+++ b/Engine/Source/Runtime/Engine/Private/SkeletalMesh.cpp
@@ -7485,6 +7485,7 @@ FPrimitiveViewRelevance FSkeletalMeshSceneProxy::GetViewRelevance(const FSceneVi
 			|| EngineShowFlags.VisualizeGPUSkinCache
 			|| (IsSelected() && (EngineShowFlags.VertexColors || EngineShowFlags.Bounds))
 			|| bForceWireframe
+			|| MeshObject->IsAnyMaterialHidden(MeshObject->GetLOD())
 #if WITH_EDITORONLY_DATA
 			|| MeshObject->SelectedEditorMaterial != -1
 			|| MeshObject->SelectedEditorSection != -1
diff --git a/Engine/Source/Runtime/Engine/Private/SkeletalRender.cpp b/Engine/Source/Runtime/Engine/Private/SkeletalRender.cpp
--- a/Engine/Source/Runtime/Engine/Private/SkeletalRender.cpp
+++ b/Engine/Source/Runtime/Engine/Private/SkeletalRender.cpp
@@ -246,6 +246,25 @@ bool FSkeletalMeshObject::IsMaterialHidden(int32 InLODIndex,int32 MaterialIdx) c
 	check(LODInfo.IsValidIndex(InLODIndex));
 	return LODInfo[InLODIndex].HiddenMaterials.IsValidIndex(MaterialIdx) && LODInfo[InLODIndex].HiddenMaterials[MaterialIdx];
 }
+
+/**
+ * Determine if any material section entry for an LOD is hidden or not
+ *
+ * @param InLODIndex - LOD entry to get hidden material flags for
+ */
+bool FSkeletalMeshObject::IsAnyMaterialHidden(int32 InLODIndex) const
+{
+	check(LODInfo.IsValidIndex(InLODIndex));
+	for (const bool bIsHidden : LODInfo[InLODIndex].HiddenMaterials)
+	{
+		if (bIsHidden)
+		{
+			return true;
+		}
+	}
+	return false;
+}
+
 /**
  * Initialize the array of LODInfo based on the settings of the current skel mesh component
  */
diff --git a/Engine/Source/Runtime/Engine/Public/SkeletalRenderPublic.h b/Engine/Source/Runtime/Engine/Public/SkeletalRenderPublic.h
--- a/Engine/Source/Runtime/Engine/Public/SkeletalRenderPublic.h
+++ b/Engine/Source/Runtime/Engine/Public/SkeletalRenderPublic.h
@@ -245,6 +245,13 @@ public:
 	 * @param MaterialIdx - index of the material section to check
 	 */
 	bool IsMaterialHidden(int32 InLODIndex,int32 MaterialIdx) const;
+
+	/**
+	 * Determine if any material section entry for an LOD is hidden or not
+	 *
+	 * @param InLODIndex - LOD entry to get hidden material flags for
+	 */
+	bool IsAnyMaterialHidden(int32 InLODIndex) const;
 	
 	/**
 	 * Initialize the array of LODInfo based on the settings of the current skel mesh component
