Print FTimerDelegate's object name in FTimerUnifiedDelegate::ToString() if available

diff --git a/Engine/Source/Runtime/Engine/Private/TimerManager.cpp b/Engine/Source/Runtime/Engine/Private/TimerManager.cpp
--- a/Engine/Source/Runtime/Engine/Private/TimerManager.cpp
+++ b/Engine/Source/Runtime/Engine/Private/TimerManager.cpp
@@ -476,7 +476,9 @@ FString FTimerUnifiedDelegate::ToString() const
 						FunctionNameStr = TEXT(" 0x0");
 					}
 				}
-				return FString::Printf(TEXT("DELEGATE,%s,%s"), TEXT("NO OBJ"), *FunctionNameStr);
+
+				const UObject* const Object = FuncDelegate.GetUObject();
+				return FString::Printf(TEXT("DELEGATE,%s,%s"), Object == nullptr ? TEXT("NO OBJ") : *Object->GetPathName(), *FunctionNameStr);
 			}
 			return TEXT("UNBOUND DELEGATE");
 		}
