Add ODSC support for default light material
Enable ODSC reload on directional lights on mobile
Default material for surfaces are not supported yet because we're missing some invalidation and it crashes inside FMeshDrawShaderBindings::SetShaderBindings

#prelight 66bb1e46ae40ad937eac5314

diff --git a/Engine/Source/Runtime/Engine/Private/Materials/MaterialShared.cpp b/Engine/Source/Runtime/Engine/Private/Materials/MaterialShared.cpp
--- a/Engine/Source/Runtime/Engine/Private/Materials/MaterialShared.cpp
+++ b/Engine/Source/Runtime/Engine/Private/Materials/MaterialShared.cpp
@@ -3731,7 +3731,7 @@ bool FMaterial::TryGetShaders(const FMaterialShaderTypes& InTypes, const FVertex
 
 #if WITH_ODSC
 	const bool bIsODSCActive = FODSCManager::IsODSCActive();
-	bool bShouldForceRecompile = bIsODSCActive && FODSCManager::ShouldForceRecompile(ShaderMap, this) && !IsDefaultMaterial();
+	bool bShouldForceRecompile = bIsODSCActive && FODSCManager::ShouldForceRecompile(ShaderMap, this);
 #else
 	constexpr bool bIsODSCActive = false;
 	constexpr bool bShouldForceRecompile = false;
diff --git a/Engine/Source/Runtime/Engine/Private/ODSC/ODSCManager.cpp b/Engine/Source/Runtime/Engine/Private/ODSC/ODSCManager.cpp
--- a/Engine/Source/Runtime/Engine/Private/ODSC/ODSCManager.cpp
+++ b/Engine/Source/Runtime/Engine/Private/ODSC/ODSCManager.cpp
@@ -327,6 +327,11 @@ bool FODSCManager::ShouldForceRecompileInternal(const FMaterialShaderMap* Materi
 		return false;
 	}
 
+	if (Material->IsDefaultMaterial() && !Material->IsLightFunction())
+	{
+		return false;
+	}
+
 	if (!MaterialNameToRecompile.IsNone())
 	{
 		EODSCMetaDataType ODSCMetaData = (EODSCMetaDataType)Material->GetODSCMetaData();
diff --git a/Engine/Source/Runtime/Engine/Private/ODSC/ODSCThread.cpp b/Engine/Source/Runtime/Engine/Private/ODSC/ODSCThread.cpp
--- a/Engine/Source/Runtime/Engine/Private/ODSC/ODSCThread.cpp
+++ b/Engine/Source/Runtime/Engine/Private/ODSC/ODSCThread.cpp
@@ -369,6 +369,7 @@ void FODSCThread::AddShaderPipelineRequest(
 )
 {
 	bool bShouldAddRequest = false;
+	bool bDefaultMaterial = false;
 
 	FString ActorPath;
 	{
@@ -380,6 +381,8 @@ void FODSCThread::AddShaderPipelineRequest(
 			CachedMaterialName = FName(Material->GetFullPath());
 		}
 		
+		bDefaultMaterial = Material->IsDefaultMaterial();
+
 		FODSCShaderMapData& ODSCShaderMapData = RequestHashes.FindOrAdd(CachedMaterialName);
 
 		for (const FShaderId& ShaderId : RequestShaderIds)
@@ -392,6 +395,12 @@ void FODSCThread::AddShaderPipelineRequest(
 			}
 		}
 
+		// for default materials, we request all the permutations anyway
+		if (bDefaultMaterial && ODSCShaderMapData.CurrentRequests.Num() > 1)
+		{
+			bShouldAddRequest = false;
+		}
+
 		if (bShouldAddRequest)
 		{
 #if WITH_ODSC && IS_MONOLITHIC
@@ -418,7 +427,7 @@ void FODSCThread::AddShaderPipelineRequest(
 
 		FString MaterialName = Material->GetFullPath();
 
-		if (!ActorPath.IsEmpty())
+		if (!bDefaultMaterial && !ActorPath.IsEmpty())
 		{
 			MaterialName += ":::";
 			MaterialName += ActorPath;
@@ -430,7 +439,16 @@ void FODSCThread::AddShaderPipelineRequest(
 			RequestString += ShaderTypeName;
 		}
 		const FString RequestHash = FMD5::HashAnsiString(*RequestString);
-		PendingMeshMaterialThreadedRequests.Enqueue(FODSCRequestPayload(ShaderPlatform, FeatureLevel, QualityLevel, MaterialName, VertexFactoryName, PipelineName, ShaderTypeNames, PermutationId, RequestHash));
+		if (bDefaultMaterial)
+		{
+			TArray<FString> MaterialsToCompile = {MaterialName};
+			FString ShaderTypesToLoad;
+			PendingMaterialThreadedRequests.Enqueue(new FODSCMessageHandler(MaterialsToCompile, ShaderTypesToLoad, ShaderPlatform, FeatureLevel, QualityLevel, ODSCRecompileCommand::Material));
+		}
+		else
+		{
+			PendingMeshMaterialThreadedRequests.Enqueue(FODSCRequestPayload(ShaderPlatform, FeatureLevel, QualityLevel, MaterialName, VertexFactoryName, PipelineName, ShaderTypeNames, PermutationId, RequestHash));
+		}
 	}
 }
 
diff --git a/Engine/Source/Runtime/Renderer/Private/MobileDeferredShadingPass.cpp b/Engine/Source/Runtime/Renderer/Private/MobileDeferredShadingPass.cpp
--- a/Engine/Source/Runtime/Renderer/Private/MobileDeferredShadingPass.cpp
+++ b/Engine/Source/Runtime/Renderer/Private/MobileDeferredShadingPass.cpp
@@ -377,6 +377,14 @@ static void GetLightMaterial(const FCachedLightMaterial& DefaultLightMaterial, c
 	// use default material
 	OutLightMaterial.Material = DefaultLightMaterial.Material;
 	OutLightMaterial.MaterialProxy = DefaultLightMaterial.MaterialProxy;
+
+	// Perform a TryGetShaders to allow ODSC to record a shader recompile request when enabled
+	if (DefaultLightMaterial.Material->TryGetShaders(ShaderTypes, nullptr, Shaders))
+	{
+		Shaders.TryGetPixelShader(OutShader);
+		return;
+	}
+
 	const FMaterialShaderMap* MaterialShaderMap = OutLightMaterial.Material->GetRenderingThreadShaderMap();
 	OutShader = MaterialShaderMap->GetShader<ShaderType>(PermutationId);
 }
