[Material Editor] New material translator push.

- Some changes to the IR Emitter API.
- Implemented AppendVector material expression.
- Fixed input names not displaying correctly.

diff --git a/Engine/Source/Runtime/Engine/Private/Materials/MaterialExpressions.cpp b/Engine/Source/Runtime/Engine/Private/Materials/MaterialExpressions.cpp
--- a/Engine/Source/Runtime/Engine/Private/Materials/MaterialExpressions.cpp
+++ b/Engine/Source/Runtime/Engine/Private/Materials/MaterialExpressions.cpp
@@ -799,7 +799,14 @@ UMaterialExpression::UMaterialExpression(const FObjectInitializer& ObjectInitial
 		}
 	}
 	CachedInputs.Shrink();
-#endif
+
+	// Initialize the input names from GetInputName())
+	for (FExpressionInputIterator It{ this }; It; ++It)
+	{
+		It.Input->InputName = GetInputName(It.Index);
+	}
+
+	#endif
 }
 
 UObject* UMaterialExpression::GetAssetOwner() const
diff --git a/Engine/Source/Runtime/Engine/Private/Materials/MaterialExpressionsIR.cpp b/Engine/Source/Runtime/Engine/Private/Materials/MaterialExpressionsIR.cpp
--- a/Engine/Source/Runtime/Engine/Private/Materials/MaterialExpressionsIR.cpp
+++ b/Engine/Source/Runtime/Engine/Private/Materials/MaterialExpressionsIR.cpp
@@ -270,7 +270,7 @@ void UMaterialExpressionConstant2Vector::Build(MIR::FEmitter& Emitter)
 	Emitter.Put(GetOutput(0), Value);
 	for (int i = 0; i < 2; ++i)
 	{
-		Emitter.Put(GetOutput(i + 1), Emitter.TryEmitSubscript(Value, i));
+		Emitter.Put(GetOutput(i + 1), Emitter.EmitSubscript(Value, i));
 	}
 }
 
@@ -280,7 +280,7 @@ void UMaterialExpressionConstant3Vector::Build(MIR::FEmitter& Emitter)
 	Emitter.Put(GetOutput(0), Value);
 	for (int i = 0; i < 3; ++i)
 	{
-		Emitter.Put(GetOutput(i + 1), Emitter.TryEmitSubscript(Value, i));
+		Emitter.Put(GetOutput(i + 1), Emitter.EmitSubscript(Value, i));
 	}
 }
 
@@ -290,10 +290,84 @@ void UMaterialExpressionConstant4Vector::Build(MIR::FEmitter& Emitter)
 	Emitter.Put(GetOutput(0), Value);
 	for (int i = 0; i < 4; ++i)
 	{
-		Emitter.Put(GetOutput(i + 1), Emitter.TryEmitSubscript(Value, i));
+		Emitter.Put(GetOutput(i + 1), Emitter.EmitSubscript(Value, i));
 	}
 }
 
+void UMaterialExpressionAppendVector::Build(MIR::FEmitter& Emitter)
+{
+	MIR::FValue* AValue = Emitter.GetArithmetic(&A);
+	MIR::FValue* BValue = Emitter.TryGet(&B) ? Emitter.GetArithmetic(&B) : nullptr;
+
+	if (Emitter.IsInvalid())
+	{
+		return;
+	}
+
+	MIR::FArithmeticTypePtr AType = AValue->Type->AsArithmetic();
+	if (AType->IsMatrix())
+	{
+		Emitter.Error(TEXT("Input A does not have scalar or vector type."));
+		return;
+	}
+
+	MIR::FArithmeticTypePtr BType = nullptr;
+	if (BValue)
+	{
+		BType = BValue->Type->AsArithmetic();
+		if (BType->IsMatrix()) 
+		{
+			Emitter.Error(TEXT("Input B does not have scalar or vector type."));
+			return;
+		}
+	}
+
+	int Dimensions = AType->NumRows + (BType ? BType->NumRows : 0);
+	if (Dimensions > 4)
+	{
+		Emitter.Errorf(TEXT("The resulting vector would have %d component (it can have at most 4)."), Dimensions);
+		return;
+	}
+
+	check(Dimensions >= 2 && Dimensions <= 4);
+
+	// Construct the output vector type.
+	MIR::FArithmeticTypePtr ResultType = MIR::FArithmeticType::GetVector(MIR::SK_Float, Dimensions);
+
+	// Set up each output vector component 
+	MIR::FValue* Components[4] = { nullptr, nullptr, nullptr, nullptr };
+	int ComponentIndex = 0;
+	for (int i = 0; i < AType->NumRows; ++i, ++ComponentIndex)
+	{
+		Components[ComponentIndex] = Emitter.EmitSubscript(AValue, i);
+	}
+
+	if (BType)
+	{
+		for (int i = 0; i < BType->NumRows; ++i, ++ComponentIndex)
+		{
+			Components[ComponentIndex] = Emitter.EmitSubscript(BValue, i);
+		}
+	}
+
+	// Create the vector value and output it.
+	MIR::FValue* Output{};
+	if (Dimensions == 2)
+	{
+		Output = Emitter.EmitVector2(Components[0], Components[1]);
+	}
+	else if (Dimensions == 3)
+	{
+		Output = Emitter.EmitVector3(Components[0], Components[1], Components[2]);
+	}
+	else
+	{
+		Output = Emitter.EmitVector4(Components[0], Components[1], Components[2], Components[3]);
+	}
+
+	Emitter.Put(0, Output);
+}
+
 /* Mathematical Operations */
 
 static void BuildBinaryArithmeticOperator(
@@ -302,12 +376,11 @@ static void BuildBinaryArithmeticOperator(
 	FExpressionInput* LhsInput,
 	float LhsConst,
 	FExpressionInput* RhsInput,
-	float RhsConst,
-	FExpressionOutput* Output)
+	float RhsConst)
 {
 	// Default inputs to their relative constants if disconnected, then get each input after checking it has arithmetic type.
-	MIR::FValue* LhsValue = Emitter.DefaultTo(LhsInput, LhsConst).TryGetArithmetic(LhsInput);
-	MIR::FValue* RhsValue = Emitter.DefaultTo(RhsInput, RhsConst).TryGetArithmetic(RhsInput);
+	MIR::FValue* LhsValue = Emitter.DefaultTo(LhsInput, LhsConst).GetArithmetic(LhsInput);
+	MIR::FValue* RhsValue = Emitter.DefaultTo(RhsInput, RhsConst).GetArithmetic(RhsInput);
 
 	if (Emitter.IsInvalid())
 	{
@@ -320,8 +393,8 @@ static void BuildBinaryArithmeticOperator(
 		: MIR::FArithmeticType::GetScalar(MIR::SK_Float);
 
 	// Convert operand values to determined result type. 
-	LhsValue = Emitter.TryEmitConstruct(ResultType, LhsValue);
-	RhsValue = Emitter.TryEmitConstruct(ResultType, RhsValue);
+	LhsValue = Emitter.EmitConstruct(ResultType, LhsValue);
+	RhsValue = Emitter.EmitConstruct(ResultType, RhsValue);
 
 	if (Emitter.IsInvalid())
 	{
@@ -332,27 +405,27 @@ static void BuildBinaryArithmeticOperator(
 	MIR::FValue* Value = Emitter.EmitBinaryOperator(Op, LhsValue, RhsValue);
 
 	// And flow it out of the expression's only output.
-	Emitter.Put(Output, Value);
+	Emitter.Put(0, Value);
 }
 
 void UMaterialExpressionAdd::Build(MIR::FEmitter& Emitter)
 { 
-	BuildBinaryArithmeticOperator(Emitter, MIR::BO_Add, &A, ConstA, &B, ConstB, GetOutput(0));
+	BuildBinaryArithmeticOperator(Emitter, MIR::BO_Add, &A, ConstA, &B, ConstB);
 }
 
 void UMaterialExpressionSubtract::Build(MIR::FEmitter& Emitter)
 { 
-	BuildBinaryArithmeticOperator(Emitter, MIR::BO_Subtract, &A, ConstA, &B, ConstB, GetOutput(0));
+	BuildBinaryArithmeticOperator(Emitter, MIR::BO_Subtract, &A, ConstA, &B, ConstB);
 }
 
 void UMaterialExpressionMultiply::Build(MIR::FEmitter& Emitter)
 { 
-	BuildBinaryArithmeticOperator(Emitter, MIR::BO_Multiply, &A, ConstA, &B, ConstB, GetOutput(0));
+	BuildBinaryArithmeticOperator(Emitter, MIR::BO_Multiply, &A, ConstA, &B, ConstB);
 }
 
 void UMaterialExpressionDivide::Build(MIR::FEmitter& Emitter)
 { 
-	BuildBinaryArithmeticOperator(Emitter, MIR::BO_Divide, &A, ConstA, &B, ConstB, GetOutput(0));
+	BuildBinaryArithmeticOperator(Emitter, MIR::BO_Divide, &A, ConstA, &B, ConstB);
 }
 
 void UMaterialExpressionIf::Build(MIR::FEmitter& Emitter)
@@ -365,11 +438,11 @@ void UMaterialExpressionIf::Build(MIR::FEmitter& Emitter)
 	Emitter.DefaultToFloatZero(&ALessThanB);
 
 	// Get input values and check their types are what we expect.
-	MIR::FValue* AValue = Emitter.TryGetScalar(&A);
-	MIR::FValue* BValue = Emitter.TryGetScalar(&B);
-	MIR::FValue* AGreaterThanBValue = Emitter.TryGetArithmetic(&AGreaterThanB);
-	MIR::FValue* AEqualsBValue = Emitter.TryGetArithmetic(&AEqualsB);
-	MIR::FValue* ALessThanBValue = Emitter.TryGetArithmetic(&ALessThanB);
+	MIR::FValue* AValue = Emitter.GetScalar(&A);
+	MIR::FValue* BValue = Emitter.GetScalar(&B);
+	MIR::FValue* AGreaterThanBValue = Emitter.GetArithmetic(&AGreaterThanB);
+	MIR::FValue* AEqualsBValue = Emitter.GetArithmetic(&AEqualsB);
+	MIR::FValue* ALessThanBValue = Emitter.GetArithmetic(&ALessThanB);
 
 	if (Emitter.IsInvalid())
 	{
@@ -377,7 +450,7 @@ void UMaterialExpressionIf::Build(MIR::FEmitter& Emitter)
 	}
 
 	// Get the arithmetic common type between the conditional arguments (e.g. if inputs are int and float, it will return float).
-	MIR::FArithmeticTypePtr ConditionArgsType = Emitter.TryGetCommonArithmeticType(AValue->Type->AsArithmetic(), BValue->Type->AsArithmetic());
+	MIR::FArithmeticTypePtr ConditionArgsType = Emitter.GetCommonArithmeticType(AValue->Type->AsArithmetic(), BValue->Type->AsArithmetic());
 
 	if (Emitter.IsInvalid())
 	{
@@ -385,12 +458,12 @@ void UMaterialExpressionIf::Build(MIR::FEmitter& Emitter)
 	}
 
 	// Convert both conditional argument values to the common type. This is expected to work as the common type could be found.
-	AValue = Emitter.TryEmitConstruct(ConditionArgsType, AValue);
-	BValue = Emitter.TryEmitConstruct(ConditionArgsType, BValue);
+	AValue = Emitter.EmitConstruct(ConditionArgsType, AValue);
+	BValue = Emitter.EmitConstruct(ConditionArgsType, BValue);
 
 	// Now determine the output type by taking the common arithmetic type between result values.
-	MIR::FArithmeticTypePtr OutputType = Emitter.TryGetCommonArithmeticType(AGreaterThanBValue->Type->AsArithmetic(), AEqualsBValue->Type->AsArithmetic());
-	OutputType = Emitter.TryGetCommonArithmeticType(OutputType, ALessThanBValue->Type->AsArithmetic());
+	MIR::FArithmeticTypePtr OutputType = Emitter.GetCommonArithmeticType(AGreaterThanBValue->Type->AsArithmetic(), AEqualsBValue->Type->AsArithmetic());
+	OutputType = Emitter.GetCommonArithmeticType(OutputType, ALessThanBValue->Type->AsArithmetic());
 
 	if (Emitter.IsInvalid())
 	{
@@ -398,9 +471,9 @@ void UMaterialExpressionIf::Build(MIR::FEmitter& Emitter)
 	}
 
 	// Convert result values to the common result type.
-	AGreaterThanBValue = Emitter.TryEmitConstruct(OutputType, AGreaterThanBValue);
-	AEqualsBValue = Emitter.TryEmitConstruct(OutputType, AEqualsBValue);
-	ALessThanBValue = Emitter.TryEmitConstruct(OutputType, ALessThanBValue);
+	AGreaterThanBValue = Emitter.EmitConstruct(OutputType, AGreaterThanBValue);
+	AEqualsBValue = Emitter.EmitConstruct(OutputType, AEqualsBValue);
+	ALessThanBValue = Emitter.EmitConstruct(OutputType, ALessThanBValue);
 
 	if (Emitter.IsInvalid())
 	{
@@ -448,13 +521,13 @@ void UMaterialExpressionTextureSample::Build(UE::MIR::FEmitter& Emitter)
 		return;
 	}
 
-	MIR::FValue* OutputValue = Emitter.TryEmitTextureSample(InputTexture, TexCoordsValue, SamplerSource, MipValueMode, SamplerType);
+	MIR::FValue* OutputValue = Emitter.EmitTextureSample(InputTexture, TexCoordsValue, SamplerSource, MipValueMode, SamplerType);
 	
-	Emitter.Put(GetOutput(0), Emitter.TryEmitSwizzle(OutputValue, MIR::FSwizzleMask::XYZ()));
-	Emitter.Put(GetOutput(1), Emitter.TryEmitSubscript(OutputValue, 0)); 
-	Emitter.Put(GetOutput(2), Emitter.TryEmitSubscript(OutputValue, 1)); 
-	Emitter.Put(GetOutput(3), Emitter.TryEmitSubscript(OutputValue, 2));
-	Emitter.Put(GetOutput(4), Emitter.TryEmitSubscript(OutputValue, 3));
+	Emitter.Put(GetOutput(0), Emitter.EmitSwizzle(OutputValue, MIR::FSwizzleMask::XYZ()));
+	Emitter.Put(GetOutput(1), Emitter.EmitSubscript(OutputValue, 0)); 
+	Emitter.Put(GetOutput(2), Emitter.EmitSubscript(OutputValue, 1)); 
+	Emitter.Put(GetOutput(3), Emitter.EmitSubscript(OutputValue, 2));
+	Emitter.Put(GetOutput(4), Emitter.EmitSubscript(OutputValue, 3));
 	Emitter.Put(GetOutput(5), OutputValue); 
 }
 
diff --git a/Engine/Source/Runtime/Engine/Private/Materials/MaterialIREmitter.cpp b/Engine/Source/Runtime/Engine/Private/Materials/MaterialIREmitter.cpp
--- a/Engine/Source/Runtime/Engine/Private/Materials/MaterialIREmitter.cpp
+++ b/Engine/Source/Runtime/Engine/Private/Materials/MaterialIREmitter.cpp
@@ -152,6 +152,26 @@ static FValue* EmitPrototype(FEmitter* Emitter, const TValueType& Prototype)
 	return Value;
 }
 
+// Finds the expression input index. Although the implementation has O(n) complexity, it is only used for error reporting.
+static int SlowFindExpressionInputIndex(UMaterialExpression* Expression, const FExpressionInput* Input)
+{
+	for (FExpressionInputIterator It{ Expression }; It; ++It)
+	{
+		if (It.Input == Input)
+		{
+			return It.Index;
+		}
+	}
+	check(false && "No input found.");
+	return -1;
+}
+
+// Finds the expression input name. Although the implementation has O(n) complexity, it is only used for error reporting.
+static FName SlowFindInputName(UMaterialExpression* Expression, const FExpressionInput* Input)
+{
+	int InputIndex = SlowFindExpressionInputIndex(Expression, Input);
+	return Expression->GetInputName(InputIndex);
+}
 
 // FEmitter API
 
@@ -162,20 +182,34 @@ FEmitter::FEmitter(FMaterialIRModuleBuilder* InBuilder, UMaterial* InMaterial, F
 	Module = InModule;
 }
 
-FValue* FEmitter::Get(const FExpressionInput* Input)
+FValue*  FEmitter::TryGet(const FExpressionInput* Input)
 {
 	FValue** Value = Builder->InputValues.Find(Input);
-	if (!Value || !*Value)
+	return Value ? *Value : nullptr;
+}
+
+FValue* FEmitter::Get(const FExpressionInput* Input)
+{
+	FValue* Value = TryGet(Input);
+	if (!Value)
 	{
+		Errorf(TEXT("Input '%s' has no input value."), *SlowFindInputName(Expression, Input).ToString());
 		return nullptr;
 	}
-	return *Value;
+	return Value;
 }
 	
-void FEmitter::Put(const FExpressionOutput* Output, FValue* Value)
+FEmitter& FEmitter::Put(int OutputIndex, FValue* Value)
+{
+	Put(Expression->GetOutput(OutputIndex), Value);
+	return *this;
+}
+
+FEmitter& FEmitter::Put(const FExpressionOutput* Output, FValue* Value)
 {
 	check(Output);
 	Builder->OutputValues.Add(Output, Value);
+	return *this;
 }
 
 FEmitter& FEmitter::DefaultToFloatZero(const FExpressionInput* Input)
@@ -193,33 +227,33 @@ FEmitter& FEmitter::DefaultTo(const FExpressionInput* Input, TFloat Float)
 	return *this;
 }
 
-FValue* FEmitter::TryGetFloat(const FExpressionInput* Input)
+FValue* FEmitter::GetFloat(const FExpressionInput* Input)
 {
 	FValue* Value = Get(Input);
 	CheckInputIsScalar(Input, Value, SK_Float);
 	return Value;
 }
 
-FValue* FEmitter::TryGetScalar(const FExpressionInput* Input)
+FValue* FEmitter::GetScalar(const FExpressionInput* Input)
 {
-	FValue* Value = TryGetArithmetic(Input);
+	FValue* Value = GetArithmetic(Input);
 	if (!Value)
 	{
 		return nullptr;
 	}
 
 	FArithmeticTypePtr ScalarType = Value->Type->AsArithmetic()->ToScalar();
-	return TryEmitConstruct(ScalarType, Value);
+	return EmitConstruct(ScalarType, Value);
 }
 
-FValue* FEmitter::TryGetArithmetic(const FExpressionInput* Input)
+FValue* FEmitter::GetArithmetic(const FExpressionInput* Input)
 {
 	FValue* Value = Get(Input);
 	CheckInputTypeIs(Input, Value, TK_Arithmetic);
 	return Value;
 }
 
-FValue* FEmitter::TryGetOfType(const FExpressionInput* Input, ETypeKind Kind)
+FValue* FEmitter::GetOfType(const FExpressionInput* Input, ETypeKind Kind)
 {
 	FValue* Value = Get(Input);
 	if (!Value)
@@ -246,7 +280,7 @@ void FEmitter::CheckInputIsScalar(const FExpressionInput* Input, FValue* InputVa
 	FArithmeticTypePtr ArithmeticType = InputValue->Type->AsArithmetic();
 	if (!ArithmeticType || !ArithmeticType->IsScalar())
 	{
-		Errorf(TEXT("Input '%s' expected to be a scalar. It is %s instead."), *Input->InputName.ToString(), InputValue->Type->GetSpelling().GetData());
+		Errorf(TEXT("Input '%s' expected to be a scalar. It is %s instead."), *SlowFindInputName(Expression, Input).ToString(), InputValue->Type->GetSpelling().GetData());
 	}
 }
 
@@ -255,7 +289,7 @@ void FEmitter::CheckInputIsScalar(const FExpressionInput* Input, FValue* InputVa
 	FArithmeticTypePtr ArithmeticType = InputValue->Type->AsArithmetic();
 	if (!ArithmeticType || !ArithmeticType->IsScalar() || ArithmeticType->ScalarKind != Kind)
 	{
-		Errorf(TEXT("Input '%s' expected to be a %s scalar. It is %s instead."), *Input->InputName.ToString(), ScalarKindToString(Kind), InputValue->Type->GetSpelling().GetData());
+		Errorf(TEXT("Input '%s' expected to be a %s scalar. It is %s instead."), *SlowFindInputName(Expression, Input).ToString(), ScalarKindToString(Kind), InputValue->Type->GetSpelling().GetData());
 	}
 }
 
@@ -263,7 +297,7 @@ void FEmitter::CheckInputTypeIs(const FExpressionInput* Input, FValue* InputValu
 {
 	if (InputValue && InputValue->Type->Kind != Kind)
 	{
-		Errorf(TEXT("Input '%s' expected to be have type %s. It is %s instead."), *Input->InputName.ToString(), TypeKindToString(Kind), InputValue->Type->GetSpelling().GetData());
+		Errorf(TEXT("Input '%s' expected to be have type %s. It is %s instead."), *SlowFindInputName(Expression, Input).ToString(), TypeKindToString(Kind), InputValue->Type->GetSpelling().GetData());
 	}
 }
 
@@ -427,7 +461,7 @@ FValue* FEmitter::GetExternalInput(EExternalInput Id)
 	return EmitPrototype(this, Prototype);
 }
 
-FValue* FEmitter::TryEmitSubscript(FValue* Value, int Index)
+FValue* FEmitter::EmitSubscript(FValue* Value, int Index)
 {
 	FArithmeticTypePtr ArithmeticType = Value->Type->AsArithmetic();
 	if (!ArithmeticType)
@@ -463,7 +497,7 @@ FValue* FEmitter::TryEmitSubscript(FValue* Value, int Index)
 	return EmitPrototype(this, Prototype);
 }
 
-FValue* FEmitter::TryEmitSwizzle(FValue* Value, FSwizzleMask Mask)
+FValue* FEmitter::EmitSwizzle(FValue* Value, FSwizzleMask Mask)
 {
 	// At least one component must have been specified.
 	check(Mask.NumComponents > 0);
@@ -489,7 +523,7 @@ FValue* FEmitter::TryEmitSwizzle(FValue* Value, FSwizzleMask Mask)
 	// If only one component is requested, we can use EmitSubscript() to return the single component.
 	if (Mask.NumComponents == 1)
 	{
-		return TryEmitSubscript(Value, (int)Mask.Components[0]);
+		return EmitSubscript(Value, (int)Mask.Components[0]);
 	}
 
 	// If the requested number of components is the same as Value and the order in which the components
@@ -518,7 +552,7 @@ FValue* FEmitter::TryEmitSwizzle(FValue* Value, FSwizzleMask Mask)
 
 	for (int i = 0; i < Mask.NumComponents; ++i)
 	{
-		Result->GetComponents()[i] = TryEmitSubscript(Value, (int)Mask.Components[i]);
+		Result->GetComponents()[i] = EmitSubscript(Value, (int)Mask.Components[i]);
 	}
 
 	return Result;
@@ -756,7 +790,7 @@ static FValue* ConstructArithmeticValue(FEmitter* Emitter, FArithmeticTypePtr Ta
 	if (TargetArithmeticType->IsScalar())
 	{
 		//
-		Initializer = Emitter->TryEmitSubscript(Initializer, 0);
+		Initializer = Emitter->EmitSubscript(Initializer, 0);
 		InitializerArithmeticType = Initializer->Type->AsArithmetic();
 		
 		//
@@ -785,7 +819,7 @@ static FValue* ConstructArithmeticValue(FEmitter* Emitter, FArithmeticTypePtr Ta
 
 		// Create a dimensional and initialize each of its components to the conversion
 		// of initializer value to the single component type.
-		FValue* Component = Emitter->TryEmitConstruct(TargetArithmeticType->ToScalar(), Initializer);
+		FValue* Component = Emitter->EmitConstruct(TargetArithmeticType->ToScalar(), Initializer);
 
 		// Initialize all result components to the same scalar.
 		for (int i = 0; i < TargetArithmeticType->GetNumComponents(); ++i)
@@ -817,7 +851,7 @@ static FValue* ConstructArithmeticValue(FEmitter* Emitter, FArithmeticTypePtr Ta
 		int MinNumComponents = FMath::Min(TargetNumComponents, InitializerNumComponents);
 		for (; Index < MinNumComponents; ++Index)
 		{
-			Result->GetComponents()[Index] = Emitter->TryEmitConstruct(ResultComponentType, Emitter->TryEmitSubscript(Initializer, Index));
+			Result->GetComponents()[Index] = Emitter->EmitConstruct(ResultComponentType, Emitter->EmitSubscript(Initializer, Index));
 		}
 
 		// Initialize remaining result dimensional components to zero.
@@ -847,7 +881,7 @@ static FValue* ConstructArithmeticValue(FEmitter* Emitter, FArithmeticTypePtr Ta
 			// Convert components from the initializer vector.
 			for (int Index = 0, Num = TargetArithmeticType->GetNumComponents(); Index < Num; ++Index)
 			{
-				Result->GetComponents()[Index] = Emitter->TryEmitConstruct(ResultComponentType, DimensionalInitializer->GetComponents()[Index]);
+				Result->GetComponents()[Index] = Emitter->EmitConstruct(ResultComponentType, DimensionalInitializer->GetComponents()[Index]);
 			}
 
 			return EmitNew(Emitter, Result);
@@ -865,7 +899,7 @@ static FValue* ConstructArithmeticValue(FEmitter* Emitter, FArithmeticTypePtr Ta
 	return nullptr;
 }
 
-FValue* FEmitter::TryEmitConstruct(FTypePtr Type, FValue* Initializer)
+FValue* FEmitter::EmitConstruct(FTypePtr Type, FValue* Initializer)
 {
 	// If target type matches initializer's, simply return the same value.
 	FTypePtr InitializerType = Initializer->Type;
@@ -889,7 +923,7 @@ FValue* FEmitter::TryEmitConstruct(FTypePtr Type, FValue* Initializer)
 	return Result;
 }
 
-FValue* FEmitter::TryEmitTextureSample(UTexture* Texture, FValue* TexCoord, ESamplerSourceMode SamplerSourceMode, ETextureMipValueMode MipValueMode, EMaterialSamplerType SamplerType)
+FValue* FEmitter::EmitTextureSample(UTexture* Texture, FValue* TexCoord, ESamplerSourceMode SamplerSourceMode, ETextureMipValueMode MipValueMode, EMaterialSamplerType SamplerType)
 {
 	MIR::FTextureSample Prototype = MakePrototype<MIR::FTextureSample>(FArithmeticType::GetFloat4());
 	Prototype.Texture = Texture;
@@ -903,7 +937,7 @@ FValue* FEmitter::TryEmitTextureSample(UTexture* Texture, FValue* TexCoord, ESam
 	return EmitPrototype(this, Prototype);
 }
 
-FArithmeticTypePtr FEmitter::TryGetCommonArithmeticType(FArithmeticTypePtr A, FArithmeticTypePtr B)
+FArithmeticTypePtr FEmitter::GetCommonArithmeticType(FArithmeticTypePtr A, FArithmeticTypePtr B)
 {
 	// Trivial case: types are equal
 	if (A == B)
diff --git a/Engine/Source/Runtime/Engine/Private/Materials/MaterialIRModuleBuilder.cpp b/Engine/Source/Runtime/Engine/Private/Materials/MaterialIRModuleBuilder.cpp
--- a/Engine/Source/Runtime/Engine/Private/Materials/MaterialIRModuleBuilder.cpp
+++ b/Engine/Source/Runtime/Engine/Private/Materials/MaterialIRModuleBuilder.cpp
@@ -144,7 +144,7 @@ struct FMaterialIRModuleBuilder::FPrivate
 				Builder.InputValues.Add(Input.Input, *ValuePtr);
 
 				MIR::FTypePtr OutputArgType = MIR::FType::FromShaderType(Input.Type);
-				Output->Arg = Emitter.TryEmitConstruct(OutputArgType, *ValuePtr);
+				Output->Arg = Emitter.EmitConstruct(OutputArgType, *ValuePtr);
 			}
 
 			if (Params.TargetInsight)
diff --git a/Engine/Source/Runtime/Engine/Private/Materials/MaterialShared.cpp b/Engine/Source/Runtime/Engine/Private/Materials/MaterialShared.cpp
--- a/Engine/Source/Runtime/Engine/Private/Materials/MaterialShared.cpp
+++ b/Engine/Source/Runtime/Engine/Private/Materials/MaterialShared.cpp
@@ -99,7 +99,7 @@ static TAutoConsoleVariable<bool> CVarMaterialEdPreshaderDumpToHLSL(
 #endif
 
 static TAutoConsoleVariable<bool> CVarUsingUseNewMaterialTranslatorPrototype(
-	TEXT("r.MaterialEditor.UseNewTranslatorPrototype"),
+	TEXT("r.Material.Translator.EnableNew"),
 	false,
 	TEXT("Controls whether to enable the new material translator prototype (WIP) ."),
 	ECVF_RenderThreadSafe);
diff --git a/Engine/Source/Runtime/Engine/Public/Materials/MaterialExpressionAppendVector.h b/Engine/Source/Runtime/Engine/Public/Materials/MaterialExpressionAppendVector.h
--- a/Engine/Source/Runtime/Engine/Public/Materials/MaterialExpressionAppendVector.h
+++ b/Engine/Source/Runtime/Engine/Public/Materials/MaterialExpressionAppendVector.h
@@ -23,6 +23,7 @@ class UMaterialExpressionAppendVector : public UMaterialExpression
 
 	//~ Begin UMaterialExpression Interface
 #if WITH_EDITOR
+	virtual void Build(UE::MIR::FEmitter& Emitter) override;
 	virtual int32 Compile(class FMaterialCompiler* Compiler, int32 OutputIndex) override;
 	virtual void GetCaption(TArray<FString>& OutCaptions) const override;
 	virtual bool GenerateHLSLExpression(FMaterialHLSLGenerator& Generator, UE::HLSLTree::FScope& Scope, int32 OutputIndex, UE::HLSLTree::FExpression const*& OutExpression) const override;
diff --git a/Engine/Source/Runtime/Engine/Public/Materials/MaterialIRCommon.h b/Engine/Source/Runtime/Engine/Public/Materials/MaterialIRCommon.h
--- a/Engine/Source/Runtime/Engine/Public/Materials/MaterialIRCommon.h
+++ b/Engine/Source/Runtime/Engine/Public/Materials/MaterialIRCommon.h
@@ -48,7 +48,6 @@ struct FValue;
 struct FGlobalInput;
 struct FInstruction;
 struct FSetMaterialOutput;
-using FValuePtr = const FValue*;
 using FInstructionPtr = const FInstruction*;
 enum class EExternalInput;
 
diff --git a/Engine/Source/Runtime/Engine/Public/Materials/MaterialIRDebug.cpp b/Engine/Source/Runtime/Engine/Public/Materials/MaterialIRDebug.cpp
--- a/Engine/Source/Runtime/Engine/Public/Materials/MaterialIRDebug.cpp
+++ b/Engine/Source/Runtime/Engine/Public/Materials/MaterialIRDebug.cpp
@@ -11,14 +11,14 @@
 #if WITH_EDITOR
 
 static TAutoConsoleVariable<bool> CVarDumpMaterialIRUseGraph(
-	TEXT("r.MaterialEditor.MaterialTranslator.DumpUseGraph"),
+	TEXT("r.Material.Translator.DumpUseGraph"),
 	true,
 	TEXT("Whether the material translator should emit the Module IR 'Uses' graph in Graphviz Dot syntax (to 'MaterialIRDumpGraph.dot')."),
 	ECVF_RenderThreadSafe);
 
 
 static TAutoConsoleVariable<bool> CVarDumpMaterialIRUseGraph_EnableNext(
-	TEXT("r.MaterialEditor.MaterialTranslator.DumpUseGraphOpts.EnableSuccessors"),
+	TEXT("r.Material.Translator.DumpUseGraphOpts.EnableSuccessors"),
 	false,
 	TEXT("Whether the Material Module IR 'Uses' graph should also display 'Instruction Next' edges."),
 	ECVF_RenderThreadSafe);
diff --git a/Engine/Source/Runtime/Engine/Public/Materials/MaterialIREmitter.h b/Engine/Source/Runtime/Engine/Public/Materials/MaterialIREmitter.h
--- a/Engine/Source/Runtime/Engine/Public/Materials/MaterialIREmitter.h
+++ b/Engine/Source/Runtime/Engine/Public/Materials/MaterialIREmitter.h
@@ -44,13 +44,21 @@ public:
 	
 	/* IO */
 
-	// Gets and returns the value flowing into specified `Input`. If disconnected,
-	// it returns nullptr. 
+	// Tries to get a value from specified input. If it has no value (for instance
+	// because it is not connected), it returns null.
+	FValue* TryGet(const FExpressionInput* Input);
+
+	// Gets and returns the value flowing into input with specified `InputIndex`.
+	// If no value is flowing in (e.g. bacause it is not connected) it reports
+	// an error and returns null.
 	FValue* Get(const FExpressionInput* Input);
+	
+	// Flows specified `Value` out of output with specified `OutputIndex`.
+	FEmitter& Put(int OutputIndex, FValue* Value);
 
 	// Flows specified `Value` out of specified expression `Output`.
-	void 	Put(const FExpressionOutput* Output, FValue* Value);
-
+	FEmitter& Put(const FExpressionOutput* Output, FValue* Value);
+	
 	/* IO Helpers */
 
 	//
@@ -60,17 +68,17 @@ public:
 	FEmitter& DefaultTo(const FExpressionInput* Input, TFloat Float);
 
 	// It gets the value flowing into it and checks that its type is float scalar.
-	FValue* TryGetFloat(const FExpressionInput* Input);
+	FValue* GetFloat(const FExpressionInput* Input);
 
 	//
-	FValue* TryGetScalar(const FExpressionInput* Input);
-
+	FValue* GetScalar(const FExpressionInput* Input);
+	
 	//
-	FValue* TryGetArithmetic(const FExpressionInput* Input);
+	FValue* GetArithmetic(const FExpressionInput* Input);
 
 	// Gets the value flowing into `Input` and returns it after checking that its
 	// type matches `Kind`.
-	FValue* TryGetOfType(const FExpressionInput* Input, ETypeKind Kind);
+	FValue* GetOfType(const FExpressionInput* Input, ETypeKind Kind);
 
 	/* Error Checking */
 
@@ -109,8 +117,8 @@ public:
 
 	/* Other Values */
 
-	FValue* TryEmitSubscript(FValue* Value, int ComponentIndex);
-	FValue* TryEmitSwizzle(FValue* Value, FSwizzleMask Mask);
+	FValue* EmitSubscript(FValue* Value, int ComponentIndex);
+	FValue* EmitSwizzle(FValue* Value, FSwizzleMask Mask);
 	FValue* GetParameter(FName Name, const FMaterialParameterMetadata& Metadata);
 
 	/* Instructions */
@@ -118,12 +126,12 @@ public:
 	FSetMaterialOutput* EmitSetMaterialOutput(EMaterialProperty InProperty, FValue* InArgValue);
 	FValue* EmitBinaryOperator(EBinaryOperator Operator, FValue* Lhs, FValue* Rhs);
 	FValue* EmitBranch(FValue* Condition, FValue* True, FValue* False);
-	FValue* TryEmitConstruct(FTypePtr Type, FValue* Initializer);
-	FValue* TryEmitTextureSample(UTexture* Texture, FValue* TexCoord, ESamplerSourceMode SamplerSourceMode, ETextureMipValueMode MipValueMode, EMaterialSamplerType SamplerType);
+	FValue* EmitConstruct(FTypePtr Type, FValue* Initializer);
+	FValue* EmitTextureSample(UTexture* Texture, FValue* TexCoord, ESamplerSourceMode SamplerSourceMode, ETextureMipValueMode MipValueMode, EMaterialSamplerType SamplerType);
 
 	/* Types */
 
-	FArithmeticTypePtr TryGetCommonArithmeticType(FArithmeticTypePtr A, FArithmeticTypePtr B);
+	FArithmeticTypePtr GetCommonArithmeticType(FArithmeticTypePtr A, FArithmeticTypePtr B);
 
 	/* Error reporting */
 
