[Pawn] Check pointer before using it.

diff --git a/Engine/Source/Runtime/Engine/Private/Components/PawnMovementComponent.cpp b/Engine/Source/Runtime/Engine/Private/Components/PawnMovementComponent.cpp
--- a/Engine/Source/Runtime/Engine/Private/Components/PawnMovementComponent.cpp
+++ b/Engine/Source/Runtime/Engine/Private/Components/PawnMovementComponent.cpp
@@ -154,13 +154,21 @@ void UPawnMovementComponent::ApplyAsyncPhysicsStateAction(const UPrimitiveCompon
 void UPawnMovementComponent::ServerAsyncPhysicsStateAction_Implementation(const UPrimitiveComponent* ActionComponent, const FName BoneName, const FAsyncPhysicsTimestamp Timestamp,
 	const EPhysicsStateAction ActionType, const FVector ActionDatas, const FVector ActionPosition)
 {
-	ExecuteAsyncPhysicsStateAction(ActionComponent, BoneName, Timestamp, ActionType, ActionDatas, ActionPosition);
-	MulticastAsyncPhysicsStateAction(ActionComponent, BoneName, Timestamp, ActionType, ActionDatas, ActionPosition);
+	if (ActionComponent)
+	{
+		ExecuteAsyncPhysicsStateAction(ActionComponent, BoneName, Timestamp, ActionType, ActionDatas, ActionPosition);
+		MulticastAsyncPhysicsStateAction(ActionComponent, BoneName, Timestamp, ActionType, ActionDatas, ActionPosition);
+	}
 }
 
 void UPawnMovementComponent::MulticastAsyncPhysicsStateAction_Implementation(const UPrimitiveComponent* ActionComponent, const FName BoneName, const FAsyncPhysicsTimestamp Timestamp,
 	const EPhysicsStateAction ActionType, const FVector ActionDatas, const FVector ActionPosition)
 {
+	if (!ActionComponent)
+	{
+		return;
+	}
+
 	if (UWorld* World = GetWorld())
 	{
 		if (APlayerController* PlayerController = World->GetFirstPlayerController())
@@ -179,6 +187,11 @@ void UPawnMovementComponent::MulticastAsyncPhysicsStateAction_Implementation(con
 void UPawnMovementComponent::ExecuteAsyncPhysicsStateAction(const UPrimitiveComponent* ActionComponent, const FName& BoneName, const FAsyncPhysicsTimestamp& Timestamp,
 	const EPhysicsStateAction ActionType, const FVector& ActionDatas, const FVector& ActionPosition)
 {
+	if (!ActionComponent)
+	{
+		return;
+	}
+
 	if (FBodyInstance* BI = ActionComponent->GetBodyInstance(BoneName))
 	{
 		if (GetOwner()->IsNetMode(NM_Client) || GetOwner()->IsNetMode(NM_Standalone))
