[NetPhysics] Don't let ClusterUnionComponent change between dynamic and sleeping ObjectState if Physics Prediction is enabled, rely on the chosen Physics Replication Mode to handle it instead (Predictive Interpolation or Resimulation)

diff --git a/Engine/Source/Runtime/Engine/Classes/PhysicsEngine/ClusterUnionComponent.h b/Engine/Source/Runtime/Engine/Classes/PhysicsEngine/ClusterUnionComponent.h
--- a/Engine/Source/Runtime/Engine/Classes/PhysicsEngine/ClusterUnionComponent.h
+++ b/Engine/Source/Runtime/Engine/Classes/PhysicsEngine/ClusterUnionComponent.h
@@ -426,6 +426,9 @@ private:
 	mutable bool bHasCachedLocalBounds;
 	mutable FBoxSphereBounds CachedLocalBounds;
 
+	// If Physics Prediction is enabled
+	bool bPhysicsPredictionEnabled;
+
 protected:
 	// Handles changes to ReplicatedRigidState. Note that this function does not handle replication of X/R since we make use
 	// of the scene component's default replication for that.
diff --git a/Engine/Source/Runtime/Engine/Private/PhysicsEngine/ClusterUnionComponent.cpp b/Engine/Source/Runtime/Engine/Private/PhysicsEngine/ClusterUnionComponent.cpp
--- a/Engine/Source/Runtime/Engine/Private/PhysicsEngine/ClusterUnionComponent.cpp
+++ b/Engine/Source/Runtime/Engine/Private/PhysicsEngine/ClusterUnionComponent.cpp
@@ -697,6 +697,9 @@ void UClusterUnionComponent::OnCreatePhysicsState()
 		return;
 	}
 
+	const FPhysicsPredictionSettings& PhysicsPredictionSettings = UPhysicsSettings::Get()->PhysicsPrediction;
+	bPhysicsPredictionEnabled = PhysicsPredictionSettings.bEnablePhysicsPrediction;
+
 	// TODO: Expose these parameters via the component.
 	Chaos::FClusterCreationParameters Parameters{ 0.3f, 100, false, false };
 	Parameters.ConnectionMethod = Chaos::FClusterCreationParameters::EConnectionMethod::PointImplicit;
@@ -1334,7 +1337,15 @@ void UClusterUnionComponent::OnRep_RigidState()
 		return;
 	}
 
-	SetRigidState(static_cast<Chaos::EObjectStateType>(ReplicatedRigidState.ObjectState));
+	Chaos::EObjectStateType ObjectState = static_cast<Chaos::EObjectStateType>(ReplicatedRigidState.ObjectState);
+
+	// If Physics Prediction is enabled, only update to and from Kinematic state (sleeping/dynamic switching will be handled via the enabled EPhysicsReplicationMode)
+	if (!bPhysicsPredictionEnabled 
+		|| PhysicsProxy->GetObjectState_External() == Chaos::EObjectStateType::Kinematic
+		|| ObjectState == Chaos::EObjectStateType::Kinematic)
+	{
+		SetRigidState(ObjectState);
+	}
 }
 
 void UClusterUnionComponent::GetLifetimeReplicatedProps(TArray<FLifetimeProperty>& OutLifetimeProps) const
