[Navigation] added copy/ move constructors and operators to fixed deprecation warnings

diff --git a/Engine/Source/Runtime/Engine/Classes/AI/Navigation/NavigationDirtyArea.h b/Engine/Source/Runtime/Engine/Classes/AI/Navigation/NavigationDirtyArea.h
--- a/Engine/Source/Runtime/Engine/Classes/AI/Navigation/NavigationDirtyArea.h
+++ b/Engine/Source/Runtime/Engine/Classes/AI/Navigation/NavigationDirtyArea.h
@@ -37,6 +37,12 @@ struct FNavigationDirtyArea
 	UE_DEPRECATED(5.5, "Use the constructor taking ENavigationDirtyFlag and FNavigationElement instead.")
 	ENGINE_API FNavigationDirtyArea(const FBox& InBounds, int32 InFlags, UObject* const InOptionalSourceObject = nullptr);
 
+	ENGINE_API FNavigationDirtyArea(const FNavigationDirtyArea& Other);
+	ENGINE_API FNavigationDirtyArea(FNavigationDirtyArea&& Other);
+
+	ENGINE_API FNavigationDirtyArea& operator=(const FNavigationDirtyArea& Other);
+	ENGINE_API FNavigationDirtyArea& operator=(FNavigationDirtyArea&& Other);
+
 	bool HasFlag(const ENavigationDirtyFlag Flag) const
 	{
 		return (Flags & Flag) != ENavigationDirtyFlag::None;
diff --git a/Engine/Source/Runtime/Engine/Classes/AI/Navigation/NavigationDirtyElement.h b/Engine/Source/Runtime/Engine/Classes/AI/Navigation/NavigationDirtyElement.h
--- a/Engine/Source/Runtime/Engine/Classes/AI/Navigation/NavigationDirtyElement.h
+++ b/Engine/Source/Runtime/Engine/Classes/AI/Navigation/NavigationDirtyElement.h
@@ -59,6 +59,10 @@ struct FNavigationDirtyElement
 	ENGINE_API FNavigationDirtyElement(const TSharedRef<const FNavigationElement>& InNavigationElement, ENavigationDirtyFlag InFlagsOverride, const bool bUseWorldPartitionedDynamicMode = false);
 	ENGINE_API explicit FNavigationDirtyElement(const TSharedRef<const FNavigationElement>& InNavigationElement, const bool bUseWorldPartitionedDynamicMode = false);
 
+	ENGINE_API FNavigationDirtyElement(const FNavigationDirtyElement& Other);
+
+	ENGINE_API FNavigationDirtyElement& operator=(const FNavigationDirtyElement& Other);
+
 	UE_DEPRECATED(5.5, "This operator will no longer be used.")
 	ENGINE_API bool operator==(const UObject*& OtherOwner) const;
 
diff --git a/Engine/Source/Runtime/Engine/Classes/AI/Navigation/NavigationRelevantData.h b/Engine/Source/Runtime/Engine/Classes/AI/Navigation/NavigationRelevantData.h
--- a/Engine/Source/Runtime/Engine/Classes/AI/Navigation/NavigationRelevantData.h
+++ b/Engine/Source/Runtime/Engine/Classes/AI/Navigation/NavigationRelevantData.h
@@ -108,6 +108,12 @@ struct FNavigationRelevantData : public TSharedFromThis<FNavigationRelevantData,
 		, bLoadedData(false)
 	{}
 
+	ENGINE_API FNavigationRelevantData(const FNavigationRelevantData& Other);
+	ENGINE_API FNavigationRelevantData(FNavigationRelevantData&& Other);
+
+	ENGINE_API FNavigationRelevantData& operator=(FNavigationRelevantData&& Other);
+	ENGINE_API FNavigationRelevantData& operator=(const FNavigationRelevantData& Other);
+
 	FORCEINLINE bool HasGeometry() const { return VoxelData.Num() || CollisionData.Num(); }
 	FORCEINLINE bool HasModifiers() const { return !Modifiers.IsEmpty(); }
 	FORCEINLINE bool HasDynamicModifiers() const { return Modifiers.IsDynamic(); }
diff --git a/Engine/Source/Runtime/Engine/Private/AI/Navigation/NavigationDirtyArea.cpp b/Engine/Source/Runtime/Engine/Private/AI/Navigation/NavigationDirtyArea.cpp
--- a/Engine/Source/Runtime/Engine/Private/AI/Navigation/NavigationDirtyArea.cpp
+++ b/Engine/Source/Runtime/Engine/Private/AI/Navigation/NavigationDirtyArea.cpp
@@ -32,3 +32,53 @@ FNavigationDirtyArea::FNavigationDirtyArea(const FBox& InBounds, const int32 InF
 	: FNavigationDirtyArea(InBounds, static_cast<ENavigationDirtyFlag>(InFlags), /*InOptionalSourceElement*/nullptr)
 {
 }
+
+FNavigationDirtyArea::FNavigationDirtyArea(const FNavigationDirtyArea& Other)
+	: Bounds(Other.Bounds)
+#if WITH_EDITORONLY_DATA
+PRAGMA_DISABLE_DEPRECATION_WARNINGS
+	, OptionalSourceObject(Other.OptionalSourceObject)
+PRAGMA_ENABLE_DEPRECATION_WARNINGS
+#endif // WITH_EDITORONLY_DATA
+	, OptionalSourceElement(Other.OptionalSourceElement)
+	, Flags(Other.Flags)
+{
+}
+
+FNavigationDirtyArea::FNavigationDirtyArea(FNavigationDirtyArea&& Other)
+	: Bounds(MoveTemp(Other.Bounds))
+#if WITH_EDITORONLY_DATA
+PRAGMA_DISABLE_DEPRECATION_WARNINGS
+	, OptionalSourceObject(MoveTemp(Other.OptionalSourceObject))
+PRAGMA_ENABLE_DEPRECATION_WARNINGS
+#endif // WITH_EDITORONLY_DATA
+	, OptionalSourceElement(MoveTemp(Other.OptionalSourceElement))
+	, Flags(Other.Flags)
+{
+}
+
+FNavigationDirtyArea& FNavigationDirtyArea::operator=(const FNavigationDirtyArea& Other)
+{
+	Bounds = Other.Bounds;
+#if WITH_EDITORONLY_DATA
+PRAGMA_DISABLE_DEPRECATION_WARNINGS
+	OptionalSourceObject = Other.OptionalSourceObject;
+PRAGMA_ENABLE_DEPRECATION_WARNINGS
+#endif // WITH_EDITORONLY_DATA
+	OptionalSourceElement = Other.OptionalSourceElement;
+	Flags = Other.Flags;
+	return *this;
+}
+
+FNavigationDirtyArea& FNavigationDirtyArea::operator=(FNavigationDirtyArea&& Other)
+{
+	Bounds = MoveTemp(Other.Bounds);
+#if WITH_EDITORONLY_DATA
+PRAGMA_DISABLE_DEPRECATION_WARNINGS
+	OptionalSourceObject = MoveTemp(Other.OptionalSourceObject);
+PRAGMA_ENABLE_DEPRECATION_WARNINGS
+#endif // WITH_EDITORONLY_DATA
+	OptionalSourceElement = MoveTemp(Other.OptionalSourceElement);
+	Flags = Other.Flags;
+	return *this;
+}
diff --git a/Engine/Source/Runtime/Engine/Private/AI/Navigation/NavigationDirtyElement.cpp b/Engine/Source/Runtime/Engine/Private/AI/Navigation/NavigationDirtyElement.cpp
--- a/Engine/Source/Runtime/Engine/Private/AI/Navigation/NavigationDirtyElement.cpp
+++ b/Engine/Source/Runtime/Engine/Private/AI/Navigation/NavigationDirtyElement.cpp
@@ -22,6 +22,45 @@ FNavigationDirtyElement::FNavigationDirtyElement(const TSharedRef<const FNavigat
 {
 }
 
+FNavigationDirtyElement::FNavigationDirtyElement(const FNavigationDirtyElement& Other)
+	: ExplicitAreasToDirty(Other.ExplicitAreasToDirty)
+	, NavigationElement(Other.NavigationElement)
+#if WITH_EDITORONLY_DATA
+PRAGMA_DISABLE_DEPRECATION_WARNINGS
+	, Owner(Other.Owner)
+	, NavInterface(Other.NavInterface)
+PRAGMA_ENABLE_DEPRECATION_WARNINGS
+#endif // WITH_EDITORONLY_DATA
+	, PrevBounds(Other.PrevBounds)
+	, FlagsOverride(Other.FlagsOverride)
+	, PrevFlags(Other.PrevFlags)
+	, bHasPrevData(Other.bHasPrevData)
+	, bInvalidRequest(Other.bInvalidRequest)
+	, bIsFromVisibilityChange(Other.bIsFromVisibilityChange)
+	, bIsInBaseNavmesh(Other.bIsInBaseNavmesh)
+{
+}
+
+FNavigationDirtyElement& FNavigationDirtyElement::operator=(const FNavigationDirtyElement& Other)
+{
+	ExplicitAreasToDirty = Other.ExplicitAreasToDirty;
+	NavigationElement = Other.NavigationElement;
+#if WITH_EDITORONLY_DATA
+PRAGMA_DISABLE_DEPRECATION_WARNINGS
+	Owner = Other.Owner;
+	NavInterface = Other.NavInterface;
+PRAGMA_ENABLE_DEPRECATION_WARNINGS
+#endif // WITH_EDITORONLY_DATA
+	PrevBounds = Other.PrevBounds;
+	FlagsOverride = Other.FlagsOverride;
+	PrevFlags = Other.PrevFlags;
+	bHasPrevData = Other.bHasPrevData;
+	bInvalidRequest = Other.bInvalidRequest;
+	bIsFromVisibilityChange = Other.bIsFromVisibilityChange;
+	bIsInBaseNavmesh = Other.bIsInBaseNavmesh;
+	return *this;
+}
+
 uint32 GetTypeHash(const FNavigationDirtyElement& Info)
 {
 	return GetTypeHash(Info.NavigationElement.Get());
diff --git a/Engine/Source/Runtime/Engine/Private/AI/Navigation/NavigationTypes.cpp b/Engine/Source/Runtime/Engine/Private/AI/Navigation/NavigationTypes.cpp
--- a/Engine/Source/Runtime/Engine/Private/AI/Navigation/NavigationTypes.cpp
+++ b/Engine/Source/Runtime/Engine/Private/AI/Navigation/NavigationTypes.cpp
@@ -426,6 +426,100 @@ FNavigationRelevantData::FNavigationRelevantData(UObject& Source)
 {
 }
 
+FNavigationRelevantData::FNavigationRelevantData(const FNavigationRelevantData& Other)
+	: TSharedFromThis<FNavigationRelevantData, ESPMode::ThreadSafe>(Other)
+	, CollisionData(Other.CollisionData)
+	, VoxelData(Other.VoxelData)
+	, Bounds(Other.Bounds)
+	, NavDataPerInstanceTransformDelegate(Other.NavDataPerInstanceTransformDelegate)
+	, ShouldUseGeometryDelegate(Other.ShouldUseGeometryDelegate)
+	, Modifiers(Other.Modifiers)
+#if WITH_EDITORONLY_DATA
+PRAGMA_DISABLE_DEPRECATION_WARNINGS
+	, SourceObject(Other.SourceObject)
+PRAGMA_ENABLE_DEPRECATION_WARNINGS
+#endif // WITH_EDITORONLY_DATA
+	, SourceElement(Other.SourceElement)
+	, bPendingLazyGeometryGathering(Other.bPendingLazyGeometryGathering)
+	, bPendingLazyModifiersGathering(Other.bPendingLazyModifiersGathering)
+	, bPendingChildLazyModifiersGathering(Other.bPendingChildLazyModifiersGathering)
+	, bSupportsGatheringGeometrySlices(Other.bSupportsGatheringGeometrySlices)
+	, bShouldSkipDirtyAreaOnAddOrRemove(Other.bShouldSkipDirtyAreaOnAddOrRemove)
+	, bLoadedData(Other.bLoadedData)
+{
+}
+
+FNavigationRelevantData::FNavigationRelevantData(FNavigationRelevantData&& Other)
+	: TSharedFromThis<FNavigationRelevantData, ESPMode::ThreadSafe>(MoveTemp(Other))
+	, CollisionData(MoveTemp(Other.CollisionData))
+	, VoxelData(MoveTemp(Other.VoxelData))
+	, Bounds(MoveTemp(Other.Bounds))
+	, NavDataPerInstanceTransformDelegate(MoveTemp(Other.NavDataPerInstanceTransformDelegate))
+	, ShouldUseGeometryDelegate(MoveTemp(Other.ShouldUseGeometryDelegate))
+	, Modifiers(MoveTemp(Other.Modifiers))
+#if WITH_EDITORONLY_DATA
+PRAGMA_DISABLE_DEPRECATION_WARNINGS
+	, SourceObject(MoveTemp(Other.SourceObject))
+PRAGMA_ENABLE_DEPRECATION_WARNINGS
+#endif // WITH_EDITORONLY_DATA
+	, SourceElement(MoveTemp(Other.SourceElement))
+	, bPendingLazyGeometryGathering(Other.bPendingLazyGeometryGathering)
+	, bPendingLazyModifiersGathering(Other.bPendingLazyModifiersGathering)
+	, bPendingChildLazyModifiersGathering(Other.bPendingChildLazyModifiersGathering)
+	, bSupportsGatheringGeometrySlices(Other.bSupportsGatheringGeometrySlices)
+	, bShouldSkipDirtyAreaOnAddOrRemove(Other.bShouldSkipDirtyAreaOnAddOrRemove)
+	, bLoadedData(Other.bLoadedData)
+{
+}
+
+FNavigationRelevantData& FNavigationRelevantData::operator=(FNavigationRelevantData&& Other)
+{
+	TSharedFromThis<FNavigationRelevantData, ESPMode::ThreadSafe>::operator =(MoveTemp(Other));
+	CollisionData = MoveTemp(Other.CollisionData);
+	VoxelData = MoveTemp(Other.VoxelData);
+	Bounds = MoveTemp(Other.Bounds);
+	NavDataPerInstanceTransformDelegate = MoveTemp(Other.NavDataPerInstanceTransformDelegate);
+	ShouldUseGeometryDelegate = MoveTemp(Other.ShouldUseGeometryDelegate);
+	Modifiers = MoveTemp(Other.Modifiers);
+#if WITH_EDITORONLY_DATA
+PRAGMA_DISABLE_DEPRECATION_WARNINGS
+	SourceObject = MoveTemp(Other.SourceObject);
+PRAGMA_ENABLE_DEPRECATION_WARNINGS
+#endif // WITH_EDITORONLY_DATA
+	SourceElement = MoveTemp(Other.SourceElement);
+	bPendingLazyGeometryGathering = Other.bPendingLazyGeometryGathering;
+	bPendingLazyModifiersGathering = Other.bPendingLazyModifiersGathering;
+	bPendingChildLazyModifiersGathering = Other.bPendingChildLazyModifiersGathering;
+	bSupportsGatheringGeometrySlices = Other.bSupportsGatheringGeometrySlices;
+	bShouldSkipDirtyAreaOnAddOrRemove = Other.bShouldSkipDirtyAreaOnAddOrRemove;
+	bLoadedData = Other.bLoadedData;
+	return *this;
+}
+
+FNavigationRelevantData& FNavigationRelevantData::operator=(const FNavigationRelevantData& Other)
+{
+	TSharedFromThis<FNavigationRelevantData, ESPMode::ThreadSafe>::operator =(Other);
+	CollisionData = Other.CollisionData;
+	VoxelData = Other.VoxelData;
+	Bounds = Other.Bounds;
+	NavDataPerInstanceTransformDelegate = Other.NavDataPerInstanceTransformDelegate;
+	ShouldUseGeometryDelegate = Other.ShouldUseGeometryDelegate;
+	Modifiers = Other.Modifiers;
+#if WITH_EDITORONLY_DATA
+PRAGMA_DISABLE_DEPRECATION_WARNINGS
+	SourceObject = Other.SourceObject;
+PRAGMA_ENABLE_DEPRECATION_WARNINGS
+#endif // WITH_EDITORONLY_DATA
+	SourceElement = Other.SourceElement;
+	bPendingLazyGeometryGathering = Other.bPendingLazyGeometryGathering;
+	bPendingLazyModifiersGathering = Other.bPendingLazyModifiersGathering;
+	bPendingChildLazyModifiersGathering = Other.bPendingChildLazyModifiersGathering;
+	bSupportsGatheringGeometrySlices = Other.bSupportsGatheringGeometrySlices;
+	bShouldSkipDirtyAreaOnAddOrRemove = Other.bShouldSkipDirtyAreaOnAddOrRemove;
+	bLoadedData = Other.bLoadedData;
+	return *this;
+}
+
 // Deprecated
 const UObject* FNavigationRelevantData::GetOwner() const
 {
