Adds a Blueprint exposed function on mesh component to set a FLinearColor parameter on a material.

- Called SetColorParameterValueOnMaterials
- I opted to not change the existing function (which only sets a vec3), since the automatic BP fixup code didn't do the correct thing on existing assets (it disconnected the input and gave a BP compile error).

#tests Created a BP actor w/ a construction script which calls the SetColorParameterValueOnMaterials

diff --git a/Engine/Source/Runtime/Engine/Classes/Components/MeshComponent.h b/Engine/Source/Runtime/Engine/Classes/Components/MeshComponent.h
--- a/Engine/Source/Runtime/Engine/Classes/Components/MeshComponent.h
+++ b/Engine/Source/Runtime/Engine/Classes/Components/MeshComponent.h
@@ -175,6 +175,10 @@ public:
 	UFUNCTION(BlueprintCallable, Category = "Rendering|Material")
 	ENGINE_API void SetVectorParameterValueOnMaterials(const FName ParameterName, const FVector ParameterValue);
 
+	/** Set all occurrences of Vector Material Parameters with ParameterName in the set of materials of the SkeletalMesh to ParameterValue */
+	UFUNCTION(BlueprintCallable, Category = "Rendering|Material")
+	ENGINE_API void SetColorParameterValueOnMaterials(const FName ParameterName, const FLinearColor ParameterValue);
+
 	/**  
 	 * Returns default value for the parameter input. 
 	 *
diff --git a/Engine/Source/Runtime/Engine/Private/Components/MeshComponent.cpp b/Engine/Source/Runtime/Engine/Private/Components/MeshComponent.cpp
--- a/Engine/Source/Runtime/Engine/Private/Components/MeshComponent.cpp
+++ b/Engine/Source/Runtime/Engine/Private/Components/MeshComponent.cpp
@@ -448,6 +448,11 @@ void UMeshComponent::SetScalarParameterValueOnMaterials(const FName ParameterNam
 }
 
 void UMeshComponent::SetVectorParameterValueOnMaterials(const FName ParameterName, const FVector ParameterValue)
+{
+	SetColorParameterValueOnMaterials(ParameterName, FLinearColor(ParameterValue));
+}
+
+void UMeshComponent::SetColorParameterValueOnMaterials(const FName ParameterName, const FLinearColor ParameterValue)
 {
 	if (!bEnableMaterialParameterCaching)
 	{
