Fixing legacy conversion issue

diff --git a/Engine/Source/Runtime/Engine/Private/Materials/MaterialExpressions.cpp b/Engine/Source/Runtime/Engine/Private/Materials/MaterialExpressions.cpp
--- a/Engine/Source/Runtime/Engine/Private/Materials/MaterialExpressions.cpp
+++ b/Engine/Source/Runtime/Engine/Private/Materials/MaterialExpressions.cpp
@@ -18449,36 +18449,39 @@ void FMaterialLayersFunctionsRuntimeData::PostSerialize(const FArchive& Ar)
 #if ENABLE_MATERIAL_LAYER_PROTOTYPE
 	if (Ar.IsLoading())
 	{	
-		// When loading, if a legacy PRE substrate tree layer version:
-		// The tree need to be recreated from the list
-		// of layers and blends assuming this is a flat hierarchy of layers:
-		// For PRE substrate tree layer version:
-		// TODO: We could use:
-		// 	if (Ar.CustomVer(FRenderingObjectVersion::GUID) < FRenderingObjectVersion::AddedMaterialLayersSubstrateSupport)
-		if ((Tree.Nodes.Num() != Layers.Num()) || (Layers.Num() > Blends.Num()))
+		if (Substrate::IsSubstrateEnabled())
 		{
-			int NumLegacyLayers = Layers.Num();
-			int NumLegacyBlends = Blends.Num();
+			// When loading, if a legacy PRE substrate tree layer version:
+			// The tree need to be recreated from the list
+			// of layers and blends assuming this is a flat hierarchy of layers:
+			// For PRE substrate tree layer version:
+			// TODO: We could use:
+			// 	if (Ar.CustomVer(FRenderingObjectVersion::GUID) < FRenderingObjectVersion::AddedMaterialLayersSubstrateSupport)
+			if ((Tree.Nodes.Num() != Layers.Num()) || (Layers.Num() > Blends.Num()))
+			{
+				int NumLegacyLayers = Layers.Num();
+				int NumLegacyBlends = Blends.Num();
 
-			Tree.Empty();
+				Tree.Empty();
 
-			// First add an extra empty Blend to match the number of layers
-			while (Blends.Num() < Layers.Num())
-			{
-				Blends.AddDefaulted(); // normally we should go through this only once
-			}
+				// First add an extra empty Blend to match the number of layers
+				while (Blends.Num() < Layers.Num())
+				{
+					Blends.AddDefaulted(); // normally we should go through this only once
+				}
 
-			Layers.AddDefaulted();
-			Blends.AddDefaulted();
+				Layers.AddDefaulted();
+				Blends.AddDefaulted();
 
-			// Add a layer node at root 
-			FLayerNodeId LayerNodeId = Tree.AddNode({ NumLegacyLayers,  NumLegacyLayers }, -1);
+				// Add a layer node at root 
+				FLayerNodeId LayerNodeId = Tree.AddNode({ NumLegacyLayers,  NumLegacyLayers }, -1);
 
-			// Third rebuiild the tree of nodes referencing the layers and blends
-			for (int32 l = 0; l < NumLegacyLayers; ++l)
-			{
-				// And then ONE MORE node to be the attributes first child
-				Tree.AddNode({ l,  l == 0 ? NumLegacyLayers - 1 : l-1}, LayerNodeId);
+				// Third rebuiild the tree of nodes referencing the layers and blends
+				for (int32 l = 0; l < NumLegacyLayers; ++l)
+				{
+					// And then ONE MORE node to be the attributes first child
+					Tree.AddNode({ l,  l == 0 ? NumLegacyLayers - 1 : l-1}, LayerNodeId);
+				}
 			}
 		}
 	}
@@ -18547,6 +18550,7 @@ void FMaterialLayersFunctions::SerializeLegacy(FArchive& Ar)
 	FString KeyString_DEPRECATED;
 	Ar << KeyString_DEPRECATED;
 }
+
 #endif // WITH_EDITOR
 
 void FMaterialLayersFunctions::PostSerialize(const FArchive& Ar)
@@ -18575,10 +18579,8 @@ void FMaterialLayersFunctions::PostSerialize(const FArchive& Ar)
 		{
 			EditorOnly.RestrictToBlendRelatives = MoveTemp(RestrictToBlendRelatives_DEPRECATED);
 		}
-		bool bIsDeprecatedLegacyAsset = false;
 		if (LayerGuids_DEPRECATED.Num() > 0)
 		{
-			bIsDeprecatedLegacyAsset = true;
 			EditorOnly.LayerGuids = MoveTemp(LayerGuids_DEPRECATED);
 		}
 		if (LayerLinkStates_DEPRECATED.Num() > 0)
@@ -18590,40 +18592,18 @@ void FMaterialLayersFunctions::PostSerialize(const FArchive& Ar)
 			EditorOnly.DeletedParentLayerGuids = MoveTemp(DeletedParentLayerGuids_DEPRECATED);
 		}
 
-#if ENABLE_MATERIAL_LAYER_PROTOTYPE
-
-		if (bIsDeprecatedLegacyAsset)
-		{
-			const int32 NumLayers = EditorOnly.LayerStates.Num();
-			if (EditorOnly.LayerGuids.Num() != NumLayers ||
-				EditorOnly.LayerLinkStates.Num() != NumLayers)
-			{
-				EditorOnly.LayerGuids.Empty(NumLayers);
-				EditorOnly.LayerLinkStates.Empty(NumLayers);
 
-				if (NumLayers > 0)
-				{
-					EditorOnly.LayerGuids.Add(BackgroundGuid);
-					EditorOnly.LayerLinkStates.Add(EMaterialLayerLinkState::Uninitialized);
-
-					for (int32 i = 1; i < NumLayers; ++i)
-					{
-						// Need to allocate deterministic guids for layers loaded from old data
-						// These guids will be saved into any child material layers that have this material as their parent,
-						// But it's possible *this* material may not actually be saved in that case
-						// If that happens, need to ensure that the guids remain consistent if this material is loaded again;
-						// otherwise they will no longer match the guids that were saved into the child material
-						EditorOnly.LayerGuids.Add(FGuid(3u, 0u, 0u, i));
-						EditorOnly.LayerLinkStates.Add(EMaterialLayerLinkState::Uninitialized);
-					}
-				}
-			}
-		}
+#if ENABLE_MATERIAL_LAYER_PROTOTYPE
+		const int32 NumLayers = EditorOnly.LayerStates.Num();
 #else
-		if (EditorOnly.LayerGuids.Num() != Layers.Num() ||
-			EditorOnly.LayerLinkStates.Num() != Layers.Num())
+		const int32 NumLayers = Layers.Num();
+		// In legacy version we use Layers.Num, but in the newer version above we compare against EdtorOnly.LayerStates.Num because
+		// Layers.Num is modified in the RuntimeData legacy conversion and incorrect for the upcoming test.
+		// then further in the call CheckAndRepairPostSerializeEditorOnlyDataForRuntimeData() the editor only fileds are converted to the new tree of layers datastructure.
+#endif // ENABLE_MATERIAL_LAYER_PROTOTYPE
+		if (EditorOnly.LayerGuids.Num() != NumLayers ||
+			EditorOnly.LayerLinkStates.Num() != NumLayers)
 		{
-			const int32 NumLayers = Layers.Num();
 			EditorOnly.LayerGuids.Empty(NumLayers);
 			EditorOnly.LayerLinkStates.Empty(NumLayers);
 
@@ -18644,32 +18624,11 @@ void FMaterialLayersFunctions::PostSerialize(const FArchive& Ar)
 				}
 			}
 		}
-#endif
 
 #if ENABLE_MATERIAL_LAYER_PROTOTYPE
-		// After the deprecated content as moved to the EditorOnly block, let's update it to be in sync with the LayersTree
-		{
-			// This method is called right after unserialization
-			// Use this fact to fix editor only data in case of a legacy version of the LayersFunctions
-			// Catch case when a legacy layer has been loaded
-			// EditorOnly data need to be updated
-			if ((Layers.Num() == Blends.Num()) && ((EditorOnly.LayerNames.Num() + 1) == Layers.Num()))
-			{
-				int32 NumLegacyLayers = EditorOnly.LayerNames.Num();
-				// Second add an extra Layer and Blend for each legacy layer.
-				{
-					EditorOnly.LayerStates.Add(true);
-					FText LayerName = FText::FromString(TEXT("Legacy Layers"));
-					EditorOnly.LayerNames.Add(LayerName);
-					EditorOnly.RestrictToLayerRelatives.Add(false);
-					EditorOnly.RestrictToBlendRelatives.Add(false);
-					EditorOnly.LayerGuids.Add(FGuid::NewGuid());
-					EditorOnly.LayerLinkStates.Add(EMaterialLayerLinkState::Uninitialized);
-				}
-			}
-
-		}
+		CheckAndRepairPostSerializeEditorOnlyDataForRuntimeData(GetRuntime(), EditorOnly);
 #endif // ENABLE_MATERIAL_LAYER_PROTOTYPE
+
 	} // Ar.IsLoading()
 
 #endif // WITH_EDITORONLY_DATA	
@@ -18678,6 +18637,34 @@ void FMaterialLayersFunctions::PostSerialize(const FArchive& Ar)
 
 #if WITH_EDITOR
 
+void FMaterialLayersFunctions::CheckAndRepairPostSerializeEditorOnlyDataForRuntimeData(FMaterialLayersFunctionsRuntimeData& Runtime, FMaterialLayersFunctionsEditorOnlyData& EditorOnly)
+{
+#if ENABLE_MATERIAL_LAYER_PROTOTYPE
+	if (Substrate::IsSubstrateEnabled())
+	{
+		// This method is called right after unserialization
+		// Use this fact to fix editor only data in case of a legacy version of the LayersFunctions
+
+		// Catch case when a legacy layer has been loaded
+		// EditorOnly data need to be updated
+		if ((Runtime.Layers.Num() == Runtime.Blends.Num()) && ((EditorOnly.LayerNames.Num() + 1) == Runtime.Layers.Num()))
+		{
+			int32 NumLegacyLayers = EditorOnly.LayerNames.Num();
+			// Second add an extra Layer and Blend for each legacy layer.
+			{
+				EditorOnly.LayerStates.Add(true);
+				FText LayerName = FText::FromString(TEXT("Legacy Layers"));
+				EditorOnly.LayerNames.Add(LayerName);
+				EditorOnly.RestrictToLayerRelatives.Add(false);
+				EditorOnly.RestrictToBlendRelatives.Add(false);
+				EditorOnly.LayerGuids.Add(FGuid::NewGuid());
+				EditorOnly.LayerLinkStates.Add(EMaterialLayerLinkState::Uninitialized);
+			}
+		}
+	}
+#endif
+}
+
 void FMaterialLayersFunctions::AddDefaultBackgroundLayer()
 {
 #if ENABLE_MATERIAL_LAYER_PROTOTYPE
@@ -18995,29 +18982,10 @@ bool FMaterialLayersFunctions::ResolveParent(const FMaterialLayersFunctionsRunti
 	FMaterialLayersFunctionsEditorOnlyData& EditorOnly,
 	TArray<int32>& OutRemapLayerIndices)
 {
-/*
-#if ENABLE_MATERIAL_LAYER_PROTOTYPE
-	// This method is called right after unserialization
-	// Use this fact to fix editor only data in case of a legacy version of the LayersFunctions
+	// For some legacy materials using FMaterialLayersFunctions, this function is called right after unserialization of FMaterialLayersFunctionsRuntimeData
+	// and the PostLoad of FMaterialLayersFunctions is not called, so we need to make sure this is called and checked once
+	CheckAndRepairPostSerializeEditorOnlyDataForRuntimeData(Runtime, EditorOnly);
 
-		// Catch case when a legacy layer has been loaded
-		// EditorOnly data need to be updated
-	if ((Runtime.Layers.Num() == Runtime.Blends.Num()) && ((EditorOnly.LayerNames.Num() + 1) == Runtime.Layers.Num()))
-	{
-		int32 NumLegacyLayers = EditorOnly.LayerNames.Num();
-		// Second add an extra Layer and Blend for each legacy layer.
-		{
-			EditorOnly.LayerStates.Add(true);
-			FText LayerName = FText::FromString(TEXT("Legacy Layers"));
-			EditorOnly.LayerNames.Add(LayerName);
-			EditorOnly.RestrictToLayerRelatives.Add(false);
-			EditorOnly.RestrictToBlendRelatives.Add(false);
-			EditorOnly.LayerGuids.Add(FGuid::NewGuid());
-			EditorOnly.LayerLinkStates.Add(EMaterialLayerLinkState::Uninitialized);
-		}
-	}
-#endif
-*/
  	check(EditorOnly.LayerGuids.Num() == Runtime.Layers.Num());
 	check(EditorOnly.LayerLinkStates.Num() == Runtime.Layers.Num());
 
diff --git a/Engine/Source/Runtime/Engine/Public/Materials/MaterialLayersFunctions.h b/Engine/Source/Runtime/Engine/Public/Materials/MaterialLayersFunctions.h
--- a/Engine/Source/Runtime/Engine/Public/Materials/MaterialLayersFunctions.h
+++ b/Engine/Source/Runtime/Engine/Public/Materials/MaterialLayersFunctions.h
@@ -571,6 +571,7 @@ struct FMaterialLayersFunctions : public FMaterialLayersFunctionsRuntimeData
 		return FMaterialLayersFunctions::ResolveParent(Parent.GetRuntime(), Parent.EditorOnly, GetRuntime(), EditorOnly, OutRemapLayerIndices);
 	}
 
+
 	static ENGINE_API void Validate(const FMaterialLayersFunctionsRuntimeData& Runtime, const FMaterialLayersFunctionsEditorOnlyData& EditorOnly);
 
 	void Validate()
@@ -579,7 +580,10 @@ struct FMaterialLayersFunctions : public FMaterialLayersFunctionsRuntimeData
 	}
 
 	ENGINE_API void SerializeLegacy(FArchive& Ar);
-	#endif // WITH_EDITOR
+
+	static ENGINE_API void CheckAndRepairPostSerializeEditorOnlyDataForRuntimeData(FMaterialLayersFunctionsRuntimeData& Runtime, FMaterialLayersFunctionsEditorOnlyData& EditorOnly);
+
+#endif // WITH_EDITOR
 
 	ENGINE_API void PostSerialize(const FArchive& Ar);
 
