Added STAT_RayTracingSkinnedBuildPrimitives to track number of triangles built for ray tracing skinned geometry.
- similar to the existing STAT_RayTracingSkinnedUpdatePrimitives

diff --git a/Engine/Source/Runtime/Engine/Private/RayTracingSkinnedGeometry.cpp b/Engine/Source/Runtime/Engine/Private/RayTracingSkinnedGeometry.cpp
--- a/Engine/Source/Runtime/Engine/Private/RayTracingSkinnedGeometry.cpp
+++ b/Engine/Source/Runtime/Engine/Private/RayTracingSkinnedGeometry.cpp
@@ -10,6 +10,7 @@
 DECLARE_GPU_STAT(SkinnedGeometryBuildBLAS);
 DECLARE_GPU_STAT(SkinnedGeometryUpdateBLAS);
 
+DECLARE_DWORD_COUNTER_STAT(TEXT("Ray tracing skinned build primitives"), STAT_RayTracingSkinnedBuildPrimitives, STATGROUP_SceneRendering);
 DECLARE_DWORD_COUNTER_STAT(TEXT("Ray tracing skinned update primitives"), STAT_RayTracingSkinnedUpdatePrimitives, STATGROUP_SceneRendering);
 
 static TAutoConsoleVariable<int32> CVarSkinCacheRayTracingMaxUpdatePrimitivesPerFrame(
@@ -207,6 +208,8 @@ void FRayTracingSkinnedGeometryUpdateQueue::Commit(FRDGBuilder& GraphBuilder, ER
 
 	uint32 BLASScratchSize = 0;
 
+	int32 NumBuiltPrimitives = 0;
+
 	{
 		FScopeLock Lock(&CS);
 
@@ -232,6 +235,8 @@ void FRayTracingSkinnedGeometryUpdateQueue::Commit(FRDGBuilder& GraphBuilder, ER
 				BLASScratchSize += UpdateInfo.ScratchSize;
 
 				RayTracingGeometry->LastUpdatedFrame = GFrameCounterRenderThread;
+
+				NumBuiltPrimitives += RayTracingGeometry->Initializer.TotalPrimitiveCount;
 			}
 			else
 			{
@@ -297,6 +302,7 @@ void FRayTracingSkinnedGeometryUpdateQueue::Commit(FRDGBuilder& GraphBuilder, ER
 		}
 	}
 
+	INC_DWORD_STAT_BY(STAT_RayTracingSkinnedBuildPrimitives, NumBuiltPrimitives);
 	INC_DWORD_STAT_BY(STAT_RayTracingSkinnedUpdatePrimitives, NumUpdatedPrimitives);
 
 	FRDGBufferRef SharedScratchBuffer = nullptr;
