UWorldPartitionRuntimeHashSet::RegisterWorldAssetStreaming() - Refactored registration to properly support injection to HLOD grids

diff --git a/Engine/Source/Editor/UnrealEd/Private/WorldPartition/WorldPartitionHLODsBuilder.cpp b/Engine/Source/Editor/UnrealEd/Private/WorldPartition/WorldPartitionHLODsBuilder.cpp
--- a/Engine/Source/Editor/UnrealEd/Private/WorldPartition/WorldPartitionHLODsBuilder.cpp
+++ b/Engine/Source/Editor/UnrealEd/Private/WorldPartition/WorldPartitionHLODsBuilder.cpp
@@ -393,42 +393,45 @@ bool UWorldPartitionHLODsBuilder::BuildHLODActors()
 			return false;
 		}
 
-		UE_LOG(LogWorldPartitionHLODsBuilder, Display, TEXT("#### Building %d HLOD actors ####"), HLODActorsToBuild.Num());
-		if (bResumeBuild)
+		if (0)
 		{
-			UE_LOG(LogWorldPartitionHLODsBuilder, Display, TEXT("#### Resuming build at %d ####"), ResumeBuildIndex);
-		}
-
-		for (int32 CurrentActor = ResumeBuildIndex; CurrentActor < HLODActorsToBuild.Num(); ++CurrentActor)
-		{
-			TRACE_BOOKMARK(TEXT("BuildHLOD Start - %d"), CurrentActor);
+			UE_LOG(LogWorldPartitionHLODsBuilder, Display, TEXT("#### Building %d HLOD actors ####"), HLODActorsToBuild.Num());
+			if (bResumeBuild)
+			{
+				UE_LOG(LogWorldPartitionHLODsBuilder, Display, TEXT("#### Resuming build at %d ####"), ResumeBuildIndex);
+			}
 
+			for (int32 CurrentActor = ResumeBuildIndex; CurrentActor < HLODActorsToBuild.Num(); ++CurrentActor)
 			{
-				const FGuid& HLODActorGuid = HLODActorsToBuild[CurrentActor];
+				TRACE_BOOKMARK(TEXT("BuildHLOD Start - %d"), CurrentActor);
 
-				FWorldPartitionReference ActorRef(WorldPartition, HLODActorGuid);
+				{
+					const FGuid& HLODActorGuid = HLODActorsToBuild[CurrentActor];
 
-				AWorldPartitionHLOD* HLODActor = CastChecked<AWorldPartitionHLOD>(ActorRef.GetActor());
+					FWorldPartitionReference ActorRef(WorldPartition, HLODActorGuid);
 
-				UE_LOG(LogWorldPartitionHLODsBuilder, Display, TEXT("[%d / %d] Building HLOD actor %s..."), CurrentActor + 1, HLODActorsToBuild.Num(), *HLODActor->GetActorLabel());
+					AWorldPartitionHLOD* HLODActor = CastChecked<AWorldPartitionHLOD>(ActorRef.GetActor());
 
-				// Simulate an engine tick to make sure engine & render resources that are queued for deletion are processed.
-				FWorldPartitionHelpers::FakeEngineTick(World);
+					UE_LOG(LogWorldPartitionHLODsBuilder, Display, TEXT("[%d / %d] Building HLOD actor %s..."), CurrentActor + 1, HLODActorsToBuild.Num(), *HLODActor->GetActorLabel());
 
-				HLODActor->BuildHLOD(bForceBuild);
+					// Simulate an engine tick to make sure engine & render resources that are queued for deletion are processed.
+					FWorldPartitionHelpers::FakeEngineTick(World);
 
-				bool bSaved = SaveHLODActor(HLODActor);
-				if (!bSaved)
-				{
-					return false;
+					HLODActor->BuildHLOD(bForceBuild);
+
+					bool bSaved = SaveHLODActor(HLODActor);
+					if (!bSaved)
+					{
+						return false;
+					}
 				}
-			}
 
-			TRACE_BOOKMARK(TEXT("BuildHLOD End - %d"), CurrentActor);
+				TRACE_BOOKMARK(TEXT("BuildHLOD End - %d"), CurrentActor);
 
-			if (FWorldPartitionHelpers::ShouldCollectGarbage())
-			{
-				FWorldPartitionHelpers::DoCollectGarbage();
+				if (FWorldPartitionHelpers::ShouldCollectGarbage())
+				{
+					FWorldPartitionHelpers::DoCollectGarbage();
+				}
 			}
 		}
 
@@ -726,13 +729,13 @@ bool UWorldPartitionHLODsBuilder::ValidateWorkload(const TArray<FGuid>& Workload
 		const FWorldPartitionActorDescInstance* ActorDescInstance = WorldPartition->GetActorDescInstance(HLODActorGuid);
 		if(!ActorDescInstance)
 		{
-			UE_LOG(LogWorldPartitionHLODsBuilder, Error, TEXT("Unknown actor guid found, your HLOD actors are probably out of date. Run with -SetupHLODs to fix this. Exiting..."));
+			UE_LOG(LogWorldPartitionHLODsBuilder, Error, TEXT("Unknown actor guid found (\"%s\"), your HLOD actors are probably out of date. Run with -SetupHLODs to fix this. Exiting..."), *HLODActorGuid.ToString());
 			return false;
 		}
 
 		if (!ActorDescInstance->GetActorNativeClass()->IsChildOf<AWorldPartitionHLOD>())
 		{
-			UE_LOG(LogWorldPartitionHLODsBuilder, Error, TEXT("Unexpected actor guid found in HLOD workload, exiting..."));
+			UE_LOG(LogWorldPartitionHLODsBuilder, Error, TEXT("Unexpected actor guid found in HLOD workload (\"%s\"), exiting..."), *HLODActorGuid.ToString());
 			return false;
 		}
 
@@ -745,7 +748,7 @@ bool UWorldPartitionHLODsBuilder::ValidateWorkload(const TArray<FGuid>& Workload
 			{
 				if (!ProcessedHLOD.Contains(ChildHLODActorGuid))
 				{
-					UE_LOG(LogWorldPartitionHLODsBuilder, Error, TEXT("Child HLOD actor missing or out of order in HLOD workload, exiting..."));
+					UE_LOG(LogWorldPartitionHLODsBuilder, Error, TEXT("Child HLOD actor (\"%s\") missing or out of order in HLOD workload, exiting..."), *HLODActorGuid.ToString());
 					return false;
 				}
 			}
diff --git a/Engine/Source/Runtime/Engine/Private/WorldPartition/RuntimeHashSet/WorldPartitionRuntimeHashSet.cpp b/Engine/Source/Runtime/Engine/Private/WorldPartition/RuntimeHashSet/WorldPartitionRuntimeHashSet.cpp
--- a/Engine/Source/Runtime/Engine/Private/WorldPartition/RuntimeHashSet/WorldPartitionRuntimeHashSet.cpp
+++ b/Engine/Source/Runtime/Engine/Private/WorldPartition/RuntimeHashSet/WorldPartitionRuntimeHashSet.cpp
@@ -249,11 +249,11 @@ void UWorldPartitionRuntimeHashSet::FlushStreamingContent()
 
 bool UWorldPartitionRuntimeHashSet::IsValidGrid(FName GridName, const UClass* ActorClass) const
 {
-	return IsValidGridInternal(GridName, ActorClass);
+	return ResolveRuntimePartition(GridName) != nullptr;
 }
 #endif
 
-bool UWorldPartitionRuntimeHashSet::IsValidGridInternal(FName GridName, const UClass* ActorClass) const
+const URuntimePartition* UWorldPartitionRuntimeHashSet::ResolveRuntimePartition(FName GridName, bool bMainPartitionLayer) const
 {
 	TArray<FName> MainPartitionTokens;
 	TArray<FName> HLODPartitionTokens;
@@ -264,81 +264,101 @@ bool UWorldPartitionRuntimeHashSet::IsValidGridInternal(FName GridName, const UC
 		// The None grid name will always map to the first runtime partition in the list
 		if (MainPartitionTokens[0].IsNone())
 		{
-			return true;
+			if (RuntimePartitions.IsValidIndex(0))
+			{
+				return RuntimePartitions[0].MainLayer;
+			}
+			else
+			{
+				return nullptr;
+			}
 		}
 
-		for (const FRuntimePartitionDesc& RuntimePartitionDesc : RuntimePartitions)
+		// Make sure the runtime partition is valid
+		const FRuntimePartitionDesc* FoundRuntimePartitionDesc = RuntimePartitions.FindByPredicate([&MainPartitionTokens](const FRuntimePartitionDesc& RuntimePartitionDesc)
+		{
+			return RuntimePartitionDesc.Name == MainPartitionTokens[0];
+		});
+		
+		if (!FoundRuntimePartitionDesc || !FoundRuntimePartitionDesc->MainLayer || !FoundRuntimePartitionDesc->MainLayer->IsValidPartitionTokens(MainPartitionTokens))
+		{
+			return nullptr;
+		}
+	
+		// If an HLOD partition token was specified ("MainPartition:HLODPartition"), make sure it's valid too
+		if (!bMainPartitionLayer && HLODPartitionTokens.IsValidIndex(0))
 		{
-			if (RuntimePartitionDesc.Name == MainPartitionTokens[0])
+			const FRuntimePartitionHLODSetup* FoundRuntimePartitionHLODSetup = FoundRuntimePartitionDesc->HLODSetups.FindByPredicate([&HLODPartitionTokens](const FRuntimePartitionHLODSetup& RuntimePartitionHLODSetup)
 			{
-				if (!RuntimePartitionDesc.MainLayer)
-				{
-					return false;
-				}
-
-				if (!RuntimePartitionDesc.MainLayer->IsValidPartitionTokens(MainPartitionTokens))
-				{
-					return false;
-				}
+				return RuntimePartitionHLODSetup.Name == HLODPartitionTokens[0];
+			});
 
-				return true;
+			if (!FoundRuntimePartitionHLODSetup || !FoundRuntimePartitionHLODSetup->PartitionLayer || !FoundRuntimePartitionHLODSetup->PartitionLayer->IsValidPartitionTokens(HLODPartitionTokens))
+			{
+				return nullptr;
 			}
+
+			return FoundRuntimePartitionHLODSetup->PartitionLayer;
 		}
-	}
 
-	return false;
+		return FoundRuntimePartitionDesc->MainLayer;
+	}			
+
+	return nullptr;
 }
 
 #if WITH_EDITOR
 bool UWorldPartitionRuntimeHashSet::IsValidHLODLayer(FName GridName, const FSoftObjectPath& HLODLayerPath) const
 {
-	if (!RuntimePartitions.Num())
-	{
-		return false;
-	}
-
-	if (const UHLODLayer* HLODLayer = Cast<UHLODLayer>(HLODLayerPath.ResolveObject()))
-	{
-		// The None grid name will always map to the first runtime partition in the list
-		int32 RuntimePartitionIndex = GridName.IsNone() ? 0 : INDEX_NONE;
-		
-		if (RuntimePartitionIndex == INDEX_NONE)
-		{
-			TArray<FName> PartitionTokens;
-			TArray<FName> HLODPartitionTokens;
-
-			// Parse the potentially dot separated grid name to identiy the associated runtime partition
-			if (ParseGridName(GridName, PartitionTokens, HLODPartitionTokens))
-			{
-				int32 RuntimePartitionIndexLookup = 0;
-				for (const FRuntimePartitionDesc& RuntimePartitionDesc : RuntimePartitions)
-				{
-					if (RuntimePartitionDesc.Name == PartitionTokens[0])
-					{
-						RuntimePartitionIndex = RuntimePartitionIndexLookup;
-						break;
-					}
-
-					RuntimePartitionIndexLookup++;
-				}
-			}
-		}
+	return ResolveRuntimePartitionForHLODLayer(GridName, HLODLayerPath) != nullptr;
+}
 
-		if (RuntimePartitionIndex == INDEX_NONE)
-		{
-			return false;
-		}
+const URuntimePartition* UWorldPartitionRuntimeHashSet::ResolveRuntimePartitionForHLODLayer(FName GridName, const FSoftObjectPath& HLODLayerPath) const
+{
+	const URuntimePartition* RuntimePartition = nullptr;
 
-		for (const FRuntimePartitionHLODSetup& HLODSetup : RuntimePartitions[RuntimePartitionIndex].HLODSetups)
-		{
-			if (HLODSetup.HLODLayers.Contains(HLODLayer))
-			{
-				return true;
-			}
-		}
+	if (!RuntimePartitions.IsEmpty())
+	{
+	    if (const UHLODLayer* HLODLayer = Cast<UHLODLayer>(HLODLayerPath.ResolveObject()))
+	    {
+		    // The None grid name will always map to the first runtime partition in the list
+		    int32 RuntimePartitionIndex = GridName.IsNone() ? 0 : INDEX_NONE;
+    
+		    if (RuntimePartitionIndex == INDEX_NONE)
+		    {
+			    TArray<FName> PartitionTokens;
+			    TArray<FName> HLODPartitionTokens;
+    
+			    // Parse the potentially dot separated grid name to identiy the associated runtime partition
+			    if (ParseGridName(GridName, PartitionTokens, HLODPartitionTokens))
+			    {
+				    int32 RuntimePartitionIndexLookup = 0;
+				    for (const FRuntimePartitionDesc& RuntimePartitionDesc : RuntimePartitions)
+				    {
+					    if (RuntimePartitionDesc.Name == PartitionTokens[0])
+					    {
+						    RuntimePartitionIndex = RuntimePartitionIndexLookup;
+						    break;
+					    }
+    
+					    RuntimePartitionIndexLookup++;
+				    }
+			    }
+		    }
+    
+		    if (RuntimePartitionIndex != INDEX_NONE)
+		    {
+		        const FRuntimePartitionHLODSetup* FoundRuntimePartitionHLODSetup = RuntimePartitions[RuntimePartitionIndex].HLODSetups.FindByPredicate([&HLODLayer](const FRuntimePartitionHLODSetup& RuntimePartitionHLODSetup)
+		        {
+			        return RuntimePartitionHLODSetup.HLODLayers.Contains(HLODLayer);
+		        });
+
+				RuntimePartition = FoundRuntimePartitionHLODSetup ? FoundRuntimePartitionHLODSetup->PartitionLayer : nullptr;
+    		}
+	    }
 	}
 
-	return false;
+	return RuntimePartition;
 }
 #endif
 
@@ -635,29 +655,47 @@ void UWorldPartitionRuntimeHashSet::ForEachStreamingCellsSources(const TArray<FW
 
 bool UWorldPartitionRuntimeHashSet::SupportsWorldAssetStreaming(const FName& InTargetGrid)
 {
-	return IsValidGridInternal(InTargetGrid, nullptr);
+	return ResolveRuntimePartition(InTargetGrid) != nullptr;
 }
 
 FGuid UWorldPartitionRuntimeHashSet::RegisterWorldAssetStreaming(const UWorldPartition::FRegisterWorldAssetStreamingParams& InParams)
 {
 	if (!InParams.IsValid())
 	{
+		UE_LOG(LogWorldPartition, Error, TEXT("RegisterWorldAssetStreaming: Invalid parameters provided."));
 		return FGuid();
 	}
 
 	if (WorldAssetStreamingObjects.Contains(InParams.Guid))
 	{
+		UE_LOG(LogWorldPartition, Error, TEXT("RegisterWorldAssetStreaming: World asset guid '%s' was already registered."), *InParams.Guid.ToString());
 		return FGuid();
 	}
 
-	if (!IsValidGridInternal(InParams.TargetGrid, nullptr) || (!InParams.TargetGridHLOD.IsNone() && !IsValidGridInternal(InParams.TargetGridHLOD, nullptr)))
+	FName RuntimePartitionName;
+	const URuntimePartition* RuntimePartition = ResolveRuntimePartition(InParams.TargetGrid);
+	if (!RuntimePartition)
 	{
+		UE_LOG(LogWorldPartition, Error, TEXT("RegisterWorldAssetStreaming: Unable to resolve TargetGrid '%s'."), *InParams.TargetGrid.ToString());
 		return FGuid();
 	}
 
+	FName HLODRuntimePartitionName;
+	const URuntimePartition* HLODRuntimePartition = nullptr;
+	if (!InParams.WorldAssetHLOD.IsNull() && !InParams.TargetGridHLOD.IsNone())
+	{
+		HLODRuntimePartition = ResolveRuntimePartition(InParams.TargetGridHLOD);
+		if (!HLODRuntimePartition)
+		{
+			UE_LOG(LogWorldPartition, Error, TEXT("RegisterWorldAssetStreaming: Unable to resolve TargetGridHLOD '%s'."), *InParams.TargetGridHLOD.ToString());
+			return FGuid();
+		}
+	}
+
 	URuntimeHashSetExternalStreamingObject* StreamingObject = CastChecked<URuntimeHashSetExternalStreamingObject>(CreateExternalStreamingObject(URuntimeHashSetExternalStreamingObject::StaticClass(), this, GetTypedOuter<UWorld>()));
 	if (!StreamingObject)
 	{
+		UE_LOG(LogWorldPartition, Error, TEXT("RegisterWorldAssetStreaming: Couldn't create ExternalStreamingObject."));
 		return FGuid();
 	}		
 
@@ -666,51 +704,24 @@ FGuid UWorldPartitionRuntimeHashSet::RegisterWorldAssetStreaming(const UWorldPar
 	// 2 passes : 1st for target grid, 2nd for HLOD target grid
 	for (const bool bIsHLODPass : { false, true })
 	{
-		const FName TargetGrid = bIsHLODPass ? InParams.TargetGridHLOD : InParams.TargetGrid;
-		if (TargetGrid.IsNone())
-		{
-			continue;
-		}
-
-		TArray<FName> MainPartitionTokens;
-		TArray<FName> HLODPartitionTokens;
-		if (!ParseGridName(TargetGrid, MainPartitionTokens, HLODPartitionTokens))
+		const URuntimePartition* TargetPartition = bIsHLODPass ? HLODRuntimePartition : RuntimePartition;
+		if (!TargetPartition)
 		{
 			continue;
 		}
 
-		if (bIsHLODPass && InParams.WorldAssetHLOD.IsNull())
-		{
-			continue;
-		}
+		bool bClientOnlyVisible = TargetPartition->bClientOnlyVisible;
+		bool bBlockOnSlowStreaming = TargetPartition->bBlockOnSlowStreaming;
 
 		FRuntimePartitionStreamingData StreamingData;
-		StreamingData.LoadingRange = -1;
-
-		bool bClientOnlyVisible = false;
-		bool bBlockOnSlowStreaming = false;
-		for (const FRuntimePartitionDesc& RuntimePartitionDesc : RuntimePartitions)
-		{
-			if (RuntimePartitionDesc.MainLayer && ((RuntimePartitionDesc.Name == MainPartitionTokens[0]) || MainPartitionTokens[0].IsNone()))
-			{
-				StreamingData.Name = RuntimePartitionDesc.Name;
-				StreamingData.LoadingRange = RuntimePartitionDesc.MainLayer->LoadingRange;
-				bClientOnlyVisible = RuntimePartitionDesc.MainLayer->bClientOnlyVisible;
-				bBlockOnSlowStreaming = RuntimePartitionDesc.MainLayer->bBlockOnSlowStreaming;
-				break;
-			}
-		}
-
-		if (StreamingData.LoadingRange == -1)
-		{
-			continue;
-		}
+		StreamingData.Name = TargetPartition->Name;
+		StreamingData.LoadingRange = TargetPartition->LoadingRange;
 
 		// Create Cell
-		FName TargetGridName = (bIsHLODPass ? InParams.TargetGridHLOD : InParams.TargetGrid);
+		FName TargetPartitionName = TargetPartition->Name;
 		FGuid InstanceGuid = InParams.Guid;
 		FArchiveMD5 ArMD5;
-		ArMD5 << TargetGridName << InstanceGuid;
+		ArMD5 << TargetPartitionName << InstanceGuid;
 		const FGuid CellGuid = ArMD5.GetGuidFromHash();
 		check(CellGuid.IsValid());
 		if (!bIsHLODPass)
@@ -741,7 +752,7 @@ FGuid UWorldPartitionRuntimeHashSet::RegisterWorldAssetStreaming(const UWorldPar
 			RuntimeCellData->ContentBounds = InParams.Bounds;
 			RuntimeCellData->HierarchicalLevel = MAX_int32;
 			RuntimeCellData->Priority = InParams.Priority;
-			RuntimeCellData->GridName = TargetGrid;
+			RuntimeCellData->GridName = TargetPartition->Name;
 			RuntimeCellData->bIs2D = false;
 
 			if (RuntimeCell->CreateAndSetLevelStreaming(WorldAsset, InParams.Transform))
diff --git a/Engine/Source/Runtime/Engine/Private/WorldPartition/RuntimeHashSet/WorldPartitionRuntimeHashSetHLODGeneration.cpp b/Engine/Source/Runtime/Engine/Private/WorldPartition/RuntimeHashSet/WorldPartitionRuntimeHashSetHLODGeneration.cpp
--- a/Engine/Source/Runtime/Engine/Private/WorldPartition/RuntimeHashSet/WorldPartitionRuntimeHashSetHLODGeneration.cpp
+++ b/Engine/Source/Runtime/Engine/Private/WorldPartition/RuntimeHashSet/WorldPartitionRuntimeHashSetHLODGeneration.cpp
@@ -217,21 +217,24 @@ bool UWorldPartitionRuntimeHashSet::SetupHLODActors(const IStreamingGenerationCo
 				// Fake tick
 				PrivateUtils::GameTick(WorldPartition->GetWorld());
 
-				TArray<FName> MainPartitionTokens;
-				TArray<FName> HLODPartitionTokens;
-				verify(ParseGridName(ActorInstances[0].ActorSetInstance->RuntimeGrid, MainPartitionTokens, HLODPartitionTokens));
-
-				if (MainPartitionTokens[0].IsNone())
-				{
-					MainPartitionTokens[0] = RuntimePartitions[0].Name;
-				}
+				// Resolve main partition
+				const URuntimePartition* MainRuntimePartition = ResolveRuntimePartition(ActorInstances[0].ActorSetInstance->RuntimeGrid, /*bMainPartitionLayer*/ true);
+				check(MainRuntimePartition);
+				
+				// Retrieve the runtime grid to use for this HLOD
+				auto GetHLODRuntimeGrid = [this, MainRuntimePartition](const UHLODLayer* InHLODLayer)
+				{					
+					const URuntimePartition* HLODRuntimePartition = ResolveRuntimePartitionForHLODLayer(MainRuntimePartition->Name, InHLODLayer);
+					check(HLODRuntimePartition);
+					return FName(*FString::Printf(TEXT("%s:%s"), *MainRuntimePartition->Name.ToString(), *HLODRuntimePartition->Name.ToString()));
+				};				
 
 				FHLODCreationParams HLODCreationParams;
 				HLODCreationParams.WorldPartition = WorldPartition;
 				HLODCreationParams.CellName = CellUniqueId.Name;
 				HLODCreationParams.CellGuid = CellUniqueId.Guid;
 				HLODCreationParams.CellBounds = CellBounds;
-				HLODCreationParams.GetRuntimeGrid = [&MainPartitionTokens](const UHLODLayer* InHLODLayer) { return FName(*FString::Printf(TEXT("%s:%s"), *FString::JoinBy(MainPartitionTokens, TEXT("."), [](const FName Token) { return Token.ToString(); }), *InHLODLayer->GetName())); };
+				HLODCreationParams.GetRuntimeGrid = GetHLODRuntimeGrid;
 				HLODCreationParams.HLODLevel = HLODLevel;
 				HLODCreationParams.MinVisibleDistance = RuntimePartition->LoadingRange;
 				HLODCreationParams.ContentBundleGuid = CellDescInstance.ContentBundleID;
diff --git a/Engine/Source/Runtime/Engine/Private/WorldPartition/RuntimeHashSet/WorldPartitionRuntimeHashSetStreamingGeneration.cpp b/Engine/Source/Runtime/Engine/Private/WorldPartition/RuntimeHashSet/WorldPartitionRuntimeHashSetStreamingGeneration.cpp
--- a/Engine/Source/Runtime/Engine/Private/WorldPartition/RuntimeHashSet/WorldPartitionRuntimeHashSetStreamingGeneration.cpp
+++ b/Engine/Source/Runtime/Engine/Private/WorldPartition/RuntimeHashSet/WorldPartitionRuntimeHashSetStreamingGeneration.cpp
@@ -20,53 +20,26 @@ bool UWorldPartitionRuntimeHashSet::GenerateRuntimePartitionsStreamingDescs(cons
 	//
 	// Split actor sets into their corresponding runtime partition implementation
 	//
-	TMap<FName, const FRuntimePartitionDesc*> NameToRuntimePartitionDescMap;
-
-	// Actors with RuntimeGrid set to None will be assigned to the default partition
-	NameToRuntimePartitionDescMap.Add(NAME_None, &RuntimePartitions[0]);
-
-	for (const FRuntimePartitionDesc& RuntimePartitionDesc : RuntimePartitions)
-	{
-		NameToRuntimePartitionDescMap.Add(RuntimePartitionDesc.Name, &RuntimePartitionDesc);
-	}
+	TMap<FName, URuntimePartition*> NameToRuntimePartitionMap;
+		
 
 	TMap<URuntimePartition*, TArray<const IStreamingGenerationContext::FActorSetInstance*>> RuntimePartitionsToActorSetMap;
-	StreamingGenerationContext->ForEachActorSetInstance([this, &NameToRuntimePartitionDescMap, &RuntimePartitionsToActorSetMap](const IStreamingGenerationContext::FActorSetInstance& ActorSetInstance)
+	StreamingGenerationContext->ForEachActorSetInstance([this, &NameToRuntimePartitionMap, &RuntimePartitionsToActorSetMap](const IStreamingGenerationContext::FActorSetInstance& ActorSetInstance)
 	{
-		TArray<FName> MainPartitionTokens;
-		TArray<FName> HLODPartitionTokens;
-		verify(ParseGridName(ActorSetInstance.RuntimeGrid, MainPartitionTokens, HLODPartitionTokens));
+		URuntimePartition* RuntimePartition = nullptr;
+		URuntimePartition** RuntimePartitionPtr = NameToRuntimePartitionMap.Find(ActorSetInstance.RuntimeGrid);
 
-		check(!MainPartitionTokens.IsEmpty());
-		if (const FRuntimePartitionDesc** RuntimePartitionDesc = NameToRuntimePartitionDescMap.Find(MainPartitionTokens[0]))
+		if (RuntimePartitionPtr)
 		{
-			if (!HLODPartitionTokens.IsEmpty())
-			{
-				bool bFoundHLODPartitionLayer = false;
-
-				for (const FRuntimePartitionHLODSetup& HLODSetup : (*RuntimePartitionDesc)->HLODSetups)
-				{
-					for (const UHLODLayer* HLODPartitionLayer : HLODSetup.HLODLayers)
-					{
-						if (HLODPartitionLayer->GetName() == HLODPartitionTokens[0])
-						{
-							RuntimePartitionsToActorSetMap.FindOrAdd(HLODSetup.PartitionLayer).Add(&ActorSetInstance);
-							bFoundHLODPartitionLayer = true;
-							break;
-						}
-					}
-
-					if (bFoundHLODPartitionLayer)
-					{
-						break;
-					}
-				}
-			}
-			else
-			{
-				RuntimePartitionsToActorSetMap.FindOrAdd((*RuntimePartitionDesc)->MainLayer).Add(&ActorSetInstance);
-			}
+			RuntimePartition = *RuntimePartitionPtr;
 		}
+		else
+		{
+			RuntimePartition = const_cast<URuntimePartition*>(ResolveRuntimePartition(ActorSetInstance.RuntimeGrid));// @todo-ow: GenerateStreaming() requires a non-const URuntimePartition object
+			NameToRuntimePartitionMap.Emplace(ActorSetInstance.RuntimeGrid, RuntimePartition); 
+		}
+
+		RuntimePartitionsToActorSetMap.FindOrAdd(RuntimePartition).Add(&ActorSetInstance);
 	});
 
 	//
diff --git a/Engine/Source/Runtime/Engine/Private/WorldPartition/WorldPartitionRuntimeSpatialHash.cpp b/Engine/Source/Runtime/Engine/Private/WorldPartition/WorldPartitionRuntimeSpatialHash.cpp
--- a/Engine/Source/Runtime/Engine/Private/WorldPartition/WorldPartitionRuntimeSpatialHash.cpp
+++ b/Engine/Source/Runtime/Engine/Private/WorldPartition/WorldPartitionRuntimeSpatialHash.cpp
@@ -1950,24 +1950,37 @@ FGuid UWorldPartitionRuntimeSpatialHash::RegisterWorldAssetStreaming(const UWorl
 {
 	if (!InParams.IsValid())
 	{
+		UE_LOG(LogWorldPartition, Error, TEXT("RegisterWorldAssetStreaming: Invalid parameters provided."));
 		return FGuid();
 	}
 
 	if (WorldAssetStreamingObjects.Contains(InParams.Guid))
 	{
+		UE_LOG(LogWorldPartition, Error, TEXT("RegisterWorldAssetStreaming: World asset guid '%s' was already registered."), *InParams.Guid.ToString());
 		return FGuid();
 	}
 
 	const FSpatialHashStreamingGrid* TargetGrid = GetStreamingGridByName(InParams.TargetGrid);
-	const FSpatialHashStreamingGrid* HLODTargetGrid = GetStreamingGridByName(InParams.TargetGridHLOD);
-	if (!TargetGrid || (!InParams.TargetGridHLOD.IsNone() && !HLODTargetGrid))
+	if (!TargetGrid)
 	{
+		UE_LOG(LogWorldPartition, Error, TEXT("RegisterWorldAssetStreaming: Unable to resolve TargetGrid '%s'."), *InParams.TargetGrid.ToString());
 		return FGuid();
 	}
 
+	const FSpatialHashStreamingGrid* HLODTargetGrid = nullptr;
+	if (!InParams.WorldAssetHLOD.IsNull() && !InParams.TargetGridHLOD.IsNone())
+	{
+		HLODTargetGrid = GetStreamingGridByName(InParams.TargetGridHLOD);
+		if (!HLODTargetGrid)
+		{
+			UE_LOG(LogWorldPartition, Error, TEXT("RegisterWorldAssetStreaming: Unable to resolve TargetGridHLOD '%s'."), *InParams.TargetGridHLOD.ToString());
+		}
+	}
+
 	URuntimeSpatialHashExternalStreamingObject* StreamingObject = CastChecked<URuntimeSpatialHashExternalStreamingObject>(CreateExternalStreamingObject(URuntimeSpatialHashExternalStreamingObject::StaticClass(), this, GetTypedOuter<UWorld>()));
 	if (!StreamingObject)
 	{
+		UE_LOG(LogWorldPartition, Error, TEXT("RegisterWorldAssetStreaming: Couldn't create ExternalStreamingObject."));
 		return FGuid();
 	}		
 
diff --git a/Engine/Source/Runtime/Engine/Public/WorldPartition/RuntimeHashSet/WorldPartitionRuntimeHashSet.h b/Engine/Source/Runtime/Engine/Public/WorldPartition/RuntimeHashSet/WorldPartitionRuntimeHashSet.h
--- a/Engine/Source/Runtime/Engine/Public/WorldPartition/RuntimeHashSet/WorldPartitionRuntimeHashSet.h
+++ b/Engine/Source/Runtime/Engine/Public/WorldPartition/RuntimeHashSet/WorldPartitionRuntimeHashSet.h
@@ -28,14 +28,15 @@ struct FRuntimePartitionHLODSetup
 {
 	GENERATED_USTRUCT_BODY()
 
-#if WITH_EDITORONLY_DATA
 	/** Name for this HLOD layer setup */
 	UPROPERTY(EditAnywhere, Category = RuntimeSettings)
 	FName Name;
 
+#if WITH_EDITORONLY_DATA
 	/** Associated HLOD Layer objects */
 	UPROPERTY(EditAnywhere, Category = RuntimeSettings)
 	TArray<TObjectPtr<const UHLODLayer>> HLODLayers;
+#endif
 
 	/** whether this HLOD setup is spatially loaded or not */
 	UPROPERTY(EditAnywhere, Category = RuntimeSettings)
@@ -43,7 +44,6 @@ struct FRuntimePartitionHLODSetup
 
 	UPROPERTY(VisibleAnywhere, Category = RuntimeSettings, Instanced, Meta = (EditCondition = "bIsSpatiallyLoaded", HideEditConditionToggle, NoResetToDefault, TitleProperty = "Name"))
 	TObjectPtr<URuntimePartition> PartitionLayer;
-#endif
 };
 
 /** Holds settings for a runtime partition instance. */
@@ -64,11 +64,9 @@ struct FRuntimePartitionDesc
 	UPROPERTY(VisibleAnywhere, Category = RuntimeSettings, Instanced, Meta = (EditCondition = "Class != nullptr", HideEditConditionToggle, NoResetToDefault, TitleProperty = "Name"))
 	TObjectPtr<URuntimePartition> MainLayer;
 
-#if WITH_EDITORONLY_DATA
 	/** HLOD setups used by this partition, one for each layers in the hierarchy */
 	UPROPERTY(EditAnywhere, Category = RuntimeSettings, Meta = (EditCondition = "Class != nullptr", HideEditConditionToggle, ForceInlineRow))
 	TArray<FRuntimePartitionHLODSetup> HLODSetups;
-#endif
 
 #if WITH_EDITOR
 	void UpdateHLODPartitionLayers();
@@ -218,7 +216,8 @@ private:
 	ENGINE_API void ForEachStreamingData(TFunctionRef<bool(const FRuntimePartitionStreamingData&)> Func) const;
 
 	ENGINE_API void UpdateRuntimeDataGridMap();
-	ENGINE_API bool IsValidGridInternal(FName GridName, const UClass* ActorClass) const;
+	const URuntimePartition* ResolveRuntimePartition(FName GridName, bool bMainPartitionLayer = false) const;
+	const URuntimePartition* ResolveRuntimePartitionForHLODLayer(FName GridName, const FSoftObjectPath& HLODLayerPath) const;
 
 	/** Array of runtime partition descriptors */
 	UPROPERTY(EditAnywhere, Category = RuntimeSettings, Meta = (TitleProperty = "Name"))
