[AutoRTFM] Make FTSTicker::AddTicker() and FTSTicker::RemoveTicker() transaction-safe

Remove two UE_AUTORTFM_ONCOMMIT2 wrappers in around calls to AddTicker in UPlayspaceManagerComponent, as FTSTicker is now transaction-safe.

diff --git a/Engine/Source/Runtime/Core/Private/Containers/Ticker.cpp b/Engine/Source/Runtime/Core/Private/Containers/Ticker.cpp
--- a/Engine/Source/Runtime/Core/Private/Containers/Ticker.cpp
+++ b/Engine/Source/Runtime/Core/Private/Containers/Ticker.cpp
@@ -1,6 +1,7 @@
 // Copyright Epic Games, Inc. All Rights Reserved.
 
 #include "Containers/Ticker.h"
+#include "AutoRTFM/AutoRTFM.h"
 #include "Stats/Stats.h"
 #include "Misc/TimeGuard.h"
 
@@ -13,14 +14,20 @@ FTSTicker& FTSTicker::GetCoreTicker()
 FTSTicker::FDelegateHandle FTSTicker::AddTicker(const FTickerDelegate& InDelegate, float InDelay)
 {
 	FElementPtr NewElement{ new FElement{ CurrentTime.load(std::memory_order_relaxed) + InDelay, InDelay, InDelegate } };
-	AddedElements.Enqueue(NewElement);
+	AutoRTFM::OnCommit([this, NewElement]
+	{
+		AddedElements.Enqueue(NewElement); 
+	});
 	return NewElement;
 }
 
 FTSTicker::FDelegateHandle FTSTicker::AddTicker(const TCHAR* InName, float InDelay, TUniqueFunction<bool(float)>&& InFunction)
 {
 	FElementPtr NewElement{ new FElement{ CurrentTime.load(std::memory_order_relaxed) + InDelay, InDelay, MoveTemp(InFunction) } };
-	AddedElements.Enqueue(NewElement);
+	AutoRTFM::OnCommit([this, NewElement]
+	{
+		AddedElements.Enqueue(NewElement); 
+	});
 	return NewElement;
 }
 
@@ -31,20 +38,23 @@ uint32 GetThreadId(uint64 State)
 
 void FTSTicker::RemoveTicker(FDelegateHandle Handle)
 {
-	if (FElementPtr Element = Handle.Pin())
+	AutoRTFM::OnCommit([Handle]
 	{
-		// mark the element as removed and if it's being ticked atm, spin-wait until its execution is finished
-		uint64 PrevState = Element->State.fetch_or(FElement::RemovedState, std::memory_order_acquire); // "acquire" to prevent potential 
-		// resource release after RemoveTicker() to be reordered before it
-		uint32 ExecutingThreadId = GetThreadId(PrevState);
-		
-		while (ExecutingThreadId != 0 && // is being executed right now
-			FPlatformTLS::GetCurrentThreadId() != ExecutingThreadId) // and is not removed from inside its execution
+		if (FElementPtr Element = Handle.Pin())
 		{
-			FPlatformProcess::Yield();
-			ExecutingThreadId = GetThreadId(Element->State.load(std::memory_order_relaxed));
+			// mark the element as removed and if it's being ticked atm, spin-wait until its execution is finished
+			uint64 PrevState = Element->State.fetch_or(FElement::RemovedState, std::memory_order_acquire); // "acquire" to prevent potential 
+			// resource release after RemoveTicker() to be reordered before it
+			uint32 ExecutingThreadId = GetThreadId(PrevState);
+			
+			while (ExecutingThreadId != 0 && // is being executed right now
+				FPlatformTLS::GetCurrentThreadId() != ExecutingThreadId) // and is not removed from inside its execution
+			{
+				FPlatformProcess::Yield();
+				ExecutingThreadId = GetThreadId(Element->State.load(std::memory_order_relaxed));
+			}
 		}
-	}
+	});
 }
 
 void FTSTicker::Tick(float DeltaTime)
diff --git a/Engine/Source/Runtime/Engine/Private/Tests/AutoRTFM/AutoRTFMTickerTest.cpp b/Engine/Source/Runtime/Engine/Private/Tests/AutoRTFM/AutoRTFMTickerTest.cpp
new file mode 100644
--- /dev/null
+++ b/Engine/Source/Runtime/Engine/Private/Tests/AutoRTFM/AutoRTFMTickerTest.cpp
@@ -0,0 +1,164 @@
+// Copyright Epic Games, Inc. All Rights Reserved.
+
+#include "Containers/Ticker.h"
+#include "Misc/AutomationTest.h"
+#include "AutoRTFM/AutoRTFM.h"
+#include <vector>
+
+#if WITH_DEV_AUTOMATION_TESTS
+
+#define CHECK_EQ(A, B) \
+	UTEST_EQUAL(TEXT(__FILE__ ":" UE_STRINGIZE(__LINE__) ": UTEST_EQUAL_EXPR(" #A ", " #B ")"), A, B)
+
+IMPLEMENT_SIMPLE_AUTOMATION_TEST(FAutoRTFMTickerTest, "AutoRTFM + FTSTicker", EAutomationTestFlags::EngineFilter | EAutomationTestFlags::ClientContext | EAutomationTestFlags::ServerContext | EAutomationTestFlags::CommandletContext)
+
+bool FAutoRTFMTickerTest::RunTest(const FString & Parameters)
+{
+	if (!AutoRTFM::ForTheRuntime::IsAutoRTFMRuntimeEnabled())
+	{
+		ExecutionInfo.AddEvent(FAutomationEvent(EAutomationEventType::Info, TEXT("SKIPPED 'FAutoRTFMTickerTest' test. AutoRTFM disabled.")));
+		return true;
+	}
+
+	FTSTicker& Ticker = FTSTicker::GetCoreTicker();
+
+	// A ticker delegate that appends the delta-time to CallbackEvents.
+	std::vector<float> CallbackEvents;
+	auto Delegate = [&](float DeltaTime)
+	{ 
+		AutoRTFM::Open([&]{ CallbackEvents.push_back(DeltaTime); });
+		return /* reschedule? */ false;
+	};
+
+	auto Section = [&](const TCHAR* Name, auto&& Test)
+	{
+		bool Result = Test();
+		CallbackEvents.clear();
+		Ticker.Reset();
+		if (!Result)
+		{
+			AddError(FString::Printf(TEXT("In section '%s'."), Name), 1);
+		}
+	};
+
+	Section(TEXT("Basic assumptions"), [&]
+	{
+		Section(TEXT("Explicit remove"), [&]
+		{
+			auto Handle = Ticker.AddTicker(TEXT(""), /* Delay */ 0.1f, Delegate);
+			FTSTicker::RemoveTicker(Handle);
+			Ticker.Tick(1.0f);
+			CHECK_EQ(CallbackEvents, {});
+			return true;
+		});
+		Section(TEXT("One-shot self removal"), [&]
+		{
+			Ticker.AddTicker(TEXT(""), /* Delay */ 0.1f, Delegate);
+			Ticker.Tick(1.0f);
+			CHECK_EQ(CallbackEvents, {1.0f});
+			Ticker.Tick(2.0f);
+			CHECK_EQ(CallbackEvents, {1.0f}); // Doesn't repeat.
+			return true;
+		});
+		return true;
+	});
+
+	Section(TEXT("Transact(AddTicker), Tick"), [&]
+	{
+		AutoRTFM::Transact([&]
+		{
+			Ticker.AddTicker(TEXT(""), /* Delay */ 0.1f, Delegate);
+		});
+		Ticker.Tick(1.0f);
+		CHECK_EQ(CallbackEvents, {1.0f});
+		return true;
+	});
+
+	Section(TEXT("Transact(AddTicker, Abort), Tick"), [&]
+	{
+		AutoRTFM::Transact([&]
+		{
+			Ticker.AddTicker(TEXT(""), /* Delay */ 0.1f, Delegate);
+			AutoRTFM::AbortTransaction();
+		});
+		Ticker.Tick(1.0f);
+		CHECK_EQ(CallbackEvents, {});
+		return true;
+	});
+
+	Section(TEXT("Transact(AddTicker, RemoveTicker), Tick"), [&]
+	{
+		AutoRTFM::Transact([&]
+		{
+			FTSTicker::FDelegateHandle Handle = Ticker.AddTicker(TEXT(""), /* Delay */ 0.1f, Delegate);
+			FTSTicker::RemoveTicker(Handle);
+		});
+		Ticker.Tick(1.0f);
+		CHECK_EQ(CallbackEvents, {});
+		return true;
+	});
+
+	Section(TEXT("Transact(AddTicker, RemoveTicker, Abort), Tick"), [&]
+	{
+		AutoRTFM::Transact([&]
+		{
+			FTSTicker::FDelegateHandle Handle = Ticker.AddTicker(TEXT(""), /* Delay */ 0.1f, Delegate);
+			FTSTicker::RemoveTicker(Handle);
+			AutoRTFM::AbortTransaction();
+		});
+		Ticker.Tick(1.0f);
+		CHECK_EQ(CallbackEvents, {});
+		return true;
+	});
+
+	Section(TEXT("Transact(AddTicker, Abort), Transact(AddTicker), Tick"), [&]
+	{
+		AutoRTFM::Transact([&]
+		{
+			Ticker.AddTicker(TEXT(""), /* Delay */ 0.1f, Delegate);
+			AutoRTFM::AbortTransaction();
+		});
+		AutoRTFM::Transact([&]
+		{
+			Ticker.AddTicker(TEXT(""), /* Delay */ 0.1f, Delegate);
+		});
+		Ticker.Tick(1.0f);
+		CHECK_EQ(CallbackEvents, {1.0f});
+		return true;
+	});
+
+	Section(TEXT("AddTicker, Transact(RemoveTicker), Tick"), [&]
+	{
+		FTSTicker::FDelegateHandle Handle = Ticker.AddTicker(TEXT(""), /* Delay */ 0.1f, Delegate);
+
+		AutoRTFM::Transact([&]
+		{
+			FTSTicker::RemoveTicker(Handle);
+		});
+
+		Ticker.Tick(1.0f);
+		CHECK_EQ(CallbackEvents, {});
+		return true;
+	});
+
+	Section(TEXT("AddTicker, Transact(RemoveTicker, Abort), Tick"), [&]
+	{
+		FTSTicker::FDelegateHandle Handle = Ticker.AddTicker(TEXT(""), /* Delay */ 0.1f, Delegate);
+
+		AutoRTFM::Transact([&]
+		{
+			FTSTicker::RemoveTicker(Handle);
+			AutoRTFM::AbortTransaction();
+		});
+
+		Ticker.Tick(1.0f);
+		CHECK_EQ(CallbackEvents, {1.0f});
+		return true;
+	});
+
+	return true;
+}
+
+#undef CHECK_EQ
+
+#endif //WITH_DEV_AUTOMATION_TESTS
