Landscape : fixed FBodyInstance being improperly detected as not synced with parent landscape when the content of the collision responses is similar but out-of-order

#tests cooked game, editor, PIE

diff --git a/Engine/Source/Runtime/Engine/Classes/PhysicsEngine/BodyInstance.h b/Engine/Source/Runtime/Engine/Classes/PhysicsEngine/BodyInstance.h
--- a/Engine/Source/Runtime/Engine/Classes/PhysicsEngine/BodyInstance.h
+++ b/Engine/Source/Runtime/Engine/Classes/PhysicsEngine/BodyInstance.h
@@ -254,6 +254,15 @@ private:
 	friend struct FBodyInstance;
 };
 
+template<>
+struct TStructOpsTypeTraits<FCollisionResponse> : public TStructOpsTypeTraitsBase2<FCollisionResponse>
+{
+	enum
+	{
+		WithIdenticalViaEquality = true
+	};
+};
+
 enum class BodyInstanceSceneState : uint8
 {
 	NotAdded,
