[FSeamlessTravelHandler] Do not add navigation system to the temporary world since there are missing worldsubsystems

diff --git a/Engine/Source/Runtime/Engine/Private/World.cpp b/Engine/Source/Runtime/Engine/Private/World.cpp
--- a/Engine/Source/Runtime/Engine/Private/World.cpp
+++ b/Engine/Source/Runtime/Engine/Private/World.cpp
@@ -8054,8 +8054,13 @@ UWorld* FSeamlessTravelHandler::Tick()
 				}
 			}
 
-			// calling it after InitializeActorsForPlay has been called to have all potential bounding boxed initialized
-			FNavigationSystem::AddNavigationSystemToWorld(*LoadedWorld, FNavigationSystemRunMode::GameMode);
+			// We don't want to add navigation system to the temporary worlds.
+			// Those don't get the required world subsystem required by the NavigationSystem
+			if (LoadedWorld->HasAnyFlags(RF_Standalone))
+			{
+				// calling it after InitializeActorsForPlay has been called to have all potential bounding boxed initialized
+				FNavigationSystem::AddNavigationSystemToWorld(*LoadedWorld, FNavigationSystemRunMode::GameMode);
+			}
 
 			FName LoadedWorldName = FName(*UWorld::RemovePIEPrefix(LoadedWorld->GetOutermost()->GetName()));
 
