[World Partition] Fix issue with IsStreamingCompleted that wouldn't properly take into account External Data Layers when checking for active streaming.

diff --git a/Engine/Source/Runtime/Engine/Private/WorldPartition/WorldPartitionStreamingPolicy.cpp b/Engine/Source/Runtime/Engine/Private/WorldPartition/WorldPartitionStreamingPolicy.cpp
--- a/Engine/Source/Runtime/Engine/Private/WorldPartition/WorldPartitionStreamingPolicy.cpp
+++ b/Engine/Source/Runtime/Engine/Private/WorldPartition/WorldPartitionStreamingPolicy.cpp
@@ -1074,15 +1074,17 @@ bool UWorldPartitionStreamingPolicy::IsStreamingCompleted(const TArray<FWorldPar
 
 bool UWorldPartitionStreamingPolicy::IsStreamingCompleted(EWorldPartitionRuntimeCellState QueryState, const TArray<FWorldPartitionStreamingQuerySource>& QuerySources, bool bExactState) const
 {
+	const FWorldPartitionStreamingContext StreamingContext = FWorldPartitionStreamingContext::Create(GetTypedOuter<UWorld>());
 	const UDataLayerManager* DataLayerManager = WorldPartition->GetDataLayerManager();
 	const bool bIsHLODEnabled = UWorldPartitionHLODRuntimeSubsystem::IsHLODEnabled();
 
 	bool bResult = true;
 	for (const FWorldPartitionStreamingQuerySource& QuerySource : QuerySources)
 	{
-		WorldPartition->RuntimeHash->ForEachStreamingCellsQuery(QuerySource, [QuerySource, QueryState, bExactState, bIsHLODEnabled, DataLayerManager, &bResult](const UWorldPartitionRuntimeCell* Cell)
+		WorldPartition->RuntimeHash->ForEachStreamingCellsQuery(QuerySource, [QuerySource, QueryState, bExactState, bIsHLODEnabled, DataLayerManager, &StreamingContext, &bResult](const UWorldPartitionRuntimeCell* Cell)
 		{
-			EWorldPartitionRuntimeCellState CellState = Cell->GetCurrentState();
+			const EWorldPartitionRuntimeCellState CellState = Cell->GetCurrentState();
+
 			if (CellState != QueryState)
 			{
 				bool bSkipCell = false;
@@ -1093,6 +1095,14 @@ bool UWorldPartitionStreamingPolicy::IsStreamingCompleted(EWorldPartitionRuntime
 					bSkipCell = Cell->GetIsHLOD();
 				}
 
+				if (!bSkipCell)
+				{
+					const EDataLayerRuntimeState CellWantedState = Cell->GetCellEffectiveWantedState(StreamingContext);
+					bSkipCell = (CellState == EWorldPartitionRuntimeCellState::Unloaded && CellWantedState == EDataLayerRuntimeState::Unloaded) || 
+								(CellState == EWorldPartitionRuntimeCellState::Loaded && CellWantedState == EDataLayerRuntimeState::Loaded) || 
+								(CellState == EWorldPartitionRuntimeCellState::Activated && CellWantedState == EDataLayerRuntimeState::Activated);
+				}
+
 				// If we are querying for Unloaded/Loaded but a Cell is part of a data layer outside of the query that is activated do not consider it
 				if (!bSkipCell && QueryState < CellState)
 				{
