Alpha: Partial revert of CL 35501865 due to missed UObject load/cook errors with existing project's propagate alpha set to "2".

This forces us to expose a combobox instead of a checkox, but maintains backward compatibility with the same cvar. Updated display name to "Alpha Output" accordingly.

diff --git a/Engine/Plugins/Experimental/Avalanche/Source/AvalancheViewport/Private/Visualizers/AvaViewportPostProcessVisualizer.cpp b/Engine/Plugins/Experimental/Avalanche/Source/AvalancheViewport/Private/Visualizers/AvaViewportPostProcessVisualizer.cpp
--- a/Engine/Plugins/Experimental/Avalanche/Source/AvalancheViewport/Private/Visualizers/AvaViewportPostProcessVisualizer.cpp
+++ b/Engine/Plugins/Experimental/Avalanche/Source/AvalancheViewport/Private/Visualizers/AvaViewportPostProcessVisualizer.cpp
@@ -69,7 +69,7 @@ bool FAvaViewportPostProcessVisualizer::CanActivate(bool bInSilent) const
 	URendererSettings* RendererSettings = GetMutableDefault<URendererSettings>();
 	check(RendererSettings);
 
-	if (RendererSettings->bEnableAlphaChannelInPostProcessing)
+	if (RendererSettings->bEnableAlphaChannelInPostProcessing == EAlphaChannelMode::Enabled)
 	{
 		return true;
 	}
@@ -84,7 +84,7 @@ bool FAvaViewportPostProcessVisualizer::CanActivate(bool bInSilent) const
 	{
 		case EAppReturnType::Yes:
 		{
-			RendererSettings->bEnableAlphaChannelInPostProcessing = true;
+			RendererSettings->bEnableAlphaChannelInPostProcessing = EAlphaChannelMode::Enabled;
 
 			if (IConsoleVariable* PropagateAlphaCVar = IConsoleManager::Get().FindConsoleVariable(TEXT("r.PostProcessing.PropagateAlpha")))
 			{
diff --git a/Engine/Plugins/MovieScene/MovieRenderPipeline/Source/MovieRenderPipelineCore/Private/Graph/MovieGraphRenderLayerSubsystem.cpp b/Engine/Plugins/MovieScene/MovieRenderPipeline/Source/MovieRenderPipelineCore/Private/Graph/MovieGraphRenderLayerSubsystem.cpp
--- a/Engine/Plugins/MovieScene/MovieRenderPipeline/Source/MovieRenderPipelineCore/Private/Graph/MovieGraphRenderLayerSubsystem.cpp
+++ b/Engine/Plugins/MovieScene/MovieRenderPipeline/Source/MovieRenderPipelineCore/Private/Graph/MovieGraphRenderLayerSubsystem.cpp
@@ -192,7 +192,7 @@ void UMovieGraphRenderPropertyModifier::ApplyModifier(const UWorld* World)
 
 	// Generate a warning if holdout is being used, but alpha is not enabled in post processing. Without that setting enabled, holdout will not work.
 	const URendererSettings* RendererSettings = GetDefault<URendererSettings>();
-	if (bHoldout && (RendererSettings->bEnableAlphaChannelInPostProcessing == false))
+	if (bHoldout && (RendererSettings->bEnableAlphaChannelInPostProcessing == EAlphaChannelMode::Disabled))
 	{
 		// TODO: Ideally this is called in a general-purpose validation step instead, but that framework does not exist yet.
 		UE_LOG(LogMovieRenderPipeline, Warning, TEXT("A modifier with 'Holdout' is active, but 'Enable Alpha Output' in "
diff --git a/Engine/Plugins/MovieScene/MovieRenderPipeline/Source/MovieRenderPipelineCore/Private/MoviePipelineRenderPass.cpp b/Engine/Plugins/MovieScene/MovieRenderPipeline/Source/MovieRenderPipelineCore/Private/MoviePipelineRenderPass.cpp
--- a/Engine/Plugins/MovieScene/MovieRenderPipeline/Source/MovieRenderPipelineCore/Private/MoviePipelineRenderPass.cpp
+++ b/Engine/Plugins/MovieScene/MovieRenderPipeline/Source/MovieRenderPipelineCore/Private/MoviePipelineRenderPass.cpp
@@ -11,7 +11,7 @@ void UMoviePipelineRenderPass::ValidateStateImpl()
 	if (IsAlphaInTonemapperRequired())
 	{
 		const URendererSettings* RenderSettings = GetDefault<URendererSettings>();
-		if (RenderSettings->bEnableAlphaChannelInPostProcessing == false)
+		if (RenderSettings->bEnableAlphaChannelInPostProcessing == EAlphaChannelMode::Disabled)
 		{
 			ValidationState = EMoviePipelineValidationState::Warnings;
 			ValidationResults.Add(NSLOCTEXT("MovieRenderPipeline", "Outputs_AlphaWithoutProjectSetting", "This option does not work without enabling the Alpha Support in Tonemapper setting via Project Settings > Rendering > Post Processing > Enable Alpha Channel Support."));
diff --git a/Engine/Source/Runtime/Engine/Classes/Engine/RendererSettings.h b/Engine/Source/Runtime/Engine/Classes/Engine/RendererSettings.h
--- a/Engine/Source/Runtime/Engine/Classes/Engine/RendererSettings.h
+++ b/Engine/Source/Runtime/Engine/Classes/Engine/RendererSettings.h
@@ -114,7 +114,7 @@ namespace EVertexDeformationOutputsVelocity
 UENUM()
 namespace EAlphaChannelMode
 {
-	enum UE_DEPRECATED(5.5, "Propagate alpha was converted back to a boolean.") Type : int
+	enum Type : int
 	{
 		/** Disabled, reducing GPU cost to the minimum. (default). */
 		Disabled = 0,
@@ -123,17 +123,15 @@ namespace EAlphaChannelMode
 		Enabled = 1,
 
 		// Deprecated
-		LinearColorSpaceOnly = Enabled UMETA(Hidden),
-		AllowThroughTonemapper = Enabled UMETA(Hidden)
+		LinearColorSpaceOnly UE_DEPRECATED(5.5, "Propagate alpha was converted back to a boolean.") = Enabled UMETA(Hidden),
+		AllowThroughTonemapper UE_DEPRECATED(5.5, "Propagate alpha was converted back to a boolean.") = Enabled UMETA(Hidden)
 	};
 }
 
-PRAGMA_DISABLE_DEPRECATION_WARNINGS
 namespace EAlphaChannelMode
 {
 	ENGINE_API EAlphaChannelMode::Type FromInt(int32 InAlphaChannelMode);
 }
-PRAGMA_ENABLE_DEPRECATION_WARNINGS
 
 /** used by FPostProcessSettings AutoExposure*/
 UENUM()
@@ -726,10 +724,10 @@ class URendererSettings : public UDeveloperSettings
 	uint32 bCustomDepthTaaJitter : 1;
 
 	UPROPERTY(config, EditAnywhere, Category = Postprocessing, meta = (
-		ConsoleVariable = "r.PostProcessing.PropagateAlpha", DisplayName = "Enable Alpha Output",
+		ConsoleVariable = "r.PostProcessing.PropagateAlpha", DisplayName = "Alpha Output",
 		ToolTip = "Enable r.PostProcessing.PropagateAlpha to enforce alpha in scene color (overriding r.SceneColorFormat if necessary) and propagate it through the renderer's post-processing chain.  The legacy \"Linear color space only\" and \"Allow through tonemapper\" options now map to true and the engine keeps full alpha precision throughout post-processing.",
 		ConfigRestartRequired = true))
-	bool bEnableAlphaChannelInPostProcessing;
+	TEnumAsByte<EAlphaChannelMode::Type> bEnableAlphaChannelInPostProcessing;
 
 	UPROPERTY(config, EditAnywhere, Category = DefaultSettings, meta = (
 		ConsoleVariable = "r.DefaultFeature.Bloom", DisplayName = "Bloom",
diff --git a/Engine/Source/Runtime/Engine/Private/RendererSettings.cpp b/Engine/Source/Runtime/Engine/Private/RendererSettings.cpp
--- a/Engine/Source/Runtime/Engine/Private/RendererSettings.cpp
+++ b/Engine/Source/Runtime/Engine/Private/RendererSettings.cpp
@@ -32,15 +32,13 @@ extern UNREALED_API class UEditorEngine* GEditor;
 
 #define LOCTEXT_NAMESPACE "RendererSettings"
 
-PRAGMA_DISABLE_DEPRECATION_WARNINGS
 namespace EAlphaChannelMode
 {
 	EAlphaChannelMode::Type FromInt(int32 InAlphaChannelMode)
 	{
-		return static_cast<EAlphaChannelMode::Type>(FMath::Clamp(InAlphaChannelMode, (int32)Disabled, (int32)AllowThroughTonemapper));
+		return static_cast<EAlphaChannelMode::Type>(FMath::Clamp(InAlphaChannelMode, (int32)Disabled, (int32)Enabled));
 	}
 }
-PRAGMA_ENABLE_DEPRECATION_WARNINGS
 
 namespace EDefaultBackBufferPixelFormat
 {
@@ -120,6 +118,9 @@ void URendererSettings::PostInitProperties()
 		ImportConsoleVariableValues();
 	}
 #endif // #if WITH_EDITOR
+
+	// Necessary clamp after 5.5 remapping of propagate alpha back to a boolean state.
+	bEnableAlphaChannelInPostProcessing = EAlphaChannelMode::FromInt(static_cast<int32>(bEnableAlphaChannelInPostProcessing));
 }
 
 #if WITH_EDITOR
