Bing back CL35497867 with fix.
-----------------------------------------------------------------
Changed an assert to an ensure when serializing a level with an actor that has its outer mismatching the owning level.

diff --git a/Engine/Source/Runtime/Engine/Private/Level.cpp b/Engine/Source/Runtime/Engine/Private/Level.cpp
--- a/Engine/Source/Runtime/Engine/Private/Level.cpp
+++ b/Engine/Source/Runtime/Engine/Private/Level.cpp
@@ -638,7 +638,10 @@ void ULevel::Serialize( FArchive& Ar )
 				return false;
 			}
 
-			check(Actor->GetLevel() == this);
+			if (!ensureMsgf(Actor->GetLevel() == this, TEXT("Actor %s is in wrong level:") LINE_TERMINATOR TEXT("\tOwner Level=%s") LINE_TERMINATOR TEXT("\tActor Level=%s)"), *GetPathNameSafe(Actor), *GetPathNameSafe(this), *GetPathNameSafe(Actor->GetLevel())))
+			{
+				return false;
+			}
 
 			if (Actor->HasAnyFlags(RF_Transient))
 			{
