Remove 16 bytes of padding on UPrimitiveComponent
Saves 1.1 mb on clients (MB2)
#rnx

diff --git a/Engine/Source/Runtime/Engine/Classes/Components/PrimitiveComponent.h b/Engine/Source/Runtime/Engine/Classes/Components/PrimitiveComponent.h
--- a/Engine/Source/Runtime/Engine/Classes/Components/PrimitiveComponent.h
+++ b/Engine/Source/Runtime/Engine/Classes/Components/PrimitiveComponent.h
@@ -761,6 +761,23 @@ public:
 	std::atomic<bool> bPSOPrecacheFinished = true;
 #endif
 
+	/**
+	* Defines how quickly it should be culled. For example buildings should have a low priority, but small dressing should have a high priority.
+	*/
+	UPROPERTY(EditAnywhere, AdvancedDisplay, BlueprintReadOnly, Category = RayTracing, meta = (DisplayAfter="RayTracingGroupId"))
+	ERayTracingGroupCullingPriority RayTracingGroupCullingPriority;
+
+	/** Mask used for stencil buffer writes. */
+	UPROPERTY(EditAnywhere, AdvancedDisplay, BlueprintReadOnly, Category = "Rendering", meta = (editcondition = "bRenderCustomDepth", DisplayAfter="CustomDepthStencilValue"))
+	ERendererStencilMask CustomDepthStencilWriteMask;
+
+private:
+	/** Which specific HLOD levels this component should be excluded from */
+	UPROPERTY(EditAnywhere, AdvancedDisplay, Category = HLOD, meta = (Bitmask, BitmaskEnum = "/Script/Engine.EHLODLevelExclusion", DisplayName = "Exclude from HLOD Levels", DisplayAfter = "bEnableAutoLODGeneration", EditConditionHides, EditCondition = "AllowHLODLevelsExclusion()"))
+	uint8 ExcludeFromHLODLevels;
+
+public:
+
 	/**
 	 * Defines run-time groups of components. For example allows to assemble multiple parts of a building at runtime.
 	 * -1 means that component doesn't belong to any group.
@@ -776,20 +793,7 @@ public:
 	UPROPERTY(EditAnywhere, AdvancedDisplay, BlueprintReadOnly, Category=Rendering,  meta=(UIMin = "0", UIMax = "255", editcondition = "bRenderCustomDepth", DisplayName = "CustomDepth Stencil Value"))
 	int32 CustomDepthStencilValue;
 
-	/**
-	* Defines how quickly it should be culled. For example buildings should have a low priority, but small dressing should have a high priority.
-	*/
-	UPROPERTY(EditAnywhere, AdvancedDisplay, BlueprintReadOnly, Category = RayTracing)
-	ERayTracingGroupCullingPriority RayTracingGroupCullingPriority;
-
-	/** Mask used for stencil buffer writes. */
-	UPROPERTY(EditAnywhere, AdvancedDisplay, BlueprintReadOnly, Category = "Rendering", meta = (editcondition = "bRenderCustomDepth"))
-	ERendererStencilMask CustomDepthStencilWriteMask;
-
 private:
-	/** Which specific HLOD levels this component should be excluded from */
-	UPROPERTY(EditAnywhere, AdvancedDisplay, Category = HLOD, meta = (Bitmask, BitmaskEnum = "/Script/Engine.EHLODLevelExclusion", DisplayName = "Exclude from HLOD Levels", DisplayAfter = "bEnableAutoLODGeneration", EditConditionHides, EditCondition = "AllowHLODLevelsExclusion()"))
-	uint8 ExcludeFromHLODLevels;
 
 	/** Optional user defined default values for the custom primitive data of this primitive */
 	UPROPERTY(EditAnywhere, Category=Rendering, meta = (DisplayName = "Custom Primitive Data Defaults"))
@@ -798,6 +802,7 @@ private:
 	/** Custom data that can be read by a material through a material parameter expression. Set data using SetCustomPrimitiveData* functions */
 	UPROPERTY(Transient)
 	FCustomPrimitiveData CustomPrimitiveDataInternal;
+	
 public:
 
 	/** If non-null, physics state creation has been deferred to ULevel::IncrementalUpdateComponents or this scene's StartFrame.*/
