Substrate - fixed default value of Opacity for UI and PostProcess domains.

#rnx

diff --git a/Engine/Shaders/Shared/SubstrateDefinitions.h b/Engine/Shaders/Shared/SubstrateDefinitions.h
--- a/Engine/Shaders/Shared/SubstrateDefinitions.h
+++ b/Engine/Shaders/Shared/SubstrateDefinitions.h
@@ -9,7 +9,7 @@
 
 #ifndef __cplusplus
 // Change this to force recompilation of all Substrate dependent shaders (for instance https://guidgenerator.com/online-guid-generator.aspx)
-#pragma message("UESHADERMETADATA_VERSION 0A6FB545-2B60-4E5D-B93C-A4CBE10274E5")
+#pragma message("UESHADERMETADATA_VERSION 1F256706-39AF-450F-A2C1-6A10D9030F4E")
 #endif
 
 // Closure offsets are packed into 32bits, each entry using SUBSTRATE_CLOSURE_OFFSET_BIT_COUNT bits
diff --git a/Engine/Source/Runtime/Engine/Private/Materials/MaterialExpressions.cpp b/Engine/Source/Runtime/Engine/Private/Materials/MaterialExpressions.cpp
--- a/Engine/Source/Runtime/Engine/Private/Materials/MaterialExpressions.cpp
+++ b/Engine/Source/Runtime/Engine/Private/Materials/MaterialExpressions.cpp
@@ -28005,7 +28005,7 @@ UMaterialExpressionSubstratePostProcess::UMaterialExpressionSubstratePostProcess
 #if WITH_EDITOR
 int32 UMaterialExpressionSubstratePostProcess::Compile(class FMaterialCompiler* Compiler, int32 OutputIndex)
 {
-	int OpacityCodeChunk = CompileWithDefaultFloat1(Compiler, Opacity, 0.0f);
+	int OpacityCodeChunk = CompileWithDefaultFloat1(Compiler, Opacity, 1.0f);
 	int TransmittanceCodeChunk = Compiler->Saturate(Compiler->Sub(Compiler->Constant(1.0f), OpacityCodeChunk));
 
 	FSubstrateOperator& SubstrateOperator = Compiler->SubstrateCompilationGetOperator(Compiler->SubstrateTreeStackGetPathUniqueId());
@@ -28082,7 +28082,7 @@ UMaterialExpressionSubstrateUI::UMaterialExpressionSubstrateUI(const FObjectInit
 #if WITH_EDITOR
 int32 UMaterialExpressionSubstrateUI::Compile(class FMaterialCompiler* Compiler, int32 OutputIndex)
 {
-	int OpacityCodeChunk = CompileWithDefaultFloat1(Compiler, Opacity, 0.0f);
+	int OpacityCodeChunk = CompileWithDefaultFloat1(Compiler, Opacity, 1.0f);
 
 	FSubstrateOperator& SubstrateOperator = Compiler->SubstrateCompilationGetOperator(Compiler->SubstrateTreeStackGetPathUniqueId());
 	int32 OutputCodeChunk = Compiler->SubstrateUIBSDF(
