[World Partition] Make sure to dirty the map package when changing the Is2D partition property.

diff --git a/Engine/Source/Runtime/Engine/Private/WorldPartition/RuntimeHashSet/RuntimePartitionLHGrid.cpp b/Engine/Source/Runtime/Engine/Private/WorldPartition/RuntimeHashSet/RuntimePartitionLHGrid.cpp
--- a/Engine/Source/Runtime/Engine/Private/WorldPartition/RuntimeHashSet/RuntimePartitionLHGrid.cpp
+++ b/Engine/Source/Runtime/Engine/Private/WorldPartition/RuntimeHashSet/RuntimePartitionLHGrid.cpp
@@ -135,7 +135,7 @@ void URuntimePartitionLHGrid::PostEditChangeProperty(FPropertyChangedEvent& InPr
 			WorldGridPreviewer = MakeUnique<FWorldGridPreviewer>(GetTypedOuter<UWorld>(), bIs2D);
 		}
 
-		if (!GPackageWasDirty)
+		if ((PropertyName == GET_MEMBER_NAME_CHECKED(URuntimePartitionLHGrid, bShowGridPreview)) && !GPackageWasDirty)
 		{
 			GetPackage()->ClearDirtyFlag();
 		}
