World Partition - Minimap - Disable sky lighting when capturing the scene

diff --git a/Engine/Source/Runtime/Engine/Private/WorldPartition/WorldPartitionMiniMapHelper.cpp b/Engine/Source/Runtime/Engine/Private/WorldPartition/WorldPartitionMiniMapHelper.cpp
--- a/Engine/Source/Runtime/Engine/Private/WorldPartition/WorldPartitionMiniMapHelper.cpp
+++ b/Engine/Source/Runtime/Engine/Private/WorldPartition/WorldPartitionMiniMapHelper.cpp
@@ -97,6 +97,7 @@ void FWorldPartitionMiniMapHelper::CaptureBoundsMiniMapToTexture(UWorld* InWorld
 	CaptureComponent->ShowFlags.Fog = false;
 	CaptureComponent->ShowFlags.VolumetricFog = false;
 	CaptureComponent->ShowFlags.DynamicShadows = false;
+	CaptureComponent->ShowFlags.SkyLighting = false;
 
 	// Disable vignetting, otherwise we'll see it as a pattern between each captured tiles
 	CaptureComponent->PostProcessSettings.bOverride_VignetteIntensity = true;
