PR #12210: Fixed SplineMeshComponent::CalcBounds to return an initialized FBoxSphereBounds when GetStaticMesh is null.

diff --git a/Engine/Source/Runtime/Engine/Private/Components/SplineMeshComponent.cpp b/Engine/Source/Runtime/Engine/Private/Components/SplineMeshComponent.cpp
--- a/Engine/Source/Runtime/Engine/Private/Components/SplineMeshComponent.cpp
+++ b/Engine/Source/Runtime/Engine/Private/Components/SplineMeshComponent.cpp
@@ -892,7 +892,7 @@ FBoxSphereBounds USplineMeshComponent::CalcBounds(const FTransform& LocalToWorld
 	const UStaticMesh* Mesh = GetStaticMesh();
 	if (Mesh == nullptr)
 	{
-		return FBox();
+		return FBoxSphereBounds(LocalToWorld.GetLocation(), FVector::ZeroVector, 0.f);;
 	}
 
 	const FBox ComputedBounds = ComputeDistortedBounds(LocalToWorld, Mesh->GetBounds());
