- Rather than retrieving CompressedData ptr use the one the compression task actually targeted
- Ensure that bUseRawDataOnly flag is set _before_ clearing the compressed data

diff --git a/Engine/Source/Runtime/Engine/Private/Animation/AnimSequence.cpp b/Engine/Source/Runtime/Engine/Private/Animation/AnimSequence.cpp
--- a/Engine/Source/Runtime/Engine/Private/Animation/AnimSequence.cpp
+++ b/Engine/Source/Runtime/Engine/Private/Animation/AnimSequence.cpp
@@ -476,7 +476,9 @@ void UAnimSequence::AddReferencedObjects(UObject* This, FReferenceCollector& Col
 #if WITH_EDITOR
 void UAnimSequence::WillNeverCacheCookedPlatformDataAgain()
 {
-	Super::WillNeverCacheCookedPlatformDataAgain();
+	PRAGMA_DISABLE_DEPRECATION_WARNINGS
+	bUseRawDataOnly = true;
+	PRAGMA_ENABLE_DEPRECATION_WARNINGS
 
 	UE::Anim::FAnimSequenceCompilingManager::Get().FinishCompilation({this});
 	// Clear out current platform, and any target platform data
@@ -484,14 +486,13 @@ void UAnimSequence::WillNeverCacheCookedPlatformDataAgain()
 		
 	CacheTasksByKeyHash.Empty();
 	DataByPlatformKeyHash.Empty();
-	PRAGMA_DISABLE_DEPRECATION_WARNINGS
-	bUseRawDataOnly = true;
-	PRAGMA_ENABLE_DEPRECATION_WARNINGS
 }
 
 void UAnimSequence::ClearAllCachedCookedPlatformData()
 {
-	Super::ClearAllCachedCookedPlatformData();
+	PRAGMA_DISABLE_DEPRECATION_WARNINGS
+	bUseRawDataOnly = true;
+	PRAGMA_ENABLE_DEPRECATION_WARNINGS
 	
 	// Delete any cache tasks first because the destructor will cancel the cache and build tasks,
 	// and drop their pointers to the data.
@@ -499,10 +500,6 @@ void UAnimSequence::ClearAllCachedCookedPlatformData()
 	DataByPlatformKeyHash.Empty();
 	CompressedData.Reset();
 	DataKeyHash = FIoHash::Zero;
-	
-	PRAGMA_DISABLE_DEPRECATION_WARNINGS
-	bUseRawDataOnly = true;
-	PRAGMA_ENABLE_DEPRECATION_WARNINGS
 }
 
 int64 UAnimSequence::GetUncompressedRawSize() const
@@ -4871,23 +4868,15 @@ void UAnimSequence::FinishAsyncTasks()
 		{
 			It->Value->Wait();
 
-			FCompressedAnimSequence* CompressedAnimData = [Hash = It->Key, this]()
-			{
-				if(Hash == DataKeyHash)
-				{
-					return &CompressedData;
-				}
-				
-				return DataByPlatformKeyHash.FindChecked(Hash).Get();
-			}();
-
-			auto ResetData = [this, CompressedAnimData, It]()
+			FCompressedAnimSequence* TaskData = It->Value->GetTargetCompressedData();
+			const bool bCompressedRunningPlatform = (TaskData == &CompressedData);
+			auto ResetData = [this, TaskData, It, bCompressedRunningPlatform]()
 			{
-				CompressedAnimData->Reset();
+				TaskData->Reset();
 				DataByPlatformKeyHash.Remove(It->Key);
 
 				// Reset running platform hash (if it got cancelled)
-				if (DataKeyHash == It->Key)
+				if (bCompressedRunningPlatform)
 				{
 					DataKeyHash = FIoHash::Zero;
 				}
@@ -4899,7 +4888,7 @@ void UAnimSequence::FinishAsyncTasks()
 			}
 			else
 			{
-				if (CompressedAnimData->IsValid(this, true))
+				if (TaskData->IsValid(this, true))
 				{
 #if WITH_EDITOR
 					//This is only safe during sync anim compression
@@ -4908,7 +4897,7 @@ void UAnimSequence::FinishAsyncTasks()
 						SetSkeletonVirtualBoneGuid(GetSkeleton()->GetVirtualBoneGuid());
 					}
 #endif
-					if (It->Key == DataKeyHash)
+					if (bCompressedRunningPlatform)
 					{
 						FAssetRegistryModule& AssetRegistryModule = FModuleManager::LoadModuleChecked<FAssetRegistryModule>(TEXT("AssetRegistry"));
 						AssetRegistryModule.Get().AssetTagsFinalized(*this);
@@ -4919,7 +4908,7 @@ void UAnimSequence::FinishAsyncTasks()
 						PRAGMA_ENABLE_DEPRECATION_WARNINGS
 					}
 				
-					check(CompressedAnimData->IsValid(this, true));
+					check(TaskData->IsValid(this, true));
 				}
 				else
 				{
diff --git a/Engine/Source/Runtime/Engine/Private/Animation/AnimationCompressionDerivedData.h b/Engine/Source/Runtime/Engine/Private/Animation/AnimationCompressionDerivedData.h
--- a/Engine/Source/Runtime/Engine/Private/Animation/AnimationCompressionDerivedData.h
+++ b/Engine/Source/Runtime/Engine/Private/Animation/AnimationCompressionDerivedData.h
@@ -83,6 +83,7 @@ namespace UE::Anim
 		bool Poll() const;
 		void Reschedule(FQueuedThreadPool* InThreadPool, EQueuedWorkPriority InPriority) const;
 		bool WasCancelled() const { return CompressibleAnimPtr->IsCancelled() || Owner.IsCanceled(); }
+		FCompressedAnimSequence* GetTargetCompressedData() const { return CompressedData; }
 	private:
 		void BeginCache(const FIoHash& KeyHash);
 		void EndCache(DerivedData::FCacheGetValueResponse&& Response);
