Deprecated direct access to VertexColor visualization global variables.
The visualization is being extended to cover all mesh painting modes and this deprecation step will help licensees who have copy/pasted setup code into their own scene proxies.

diff --git a/Engine/Plugins/MeshPainting/Source/MeshPaintEditorMode/Private/MeshPaintModeHelpers.cpp b/Engine/Plugins/MeshPainting/Source/MeshPaintEditorMode/Private/MeshPaintModeHelpers.cpp
--- a/Engine/Plugins/MeshPainting/Source/MeshPaintEditorMode/Private/MeshPaintModeHelpers.cpp
+++ b/Engine/Plugins/MeshPainting/Source/MeshPaintEditorMode/Private/MeshPaintModeHelpers.cpp
@@ -70,7 +70,7 @@ void UMeshPaintModeSubsystem::SetViewportColorMode(EMeshPaintActiveMode ActiveMo
 
 					// Restore the vertex color mode flags that were set when we last entered vertex color mode
 					ApplyViewMode(ViewportClient->GetViewMode(), ViewportClient->IsPerspective(), ViewportClient->EngineShowFlags);
-					GVertexColorViewMode = EVertexColorViewMode::Color;
+					SetMeshPaintVisualizeChannels(EVertexColorViewMode::Color);
 				}
 			}
 			else
@@ -86,36 +86,36 @@ void UMeshPaintModeSubsystem::SetViewportColorMode(EMeshPaintActiveMode ActiveMo
 				{
 				case EMeshPaintDataColorViewMode::RGB:
 				{
-					GVertexColorViewMode = EVertexColorViewMode::Color;
+					SetMeshPaintVisualizeChannels(EVertexColorViewMode::Color);
 				}
 				break;
 
 				case EMeshPaintDataColorViewMode::Alpha:
 				{
-					GVertexColorViewMode = EVertexColorViewMode::Alpha;
+					SetMeshPaintVisualizeChannels(EVertexColorViewMode::Alpha);
 				}
 				break;
 
 				case EMeshPaintDataColorViewMode::Red:
 				{
-					GVertexColorViewMode = EVertexColorViewMode::Red;
+					SetMeshPaintVisualizeChannels(EVertexColorViewMode::Red);
 				}
 				break;
 
 				case EMeshPaintDataColorViewMode::Green:
 				{
-					GVertexColorViewMode = EVertexColorViewMode::Green;
+					SetMeshPaintVisualizeChannels(EVertexColorViewMode::Green);
 				}
 				break;
 
 				case EMeshPaintDataColorViewMode::Blue:
 				{
-					GVertexColorViewMode = EVertexColorViewMode::Blue;
+					SetMeshPaintVisualizeChannels(EVertexColorViewMode::Blue);
 				}
 				break;
 				}
 				UTexture* SelectedTexture = nullptr;
-				float UVChannel = 0.0f; // Keep as float since it must be a fed to the material as a scalar parameter
+				int32 UVChannel = 0;
 				if (ActiveMode == EMeshPaintActiveMode::Texture)
 				{
 					UMeshPaintingSubsystem* MeshPaintingSubsystem = GEngine->GetEngineSubsystem<UMeshPaintingSubsystem>();
@@ -133,16 +133,10 @@ void UMeshPaintModeSubsystem::SetViewportColorMode(EMeshPaintActiveMode ActiveMo
 						}
 						UVChannel = Settings->UVChannel;
 					}
-
-					const UMeshComponent* LastPaintedComponent = MeshPaintingSubsystem ? MeshPaintingSubsystem->LastPaintedComponent : nullptr;
-					if (LastPaintedComponent)
-					{
-						GVertexViewModeOverrideOwnerName = *LastPaintedComponent->GetOwner()->GetName();
-					}
 				}
 
-				GVertexViewModeOverrideTexture = SelectedTexture;
-				GVertexViewModeOverrideUVChannel = UVChannel;
+				SetMeshPaintVisualizeTexture(SelectedTexture);
+				SetMeshPaintVisualizeTextureCoordinateIndex(UVChannel);
 			}
 		}
 	}
diff --git a/Engine/Plugins/MeshPainting/Source/MeshPaintingToolset/Private/MeshTexturePaintingTool.cpp b/Engine/Plugins/MeshPainting/Source/MeshPaintingToolset/Private/MeshTexturePaintingTool.cpp
--- a/Engine/Plugins/MeshPainting/Source/MeshPaintingToolset/Private/MeshTexturePaintingTool.cpp
+++ b/Engine/Plugins/MeshPainting/Source/MeshPaintingToolset/Private/MeshTexturePaintingTool.cpp
@@ -27,7 +27,6 @@
 
 #include "Editor/EditorEngine.h"
 extern UNREALED_API class UEditorEngine* GEditor;
-extern ENGINE_API FString GVertexViewModeOverrideOwnerName;
 
 #include UE_INLINE_GENERATED_CPP_BY_NAME(MeshTexturePaintingTool)
 
@@ -719,7 +718,6 @@ FInputRayHit UMeshTexturePaintingTool::CanBeginClickDragSequence(const FInputDev
 	if (MeshPaintingSubsystem && LastBestHitResult.Component != nullptr && MeshPaintingSubsystem->LastPaintedComponent != LastBestHitResult.Component)
 	{
 		MeshPaintingSubsystem->LastPaintedComponent = (UMeshComponent*)LastBestHitResult.Component.Get();
-		GVertexViewModeOverrideOwnerName = *LastBestHitResult.Component->GetOwner()->GetName();
 		//ApplyAllPaintedTextures();
 		//MeshPaintingSubsystem->bNeedsRecache = true;
 	}
diff --git a/Engine/Source/Editor/MeshPaint/Private/MeshPaintHelpers.cpp b/Engine/Source/Editor/MeshPaint/Private/MeshPaintHelpers.cpp
--- a/Engine/Source/Editor/MeshPaint/Private/MeshPaintHelpers.cpp
+++ b/Engine/Source/Editor/MeshPaint/Private/MeshPaintHelpers.cpp
@@ -1047,7 +1047,7 @@ void MeshPaintHelpers::SetViewportColorMode(EMeshPaintColorViewMode ColorViewMod
 
 					// Restore the vertex color mode flags that were set when we last entered vertex color mode
 					ApplyViewMode(ViewportClient->GetViewMode(), ViewportClient->IsPerspective(), ViewportClient->EngineShowFlags);
-					GVertexColorViewMode = EVertexColorViewMode::Color;
+					SetMeshPaintVisualizeChannels(EVertexColorViewMode::Color);
 				}
 			}
 			else
@@ -1063,31 +1063,31 @@ void MeshPaintHelpers::SetViewportColorMode(EMeshPaintColorViewMode ColorViewMod
 				{
 					case EMeshPaintColorViewMode::RGB:
 					{
-						GVertexColorViewMode = EVertexColorViewMode::Color;
+						SetMeshPaintVisualizeChannels(EVertexColorViewMode::Color);
 					}
 					break;
 
 					case EMeshPaintColorViewMode::Alpha:
 					{
-						GVertexColorViewMode = EVertexColorViewMode::Alpha;
+						SetMeshPaintVisualizeChannels(EVertexColorViewMode::Alpha);
 					}
 					break;
 
 					case EMeshPaintColorViewMode::Red:
 					{
-						GVertexColorViewMode = EVertexColorViewMode::Red;
+						SetMeshPaintVisualizeChannels(EVertexColorViewMode::Red);
 					}
 					break;
 
 					case EMeshPaintColorViewMode::Green:
 					{
-						GVertexColorViewMode = EVertexColorViewMode::Green;
+						SetMeshPaintVisualizeChannels(EVertexColorViewMode::Green);
 					}
 					break;
 
 					case EMeshPaintColorViewMode::Blue:
 					{
-						GVertexColorViewMode = EVertexColorViewMode::Blue;
+						SetMeshPaintVisualizeChannels(EVertexColorViewMode::Blue);
 					}
 					break;
 				}
diff --git a/Engine/Source/Runtime/Engine/Private/SceneView.cpp b/Engine/Source/Runtime/Engine/Private/SceneView.cpp
--- a/Engine/Source/Runtime/Engine/Private/SceneView.cpp
+++ b/Engine/Source/Runtime/Engine/Private/SceneView.cpp
@@ -389,22 +389,37 @@ static TAutoConsoleVariable<float> CVarOverrideTimeMaterialExpressions(
 
 #endif
 
-/** Global vertex color view mode setting when SHOW_VertexColors show flag is set */
+
+/** Deprecated vertex color global settings. */
+PRAGMA_DISABLE_DEPRECATION_WARNINGS
 EVertexColorViewMode::Type GVertexColorViewMode = EVertexColorViewMode::Color;
-TWeakObjectPtr<UTexture> GVertexViewModeOverrideTexture = nullptr;
-float GVertexViewModeOverrideUVChannel = 0.0f; // Scalar parameter, so keep as float
+TWeakObjectPtr<UTexture> GVertexViewModeOverrideTexture;
+float GVertexViewModeOverrideUVChannel = 0.0f;
 FString GVertexViewModeOverrideOwnerName;
+bool ShouldProxyUseVertexColorVisualization(FName OwnerName) { return false; }
+PRAGMA_ENABLE_DEPRECATION_WARNINGS
 
-bool ShouldProxyUseVertexColorVisualization(FName OwnerName)
+/** Global mesh paint visualization settings used when the SHOW_VertexColors show flag is set. */
+static EVertexColorViewMode::Type GMeshPaintVisualizeChannels = EVertexColorViewMode::Color;
+void SetMeshPaintVisualizeChannels(EVertexColorViewMode::Type VisualizeChannels)
+{
+	GMeshPaintVisualizeChannels = VisualizeChannels;
+}
+static TWeakObjectPtr<UTexture> GMeshPaintVisualizeTexture = nullptr;
+void SetMeshPaintVisualizeTexture(TWeakObjectPtr<UTexture> Texture)
+{
+	GMeshPaintVisualizeTexture = Texture;
+}
+static int32 GMeshPaintVisualizeUVChannel = 0;
+void SetMeshPaintVisualizeTextureCoordinateIndex(int32 Index)
 {
-	bool bUsingTextureOverride = GVertexViewModeOverrideTexture.Get() != nullptr;
-	return !bUsingTextureOverride || OwnerName.ToString().Compare(GVertexViewModeOverrideOwnerName) == 0;
+	GMeshPaintVisualizeUVChannel = Index;
 }
 
-FMaterialRenderProxy* GetVertexColorRenderProxy(bool bIsSelected, bool bIsHovered)
+FMaterialRenderProxy* GetMeshPaintVisualizeMaterialRenderProxy(bool bIsSelected, bool bIsHovered)
 {
 	UMaterial* VertexColorVisualizationMaterial = nullptr;
-	switch (GVertexColorViewMode)
+	switch (GMeshPaintVisualizeChannels)
 	{
 	case EVertexColorViewMode::Color:
 		VertexColorVisualizationMaterial = GEngine->VertexColorViewModeMaterial_ColorOnly;
@@ -432,11 +447,11 @@ FMaterialRenderProxy* GetVertexColorRenderProxy(bool bIsSelected, bool bIsHovere
 	FMaterialRenderProxy* VertexColorVisualizationMaterialInstance = nullptr;
 
 #if WITH_EDITORONLY_DATA
-	if (GVertexViewModeOverrideTexture.IsValid())
+	if (GMeshPaintVisualizeTexture.IsValid())
 	{
 		FLinearColor MaterialColor = FLinearColor::White;
 
-		switch (GVertexColorViewMode)
+		switch (GMeshPaintVisualizeChannels)
 		{
 		case EVertexColorViewMode::Color:
 			MaterialColor = FLinearColor(1.0f, 1.0f, 1.0f, 0.0f);
@@ -463,10 +478,10 @@ FMaterialRenderProxy* GetVertexColorRenderProxy(bool bIsSelected, bool bIsHovere
 			GEngine->TexturePaintingMaskMaterial->GetRenderProxy(),
 			MaterialColor,
 			NAME_Color,
-			GVertexViewModeOverrideTexture.Get(),
+			GMeshPaintVisualizeTexture.Get(),
 			NAME_LinearColor);
 
-		NewVertexColorVisualizationMaterialInstance->UVChannel = GVertexViewModeOverrideUVChannel;
+		NewVertexColorVisualizationMaterialInstance->UVChannel = GMeshPaintVisualizeUVChannel;
 		NewVertexColorVisualizationMaterialInstance->UVChannelParamName = FName(TEXT("UVChannel"));
 
 		VertexColorVisualizationMaterialInstance = NewVertexColorVisualizationMaterialInstance;
diff --git a/Engine/Source/Runtime/Engine/Private/SkeletalMesh.cpp b/Engine/Source/Runtime/Engine/Private/SkeletalMesh.cpp
--- a/Engine/Source/Runtime/Engine/Private/SkeletalMesh.cpp
+++ b/Engine/Source/Runtime/Engine/Private/SkeletalMesh.cpp
@@ -7103,10 +7103,10 @@ void FSkeletalMeshSceneProxy::GetDynamicElementsSection(const TArray<const FScen
 
 #if WITH_EDITORONLY_DATA
 #if !(UE_BUILD_SHIPPING || UE_BUILD_TEST)
-			if (bIsSelected && ViewFamily.EngineShowFlags.VertexColors && AllowDebugViewmodes() && ShouldProxyUseVertexColorVisualization(GetOwnerName()))
+			if (bIsSelected && ViewFamily.EngineShowFlags.VertexColors && AllowDebugViewmodes())
 			{
 				// Note: static mesh renderer does something more complicated involving per-section selection, but whole component selection seems ok for now.
-				if (FMaterialRenderProxy* VertexColorVisualizationMaterialInstance = GetVertexColorRenderProxy(bIsSelected, IsHovered()))
+				if (FMaterialRenderProxy* VertexColorVisualizationMaterialInstance = GetMeshPaintVisualizeMaterialRenderProxy(bIsSelected, IsHovered()))
 				{
 					Collector.RegisterOneFrameMaterialProxy(VertexColorVisualizationMaterialInstance);
 					Mesh.MaterialRenderProxy = VertexColorVisualizationMaterialInstance;
diff --git a/Engine/Source/Runtime/Engine/Private/StaticMeshRender.cpp b/Engine/Source/Runtime/Engine/Private/StaticMeshRender.cpp
--- a/Engine/Source/Runtime/Engine/Private/StaticMeshRender.cpp
+++ b/Engine/Source/Runtime/Engine/Private/StaticMeshRender.cpp
@@ -1604,9 +1604,9 @@ void FStaticMeshSceneProxy::GetDynamicMeshElements(const TArray<const FSceneView
 
 	#endif // WITH_EDITOR
 										// Override the mesh's material with our material that draws vertex color
-										if (!bDebugMaterialRenderProxySet && bProxyIsSelected && EngineShowFlags.VertexColors && AllowDebugViewmodes() && ShouldProxyUseVertexColorVisualization(GetOwnerName()))
+										if (!bDebugMaterialRenderProxySet && bProxyIsSelected && EngineShowFlags.VertexColors && AllowDebugViewmodes())
 										{
-											if (FMaterialRenderProxy* VertexColorVisualizationMaterialInstance = GetVertexColorRenderProxy(bSectionIsSelected, IsHovered()))
+											if (FMaterialRenderProxy* VertexColorVisualizationMaterialInstance = GetMeshPaintVisualizeMaterialRenderProxy(bSectionIsSelected, IsHovered()))
 											{
 												Collector.RegisterOneFrameMaterialProxy(VertexColorVisualizationMaterialInstance);
 												MeshElement.MaterialRenderProxy = VertexColorVisualizationMaterialInstance;
diff --git a/Engine/Source/Runtime/Engine/Public/SceneManagement.h b/Engine/Source/Runtime/Engine/Public/SceneManagement.h
--- a/Engine/Source/Runtime/Engine/Public/SceneManagement.h
+++ b/Engine/Source/Runtime/Engine/Public/SceneManagement.h
@@ -2739,40 +2739,43 @@ void BuildCylinderVerts(const FVector& Base, const FVector& XAxis, const FVector
 extern ENGINE_API FLinearColor GetSelectionColor(const FLinearColor& BaseColor,bool bSelected,bool bHovered, bool bUseOverlayIntensity = true);
 extern ENGINE_API FLinearColor GetViewSelectionColor(const FLinearColor& BaseColor, const FSceneView& View, bool bSelected, bool bHovered, bool bUseOverlayIntensity, bool bIndividuallySelected);
 
-
-/** Vertex Color view modes */
+/** Mesh painting visualization channels. (Used for all mesh painting modes not just vertex color.) */
 namespace EVertexColorViewMode
 {
 	enum Type
 	{
 		/** Invalid or undefined */
 		Invalid,
-
 		/** Color only */
 		Color,
-		
 		/** Alpha only */
 		Alpha,
-
 		/** Red only */
 		Red,
-
 		/** Green only */
 		Green,
-
 		/** Blue only */
 		Blue,
 	};
 }
 
-
-/** Global vertex color view mode setting when SHOW_VertexColors show flag is set */
+UE_DEPRECATED(5.5, "Use SetMeshPaintVisualizeChannels() instead.")
 extern ENGINE_API EVertexColorViewMode::Type GVertexColorViewMode;
+UE_DEPRECATED(5.5, "Use SetMeshPaintVisualizeTexture() instead.")
 extern ENGINE_API TWeakObjectPtr<UTexture> GVertexViewModeOverrideTexture;
+UE_DEPRECATED(5.5, "Use SetMeshPaintVisualizeTextureCoordinateIndex() instead.")
 extern ENGINE_API float GVertexViewModeOverrideUVChannel;
+UE_DEPRECATED(5.5, "We no longer use names to enable visualization.")
 extern ENGINE_API FString GVertexViewModeOverrideOwnerName;
-extern ENGINE_API bool ShouldProxyUseVertexColorVisualization(FName OwnerName);
-extern ENGINE_API FMaterialRenderProxy* GetVertexColorRenderProxy(bool bIsSelected, bool bIsHovered);
+UE_DEPRECATED(5.5, "We no longer use names to enable visualization.")
+extern bool ShouldProxyUseVertexColorVisualization(FName OwnerName);
+
+/** Interface to set mesh paint visualization settings used when the SHOW_VertexColors show flag is set. */
+extern ENGINE_API void SetMeshPaintVisualizeChannels(EVertexColorViewMode::Type VisualizeChannels);
+extern ENGINE_API void SetMeshPaintVisualizeTexture(TWeakObjectPtr<UTexture> Texture);
+extern ENGINE_API void SetMeshPaintVisualizeTextureCoordinateIndex(int32 Index);
+/** Get the mesh paint visualization material proxy based on the current global settings. */
+extern ENGINE_API FMaterialRenderProxy* GetMeshPaintVisualizeMaterialRenderProxy(bool bIsSelected, bool bIsHovered);
 
 /**
  * Returns true if the given view is "rich", and all primitives should be forced down the dynamic drawing path so that ApplyViewModeOverrides can implement the rich view feature.
diff --git a/Engine/Source/Runtime/Experimental/GeometryCollectionEngine/Private/GeometryCollection/GeometryCollectionSceneProxy.cpp b/Engine/Source/Runtime/Experimental/GeometryCollectionEngine/Private/GeometryCollection/GeometryCollectionSceneProxy.cpp
--- a/Engine/Source/Runtime/Experimental/GeometryCollectionEngine/Private/GeometryCollection/GeometryCollectionSceneProxy.cpp
+++ b/Engine/Source/Runtime/Experimental/GeometryCollectionEngine/Private/GeometryCollection/GeometryCollectionSceneProxy.cpp
@@ -674,7 +674,7 @@ void FGeometryCollectionSceneProxy::GetDynamicMeshElements(const TArray<const FS
 		if (!bDebugMaterialRenderProxySet && bProxyIsSelected && EngineShowFlags.VertexColors && AllowDebugViewmodes())
 		{
 			// Note: static mesh renderer does something more complicated involving per-section selection, but whole component selection seems ok for now.
-			if (FMaterialRenderProxy* VertexColorVisualizationMaterialInstance = GetVertexColorRenderProxy(bProxyIsSelected, IsHovered()))
+			if (FMaterialRenderProxy* VertexColorVisualizationMaterialInstance = GetMeshPaintVisualizeMaterialRenderProxy(bProxyIsSelected, IsHovered()))
 			{
 				Collector.RegisterOneFrameMaterialProxy(VertexColorVisualizationMaterialInstance);
 				Mesh.MaterialRenderProxy = VertexColorVisualizationMaterialInstance;
diff --git a/Engine/Source/Runtime/GeometryFramework/Private/Components/BaseDynamicMeshSceneProxy.cpp b/Engine/Source/Runtime/GeometryFramework/Private/Components/BaseDynamicMeshSceneProxy.cpp
--- a/Engine/Source/Runtime/GeometryFramework/Private/Components/BaseDynamicMeshSceneProxy.cpp
+++ b/Engine/Source/Runtime/GeometryFramework/Private/Components/BaseDynamicMeshSceneProxy.cpp
@@ -156,7 +156,7 @@ FMaterialRenderProxy* FBaseDynamicMeshSceneProxy::GetEngineVertexColorMaterialPr
 	if (bProxyIsSelected && EngineShowFlags.VertexColors && AllowDebugViewmodes())
 	{
 		// Note: static mesh renderer does something more complicated involving per-section selection, but whole component selection seems ok for now.
-		if (FMaterialRenderProxy* VertexColorVisualizationMaterialInstance = GetVertexColorRenderProxy(bProxyIsSelected, bIsHovered))
+		if (FMaterialRenderProxy* VertexColorVisualizationMaterialInstance = GetMeshPaintVisualizeMaterialRenderProxy(bProxyIsSelected, bIsHovered))
 		{
 			Collector.RegisterOneFrameMaterialProxy(VertexColorVisualizationMaterialInstance);
 			ForceOverrideMaterialProxy = VertexColorVisualizationMaterialInstance;
