[NavigationSystem] reverted navigation bound update from SceneComponent::UpdateBounds since this can be called on multiple components from multiple threads and the navigation system delegate is not thread safe.

diff --git a/Engine/Source/Runtime/Engine/Private/Components/SceneComponent.cpp b/Engine/Source/Runtime/Engine/Private/Components/SceneComponent.cpp
--- a/Engine/Source/Runtime/Engine/Private/Components/SceneComponent.cpp
+++ b/Engine/Source/Runtime/Engine/Private/Components/SceneComponent.cpp
@@ -1422,13 +1422,6 @@ void USceneComponent::UpdateBounds()
 		}
 	}
 
-	// Refresh navigation data if component is relevant to navigation.
-	// Note that 'UpdateNavigationData' will ignore unregistered components.
-	if (bNavigationRelevant)
-	{
-		UpdateNavigationData();
-	}
-
 #if ENABLE_NAN_DIAGNOSTIC
 	if (Bounds.ContainsNaN())
 	{
