Add TextureColor support to mesh paint visualization modes.
Also fixed getting the selected UTexture for the mesh paint texture asset visualization mode. Now we fetch it per frame and correctly select the overriden render target when necessary.

diff --git a/Engine/Plugins/MeshPainting/Source/MeshPaintEditorMode/Private/MeshPaintMode.cpp b/Engine/Plugins/MeshPainting/Source/MeshPaintEditorMode/Private/MeshPaintMode.cpp
--- a/Engine/Plugins/MeshPainting/Source/MeshPaintEditorMode/Private/MeshPaintMode.cpp
+++ b/Engine/Plugins/MeshPainting/Source/MeshPaintEditorMode/Private/MeshPaintMode.cpp
@@ -204,25 +204,27 @@ void UMeshPaintMode::Tick(FEditorViewportClient* ViewportClient, float DeltaTime
 
 	// Make sure that correct tab is visible for the current tool
 	// Note that currently Color and Weight mode share the same Select tool
-	FString ActiveTool = GetToolManager()->GetActiveToolName(EToolSide::Mouse);
+	UInteractiveTool const* ActiveTool = GetToolManager()->GetActiveTool(EToolSide::Mouse);
+	const FString ActiveToolName = GetToolManager()->GetActiveToolName(EToolSide::Mouse);
+
 	FName ActiveTab = Toolkit->GetCurrentPalette();
 	FName TargetTab = ActiveTab;
-	if (ActiveTool == VertexColorPaintToolName)
+	if (ActiveToolName == VertexColorPaintToolName)
 	{
 		TargetTab = MeshPaintMode_VertexColor;
 		CurrentActiveMode = EMeshPaintActiveMode::VertexColor;
 	}
-	else if (ActiveTool == VertexWeightPaintToolName)
+	else if (ActiveToolName == VertexWeightPaintToolName)
 	{
 		TargetTab = MeshPaintMode_VertexWeights;
 		CurrentActiveMode = EMeshPaintActiveMode::VertexWeights;
 	}
-	else if (ActiveTool == TextureColorPaintToolName)
+	else if (ActiveToolName == TextureColorPaintToolName || ActiveToolName == TextureColorSelectToolName)
 	{
 		TargetTab = MeshPaintMode_TextureColor;
 		CurrentActiveMode = EMeshPaintActiveMode::TextureColor;
 	}
-	else if (ActiveTool == TextureAssetPaintToolName)
+	else if (ActiveToolName == TextureAssetPaintToolName || ActiveToolName == TextureAssetSelectToolName)
 	{
 		TargetTab = MeshPaintMode_TextureAsset;
 		CurrentActiveMode = EMeshPaintActiveMode::Texture;
@@ -237,7 +239,7 @@ void UMeshPaintMode::Tick(FEditorViewportClient* ViewportClient, float DeltaTime
 		GEditor->GetEditorSubsystem<UMeshPaintModeSubsystem>()->SetRealtimeViewport(true);
 
 		// Set viewport show flags
-		GEditor->GetEditorSubsystem<UMeshPaintModeSubsystem>()->SetViewportColorMode(CurrentActiveMode, ModeSettings->ColorViewMode, ViewportClient);
+		GEditor->GetEditorSubsystem<UMeshPaintModeSubsystem>()->SetViewportColorMode(CurrentActiveMode, ModeSettings->ColorViewMode, ViewportClient, ActiveTool);
 	}
 }
 
@@ -1003,7 +1005,7 @@ void UMeshPaintMode::OnResetViewMode()
 			continue;
 		}
 
-		GEditor->GetEditorSubsystem<UMeshPaintModeSubsystem>()->SetViewportColorMode(EMeshPaintActiveMode::VertexColor, EMeshPaintDataColorViewMode::Normal, ViewportClient);
+		GEditor->GetEditorSubsystem<UMeshPaintModeSubsystem>()->SetViewportColorMode(EMeshPaintActiveMode::VertexColor, EMeshPaintDataColorViewMode::Normal, ViewportClient, nullptr);
 	}
 }
 
diff --git a/Engine/Plugins/MeshPainting/Source/MeshPaintEditorMode/Private/MeshPaintModeHelpers.cpp b/Engine/Plugins/MeshPainting/Source/MeshPaintEditorMode/Private/MeshPaintModeHelpers.cpp
--- a/Engine/Plugins/MeshPainting/Source/MeshPaintEditorMode/Private/MeshPaintModeHelpers.cpp
+++ b/Engine/Plugins/MeshPainting/Source/MeshPaintEditorMode/Private/MeshPaintModeHelpers.cpp
@@ -48,7 +48,7 @@
 #include "VT/MeshPaintVirtualTexture.h"
 
 
-void UMeshPaintModeSubsystem::SetViewportColorMode(EMeshPaintActiveMode ActiveMode, EMeshPaintDataColorViewMode ColorViewMode, FEditorViewportClient* ViewportClient)
+void UMeshPaintModeSubsystem::SetViewportColorMode(EMeshPaintActiveMode ActiveMode, EMeshPaintDataColorViewMode ColorViewMode, FEditorViewportClient* ViewportClient, UInteractiveTool const* ActiveTool)
 {
 	if (ViewportClient->IsPerspective())
 	{
@@ -70,6 +70,7 @@ void UMeshPaintModeSubsystem::SetViewportColorMode(EMeshPaintActiveMode ActiveMo
 
 					// Restore the vertex color mode flags that were set when we last entered vertex color mode
 					ApplyViewMode(ViewportClient->GetViewMode(), ViewportClient->IsPerspective(), ViewportClient->EngineShowFlags);
+					SetMeshPaintVisualizeMode(EMeshPaintVisualizeMode::VertexColor);
 					SetMeshPaintVisualizeChannels(EVertexColorViewMode::Color);
 				}
 			}
@@ -82,56 +83,49 @@ void UMeshPaintModeSubsystem::SetViewportColorMode(EMeshPaintActiveMode ActiveMo
 				ViewportClient->EngineShowFlags.SetPostProcessing(false);
 				ViewportClient->EngineShowFlags.SetHMDDistortion(false);
 
+				switch (ActiveMode)
+				{
+				case EMeshPaintActiveMode::VertexColor:
+				case EMeshPaintActiveMode::VertexWeights:
+					SetMeshPaintVisualizeMode(EMeshPaintVisualizeMode::VertexColor);
+					break;
+				case EMeshPaintActiveMode::TextureColor:
+					SetMeshPaintVisualizeMode(EMeshPaintVisualizeMode::TextureColor);
+					break;
+				case EMeshPaintActiveMode::Texture:
+					SetMeshPaintVisualizeMode(EMeshPaintVisualizeMode::TextureAsset);
+					break;
+				}
+
 				switch (ColorViewMode)
 				{
 				case EMeshPaintDataColorViewMode::RGB:
-				{
 					SetMeshPaintVisualizeChannels(EVertexColorViewMode::Color);
-				}
-				break;
-
+					break;
 				case EMeshPaintDataColorViewMode::Alpha:
-				{
 					SetMeshPaintVisualizeChannels(EVertexColorViewMode::Alpha);
-				}
-				break;
-
+					break;
 				case EMeshPaintDataColorViewMode::Red:
-				{
 					SetMeshPaintVisualizeChannels(EVertexColorViewMode::Red);
-				}
-				break;
-
+					break;
 				case EMeshPaintDataColorViewMode::Green:
-				{
 					SetMeshPaintVisualizeChannels(EVertexColorViewMode::Green);
-				}
-				break;
-
+					break;
 				case EMeshPaintDataColorViewMode::Blue:
-				{
 					SetMeshPaintVisualizeChannels(EVertexColorViewMode::Blue);
+					break;
 				}
-				break;
-				}
+
 				UTexture* SelectedTexture = nullptr;
 				int32 UVChannel = 0;
+
 				if (ActiveMode == EMeshPaintActiveMode::Texture)
 				{
-					UMeshPaintingSubsystem* MeshPaintingSubsystem = GEngine->GetEngineSubsystem<UMeshPaintingSubsystem>();
-					UMeshTextureAssetPaintingToolProperties* Settings = UMeshPaintMode::GetTextureAssetToolProperties();
-					if (MeshPaintingSubsystem && MeshPaintingSubsystem->OverridePaintTexture.IsValid())
+					UMeshTextureAssetPaintingTool const* TextureTool = Cast<UMeshTextureAssetPaintingTool>(ActiveTool);
+					if (TextureTool != nullptr)
 					{
-						SelectedTexture = MeshPaintingSubsystem->OverridePaintTexture.Get();
-					}
-					
-					if (Settings)
-					{
-						if (!SelectedTexture)
-						{
-							SelectedTexture = Settings->PaintTexture;
-						}
-						UVChannel = Settings->UVChannel;
+						SelectedTexture = TextureTool->GetSelectedPaintTextureWithOverride();
+						UVChannel = TextureTool->GetSelectedUVChannel(nullptr);
 					}
 				}
 
diff --git a/Engine/Plugins/MeshPainting/Source/MeshPaintEditorMode/Public/MeshPaintModeHelpers.h b/Engine/Plugins/MeshPainting/Source/MeshPaintEditorMode/Public/MeshPaintModeHelpers.h
--- a/Engine/Plugins/MeshPainting/Source/MeshPaintEditorMode/Public/MeshPaintModeHelpers.h
+++ b/Engine/Plugins/MeshPainting/Source/MeshPaintEditorMode/Public/MeshPaintModeHelpers.h
@@ -20,6 +20,7 @@ class UViewportInteractor;
 class FViewport;
 class FPrimitiveDrawInterface;
 class FSceneView;
+class UInteractiveTool;
 
 struct FStaticMeshComponentLODInfo;
 struct FPerComponentVertexColorData;
@@ -39,7 +40,7 @@ class MESHPAINTEDITORMODE_API UMeshPaintModeSubsystem : public UEditorSubsystem
 {
 public:
 	/** Forces the Viewport Client to render using the given Viewport Color ViewMode */
-	void SetViewportColorMode(EMeshPaintActiveMode ActiveMode, EMeshPaintDataColorViewMode ColorViewMode, FEditorViewportClient* ViewportClient);
+	void SetViewportColorMode(EMeshPaintActiveMode ActiveMode, EMeshPaintDataColorViewMode ColorViewMode, FEditorViewportClient* ViewportClient, UInteractiveTool const* ActiveTool);
 
 	/** Sets whether or not the level viewport should be real time rendered move or viewport as parameter? */
 	void SetRealtimeViewport(bool bRealtime);
diff --git a/Engine/Plugins/MeshPainting/Source/MeshPaintingToolset/Private/MeshTexturePaintingTool.cpp b/Engine/Plugins/MeshPainting/Source/MeshPaintingToolset/Private/MeshTexturePaintingTool.cpp
--- a/Engine/Plugins/MeshPainting/Source/MeshPaintingToolset/Private/MeshTexturePaintingTool.cpp
+++ b/Engine/Plugins/MeshPainting/Source/MeshPaintingToolset/Private/MeshTexturePaintingTool.cpp
@@ -361,15 +361,6 @@ void UMeshTexturePaintingTool::OnTick(float DeltaTime)
 			// if any overrides are necessary, they will re-register later
 			ClearAllTextureOverrides();
 		}
-		
-		if (PaintingTexture2D)
-		{
-			FPaintTexture2DData* TextureData = GetPaintTargetData(PaintingTexture2D);
-			if (TextureData && TextureData->PaintRenderTargetTexture)
-			{
-				MeshPaintingSubsystem->OverridePaintTexture = TextureData->PaintRenderTargetTexture;
-			}
-		}
 	}
 
 	if (bStampPending)
@@ -1878,5 +1869,18 @@ void UMeshTextureAssetPaintingTool::CacheTexturePaintData()
 	}
 }
 
-#undef LOCTEXT_NAMESPACE
+UTexture* UMeshTextureAssetPaintingTool::GetSelectedPaintTextureWithOverride() const
+{
+	UTexture* SelectedTexture = AssetProperties->PaintTexture;
+	if (AssetProperties->PaintTexture != nullptr)
+	{
+		FPaintTexture2DData const* TextureData = PaintTargetData.Find(AssetProperties->PaintTexture);
+		if (TextureData != nullptr && TextureData->PaintRenderTargetTexture)
+		{
+			SelectedTexture = TextureData->PaintRenderTargetTexture;
+		}
+	}
+	return SelectedTexture;
+}
 
+#undef LOCTEXT_NAMESPACE
diff --git a/Engine/Plugins/MeshPainting/Source/MeshPaintingToolset/Public/MeshPaintHelpers.h b/Engine/Plugins/MeshPainting/Source/MeshPaintingToolset/Public/MeshPaintHelpers.h
--- a/Engine/Plugins/MeshPainting/Source/MeshPaintingToolset/Public/MeshPaintHelpers.h
+++ b/Engine/Plugins/MeshPainting/Source/MeshPaintingToolset/Public/MeshPaintHelpers.h
@@ -274,10 +274,10 @@ public:
 
 public:
 	bool bNeedsRecache;
+
 	UPROPERTY(Transient)
-	TWeakObjectPtr<UTexture> OverridePaintTexture;
+	TObjectPtr<UMeshComponent> LastPaintedComponent;
 
-	const UMeshComponent* LastPaintedComponent = nullptr;
 protected:
 	bool bSelectionHasMaterialValidForTexturePaint;
 
diff --git a/Engine/Plugins/MeshPainting/Source/MeshPaintingToolset/Public/MeshTexturePaintingTool.h b/Engine/Plugins/MeshPainting/Source/MeshPaintingToolset/Public/MeshTexturePaintingTool.h
--- a/Engine/Plugins/MeshPainting/Source/MeshPaintingToolset/Public/MeshTexturePaintingTool.h
+++ b/Engine/Plugins/MeshPainting/Source/MeshPaintingToolset/Public/MeshTexturePaintingTool.h
@@ -317,6 +317,9 @@ public:
 	virtual void CacheTexturePaintData() override;
 	// End UMeshTexturePaintingTool Interface.
 
+	/** Get the selected paint texture, and return the modified overriden texture if currently painting. */
+	UTexture* GetSelectedPaintTextureWithOverride() const;
+
 	/** Change selected texture to previous or next available. */
 	void CycleTextures(int32 Direction);
 	/** Returns true if asset shouldn't be shown in UI because it is not in our paintable texture array. */
diff --git a/Engine/Source/Runtime/Engine/Private/SceneView.cpp b/Engine/Source/Runtime/Engine/Private/SceneView.cpp
--- a/Engine/Source/Runtime/Engine/Private/SceneView.cpp
+++ b/Engine/Source/Runtime/Engine/Private/SceneView.cpp
@@ -400,6 +400,11 @@ bool ShouldProxyUseVertexColorVisualization(FName OwnerName) { return false; }
 PRAGMA_ENABLE_DEPRECATION_WARNINGS
 
 /** Global mesh paint visualization settings used when the SHOW_VertexColors show flag is set. */
+static EMeshPaintVisualizeMode::Type GMeshPaintVisualizeMode = EMeshPaintVisualizeMode::VertexColor;
+void SetMeshPaintVisualizeMode(EMeshPaintVisualizeMode::Type VisualizeMode)
+{
+	GMeshPaintVisualizeMode = VisualizeMode;
+}
 static EVertexColorViewMode::Type GMeshPaintVisualizeChannels = EVertexColorViewMode::Color;
 void SetMeshPaintVisualizeChannels(EVertexColorViewMode::Type VisualizeChannels)
 {
@@ -419,82 +424,64 @@ void SetMeshPaintVisualizeTextureCoordinateIndex(int32 Index)
 FMaterialRenderProxy* GetMeshPaintVisualizeMaterialRenderProxy(bool bIsSelected, bool bIsHovered)
 {
 	UMaterial* VertexColorVisualizationMaterial = nullptr;
+	FLinearColor MaterialColor = FLinearColor::White;
+
 	switch (GMeshPaintVisualizeChannels)
 	{
 	case EVertexColorViewMode::Color:
 		VertexColorVisualizationMaterial = GEngine->VertexColorViewModeMaterial_ColorOnly;
+		MaterialColor = FLinearColor(1.0f, 1.0f, 1.0f, 0.0f);
 		break;
-
 	case EVertexColorViewMode::Alpha:
 		VertexColorVisualizationMaterial = GEngine->VertexColorViewModeMaterial_AlphaAsColor;
+		MaterialColor = FLinearColor(0.0f, 0.0f, 0.0f, 1.0f);
 		break;
-
 	case EVertexColorViewMode::Red:
 		VertexColorVisualizationMaterial = GEngine->VertexColorViewModeMaterial_RedOnly;
+		MaterialColor = FLinearColor(1.0f, 0.0f, 0.0f, 0.0f);
 		break;
-
 	case EVertexColorViewMode::Green:
 		VertexColorVisualizationMaterial = GEngine->VertexColorViewModeMaterial_GreenOnly;
+		MaterialColor = FLinearColor(0.0f, 1.0f, 0.0f, 0.0f);
 		break;
-
 	case EVertexColorViewMode::Blue:
 		VertexColorVisualizationMaterial = GEngine->VertexColorViewModeMaterial_BlueOnly;
+		MaterialColor = FLinearColor(0.0f, 0.0f, 1.0f, 0.0f);
 		break;
-	default:
-		return nullptr;
 	}
 
-	FMaterialRenderProxy* VertexColorVisualizationMaterialInstance = nullptr;
-
-#if WITH_EDITORONLY_DATA
-	if (GMeshPaintVisualizeTexture.IsValid())
+	if (GMeshPaintVisualizeMode == EMeshPaintVisualizeMode::VertexColor && VertexColorVisualizationMaterial != nullptr)
 	{
-		FLinearColor MaterialColor = FLinearColor::White;
-
-		switch (GMeshPaintVisualizeChannels)
-		{
-		case EVertexColorViewMode::Color:
-			MaterialColor = FLinearColor(1.0f, 1.0f, 1.0f, 0.0f);
-			break;
-
-		case EVertexColorViewMode::Alpha:
-			MaterialColor = FLinearColor(0.0f, 0.0f, 0.0f, 1.0f);
-			break;
-
-		case EVertexColorViewMode::Red:
-			MaterialColor = FLinearColor(1.0f, 0.0f, 0.0f, 0.0f);
-			break;
+		return new FColoredMaterialRenderProxy(
+			VertexColorVisualizationMaterial->GetRenderProxy(),
+			GetSelectionColor(FLinearColor::White, bIsSelected, bIsHovered));
+	}
 
-		case EVertexColorViewMode::Green:
-			MaterialColor = FLinearColor(0.0f, 1.0f, 0.0f, 0.0f);
-			break;
+	if (GMeshPaintVisualizeMode == EMeshPaintVisualizeMode::TextureColor)
+	{
+ 		return new FColoredMaterialRenderProxy(
+ 			GEngine->TextureColorViewModeMaterial->GetRenderProxy(),
+ 			MaterialColor);
+	}
 
-		case EVertexColorViewMode::Blue:
-			MaterialColor = FLinearColor(0.0f, 0.0f, 1.0f, 0.0f);
-			break;
-		}
-		
-		FColoredTexturedMaterialRenderProxy* NewVertexColorVisualizationMaterialInstance = new FColoredTexturedMaterialRenderProxy(
+#if WITH_EDITORONLY_DATA
+	if (GMeshPaintVisualizeMode == EMeshPaintVisualizeMode::TextureAsset && GMeshPaintVisualizeTexture.IsValid())
+	{
+		FColoredTexturedMaterialRenderProxy* TextureColorVisualizationMaterialInstance = new FColoredTexturedMaterialRenderProxy(
 			GEngine->TexturePaintingMaskMaterial->GetRenderProxy(),
 			MaterialColor,
 			NAME_Color,
 			GMeshPaintVisualizeTexture.Get(),
 			NAME_LinearColor);
 
-		NewVertexColorVisualizationMaterialInstance->UVChannel = GMeshPaintVisualizeUVChannel;
-		NewVertexColorVisualizationMaterialInstance->UVChannelParamName = FName(TEXT("UVChannel"));
+		TextureColorVisualizationMaterialInstance->UVChannel = (float)GMeshPaintVisualizeUVChannel;
+		TextureColorVisualizationMaterialInstance->UVChannelParamName = FName(TEXT("UVChannel"));
 
-		VertexColorVisualizationMaterialInstance = NewVertexColorVisualizationMaterialInstance;
+		return TextureColorVisualizationMaterialInstance;
 	}
-	else
 #endif
-	{
-		VertexColorVisualizationMaterialInstance = new FColoredMaterialRenderProxy(
-			VertexColorVisualizationMaterial->GetRenderProxy(),
-			GetSelectionColor(FLinearColor::White, bIsSelected, bIsHovered));
-	}
-	
-	return VertexColorVisualizationMaterialInstance;
+
+	return nullptr;
 }
 
 /** Global primitive uniform buffer resource containing identity transformations. */
diff --git a/Engine/Source/Runtime/Engine/Public/SceneManagement.h b/Engine/Source/Runtime/Engine/Public/SceneManagement.h
--- a/Engine/Source/Runtime/Engine/Public/SceneManagement.h
+++ b/Engine/Source/Runtime/Engine/Public/SceneManagement.h
@@ -2770,7 +2770,19 @@ extern ENGINE_API FString GVertexViewModeOverrideOwnerName;
 UE_DEPRECATED(5.5, "We no longer use names to enable visualization.")
 extern bool ShouldProxyUseVertexColorVisualization(FName OwnerName);
 
+/** Mesh painting visualization modes. */
+namespace EMeshPaintVisualizeMode
+{
+	enum Type
+	{
+		VertexColor,
+		TextureColor,
+		TextureAsset,
+	};
+}
+
 /** Interface to set mesh paint visualization settings used when the SHOW_VertexColors show flag is set. */
+extern ENGINE_API void SetMeshPaintVisualizeMode(EMeshPaintVisualizeMode::Type VisualizeMode);
 extern ENGINE_API void SetMeshPaintVisualizeChannels(EVertexColorViewMode::Type VisualizeChannels);
 extern ENGINE_API void SetMeshPaintVisualizeTexture(TWeakObjectPtr<UTexture> Texture);
 extern ENGINE_API void SetMeshPaintVisualizeTextureCoordinateIndex(int32 Index);
