Fix it so the temporary seamless travel world is initialized with the correct world type outside of PIE, this ensures it creates subsystems the same way it does for a loaded map.
Change it to never create a nav system for the transition world, even if it is a loaded map.
#tests local lyra cook

diff --git a/Engine/Source/Runtime/Engine/Private/World.cpp b/Engine/Source/Runtime/Engine/Private/World.cpp
--- a/Engine/Source/Runtime/Engine/Private/World.cpp
+++ b/Engine/Source/Runtime/Engine/Private/World.cpp
@@ -7383,9 +7383,8 @@ bool FSeamlessTravelHandler::StartTravel(UWorld* InCurrentWorld, const FURL& InU
 				
 				if (TransitionMap.IsEmpty())
 				{
-					// If a default transition map doesn't exist, create a dummy World to use as the transition
-					EWorldType::Type TransitionWorldType = CurrentWorld->WorldType == EWorldType::PIE ? EWorldType::PIE : EWorldType::None;
-					SetHandlerLoadedData(nullptr, UWorld::CreateWorld(TransitionWorldType, false));
+					// If a default transition map doesn't exist, create a dummy World of the right type to use as the transition
+					SetHandlerLoadedData(nullptr, UWorld::CreateWorld(CurrentWorld->WorldType, false));
 				}
 				else
 				{
@@ -8054,9 +8053,8 @@ UWorld* FSeamlessTravelHandler::Tick()
 				}
 			}
 
-			// We don't want to add navigation system to the temporary worlds.
-			// Those don't get the required world subsystem required by the NavigationSystem
-			if (LoadedWorld->HasAnyFlags(RF_Standalone))
+			// We don't want to add navigation system to the transition map as it never starts gameplay
+			if (bSwitchedToDefaultMap)
 			{
 				// calling it after InitializeActorsForPlay has been called to have all potential bounding boxed initialized
 				FNavigationSystem::AddNavigationSystemToWorld(*LoadedWorld, FNavigationSystemRunMode::GameMode);
