Add EMeshPaintVisualizeShowMode enumeration to allow setting whether mesh paint visualization should be on everything or only on selected items.

diff --git a/Engine/Plugins/MeshPainting/Source/MeshPaintEditorMode/Private/MeshPaintModeHelpers.cpp b/Engine/Plugins/MeshPainting/Source/MeshPaintEditorMode/Private/MeshPaintModeHelpers.cpp
--- a/Engine/Plugins/MeshPainting/Source/MeshPaintEditorMode/Private/MeshPaintModeHelpers.cpp
+++ b/Engine/Plugins/MeshPainting/Source/MeshPaintEditorMode/Private/MeshPaintModeHelpers.cpp
@@ -70,7 +70,7 @@ void UMeshPaintModeSubsystem::SetViewportColorMode(EMeshPaintActiveMode ActiveMo
 
 					// Restore the vertex color mode flags that were set when we last entered vertex color mode
 					ApplyViewMode(ViewportClient->GetViewMode(), ViewportClient->IsPerspective(), ViewportClient->EngineShowFlags);
-					SetMeshPaintVisualizeMode(EMeshPaintVisualizeMode::VertexColor);
+					SetMeshPaintVisualizeMode(EMeshPaintVisualizePaintMode::VertexColor);
 					SetMeshPaintVisualizeChannels(EVertexColorViewMode::Color);
 				}
 			}
@@ -87,13 +87,13 @@ void UMeshPaintModeSubsystem::SetViewportColorMode(EMeshPaintActiveMode ActiveMo
 				{
 				case EMeshPaintActiveMode::VertexColor:
 				case EMeshPaintActiveMode::VertexWeights:
-					SetMeshPaintVisualizeMode(EMeshPaintVisualizeMode::VertexColor);
+					SetMeshPaintVisualizeMode(EMeshPaintVisualizePaintMode::VertexColor);
 					break;
 				case EMeshPaintActiveMode::TextureColor:
-					SetMeshPaintVisualizeMode(EMeshPaintVisualizeMode::TextureColor);
+					SetMeshPaintVisualizeMode(EMeshPaintVisualizePaintMode::TextureColor);
 					break;
 				case EMeshPaintActiveMode::Texture:
-					SetMeshPaintVisualizeMode(EMeshPaintVisualizeMode::TextureAsset);
+					SetMeshPaintVisualizeMode(EMeshPaintVisualizePaintMode::TextureAsset);
 					break;
 				}
 
diff --git a/Engine/Source/Runtime/Engine/Private/SceneView.cpp b/Engine/Source/Runtime/Engine/Private/SceneView.cpp
--- a/Engine/Source/Runtime/Engine/Private/SceneView.cpp
+++ b/Engine/Source/Runtime/Engine/Private/SceneView.cpp
@@ -400,15 +400,20 @@ bool ShouldProxyUseVertexColorVisualization(FName OwnerName) { return false; }
 PRAGMA_ENABLE_DEPRECATION_WARNINGS
 
 /** Global mesh paint visualization settings used when the SHOW_VertexColors show flag is set. */
-static EMeshPaintVisualizeMode::Type GMeshPaintVisualizeMode = EMeshPaintVisualizeMode::VertexColor;
-void SetMeshPaintVisualizeMode(EMeshPaintVisualizeMode::Type VisualizeMode)
+static EMeshPaintVisualizePaintMode::Type GMeshPaintVisualizeMode = EMeshPaintVisualizePaintMode::VertexColor;
+void SetMeshPaintVisualizeMode(EMeshPaintVisualizePaintMode::Type PaintMode)
 {
-	GMeshPaintVisualizeMode = VisualizeMode;
+	GMeshPaintVisualizeMode = PaintMode;
+}
+static EMeshPaintVisualizeShowMode::Type GMeshPaintVisualizeShowMode = EMeshPaintVisualizeShowMode::ShowAll;
+extern ENGINE_API void SetMeshPaintVisualizeShowMode(EMeshPaintVisualizeShowMode::Type ShowMode)
+{
+	GMeshPaintVisualizeShowMode = ShowMode;
 }
 static EVertexColorViewMode::Type GMeshPaintVisualizeChannels = EVertexColorViewMode::Color;
-void SetMeshPaintVisualizeChannels(EVertexColorViewMode::Type VisualizeChannels)
+void SetMeshPaintVisualizeChannels(EVertexColorViewMode::Type Channels)
 {
-	GMeshPaintVisualizeChannels = VisualizeChannels;
+	GMeshPaintVisualizeChannels = Channels;
 }
 static TWeakObjectPtr<UTexture> GMeshPaintVisualizeTexture = nullptr;
 void SetMeshPaintVisualizeTexture(TWeakObjectPtr<UTexture> Texture)
@@ -450,14 +455,14 @@ FMaterialRenderProxy* GetMeshPaintVisualizeMaterialRenderProxy(bool bIsSelected,
 		break;
 	}
 
-	if (GMeshPaintVisualizeMode == EMeshPaintVisualizeMode::VertexColor && VertexColorVisualizationMaterial != nullptr)
+	if (GMeshPaintVisualizeMode == EMeshPaintVisualizePaintMode::VertexColor && VertexColorVisualizationMaterial != nullptr)
 	{
 		return new FColoredMaterialRenderProxy(
 			VertexColorVisualizationMaterial->GetRenderProxy(),
 			GetSelectionColor(FLinearColor::White, bIsSelected, bIsHovered));
 	}
 
-	if (GMeshPaintVisualizeMode == EMeshPaintVisualizeMode::TextureColor)
+	if (GMeshPaintVisualizeMode == EMeshPaintVisualizePaintMode::TextureColor)
 	{
  		return new FColoredMaterialRenderProxy(
  			GEngine->TextureColorViewModeMaterial->GetRenderProxy(),
@@ -465,7 +470,7 @@ FMaterialRenderProxy* GetMeshPaintVisualizeMaterialRenderProxy(bool bIsSelected,
 	}
 
 #if WITH_EDITORONLY_DATA
-	if (GMeshPaintVisualizeMode == EMeshPaintVisualizeMode::TextureAsset && GMeshPaintVisualizeTexture.IsValid())
+	if (GMeshPaintVisualizeMode == EMeshPaintVisualizePaintMode::TextureAsset && GMeshPaintVisualizeTexture.IsValid())
 	{
 		FColoredTexturedMaterialRenderProxy* TextureColorVisualizationMaterialInstance = new FColoredTexturedMaterialRenderProxy(
 			GEngine->TexturePaintingMaskMaterial->GetRenderProxy(),
diff --git a/Engine/Source/Runtime/Engine/Public/SceneManagement.h b/Engine/Source/Runtime/Engine/Public/SceneManagement.h
--- a/Engine/Source/Runtime/Engine/Public/SceneManagement.h
+++ b/Engine/Source/Runtime/Engine/Public/SceneManagement.h
@@ -2770,8 +2770,8 @@ extern ENGINE_API FString GVertexViewModeOverrideOwnerName;
 UE_DEPRECATED(5.5, "We no longer use names to enable visualization.")
 extern bool ShouldProxyUseVertexColorVisualization(FName OwnerName);
 
-/** Mesh painting visualization modes. */
-namespace EMeshPaintVisualizeMode
+/** Visualization modes for different mesh painting tools. */
+namespace EMeshPaintVisualizePaintMode
 {
 	enum Type
 	{
@@ -2781,9 +2781,22 @@ namespace EMeshPaintVisualizeMode
 	};
 }
 
+/** Visualization modes for mesh painting tools to define where the visualisation is applied. */
+namespace EMeshPaintVisualizeShowMode
+{
+	enum Type
+	{
+		/* Apply visualization to all items. */
+		ShowAll,
+		/* Only apply visualization to selected items. */
+		ShowSelected,
+	};
+}
+
 /** Interface to set mesh paint visualization settings used when the SHOW_VertexColors show flag is set. */
-extern ENGINE_API void SetMeshPaintVisualizeMode(EMeshPaintVisualizeMode::Type VisualizeMode);
-extern ENGINE_API void SetMeshPaintVisualizeChannels(EVertexColorViewMode::Type VisualizeChannels);
+extern ENGINE_API void SetMeshPaintVisualizeMode(EMeshPaintVisualizePaintMode::Type PaintMode);
+extern ENGINE_API void SetMeshPaintVisualizeShowMode(EMeshPaintVisualizeShowMode::Type ShowMode);
+extern ENGINE_API void SetMeshPaintVisualizeChannels(EVertexColorViewMode::Type Channels);
 extern ENGINE_API void SetMeshPaintVisualizeTexture(TWeakObjectPtr<UTexture> Texture);
 extern ENGINE_API void SetMeshPaintVisualizeTextureCoordinateIndex(int32 Index);
 /** Get the mesh paint visualization material proxy based on the current global settings. */
