Wait on the last PSO precache jobs. Only checking the atomic bool is not enough since its vulnerable to multiple precaching jobs fired on the same frame by the component.

#tests Local replayruns on pc

diff --git a/Engine/Source/Runtime/Engine/Classes/Components/DecalComponent.h b/Engine/Source/Runtime/Engine/Classes/Components/DecalComponent.h
--- a/Engine/Source/Runtime/Engine/Classes/Components/DecalComponent.h
+++ b/Engine/Source/Runtime/Engine/Classes/Components/DecalComponent.h
@@ -126,8 +126,9 @@ public:
 
 public:
 #if UE_WITH_PSO_PRECACHING
-	/** Atomic bool to track all the PSO precache events - used for delayed proxy creation */
-	std::atomic<bool> bPSOPrecacheFinished = true;
+	/** Atomic int used to track the last PSO precache events */
+	std::atomic<int> LatestPSOPrecacheJobSetCompleted = 0;
+	int32 LatestPSOPrecacheJobSet = 0;
 #endif
 
 	/** The decal proxy. */
diff --git a/Engine/Source/Runtime/Engine/Classes/Components/PrimitiveComponent.h b/Engine/Source/Runtime/Engine/Classes/Components/PrimitiveComponent.h
--- a/Engine/Source/Runtime/Engine/Classes/Components/PrimitiveComponent.h
+++ b/Engine/Source/Runtime/Engine/Classes/Components/PrimitiveComponent.h
@@ -757,8 +757,9 @@ public:
 	/** Cached array of material PSO requests which can be used to boost the priority */
 	TArray<FMaterialPSOPrecacheRequestID> MaterialPSOPrecacheRequestIDs;
 
-	/** Atomic bool used to track all the PSO precache events */
-	std::atomic<bool> bPSOPrecacheFinished = true;
+	/** Atomic int used to track the last PSO precache events */
+	std::atomic<int> LatestPSOPrecacheJobSetCompleted = 0;
+	int32 LatestPSOPrecacheJobSet = 0;
 #endif
 
 	/**
diff --git a/Engine/Source/Runtime/Engine/Private/Components/DecalComponent.cpp b/Engine/Source/Runtime/Engine/Private/Components/DecalComponent.cpp
--- a/Engine/Source/Runtime/Engine/Private/Components/DecalComponent.cpp
+++ b/Engine/Source/Runtime/Engine/Private/Components/DecalComponent.cpp
@@ -330,9 +330,6 @@ void UDecalComponent::PrecachePSOs()
 		return;
 	}
 
-	// clear the current request data
-	bPSOPrecacheFinished = true;
-
 	if (DecalMaterial && !DecalMaterial->HasAnyFlags(RF_NeedPostLoad))
 	{
 		FPSOPrecacheParams PSOPrecacheParams;		
@@ -344,12 +341,13 @@ void UDecalComponent::PrecachePSOs()
 		FGraphEventArray GraphEvents = DecalMaterial->PrecachePSOs(VertexFactoryDataList, PSOPrecacheParams, EPSOPrecachePriority::High, MaterialPSOPrecacheRequestIDs);
 
 		// Request recreate of the render state when the PSO compilation is ready (if we want to delay proxy creation)
-		if (GraphEvents.Num() > 0 && GetPSOPrecacheProxyCreationStrategy() != EPSOPrecacheProxyCreationStrategy::AlwaysCreate)
+		if (GetPSOPrecacheProxyCreationStrategy() != EPSOPrecacheProxyCreationStrategy::AlwaysCreate)
 		{
 			struct FPSODecalPrecacheFinishedTask
 			{
-				explicit FPSODecalPrecacheFinishedTask(UDecalComponent* InDecalComponent)
-					: WeakDecalComponent(InDecalComponent)
+				explicit FPSODecalPrecacheFinishedTask(UDecalComponent* InDecalComponent, int32 InJobSetThatJustCompleted)
+					: WeakDecalComponent(InDecalComponent),
+					JobSetThatJustCompleted(InJobSetThatJustCompleted)
 				{
 				}
 
@@ -361,16 +359,26 @@ void UDecalComponent::PrecachePSOs()
 				{
 					if (UDecalComponent* DC = WeakDecalComponent.Get())
 					{
-						DC->bPSOPrecacheFinished = true;
+						int32 CurrJobSetCompleted = DC->LatestPSOPrecacheJobSetCompleted.load();
+						while (CurrJobSetCompleted < JobSetThatJustCompleted && !DC->LatestPSOPrecacheJobSetCompleted.compare_exchange_weak(CurrJobSetCompleted, JobSetThatJustCompleted)) {}
 						DC->MarkRenderStateDirty();
 					}
 				}
 
 				TWeakObjectPtr<UDecalComponent> WeakDecalComponent;
+				int32 JobSetThatJustCompleted;
 			};
 
-			bPSOPrecacheFinished = false;
-			TGraphTask<FPSODecalPrecacheFinishedTask>::CreateTask(&GraphEvents).ConstructAndDispatchWhenReady(this);
+			LatestPSOPrecacheJobSet++;
+			if (GraphEvents.Num() > 0)
+			{
+				TGraphTask<FPSODecalPrecacheFinishedTask>::CreateTask(&GraphEvents).ConstructAndDispatchWhenReady(this, LatestPSOPrecacheJobSet);
+			}
+			else
+			{ 
+				// No graph events to wait on, the job set can be considered complete.
+				LatestPSOPrecacheJobSetCompleted = LatestPSOPrecacheJobSet;
+			}
 		}
 	}
 #endif
@@ -417,7 +425,7 @@ FDeferredDecalProxy* UDecalComponent::CreateSceneProxy()
 	LLM_SCOPE(ELLMTag::SceneRender);
 
 #if UE_WITH_PSO_PRECACHING
-	if (!bPSOPrecacheFinished && GetPSOPrecacheProxyCreationStrategy() == EPSOPrecacheProxyCreationStrategy::DelayUntilPSOPrecached)
+	if (LatestPSOPrecacheJobSetCompleted != LatestPSOPrecacheJobSet && GetPSOPrecacheProxyCreationStrategy() == EPSOPrecacheProxyCreationStrategy::DelayUntilPSOPrecached)
 	{
 		return nullptr;
 	}
diff --git a/Engine/Source/Runtime/Engine/Private/Components/PrimitiveComponent.cpp b/Engine/Source/Runtime/Engine/Private/Components/PrimitiveComponent.cpp
--- a/Engine/Source/Runtime/Engine/Private/Components/PrimitiveComponent.cpp
+++ b/Engine/Source/Runtime/Engine/Private/Components/PrimitiveComponent.cpp
@@ -4780,7 +4780,6 @@ void UPrimitiveComponent::PrecachePSOs()
 
 	// clear the current request data
 	MaterialPSOPrecacheRequestIDs.Empty();
-	bPSOPrecacheFinished = true;
 	PSOPrecacheRequestPriority = EPSOPrecachePriority::Medium;
 
 	// Collect the data from the derived classes
@@ -4799,8 +4798,10 @@ void UPrimitiveComponent::PrecachePSOs()
 #if UE_WITH_PSO_PRECACHING
 struct FPSOPrecacheFinishedTask
 {
-	explicit FPSOPrecacheFinishedTask(UPrimitiveComponent* InPrimitiveComponent)
-		: WeakPrimitiveComponent(InPrimitiveComponent)
+	explicit FPSOPrecacheFinishedTask(UPrimitiveComponent* InPrimitiveComponent, int32 InJobSetThatJustCompleted)
+		: WeakPrimitiveComponent(InPrimitiveComponent),
+		JobSetThatJustCompleted(InJobSetThatJustCompleted)
+
 	{
 	}
 
@@ -4813,12 +4814,14 @@ struct FPSOPrecacheFinishedTask
 		if (UPrimitiveComponent* PC = WeakPrimitiveComponent.Get())
 		{
 			QUICK_SCOPE_CYCLE_COUNTER(STAT_PSOPrecacheFinishedTask);
-			PC->bPSOPrecacheFinished = true;
+			int32 CurrJobSetCompleted = PC->LatestPSOPrecacheJobSetCompleted.load();
+			while (CurrJobSetCompleted < JobSetThatJustCompleted && !PC->LatestPSOPrecacheJobSetCompleted.compare_exchange_weak(CurrJobSetCompleted, JobSetThatJustCompleted)){}
 			PC->MarkRenderStateDirty();
 		}
 	}
 
 	TWeakObjectPtr<UPrimitiveComponent> WeakPrimitiveComponent;
+	int32 JobSetThatJustCompleted;
 };
 #endif
 
@@ -4827,10 +4830,18 @@ void UPrimitiveComponent::RequestRecreateRenderStateWhenPSOPrecacheFinished(cons
 #if UE_WITH_PSO_PRECACHING
 	// If the proxy creation strategy relies on knowing when the precached PSO has been compiled,
 	// schedule a task to mark the render state dirty when all PSOs are compiled so the proxy gets recreated.
-	if (UsePSOPrecacheRenderProxyDelay() && GetPSOPrecacheProxyCreationStrategy() != EPSOPrecacheProxyCreationStrategy::AlwaysCreate && !PSOPrecacheCompileEvents.IsEmpty())
+	if (UsePSOPrecacheRenderProxyDelay() && GetPSOPrecacheProxyCreationStrategy() != EPSOPrecacheProxyCreationStrategy::AlwaysCreate)
 	{
-		bPSOPrecacheFinished = false;
-		TGraphTask<FPSOPrecacheFinishedTask>::CreateTask(&PSOPrecacheCompileEvents).ConstructAndDispatchWhenReady(this);
+		LatestPSOPrecacheJobSet++;
+		if(!PSOPrecacheCompileEvents.IsEmpty())
+		{
+			TGraphTask<FPSOPrecacheFinishedTask>::CreateTask(&PSOPrecacheCompileEvents).ConstructAndDispatchWhenReady(this, LatestPSOPrecacheJobSet);
+		}
+		else
+		{
+			// No graph events to wait on, the job set can be considered complete.
+			LatestPSOPrecacheJobSetCompleted = LatestPSOPrecacheJobSet;
+		}
 	}
 
 	bPSOPrecacheCalled = true;
@@ -4849,7 +4860,7 @@ bool UPrimitiveComponent::UsePSOPrecacheRenderProxyDelay() const
 bool UPrimitiveComponent::IsPSOPrecaching() const
 {
 #if UE_WITH_PSO_PRECACHING
-	return !bPSOPrecacheFinished;
+	return LatestPSOPrecacheJobSetCompleted != LatestPSOPrecacheJobSet;
 #else
 	return false;
 #endif // UE_WITH_PSO_PRECACHING
diff --git a/Engine/Source/Runtime/Engine/Private/Particles/ParticleComponents.cpp b/Engine/Source/Runtime/Engine/Private/Particles/ParticleComponents.cpp
--- a/Engine/Source/Runtime/Engine/Private/Particles/ParticleComponents.cpp
+++ b/Engine/Source/Runtime/Engine/Private/Particles/ParticleComponents.cpp
@@ -3566,20 +3566,19 @@ void UFXSystemComponent::PrecacheAssetPSOs(UFXSystemAsset* FXSystemAsset)
 	check(IsInGameThread() || IsInParallelGameThread());
 
 	MaterialPSOPrecacheRequestIDs.Empty();
-	bPSOPrecacheFinished = true;
 	PSOPrecacheRequestPriority = EPSOPrecachePriority::Medium;
 
 	// The asset will keep the Precache events alive, but these might be over. Avoid delaying scene proxy creation if everything is finished
 	bool bAllEventsDone = GraphEvent == nullptr || GraphEvent->IsComplete();
+
+	FGraphEventArray Events;
 	if (!bAllEventsDone)
 	{
 		MaterialPSOPrecacheRequestIDs.Append(FXSystemAsset->GetMaterialPSOPrecacheRequestIDs());
-
-		FGraphEventArray Events;
 		Events.Add(GraphEvent);
-		RequestRecreateRenderStateWhenPSOPrecacheFinished(Events);
 	}
 
+	RequestRecreateRenderStateWhenPSOPrecacheFinished(Events);
 	bPSOPrecacheCalled = true;
 #endif // UE_WITH_PSO_PRECACHING
 }
