UE-148290 - [Console Commands] - au.debug.sounds.sort priority shows all sounds as priority 1.0 when volume scaling is bypassed or sound is 2D

diff --git a/Engine/Source/Runtime/Engine/Private/ActiveSound.cpp b/Engine/Source/Runtime/Engine/Private/ActiveSound.cpp
--- a/Engine/Source/Runtime/Engine/Private/ActiveSound.cpp
+++ b/Engine/Source/Runtime/Engine/Private/ActiveSound.cpp
@@ -365,6 +365,18 @@ bool FActiveSound::IsPlayWhenSilent() const
 	return Sound && Sound->IsPlayWhenSilent();
 }
 
+float FActiveSound::GetConcurrencyPriority() const
+{
+	if (GetAlwaysPlay())
+	{
+		return GetHighestPriority(true) * Priority + MAX_SOUND_PRIORITY + 1.0f;
+	}
+	else
+	{
+		return GetHighestPriority() * Priority;
+	}
+}
+
 void FActiveSound::ClearAudioComponent()
 {
 	AudioComponentID = 0;
diff --git a/Engine/Source/Runtime/Engine/Private/AudioDebug.cpp b/Engine/Source/Runtime/Engine/Private/AudioDebug.cpp
--- a/Engine/Source/Runtime/Engine/Private/AudioDebug.cpp
+++ b/Engine/Source/Runtime/Engine/Private/AudioDebug.cpp
@@ -2330,7 +2330,7 @@ namespace Audio
 							StatSoundInfo.Distance = AudioDevice.GetDistanceToNearestListener(ActiveSound->Transform.GetTranslation());
 						}
 						StatSoundInfo.PlaybackTime = ActiveSound->PlaybackTime;
-						StatSoundInfo.Priority = ActiveSound->GetHighestPriority();
+						StatSoundInfo.Priority = ActiveSound->GetConcurrencyPriority();
 						StatSoundInfo.PlaybackTimeNonVirtualized = ActiveSound->PlaybackTimeNonVirtualized;
 					
 						StatSoundInfo.Volume = 0.0f;
diff --git a/Engine/Source/Runtime/Engine/Private/SoundConcurrency.cpp b/Engine/Source/Runtime/Engine/Private/SoundConcurrency.cpp
--- a/Engine/Source/Runtime/Engine/Private/SoundConcurrency.cpp
+++ b/Engine/Source/Runtime/Engine/Private/SoundConcurrency.cpp
@@ -267,19 +267,8 @@ void FConcurrencyGroup::UpdateGeneration(FActiveSound* NewActiveSound)
 				case EConcurrencyVolumeScaleMode::Priority:
 				{
 					// Ensures sounds set to always play are sorted above those that aren't, but are sorted appropriately between one another
-					auto GetPriority = [](const FActiveSound& InActiveSound)
-					{
-						if (InActiveSound.GetAlwaysPlay())
-						{
-							return InActiveSound.GetHighestPriority(true) * InActiveSound.Priority + MAX_SOUND_PRIORITY + 1.0f;
-						}
-						else
-						{
-							return InActiveSound.GetHighestPriority() * InActiveSound.Priority;
-						}
-					};
-					const float APriority = GetPriority(A);
-					const float BPriority = GetPriority(B);
+					const float APriority = A.GetConcurrencyPriority();
+					const float BPriority = B.GetConcurrencyPriority();
 
 					// If sounds share the same priority, newer sounds will be sorted last to avoid volume ping-ponging 
 					if (FMath::IsNearlyEqual(APriority, BPriority, UE_KINDA_SMALL_NUMBER))
@@ -357,19 +346,8 @@ void FConcurrencyGroup::CullSoundsDueToMaxConcurrency()
 				case EMaxConcurrentResolutionRule::StopLowestPriorityThenPreventNew:
 				{
 					// Ensures sounds set to always play are sorted above those that aren't, but are sorted appropriately between one another
-					auto GetPriority = [](const FActiveSound& InActiveSound)
-					{
-						if (InActiveSound.GetAlwaysPlay())
-						{
-							return InActiveSound.GetHighestPriority(true) * InActiveSound.Priority + MAX_SOUND_PRIORITY + 1.0f;
-						}
-						else
-						{
-							return InActiveSound.GetHighestPriority() * InActiveSound.Priority;
-						}
-					};
-					const float APriority = GetPriority(A);
-					const float BPriority = GetPriority(B);
+					const float APriority = A.GetConcurrencyPriority();
+					const float BPriority = B.GetConcurrencyPriority();
 
 					if (!FMath::IsNearlyEqual(APriority, BPriority, UE_KINDA_SMALL_NUMBER))
 					{
diff --git a/Engine/Source/Runtime/Engine/Public/ActiveSound.h b/Engine/Source/Runtime/Engine/Public/ActiveSound.h
--- a/Engine/Source/Runtime/Engine/Public/ActiveSound.h
+++ b/Engine/Source/Runtime/Engine/Public/ActiveSound.h
@@ -374,6 +374,9 @@ public:
 	/** Whether or not sound reference is valid and set to play when silent. */
 	ENGINE_API bool IsPlayWhenSilent() const;
 
+	/** The priority used for concurrency calculations */
+	float GetConcurrencyPriority() const;
+
 	FAudioDevice* AudioDevice;
 
 	/** The concurrent groups that this sound is actively playing in. */
