PR #12179: Suppress Stats Messaging when bShowScreenMessages is false

#rnx

diff --git a/Engine/Source/Runtime/Engine/Private/UnrealEngine.cpp b/Engine/Source/Runtime/Engine/Private/UnrealEngine.cpp
--- a/Engine/Source/Runtime/Engine/Private/UnrealEngine.cpp
+++ b/Engine/Source/Runtime/Engine/Private/UnrealEngine.cpp
@@ -12931,6 +12931,9 @@ void DrawStatsHUD( UWorld* World, FViewport* Viewport, FCanvas* Canvas, UCanvas*
 
 #endif // UE_BUILD_SHIPPING 
 
+#if !UE_BUILD_SHIPPING
+	if (bShowScreenMessages)
+#endif
 	{
 		int32 RightSideX = ((CanvasObject) ? CanvasObject->SizeX : TextureSize.X) / Canvas->GetDPIScale() - FPSXOffset;
 		int32 RightSideY = FMath::TruncToInt(TextureSize.Y * 0.20f) / Canvas->GetDPIScale();
