World Partition - FWorldPartitionActorDesc::ToString() was missing the "HLODRelevant" field since refactor in 33922379

diff --git a/Engine/Source/Runtime/Engine/Private/WorldPartition/WorldPartitionActorDesc.cpp b/Engine/Source/Runtime/Engine/Private/WorldPartition/WorldPartitionActorDesc.cpp
--- a/Engine/Source/Runtime/Engine/Private/WorldPartition/WorldPartitionActorDesc.cpp
+++ b/Engine/Source/Runtime/Engine/Private/WorldPartition/WorldPartitionActorDesc.cpp
@@ -520,6 +520,7 @@ FString FWorldPartitionActorDesc::ToString(EToStringMode Mode) const
 		AppendToString(TEXT("RuntimeGrid"), RuntimeGrid);
 		AppendFromBool(TEXT("EditorOnly"), bActorIsEditorOnly);
 		AppendFromBool(TEXT("RuntimeOnly"), bActorIsRuntimeOnly);
+		AppendFromBool(TEXT("HLODRelevant"), bActorIsHLODRelevant);
 		AppendFromBool(TEXT("ListedInSceneOutliner"), bActorIsListedInSceneOutliner);
 		AppendFromBool(TEXT("IsMainWorldOnly"), IsMainWorldOnly());
 
