Added explicit call to FActorInstanceGuid::ReleaseLevelInstanceGuid() in FLevelInstanceActorImpl::UnregisterLevelInstance() so that we don't rely on beeing GCed to cleanup the annotation when we reuse Levels instead of destroying them.

diff --git a/Engine/Source/Runtime/Engine/Private/LevelInstance/LevelInstanceActorImpl.cpp b/Engine/Source/Runtime/Engine/Private/LevelInstance/LevelInstanceActorImpl.cpp
--- a/Engine/Source/Runtime/Engine/Private/LevelInstance/LevelInstanceActorImpl.cpp
+++ b/Engine/Source/Runtime/Engine/Private/LevelInstance/LevelInstanceActorImpl.cpp
@@ -51,6 +51,10 @@ void FLevelInstanceActorImpl::UnregisterLevelInstance()
 	if (ULevelInstanceSubsystem* LevelInstanceSubsystem = LevelInstance->GetLevelInstanceSubsystem())
 	{
 		LevelInstanceSubsystem->UnregisterLevelInstance(LevelInstance);
+		if (ULevel* LoadedLevel = LevelInstance->GetLoadedLevel())
+		{
+			FActorInstanceGuid::ReleaseLevelInstanceGuid(LoadedLevel);
+		}
 
 		LevelInstance->UnloadLevelInstance();
 
diff --git a/Engine/Source/Runtime/Engine/Private/WorldPartition/ActorInstanceGuids.cpp b/Engine/Source/Runtime/Engine/Private/WorldPartition/ActorInstanceGuids.cpp
--- a/Engine/Source/Runtime/Engine/Private/WorldPartition/ActorInstanceGuids.cpp
+++ b/Engine/Source/Runtime/Engine/Private/WorldPartition/ActorInstanceGuids.cpp
@@ -23,6 +23,11 @@ struct FLevelInstanceGuid
 FUObjectAnnotationSparse<FActorInstanceGuid, true>	GActorGuids;
 FUObjectAnnotationSparse<FLevelInstanceGuid, true>	GLevelInstanceGuids;
 
+void FActorInstanceGuid::ReleaseLevelInstanceGuid(ULevel* Level)
+{
+	GLevelInstanceGuids.RemoveAnnotation(Level);
+}
+
 void FActorInstanceGuid::SetLevelInstanceGuid(ULevel* Level, ULevel* OwnerLevel, const FGuid& Guid, const FGuid& ResolvedGuid)
 {
 	// double registration is often an order issue (LevelInstance GUID tried to be accessed before FLevelInstanceActorImpl::OnLevelInstanceLoaded was called)
diff --git a/Engine/Source/Runtime/Engine/Public/WorldPartition/ActorInstanceGuids.h b/Engine/Source/Runtime/Engine/Public/WorldPartition/ActorInstanceGuids.h
--- a/Engine/Source/Runtime/Engine/Public/WorldPartition/ActorInstanceGuids.h
+++ b/Engine/Source/Runtime/Engine/Public/WorldPartition/ActorInstanceGuids.h
@@ -18,6 +18,7 @@ struct FActorInstanceGuid
 	void InitializeFrom(AActor& InActor);
 #endif
 
+	ENGINE_API static void ReleaseLevelInstanceGuid(ULevel* Level);
 	ENGINE_API static void SetLevelInstanceGuid(ULevel* Level, ULevel* OwnerLevel, const FGuid& Guid, const FGuid& ResolvedGuid = FGuid());
 	ENGINE_API static FGuid GetLevelInstanceGuid(ULevel* Level);
 
