Fixed a crash at SkeletalMeshCompoent ResetToRefPose, when animation is disabled and there is no skeletal mesh asset set

#rnx

diff --git a/Engine/Source/Runtime/Engine/Private/Components/SkeletalMeshComponent.cpp b/Engine/Source/Runtime/Engine/Private/Components/SkeletalMeshComponent.cpp
--- a/Engine/Source/Runtime/Engine/Private/Components/SkeletalMeshComponent.cpp
+++ b/Engine/Source/Runtime/Engine/Private/Components/SkeletalMeshComponent.cpp
@@ -3272,13 +3272,16 @@ void USkeletalMeshComponent::ResetLinkedAnimInstances()
 
 void USkeletalMeshComponent::ResetToRefPose()
 {
-	PRAGMA_DISABLE_DEPRECATION_WARNINGS
-	BoneSpaceTransforms = GetSkeletalMeshAsset()->GetRefSkeleton().GetRefBonePose();
-	//Mini RefreshBoneTransforms (the bit we actually care about)
-	GetSkeletalMeshAsset()->FillComponentSpaceTransforms(BoneSpaceTransforms, FillComponentSpaceTransformsRequiredBones, GetEditableComponentSpaceTransforms());
-	PRAGMA_ENABLE_DEPRECATION_WARNINGS
-	bNeedToFlipSpaceBaseBuffers = true; // Have updated space bases so need to flip
-	FlipEditableSpaceBases();
+	if (USkeletalMesh* SkelMesh = GetSkeletalMeshAsset())
+	{
+PRAGMA_DISABLE_DEPRECATION_WARNINGS
+		BoneSpaceTransforms = SkelMesh->GetRefSkeleton().GetRefBonePose();
+		//Mini RefreshBoneTransforms (the bit we actually care about)
+		SkelMesh->FillComponentSpaceTransforms(BoneSpaceTransforms, FillComponentSpaceTransformsRequiredBones, GetEditableComponentSpaceTransforms());
+PRAGMA_ENABLE_DEPRECATION_WARNINGS
+		bNeedToFlipSpaceBaseBuffers = true; // Have updated space bases so need to flip
+		FlipEditableSpaceBases();
+	}
 }
 
 void USkeletalMeshComponent::AllowQueuedAnimEventsNextDispatch()
