Fix incorrect material's connected inputs, causing decals to select an blending mode at runtime in cook build with Substrate.

diff --git a/Engine/Source/Runtime/Engine/Private/Materials/Material.cpp b/Engine/Source/Runtime/Engine/Private/Materials/Material.cpp
--- a/Engine/Source/Runtime/Engine/Private/Materials/Material.cpp
+++ b/Engine/Source/Runtime/Engine/Private/Materials/Material.cpp
@@ -2087,6 +2087,49 @@ void UMaterial::UpdateTransientExpressionData()
 #endif // WITH_EDITORONLY_DATA
 
 #if WITH_EDITOR
+
+static FSubstrateMaterialInfo GetSubstrateMaterialInfo(UMaterialEditorOnlyData* EditorOnly)
+{
+	FSubstrateMaterialInfo Out;
+	if (Substrate::IsSubstrateEnabled() && EditorOnly->FrontMaterial.IsConnected())
+	{
+		check(EditorOnly->FrontMaterial.Expression);
+		if (EditorOnly->FrontMaterial.Expression->IsResultSubstrateMaterial(EditorOnly->FrontMaterial.OutputIndex))
+		{
+			EditorOnly->FrontMaterial.Expression->GatherSubstrateMaterialInfo(Out, EditorOnly->FrontMaterial.OutputIndex);
+		}
+	}
+	return Out;
+}
+
+static void UpdatePropertyConnectedMask(const FSubstrateMaterialInfo& InSubstrateMaterialInfo, FMaterialCachedExpressionData* Out)
+{
+	if (Substrate::IsSubstrateEnabled() && InSubstrateMaterialInfo.IsValid() && Out)
+	{
+		// Mask of all input collected by Substrate BSDF nodes
+		static const uint64 ConnectionMask =
+			  (1ull << MP_BaseColor)
+			| (1ull << MP_Metallic)
+			| (1ull << MP_Specular)
+			| (1ull << MP_Roughness)
+			| (1ull << MP_Anisotropy)
+			| (1ull << MP_EmissiveColor)
+			| (1ull << MP_Normal)
+			| (1ull << MP_Tangent)
+			| (1ull << MP_SubsurfaceColor)
+			| (1ull << MP_CustomData0)
+			| (1ull << MP_CustomData1)
+			| (1ull << MP_Opacity)
+			| (1ull << MP_ShadingModel)
+			| (1ull << MP_DiffuseColor)
+			| (1ull << MP_SpecularColor);
+
+		// Override the cached expression data with collected connection from InSubstrateMaterialInfo, but preserve all other input (e.g., refraction)
+		Out->PropertyConnectedMask &= ~ConnectionMask;
+		Out->PropertyConnectedMask |= InSubstrateMaterialInfo.GetPropertyConnected();
+	}
+}
+
 void UMaterial::UpdateCachedExpressionData()
 {
 	//@note FH: temporary preemptive PostLoad until zenloader load ordering improvements
@@ -2124,6 +2167,13 @@ void UMaterial::UpdateCachedExpressionData()
 
 	LocalCachedExpressionData->Validate(*this);
 
+	if (Substrate::IsSubstrateEnabled())
+	{
+		UMaterialEditorOnlyData* EditorOnly = GetEditorOnlyData();
+		FSubstrateMaterialInfo SubstrateMaterialInfo = GetSubstrateMaterialInfo(EditorOnly);
+		UpdatePropertyConnectedMask(SubstrateMaterialInfo, LocalCachedExpressionData);
+	}
+
 	CachedExpressionData.Reset(LocalCachedExpressionData);
 	CachedHLSLTree.Reset(LocalCachedTree);
 	EditorOnlyData->CachedExpressionData = LocalCachedExpressionData->EditorOnlyData;
@@ -5026,45 +5076,18 @@ void UMaterial::RebuildShadingModelField()
 
 	if (Substrate::IsSubstrateEnabled() && EditorOnly->FrontMaterial.IsConnected())
 	{
-		FSubstrateMaterialInfo SubstrateMaterialInfo;
-		check(EditorOnly->FrontMaterial.Expression);
-		if (EditorOnly->FrontMaterial.Expression->IsResultSubstrateMaterial(EditorOnly->FrontMaterial.OutputIndex))
+		FSubstrateMaterialInfo SubstrateMaterialInfo = GetSubstrateMaterialInfo(EditorOnly);
+		UpdatePropertyConnectedMask(SubstrateMaterialInfo, this->CachedExpressionData.Get());
+
+		if (SubstrateMaterialInfo.IsValid() && SubstrateMaterialInfo.GetSubstrateTreeOutOfStackDepthOccurred())
 		{
-			// Mask of all input collected by Substrate BSDF nodes
-			static const uint64 ConnectionMask = 
-				  (1ull << MP_BaseColor)
-				| (1ull << MP_Metallic)
-				| (1ull << MP_Specular)
-				| (1ull << MP_Roughness)
-				| (1ull << MP_Anisotropy)
-				| (1ull << MP_EmissiveColor)
-				| (1ull << MP_Normal)
-				| (1ull << MP_Tangent)
-				| (1ull << MP_SubsurfaceColor)
-				| (1ull << MP_CustomData0)
-				| (1ull << MP_CustomData1)
-				| (1ull << MP_Opacity)
-				| (1ull << MP_ShadingModel)
-				| (1ull << MP_DiffuseColor)
-				| (1ull << MP_SpecularColor);
-
-			EditorOnly->FrontMaterial.Expression->GatherSubstrateMaterialInfo(SubstrateMaterialInfo, EditorOnly->FrontMaterial.OutputIndex);
-
-			// Override the cached expression data with collected connection from SubstrateMaterialInfo, but preserve all other input (e.g., refraction)
-			check(this->CachedExpressionData);
-			this->CachedExpressionData->PropertyConnectedMask &= ~ConnectionMask;
-			this->CachedExpressionData->PropertyConnectedMask |= SubstrateMaterialInfo.GetPropertyConnected();
-
-			if (SubstrateMaterialInfo.GetSubstrateTreeOutOfStackDepthOccurred())
-			{
-				SubstrateMaterialInfo.AddShadingModel(SSM_Unlit);
-				MaterialDomain = EMaterialDomain::MD_Surface;
-				ShadingModel = MSM_Unlit;
-				ShadingModels.AddShadingModel(MSM_Unlit);
-				CancelOutstandingCompilation();
-				UE_LOG(LogMaterial, Error, TEXT("%s: Substrate - Cyclic graph detected when we only support acyclic graph."), *GetName());
-				return;
-			}
+			SubstrateMaterialInfo.AddShadingModel(SSM_Unlit);
+			MaterialDomain = EMaterialDomain::MD_Surface;
+			ShadingModel = MSM_Unlit;
+			ShadingModels.AddShadingModel(MSM_Unlit);
+			CancelOutstandingCompilation();
+			UE_LOG(LogMaterial, Error, TEXT("%s: Substrate - Cyclic graph detected when we only support acyclic graph."), *GetName());
+			return;
 		}
 
 		bool bSanitizeMaterial = false;
