[World Partition] Keep container instances actors as part of the streaming generation process if they are not editor-only.
This allows having components on the actor that will stream in the world instead of losing the actor after embedding.

fyi pamela.figueroa

diff --git a/Engine/Source/Runtime/Engine/Private/LevelInstance/LevelInstanceActor.cpp b/Engine/Source/Runtime/Engine/Private/LevelInstance/LevelInstanceActor.cpp
--- a/Engine/Source/Runtime/Engine/Private/LevelInstance/LevelInstanceActor.cpp
+++ b/Engine/Source/Runtime/Engine/Private/LevelInstance/LevelInstanceActor.cpp
@@ -27,6 +27,7 @@ ALevelInstance::ALevelInstance()
 	RootComponent->Mobility = EComponentMobility::Static;
 
 #if WITH_EDITORONLY_DATA
+	bIsEditorOnlyActor = true;
 	DesiredRuntimeBehavior = ELevelInstanceRuntimeBehavior::Partitioned;
 #endif
 }
@@ -366,6 +367,16 @@ bool ALevelInstance::EditorCanAttachFrom(const AActor* InChild, FText& OutReason
 	return false;
 }
 
+bool ALevelInstance::IsEditorOnly() const
+{
+	if (DesiredRuntimeBehavior == ELevelInstanceRuntimeBehavior::Partitioned)
+	{
+		return Super::IsEditorOnly();
+	}
+
+	return false;
+}
+
 FBox ALevelInstance::GetComponentsBoundingBox(bool bNonColliding, bool bIncludeFromChildActors) const
 {
 	FBox Box = Super::GetComponentsBoundingBox(bNonColliding, bIncludeFromChildActors);
diff --git a/Engine/Source/Runtime/Engine/Private/PackedLevelActor/PackedLevelActor.cpp b/Engine/Source/Runtime/Engine/Private/PackedLevelActor/PackedLevelActor.cpp
--- a/Engine/Source/Runtime/Engine/Private/PackedLevelActor/PackedLevelActor.cpp
+++ b/Engine/Source/Runtime/Engine/Private/PackedLevelActor/PackedLevelActor.cpp
@@ -22,6 +22,7 @@ APackedLevelActor::APackedLevelActor()
 #if WITH_EDITORONLY_DATA
 	// Packed Level Instances don't support level streaming or sub actors
 	DesiredRuntimeBehavior = ELevelInstanceRuntimeBehavior::None;
+	bIsEditorOnlyActor = false;
 #endif
 }
 
diff --git a/Engine/Source/Runtime/Engine/Private/WorldPartition/WorldPartitionStreamingGeneration.cpp b/Engine/Source/Runtime/Engine/Private/WorldPartition/WorldPartitionStreamingGeneration.cpp
--- a/Engine/Source/Runtime/Engine/Private/WorldPartition/WorldPartitionStreamingGeneration.cpp
+++ b/Engine/Source/Runtime/Engine/Private/WorldPartition/WorldPartitionStreamingGeneration.cpp
@@ -781,24 +781,28 @@ class FWorldPartitionStreamingGenerator
 				{
 					*OutActor = Actor;
 				}
-				return !Actor->IsEditorOnly();
 			}
 
-			return !InActorDescInstance->GetActorIsEditorOnly();
+			return true;
 		};
 
 		// Register the actor descriptor view
-		auto RegisterActorDescView = [this, &OutActorDescViewMap, &OutContainerInstances](FStreamingGenerationActorDescView&& InActorDescView)
+		auto RegisterActorDescView = [this, &OutActorDescViewMap, &OutContainerInstances](FStreamingGenerationActorDescView&& InActorDescView, TSet<FGuid>* OutEditorOnlyActorDescSet = nullptr, AActor* InActor = nullptr)
 		{
 			if (InActorDescView.IsChildContainerInstance())
 			{
 				OutContainerInstances.Add(InActorDescView);
 			}
-			else
+
+			if (!InActorDescView.GetActorIsEditorOnly() || (InActor && !InActor->IsEditorOnly()))
 			{
 				const FGuid ActorGuid = InActorDescView.GetGuid();
 				OutActorDescViewMap.Emplace(ActorGuid, MoveTemp(InActorDescView));
 			}
+			else if (OutEditorOnlyActorDescSet)
+			{
+				OutEditorOnlyActorDescSet->Add(InActorDescView.GetGuid());
+			}
 		};
 
 		for (FStreamingGenerationContainerInstanceCollection::TConstIterator<> Iterator(&InActorDescCollection); Iterator; ++Iterator)
@@ -832,17 +836,13 @@ class FWorldPartitionStreamingGenerator
 					{
 						// Dirty, unsaved actor for PIE
 						TUniquePtr<FStreamingGenerationUnsavedDirtyActorDescInstance>& UnsavedDirtyRef = OutUnsavedDirtyInstances.Add_GetRef(FStreamingGenerationUnsavedDirtyActorDescInstance::Create(Actor, InActorDescCollection));
-						RegisterActorDescView(FStreamingGenerationActorDescView(OutActorDescViewMap, UnsavedDirtyRef.Get(), true));
+						RegisterActorDescView(FStreamingGenerationActorDescView(OutActorDescViewMap, UnsavedDirtyRef.Get(), true), &OutEditorOnlyActorDescSet, Actor);
 						continue;
 					}
 				}
 
 				// Non-dirty actor
-				RegisterActorDescView(FStreamingGenerationActorDescView(OutActorDescViewMap, *Iterator));
-			}
-			else
-			{
-				OutEditorOnlyActorDescSet.Add(Iterator->GetGuid());
+				RegisterActorDescView(FStreamingGenerationActorDescView(OutActorDescViewMap, *Iterator), &OutEditorOnlyActorDescSet);
 			}
 		}
 
diff --git a/Engine/Source/Runtime/Engine/Public/LevelInstance/LevelInstanceActor.h b/Engine/Source/Runtime/Engine/Public/LevelInstance/LevelInstanceActor.h
--- a/Engine/Source/Runtime/Engine/Public/LevelInstance/LevelInstanceActor.h
+++ b/Engine/Source/Runtime/Engine/Public/LevelInstance/LevelInstanceActor.h
@@ -110,6 +110,7 @@ public:
 	ENGINE_API virtual bool GetSoftReferencedContentObjects(TArray<FSoftObjectPath>& SoftObjects) const override;
 	ENGINE_API virtual bool OpenAssetEditor() override;
 	ENGINE_API virtual bool EditorCanAttachFrom(const AActor* InChild, FText& OutReason) const override;	
+	ENGINE_API virtual bool IsEditorOnly() const override;
 	ENGINE_API virtual bool IsUserManaged() const override;
 	ENGINE_API virtual bool ShouldExport() override;
 	ENGINE_API virtual bool SupportsSubRootSelection() const override { return true; }
