Fix ray tracing objects being created on platforms that don't use ray tracing in some code paths.
- Updated UStaticMesh::BuildFromMeshDescriptions(...) to check IsRayTracingAllowed() on non-editor builds.
- Added logic to serialization functions to strip ray tracing representation during cooking if target platform doesn't use ray tracing.

diff --git a/Engine/Source/Runtime/Core/Private/UObject/FortniteMainBranchObjectVersion.cpp b/Engine/Source/Runtime/Core/Private/UObject/FortniteMainBranchObjectVersion.cpp
--- a/Engine/Source/Runtime/Core/Private/UObject/FortniteMainBranchObjectVersion.cpp
+++ b/Engine/Source/Runtime/Core/Private/UObject/FortniteMainBranchObjectVersion.cpp
@@ -13,7 +13,7 @@ TMap<FGuid, FGuid> FFortniteMainBranchObjectVersion::GetSystemGuids()
 	SystemGuids.Add(DevGuids.NIAGARASHADERMAP_DERIVEDDATA_VER, FGuid("C417A50271B3427D98785388266E3B66"));
 	SystemGuids.Add(DevGuids.Niagara_LatestScriptCompileVersion, FGuid("E486231DA967C64B95D4E73131056FA0"));
 	SystemGuids.Add(DevGuids.SkeletalMeshDerivedDataVersion, FGuid("B9AEAA3EE7AC4FCFB3860E8F3E06DD75"));
-	SystemGuids.Add(DevGuids.STATICMESH_DERIVEDDATA_VER, FGuid("68606604DE7F4CE7939725B8B06015B3"));
+	SystemGuids.Add(DevGuids.STATICMESH_DERIVEDDATA_VER, FGuid("1C2EBE3310464FD7995ADE1D6AB92BDC"));
 	SystemGuids.Add(DevGuids.MaterialTranslationDDCVersion, FGuid("5F69E03D3B204DA38D8CF564CB24F23F"));
 	return SystemGuids;
 }
diff --git a/Engine/Source/Runtime/Engine/Private/StaticMesh.cpp b/Engine/Source/Runtime/Engine/Private/StaticMesh.cpp
--- a/Engine/Source/Runtime/Engine/Private/StaticMesh.cpp
+++ b/Engine/Source/Runtime/Engine/Private/StaticMesh.cpp
@@ -760,6 +760,13 @@ void FStaticMeshLODResources::Serialize(FArchive& Ar, UObject* Owner, int32 Inde
 		{
 			bool bHasRayTracingGeometry = (RayTracingGeometry != nullptr);
 
+#if WITH_EDITOR
+			if (Ar.IsCooking() && !Ar.CookingTarget()->UsesRayTracing())
+			{
+				bHasRayTracingGeometry = false; // strip ray tracing geometry
+			}
+#endif
+
 			Ar << bHasRayTracingGeometry;
 
 			if (bHasRayTracingGeometry && RayTracingGeometry == nullptr)
@@ -2113,6 +2120,13 @@ void FStaticMeshRenderData::Serialize(FArchive& Ar, UStaticMesh* Owner, bool bCo
 	{
 		bool bHasRayTracingProxy = (RayTracingProxy != nullptr);
 
+#if WITH_EDITOR
+		if (Ar.IsCooking() && !Ar.CookingTarget()->UsesRayTracing())
+		{
+			bHasRayTracingProxy = false; // strip ray tracing proxy
+		}
+#endif
+
 		Ar << bHasRayTracingProxy;
 
 		if (bHasRayTracingProxy)
@@ -7687,7 +7701,21 @@ bool UStaticMesh::BuildFromMeshDescriptions(const TArray<const FMeshDescription*
 			BuildFromMeshDescription(*MeshDescriptions[LODIndex], LODResources);
 		}
 
-		GetRenderData()->InitializeRayTracingRepresentationFromRenderingLODs();
+		bool bInitializeRayTracingRepresentation = bSupportRayTracing;
+
+#if !WITH_EDITOR
+		// non-editor build -> only initialize ray tracing representation if ray tracing is allowed
+		// editor build -> initialize because we don't know TargetPlatform, but might be stripped later during serialization
+		if (!IsRayTracingAllowed())
+		{
+			bInitializeRayTracingRepresentation = false;
+		}
+#endif
+
+		if(bInitializeRayTracingRepresentation)
+		{
+			GetRenderData()->InitializeRayTracingRepresentationFromRenderingLODs();
+		}
 
 		InitResources();
 
