Move UPrimitiveComponent from static instances to dynamic component manager when primitive is updated. This fix primitive bounds to be duplicated in multiple containers and being tested multiple times for nothing in the streaming loop (plus it saves memory)
Added TextureStreaming/RenderAssetStreamingUpdate csv stat

#rnx

diff --git a/Engine/Source/Runtime/Engine/Private/Streaming/AsyncTextureStreaming.cpp b/Engine/Source/Runtime/Engine/Private/Streaming/AsyncTextureStreaming.cpp
--- a/Engine/Source/Runtime/Engine/Private/Streaming/AsyncTextureStreaming.cpp
+++ b/Engine/Source/Runtime/Engine/Private/Streaming/AsyncTextureStreaming.cpp
@@ -8,6 +8,9 @@ AsyncTextureStreaming.cpp: Definitions of classes used for texture streaming asy
 #include "RHI.h"
 #include "Streaming/StreamingManagerTexture.h"
 #include "Engine/Level.h"
+#include "ProfilingDebugging/CsvProfiler.h"
+
+CSV_DECLARE_CATEGORY_EXTERN(TextureStreaming);
 
 void FAsyncRenderAssetStreamingData::Init(
 	TArray<FStreamingViewInfo> InViewInfos,
@@ -909,14 +912,24 @@ void FRenderAssetStreamingMipCalcTask::DoWork()
 	
 	ApplyPakStateChanges_Async();
 
-	for (FStreamingRenderAsset& StreamingRenderAsset : StreamingRenderAssets)
 	{
-		if (IsAborted()) break;
+		uint64 StartTime = FPlatformTime::Cycles64();
+
+		for (FStreamingRenderAsset& StreamingRenderAsset : StreamingRenderAssets)
+		{
+			if (IsAborted()) break;
+
+			StreamingRenderAsset.UpdateOptionalMipsState_Async();
+
+			StreamingData.UpdatePerfectWantedMips_Async(StreamingRenderAsset, Settings);
+			StreamingRenderAsset.DynamicBoostFactor = 1.f; // Reset after every computation.
+		}
+
+		uint64 LenghtCycles64 = (FPlatformTime::Cycles64() - StartTime);
+		double ElapsedMSTime = FPlatformTime::ToMilliseconds64(LenghtCycles64);
 
-		StreamingRenderAsset.UpdateOptionalMipsState_Async();
-		
-		StreamingData.UpdatePerfectWantedMips_Async(StreamingRenderAsset, Settings);
-		StreamingRenderAsset.DynamicBoostFactor = 1.f; // Reset after every computation.
+		// use of a custom stat to have the elapsed time as a global stat and not a stat split accross multiple threads
+		CSV_CUSTOM_STAT(TextureStreaming, RenderAssetStreamingUpdate, float(ElapsedMSTime), ECsvCustomStatOp::Set);
 	}
 
 	// According to budget, make relevant sacrifices and keep possible unwanted mips
diff --git a/Engine/Source/Runtime/Engine/Private/Streaming/StreamingManagerTexture.cpp b/Engine/Source/Runtime/Engine/Private/Streaming/StreamingManagerTexture.cpp
--- a/Engine/Source/Runtime/Engine/Private/Streaming/StreamingManagerTexture.cpp
+++ b/Engine/Source/Runtime/Engine/Private/Streaming/StreamingManagerTexture.cpp
@@ -83,6 +83,13 @@ static TAutoConsoleVariable<int32> CVarFlushDeferredMipLevelChangeCallbacksBefor
 	TEXT("Whether to flush deferred mip level change callbacks before GC."),
 	ECVF_Default);
 
+static int32 GRenderAssetStreamingEnableStrippingDuplicatePrimitiveComponent = 1;
+static FAutoConsoleVariableRef CVarStreamingEnableStrippingDuplicatePrimitiveComponent(
+	TEXT("r.Streaming.EnableStrippingDuplicatePrimitiveComponent"),
+	GRenderAssetStreamingEnableStrippingDuplicatePrimitiveComponent,
+	TEXT("Move UPrimitiveComponent from static instances to dynamic component manager when primitive is updated."),
+	ECVF_Default);
+
 // TODO: Remove once these calls have been proven safe in production
 static TAutoConsoleVariable<int32> CVarProcessAddedRenderAssetsAfterAsyncWork(
 	TEXT("r.Streaming.ProcessAddedRenderAssetsAfterAsyncWork"),
@@ -1180,7 +1187,13 @@ void FRenderAssetStreamingManager::NotifyPrimitiveUpdated_Concurrent( const UPri
 	{
 		FScopeLock ScopeLock(&CriticalSection);
 		FStreamingTextureLevelContext LevelContext(EMaterialQualityLevel::Num);
-		DynamicComponentManager.Add(Primitive, LevelContext);
+		if (DynamicComponentManager.Add(Primitive, LevelContext) == EAddComponentResult::Success)
+		{
+			if (GRenderAssetStreamingEnableStrippingDuplicatePrimitiveComponent>0)
+			{
+				RemoveStaticReferences(Primitive);
+			}
+		}
 	}
 
 	STAT(CallbackCycle += (int32)FPlatformTime::Cycles();)
