Fix FActorIteratorState when iterating over specific actor class, which could return actors not part of any level when unloading world areas.
This reintroduces a major loading time regression when loading large parts of the world in the editor because UFoliageEditorSubsystem::OnActorMoved will be invoked for every actor loading in the editor, which the does a TActorIterator. An upcoming optimization for this case will come later on.

diff --git a/Engine/Source/Runtime/Engine/Public/EngineUtils.h b/Engine/Source/Runtime/Engine/Public/EngineUtils.h
--- a/Engine/Source/Runtime/Engine/Public/EngineUtils.h
+++ b/Engine/Source/Runtime/Engine/Public/EngineUtils.h
@@ -189,36 +189,33 @@ public:
 		// if the filter class is AActor, just use the actors that are in the world you asked for.
 		// This could be useful in runtime code as well if there are many worlds in memory, but for now we will leave
 		// it in editor code.
-		if (InClass == AActor::StaticClass())
-		{
-			// First determine the number of actors in the world to reduce reallocations when we append them to the array below.
-			int32 NumActors = 0;
-			for (ULevel* Level : InWorld->GetLevels())
-			{
-				if (Level)
-				{
-					NumActors += Level->Actors.Num();
-				}
-			}
 
-			// Presize the array
-			ObjectArray.Reserve(NumActors);
+		// First determine the number of actors in the world to reduce reallocations when we append them to the array below.
+		const int32 MaxActors = Algo::Accumulate(InWorld->GetLevels(), 0, [](int32 InTotalNum, const ULevel* InLevel) { return InTotalNum + (InLevel ? InLevel->Actors.Num() : 0); });
+
+		// Presize the array
+		ObjectArray.Reserve(MaxActors);
 
-			// Fill the array
-			for (ULevel* Level : InWorld->GetLevels())
+		for (ULevel* Level : InWorld->GetLevels())
+		{
+			if (Level)
 			{
-				if (Level)
+				if(InClass == AActor::StaticClass())
 				{
 					ObjectArray.Append(Level->Actors);
 				}
+				else
+				{
+					Algo::CopyIf(Level->Actors, ObjectArray, [InClass](AActor* InActor) { return InActor && InActor->IsA(InClass); });
+				}
 			}
 		}
-		else
-#endif // WITH_EDITOR
+#else // WITH_EDITOR
 		{
 			constexpr EObjectFlags ExcludeFlags = RF_ClassDefaultObject;
 			GetObjectsOfClass(InClass, ObjectArray, true, ExcludeFlags, EInternalObjectFlags::Garbage);
 		}
+#endif
 
 		const auto ActorSpawnedDelegate = FOnActorSpawned::FDelegate::CreateRaw(this, &FActorIteratorState::OnActorSpawned);
 		ActorSpawnedDelegateHandle = CurrentWorld->AddOnActorSpawnedHandler(ActorSpawnedDelegate);
@@ -973,4 +970,4 @@ namespace VirtualTextureUtils
 {
 	/** Function that will test if the passed in texture is VT. If so, print to Messagelog that the property does not support VT textures */
 	ENGINE_API void CheckAndReportInvalidUsage(const UObject* Owner, const FName& PropertyName, const UTexture* Texture);
-}
+}
\ No newline at end of file
