Only call the net driver delegate after the actor has been initialized.

diff --git a/Engine/Source/Runtime/Engine/Private/Actor.cpp b/Engine/Source/Runtime/Engine/Private/Actor.cpp
--- a/Engine/Source/Runtime/Engine/Private/Actor.cpp
+++ b/Engine/Source/Runtime/Engine/Private/Actor.cpp
@@ -6261,13 +6261,16 @@ PRAGMA_ENABLE_DEPRECATION_WARNINGS
 
 	UE_LOG(LogNet, Verbose, TEXT("AActor::SetNetUpdateFrequency(): %s Frequency=%f HasActorBegunPlay()=%d"), *this->GetFullName(), Frequency, HasActorBegunPlay());
 
-	ActorUtils::ForEachNetDriver(GEngine, GetWorld(), [this](UNetDriver* NetDriver)
+	if (IsActorInitialized())
 	{
-		if (NetDriver)
+		ActorUtils::ForEachNetDriver(GEngine, GetWorld(), [this](UNetDriver* NetDriver)
 		{
-			NetDriver->GetOnNetUpdateFrequencyChanged().Broadcast(this);
-		}
-	});
+			if (NetDriver)
+			{
+				NetDriver->GetOnNetUpdateFrequencyChanged().Broadcast(this);
+			}
+		});
+	}
 }
 
 float AActor::GetNetUpdateFrequency() const
