Fix indeterminism in ActorGuid when creating WorldSettings actors in generated packages or due to RepairWorldSettings. Actors created at cook time should get a deterministic ActorGuid.
This fixes indeterminism in Lyra generated and maps that convert WorldSettings to LyraWorldSettings.
#rnx

diff --git a/Engine/Source/Runtime/Engine/Private/LevelActor.cpp b/Engine/Source/Runtime/Engine/Private/LevelActor.cpp
--- a/Engine/Source/Runtime/Engine/Private/LevelActor.cpp
+++ b/Engine/Source/Runtime/Engine/Private/LevelActor.cpp
@@ -534,17 +534,6 @@ AActor* UWorld::SpawnActor( UClass* Class, FTransform const* UserTransformPtr, c
 	bool bNeedGloballyUniqueName = false;
 
 #if WITH_EDITOR
-	// Generate the actor's Guid
-	FGuid ActorGuid;
-	if (SpawnParameters.OverrideActorGuid.IsValid())
-	{
-		ActorGuid = SpawnParameters.OverrideActorGuid;
-	}
-	else
-	{
-		ActorGuid = FGuid::NewGuid();
-	}
-
 	// Generate and set the actor's external package if needed
 	if (SpawnParameters.OverridePackage)
 	{
@@ -691,6 +680,22 @@ AActor* UWorld::SpawnActor( UClass* Class, FTransform const* UserTransformPtr, c
 #if WITH_EDITOR
 	Actor->ClearActorLabel(); // Clear label on newly spawned actors
 
+	// Generate the actor's Guid
+	FGuid ActorGuid;
+	if (SpawnParameters.OverrideActorGuid.IsValid())
+	{
+		ActorGuid = SpawnParameters.OverrideActorGuid;
+	}
+	else if (IsRunningCookCommandlet())
+	{
+		// Generate the guid deterministically when creating actors at cook time
+		ActorGuid = FGuid::NewDeterministicGuid(Actor->GetPathName());
+	}
+	else
+	{
+		ActorGuid = FGuid::NewGuid();
+	}
+
 	// Set the actor's guid
 	FSetActorGuid SetActorGuid(Actor, ActorGuid);
 #endif
