Substrate - better compilation error messages when using some BSDF and operators.
Tested ElectricDreams

#rnx

diff --git a/Engine/Shaders/Public/ShaderVersion.ush b/Engine/Shaders/Public/ShaderVersion.ush
--- a/Engine/Shaders/Public/ShaderVersion.ush
+++ b/Engine/Shaders/Public/ShaderVersion.ush
@@ -3,4 +3,4 @@
 // Each time the console command is executed it generates a new GUID. As this file is included
 // in Platform.ush (which should be included in any shader) it allows to invalidate the shader DDC.
 // 
-#pragma message("UESHADERMETADATA_VERSION E2287832-4546-2EF9-CAA2-7899945AFC72")
\ No newline at end of file
+#pragma message("UESHADERMETADATA_VERSION 000EF489-EECA-4D26-84B4-15B3EC5B68FF")
\ No newline at end of file
diff --git a/Engine/Shaders/Shared/SubstrateDefinitions.h b/Engine/Shaders/Shared/SubstrateDefinitions.h
--- a/Engine/Shaders/Shared/SubstrateDefinitions.h
+++ b/Engine/Shaders/Shared/SubstrateDefinitions.h
@@ -59,6 +59,13 @@
 #define SUBSTRATE_OPERATOR_BSDF							4
 #define SUBSTRATE_OPERATOR_BSDF_LEGACY					5
 
+// Sometimes, Unlit or Weight operators are used to transport data for other meaning (e.g. Light Function or ConvertToDecal)
+#define SUBSTRATE_OPERATOR_SUBUSAGE_NONE				0
+#define SUBSTRATE_OPERATOR_SUBUSAGE_UI					1
+#define SUBSTRATE_OPERATOR_SUBUSAGE_LIGHTFUNCTION		2
+#define SUBSTRATE_OPERATOR_SUBUSAGE_POSTPROCESS			3
+#define SUBSTRATE_OPERATOR_SUBUSAGE_DECAL				4
+
 // This must map directly to ESubstrateTileMaterialType
 #define SUBSTRATE_TILE_TYPE_SIMPLE						0
 #define SUBSTRATE_TILE_TYPE_SINGLE						1
diff --git a/Engine/Source/Runtime/Core/Private/UObject/FortniteMainBranchObjectVersion.cpp b/Engine/Source/Runtime/Core/Private/UObject/FortniteMainBranchObjectVersion.cpp
--- a/Engine/Source/Runtime/Core/Private/UObject/FortniteMainBranchObjectVersion.cpp
+++ b/Engine/Source/Runtime/Core/Private/UObject/FortniteMainBranchObjectVersion.cpp
@@ -14,6 +14,6 @@ TMap<FGuid, FGuid> FFortniteMainBranchObjectVersion::GetSystemGuids()
 	SystemGuids.Add(DevGuids.Niagara_LatestScriptCompileVersion, FGuid("E486231DA967C64B95D4E73131056FA0"));
 	SystemGuids.Add(DevGuids.SkeletalMeshDerivedDataVersion, FGuid("B9AEAA3EE7AC4FCFB3860E8F3E06DD75"));
 	SystemGuids.Add(DevGuids.STATICMESH_DERIVEDDATA_VER, FGuid("1C2EBE3310464FD7995ADE1D6AB92BDC"));
-	SystemGuids.Add(DevGuids.MaterialTranslationDDCVersion, FGuid("5F69E03D3B204DA38D8CF564CB24F23F"));
+	SystemGuids.Add(DevGuids.MaterialTranslationDDCVersion, FGuid("C42A7403BBDD420EBA946325D6F79B1B"));
 	return SystemGuids;
 }
diff --git a/Engine/Source/Runtime/Engine/Private/Materials/HLSLMaterialTranslator.cpp b/Engine/Source/Runtime/Engine/Private/Materials/HLSLMaterialTranslator.cpp
--- a/Engine/Source/Runtime/Engine/Private/Materials/HLSLMaterialTranslator.cpp
+++ b/Engine/Source/Runtime/Engine/Private/Materials/HLSLMaterialTranslator.cpp
@@ -12288,6 +12288,7 @@ FSubstrateOperator& FHLSLMaterialTranslator::SubstrateCompilationRegisterOperato
 	NewOperator.RightIndex   = INDEX_NONE;
 
 	NewOperator.BSDFIndex = INDEX_NONE;	// Allocated later to be able to account for inline
+	NewOperator.SubUsage = SUBSTRATE_OPERATOR_SUBUSAGE_NONE;
 
 	NewOperator.MaxDistanceFromLeaves = 0;
 	NewOperator.bIsBottom = false;
@@ -12704,6 +12705,10 @@ bool FHLSLMaterialTranslator::FSubstrateCompilationContext::SubstrateGenerateDer
 		bool bHasEye = false;
 		bool bHasSLW = false;
 		bool bHasSlab = false;
+		bool bHasUI = false;
+		bool bHasDecal = false;
+		bool bHasPostProcess = false;
+		bool bHasLightFunction = false;
 		{
 			int VOpTopBranchCountTaken = 0;
 			int VOpBottomBranchCountTaken = 0;
@@ -12725,6 +12730,12 @@ bool FHLSLMaterialTranslator::FSubstrateCompilationContext::SubstrateGenerateDer
 					CurrentOperator.bRootOfParameterBlendingSubTree = bRootOfParameterBlendingSubTree;
 				}
 
+				// this can show up on Unlit and Weight operators for Decals.
+				bHasUI				|= CurrentOperator.SubUsage == SUBSTRATE_OPERATOR_SUBUSAGE_UI;
+				bHasDecal			|= CurrentOperator.SubUsage == SUBSTRATE_OPERATOR_SUBUSAGE_DECAL;
+				bHasPostProcess		|= CurrentOperator.SubUsage == SUBSTRATE_OPERATOR_SUBUSAGE_POSTPROCESS;
+				bHasLightFunction	|= CurrentOperator.SubUsage == SUBSTRATE_OPERATOR_SUBUSAGE_LIGHTFUNCTION;
+
 				switch (CurrentOperator.OperatorType)
 				{
 				case SUBSTRATE_OPERATOR_VERTICAL:
@@ -12790,18 +12801,27 @@ bool FHLSLMaterialTranslator::FSubstrateCompilationContext::SubstrateGenerateDer
 			}
 			bSubstrateMaterialIsUnlitNode = bHasUnlit;
 
-			if (((bHasVFogCloud || bHasHair || bHasEye || bHasSLW) && bOperatorEncountered)
-				|| (bHasUnlit && bOperatorEncounteredButNotWeight)
-				)
+			if ((bHasVFogCloud || bHasHair || bHasEye || bHasSLW) && bOperatorEncountered)
+			{
+				Compiler->Errorf(TEXT("Fog/Cloud, Hair, Eye or SingleLayerWater cannot be used with operators. See %s (asset: %s).\r\n"), *CompilerMaterial->GetDebugName(), *CompilerMaterial->GetAssetPath().ToString());
+				return false;
+			}
+
+			if ((bHasUI ||  bHasPostProcess || bHasLightFunction) && bOperatorEncountered)
+			{
+				Compiler->Errorf(TEXT("UI, Post Process or Light Function materials cannot be used with operators. See %s (asset: %s).\r\n"), *CompilerMaterial->GetDebugName(), *CompilerMaterial->GetAssetPath().ToString());
+				return false;
+			}
+
+			if (bHasUnlit && bOperatorEncounteredButNotWeight)
 			{
-				Compiler->Errorf(TEXT("Fog/Cloud, Hair or SingleLayerWater cannot be used with operators.\r\nUnlit can only be used with coverage operators. See %s (asset: %s).\r\n"), *CompilerMaterial->GetDebugName(), *CompilerMaterial->GetAssetPath().ToString());
-				// This is because it will results in lighting loops dedicated on slab and simply a few other isolated lighting cases.
+				Compiler->Errorf(TEXT("Unlit can only be used with coverage operators. See %s (asset: %s).\r\n"), *CompilerMaterial->GetDebugName(), *CompilerMaterial->GetAssetPath().ToString());
 				return false;
 			}
 
 			if ((bHasUnlit || bHasVFogCloud || bHasHair || bHasEye || bHasSLW) && SubstrateMaterialEffectiveClosureCount > 1)
 			{
-				Compiler->Errorf(TEXT("Unlit, Fog/Cloud, Hair or SingleLayerWater must be used in isolation. See %s (asset: %s).\r\n"), *CompilerMaterial->GetDebugName(), *CompilerMaterial->GetAssetPath().ToString());
+				Compiler->Errorf(TEXT("Unlit, Fog/Cloud, Hair or SingleLayerWater must result in a single one Closure (Found %d). See %s (asset: %s).\r\n"), SubstrateMaterialEffectiveClosureCount, *CompilerMaterial->GetDebugName(), *CompilerMaterial->GetAssetPath().ToString());
 				// Even though technically we could support Unlit combined with Slab.
 				return false;
 			}
diff --git a/Engine/Source/Runtime/Engine/Private/Materials/MaterialExpressions.cpp b/Engine/Source/Runtime/Engine/Private/Materials/MaterialExpressions.cpp
--- a/Engine/Source/Runtime/Engine/Private/Materials/MaterialExpressions.cpp
+++ b/Engine/Source/Runtime/Engine/Private/Materials/MaterialExpressions.cpp
@@ -27981,6 +27981,7 @@ FSubstrateOperator* UMaterialExpressionSubstrateLightFunction::SubstrateGenerate
 {
 	FSubstrateOperator& SubstrateOperator = Compiler->SubstrateCompilationRegisterOperator(SUBSTRATE_OPERATOR_BSDF, Compiler->SubstrateTreeStackGetPathUniqueId(), this, Parent, Compiler->SubstrateTreeStackGetParentPathUniqueId());
 	SubstrateOperator.BSDFType = SUBSTRATE_BSDF_TYPE_UNLIT;
+	SubstrateOperator.SubUsage = SUBSTRATE_OPERATOR_SUBUSAGE_LIGHTFUNCTION;
 	SubstrateOperator.ThicknessIndex = Compiler->SubstrateThicknessStackGetThicknessIndex();
 	return &SubstrateOperator;
 }
@@ -28058,6 +28059,7 @@ FSubstrateOperator* UMaterialExpressionSubstratePostProcess::SubstrateGenerateMa
 {
 	FSubstrateOperator& SubstrateOperator = Compiler->SubstrateCompilationRegisterOperator(SUBSTRATE_OPERATOR_BSDF, Compiler->SubstrateTreeStackGetPathUniqueId(), this, Parent, Compiler->SubstrateTreeStackGetParentPathUniqueId());
 	SubstrateOperator.BSDFType = SUBSTRATE_BSDF_TYPE_UNLIT;
+	SubstrateOperator.SubUsage = SUBSTRATE_OPERATOR_SUBUSAGE_POSTPROCESS;
 	SubstrateOperator.ThicknessIndex = Compiler->SubstrateThicknessStackGetThicknessIndex();
 	return &SubstrateOperator;
 }
@@ -28133,6 +28135,7 @@ FSubstrateOperator* UMaterialExpressionSubstrateUI::SubstrateGenerateMaterialTop
 {
 	FSubstrateOperator& SubstrateOperator = Compiler->SubstrateCompilationRegisterOperator(SUBSTRATE_OPERATOR_BSDF, Compiler->SubstrateTreeStackGetPathUniqueId(), this, Parent, Compiler->SubstrateTreeStackGetParentPathUniqueId());
 	SubstrateOperator.BSDFType = SUBSTRATE_BSDF_TYPE_UNLIT;
+	SubstrateOperator.SubUsage = SUBSTRATE_OPERATOR_SUBUSAGE_UI;
 	SubstrateOperator.ThicknessIndex = Compiler->SubstrateThicknessStackGetThicknessIndex();
 	return &SubstrateOperator;
 }
@@ -28242,6 +28245,7 @@ FSubstrateOperator* UMaterialExpressionSubstrateConvertToDecal::SubstrateGenerat
 {
 	const bool bUseParameterBlending = true;
 	FSubstrateOperator& SubstrateOperator = Compiler->SubstrateCompilationRegisterOperator(SUBSTRATE_OPERATOR_WEIGHT, Compiler->SubstrateTreeStackGetPathUniqueId(), this, Parent, Compiler->SubstrateTreeStackGetParentPathUniqueId(), bUseParameterBlending);
+	SubstrateOperator.SubUsage = SUBSTRATE_OPERATOR_SUBUSAGE_DECAL;
 	if (Compiler->GetSubstrateTreeOutOfStackDepthOccurred())
 	{
 		return &SubstrateOperator; // Out ot stack space, return now to fail the compilation
diff --git a/Engine/Source/Runtime/Engine/Public/Rendering/SubstrateMaterialShared.h b/Engine/Source/Runtime/Engine/Public/Rendering/SubstrateMaterialShared.h
--- a/Engine/Source/Runtime/Engine/Public/Rendering/SubstrateMaterialShared.h
+++ b/Engine/Source/Runtime/Engine/Public/Rendering/SubstrateMaterialShared.h
@@ -56,6 +56,7 @@ public:
 	LAYOUT_FIELD_EDITORONLY(uint8, bBSDFHasAnisotropy);
 	LAYOUT_FIELD_EDITORONLY(uint8, bBSDFHasGlint);
 	LAYOUT_FIELD_EDITORONLY(uint8, bBSDFHasSpecularProfile);
+	LAYOUT_FIELD_EDITORONLY(uint8, SubUsage);		// Sometimes, Unlit or Weight operators are used to transport data for other meaning (e.g. Light Function or ConvertToDecal)
 
 	LAYOUT_FIELD_EDITORONLY(uint8, bBSDFWritesEmissive);
 	LAYOUT_FIELD_EDITORONLY(uint8, bBSDFWritesAmbientOcclusion);
