- Mark the objects as modified when calling UStreamableSparseVolumeTexture::BeginInitialize otherwise the user will not be prompted to save the asset

#rnx
#tests Baked an SVT over an existing asset to confirm it was marked modified

diff --git a/Engine/Source/Runtime/Engine/Private/SparseVolumeTexture/SparseVolumeTexture.cpp b/Engine/Source/Runtime/Engine/Private/SparseVolumeTexture/SparseVolumeTexture.cpp
--- a/Engine/Source/Runtime/Engine/Private/SparseVolumeTexture/SparseVolumeTexture.cpp
+++ b/Engine/Source/Runtime/Engine/Private/SparseVolumeTexture/SparseVolumeTexture.cpp
@@ -1282,6 +1282,8 @@ bool UStreamableSparseVolumeTexture::BeginInitialize(int32 NumExpectedFrames)
 		return false;
 	}
 
+	Modify();
+
 	check(Frames.IsEmpty());
 	Frames.Empty(NumExpectedFrames);
 	VolumeBoundsMin = FIntVector(INT32_MAX, INT32_MAX, INT32_MAX);
