Substrate - fix for crash on mac having mismatched  RenderTarget / Shader output format.

#rnx

diff --git a/Engine/Shaders/Public/ShaderVersion.ush b/Engine/Shaders/Public/ShaderVersion.ush
--- a/Engine/Shaders/Public/ShaderVersion.ush
+++ b/Engine/Shaders/Public/ShaderVersion.ush
@@ -3,4 +3,4 @@
 // Each time the console command is executed it generates a new GUID. As this file is included
 // in Platform.ush (which should be included in any shader) it allows to invalidate the shader DDC.
 // 
-#pragma message("UESHADERMETADATA_VERSION 000EF489-EECA-4D26-84B4-15B3EC5B68FF")
\ No newline at end of file
+#pragma message("UESHADERMETADATA_VERSION 42D3EF03-C067-42EF-9081-1471131640AC")
\ No newline at end of file
diff --git a/Engine/Source/Runtime/Engine/Private/ShaderCompiler/ShaderGenerationUtil.cpp b/Engine/Source/Runtime/Engine/Private/ShaderCompiler/ShaderGenerationUtil.cpp
--- a/Engine/Source/Runtime/Engine/Private/ShaderCompiler/ShaderGenerationUtil.cpp
+++ b/Engine/Source/Runtime/Engine/Private/ShaderCompiler/ShaderGenerationUtil.cpp
@@ -1821,7 +1821,10 @@ static void DetermineUsedMaterialSlots(
 	// we have to use if statements, not switch or if/else statements because we can have multiple shader model ids.
 	if (Mat.MATERIAL_SHADINGMODEL_UNLIT)
 	{
-		SetStandardGBufferSlots(Slots, true, false, bHasVelocity, bWritesVelocity, false, bIsSubstrateMaterial);
+		// With Substrate, we still need to write to the GBufferE when static lighting is used to maintain the correct offset for SubstrateFirstMRT (and match the bound render targets).
+		// Without Substrate, the buffer is always bound and written to. It is just not used and a default value was used.
+		const bool bHasStaticLightingWritten = bIsSubstrateMaterial ? bHasStaticLighting : false;
+		SetStandardGBufferSlots(Slots, true, false, bHasVelocity, bWritesVelocity, bHasStaticLightingWritten, bIsSubstrateMaterial);
 	}
 
 	if (Mat.MATERIAL_SHADINGMODEL_DEFAULT_LIT)
