Fix for live locking the cook when texture builds fail.

#rnx

diff --git a/Engine/Source/Runtime/Engine/Classes/Engine/Texture.h b/Engine/Source/Runtime/Engine/Classes/Engine/Texture.h
--- a/Engine/Source/Runtime/Engine/Classes/Engine/Texture.h
+++ b/Engine/Source/Runtime/Engine/Classes/Engine/Texture.h
@@ -1783,9 +1783,9 @@ public:
 	ENGINE_API virtual void BeginCacheForCookedPlatformData(  const ITargetPlatform* TargetPlatform ) override;
 
 	/**
-	 * Have we finished loading all the cooked platform data for the target platforms requested in BeginCacheForCookedPlatformData
-	 * 
-	 * @param	TargetPlatform target platform to check for cooked platform data
+	 * Have we finished loading all the cooked platform data for the given platform.
+	 * Note that this return true when there's no more work to be done, which includes failure
+	 * due to errors. Never returning true will livelock the cook!
 	 */
 	ENGINE_API virtual bool IsCachedCookedPlatformDataLoaded( const ITargetPlatform* TargetPlatform ) override;
 
diff --git a/Engine/Source/Runtime/Engine/Private/TextureDerivedData.cpp b/Engine/Source/Runtime/Engine/Private/TextureDerivedData.cpp
--- a/Engine/Source/Runtime/Engine/Private/TextureDerivedData.cpp
+++ b/Engine/Source/Runtime/Engine/Private/TextureDerivedData.cpp
@@ -3895,7 +3895,7 @@ bool UTexture::IsCachedCookedPlatformDataLoaded(const ITargetPlatform* TargetPla
 
 	if ( ! CanBuildPlatformData(TargetPlatform) )
 	{
-		return false;
+		return true; // signify that the cook should move on without us.
 	}
 
 	// CookedPlatformData is keyed off of FetchOrBuild settings.
