[World Partition] Fix regression for unsaved level instances following CL 35634490.

diff --git a/Engine/Source/Runtime/Engine/Private/WorldPartition/WorldPartitionStreamingGeneration.cpp b/Engine/Source/Runtime/Engine/Private/WorldPartition/WorldPartitionStreamingGeneration.cpp
--- a/Engine/Source/Runtime/Engine/Private/WorldPartition/WorldPartitionStreamingGeneration.cpp
+++ b/Engine/Source/Runtime/Engine/Private/WorldPartition/WorldPartitionStreamingGeneration.cpp
@@ -878,7 +878,7 @@ class FWorldPartitionStreamingGenerator
 				// Here, FindHandlingContainer is used to make sure that the actor is handled by the collection
 				// The main reason is that UWorldPartition::CheckForErrors currently builds a collection per ActorDescContainer of the WorldPartition.
 				// This is probably a limitation introduced by ContentBundles. 
-				if (IsValid(Actor) && Actor->IsPackageExternal() && Actor->IsMainPackageActor() && !Actor->IsEditorOnly()
+				if (IsValid(Actor) && Actor->IsPackageExternal() && Actor->IsMainPackageActor()
 					&& InActorDescCollection.FindHandlingContainerInstance(Actor)
 					&& !InActorDescCollection.GetActorDescInstance(Actor->GetActorGuid()))
 				{
