Blendspace: enable analysis on drag/drop sample by default

diff --git a/Engine/Source/Runtime/Engine/Classes/Animation/BlendSpace.h b/Engine/Source/Runtime/Engine/Classes/Animation/BlendSpace.h
--- a/Engine/Source/Runtime/Engine/Classes/Animation/BlendSpace.h
+++ b/Engine/Source/Runtime/Engine/Classes/Animation/BlendSpace.h
@@ -208,6 +208,7 @@ struct FBlendSample
 		, SampleValue(InValue)
 		, RateScale(1.0f)
 #if WITH_EDITORONLY_DATA
+		, bIncludeInAnalyseAll(true)
 		, bIsValid(bInIsValid)
 		, CachedMarkerDataUpdateCounter(INDEX_NONE)
 #endif // WITH_EDITORONLY_DATA
