[World Partition] Fixed reloading of streaming cells in PIE containing actors from non-OFPA level instance packages.
Instead of keeping track of when a non-OFPA package is considered empty and can be trashed, we simply append a counter to the package name and let the GC do its job.

Bonus changes:
- Make sure instanced PIE packages have a deterministic name to ease debugging of such issues in the future.
- When unable to find streamed actors in PIE, trigger an ensure instead of simply log to at least have traces of what happened.

diff --git a/Engine/Source/Runtime/Engine/Private/WorldPartition/WorldPartitionLevelHelper.cpp b/Engine/Source/Runtime/Engine/Private/WorldPartition/WorldPartitionLevelHelper.cpp
--- a/Engine/Source/Runtime/Engine/Private/WorldPartition/WorldPartitionLevelHelper.cpp
+++ b/Engine/Source/Runtime/Engine/Private/WorldPartition/WorldPartitionLevelHelper.cpp
@@ -590,19 +590,24 @@ bool FWorldPartitionLevelHelper::LoadActorsInternal(FLoadActorsParams&& InParams
 	TArray<FWorldPartitionRuntimeCellObjectMapping*> ActorPackagesToLoad;
 	TMap<FActorContainerID, FLinkerInstancingContext> LinkerInstancingContexts;
 
+	// Generate a unique name to load a level instance embedded actor if there are multiple instances of this level instance and possibly across 
+	// multiple instances of the WP world:
 	auto GetContainerPackage = [](const FActorContainerID& InContainerID, const FString& InPackageName, const UObject* InContextObject) -> FName
 	{
-		// Generate a unique name to load a Level Instance embedded actor if there are multiple instances of this Level Instance and possibly across multiple instances of the World Partition world
-		// InContainerID will distinguish between instances of the same Level Instance
-		// InContextObject will distinguish between instances of the same top level World Partition world (Only needed in PIE, In cook we always cook the source WP and not an instance and Actor packages no longer exist at runtime)
-		uint64 DestLevelID = 0;
 		TStringBuilder<512> PackageNameBuilder;
+		
+		// Distinguish between instances of the same level instance	
 		PackageNameBuilder.Appendf(TEXT("/Temp%s_%s"), *InPackageName, *InContainerID.ToShortString());
 		
-		if (InContextObject)
-		{
-			PackageNameBuilder.Appendf(TEXT("_%016llx"), InContextObject);
-		}
+		// Distinguish between instances of the same top level WP world; only for PIE, in cook we always cook the source WP and not an instance 
+		// and actor packages no longer exist at runtime)
+		const FString ContextObjectPathName = GetPathNameSafe(InContextObject);
+		const uint64 ContextObjectPathNameHash = CityHash64(TCHAR_TO_ANSI(*ContextObjectPathName), ContextObjectPathName.Len());
+		PackageNameBuilder.Appendf(TEXT("_%llx"), ContextObjectPathNameHash);
+
+		// Distinguish between loading the same package after a reload between GCs (only for PIE)
+		static uint32 ContextObjectUniqueID = 0;
+		PackageNameBuilder.Appendf(TEXT("_%llx"), ContextObjectUniqueID++);
 
 		return PackageNameBuilder.ToString();
 	};
@@ -850,11 +855,9 @@ bool FWorldPartitionLevelHelper::LoadActorsInternal(FLoadActorsParams&& InParams
 			}
 			else
 			{
-				UE_LOG(LogWorldPartition, Warning, TEXT("Failed to find actor '%s' in package '%s'."), *ActorName.ToString(), *LoadedPackageName.ToString());
-
 				if (LoadedPackage)
 				{
-					UE_LOG(LogWorldPartition, Warning, TEXT("\tPackage Content:"), *ActorName.ToString(), *LoadedPackageName.ToString());
+					UE_LOG(LogWorldPartition, Warning, TEXT("\tPackage Content for '%s:"), *LoadedPackage->GetName());
 					ForEachObjectWithOuter(LoadedPackage, [](UObject* Object)
 					{
 						UE_LOG(LogWorldPartition, Warning, TEXT("\t\tObject %s, Flags 0x%llx"), *Object->GetPathName(), static_cast<uint64>(Object->GetFlags()));
@@ -862,6 +865,8 @@ bool FWorldPartitionLevelHelper::LoadActorsInternal(FLoadActorsParams&& InParams
 					}, true);
 				}
 
+				ensureMsgf(false, TEXT("Failed to find actor '%s' in package '%s'."), *ActorName.ToString(), *LoadedPackageName.ToString());
+
 				LoadProgress->NumFailedLoadedRequests++;
 			}
 
