Extend IAudioFormat to allow rejection of formats based on samplerate and channel count as well as call out hardware formats.

[RN] Extend IAudioFormat to allow rejection of formats based on samplerate and channel count as well as call out hardware formats.

diff --git a/Engine/Source/Developer/AudioFormatRad/Private/AudioFormatRad.cpp b/Engine/Source/Developer/AudioFormatRad/Private/AudioFormatRad.cpp
--- a/Engine/Source/Developer/AudioFormatRad/Private/AudioFormatRad.cpp
+++ b/Engine/Source/Developer/AudioFormatRad/Private/AudioFormatRad.cpp
@@ -60,6 +60,12 @@ public:
 	{
 		OutFormats.Add(NAME_RADA);
 	}
+	
+	virtual TConstArrayView<int32> GetSupportedSampleRates() const override
+	{
+		static constexpr int32 Supported[] = { 48000, 44100, 32000, 24000 };
+		return Supported;
+	}
 
 	static void* RadAlloc(const size_t Bytes)
 	{
diff --git a/Engine/Source/Developer/TargetPlatform/Public/Interfaces/IAudioFormat.h b/Engine/Source/Developer/TargetPlatform/Public/Interfaces/IAudioFormat.h
--- a/Engine/Source/Developer/TargetPlatform/Public/Interfaces/IAudioFormat.h
+++ b/Engine/Source/Developer/TargetPlatform/Public/Interfaces/IAudioFormat.h
@@ -4,6 +4,10 @@
 
 #include "CoreMinimal.h"
 
+#include "Algo/IndexOf.h"
+#include "Features/IModularFeatures.h"
+#include "Features/IModularFeature.h"
+
 /** 
  * A non-UObject based structure used to pass data about a sound
  * node wave around the engine and tools.
@@ -36,10 +40,21 @@ struct FSoundQualityInfo
 /**
  * Interface for audio formats.
  */
-class IAudioFormat
+class IAudioFormat : public IModularFeature
 {
 public:
 
+	/**
+	 * Returns the name of the modular feature.
+	 * @return FName of the feature. To be used to query IModularFeatures
+	 */
+	FORCEINLINE static FName GetModularFeatureName()
+	{
+		static const FName Name(TEXT("IAudioFormat"));
+		return Name;
+	}
+	
+	
 	/**
 	 * Checks whether parallel audio cooking is allowed.
 	 *
@@ -140,8 +155,76 @@ public:
 	*/
 	virtual bool ExtractSeekTableForStreaming(TArray<uint8>& InOutBuffer, FSeekTable& OutSeektable) const { return false; }
 
-public:
+	/**
+	 * Returns true if the target platform name is supported 	 
+	 */
+	virtual bool IsPlatformSupported(const FName InPlatformName) const { return true; };
+	
+	/**
+	 * Returns true if this format/codec requires hardware. 
+	 **/
+	virtual bool IsHardwareFormat() const { return false; }
 
-	/** Virtual destructor. */
-	virtual ~IAudioFormat() { }
+	/**
+	 * Returns true if the sample rate is supported by this codec.
+	 * @param InSampleRate The Rate to test
+	 * @return true or false if supported
+	 */
+	bool IsSampleRateSupported(const int32 InSampleRate) const
+	{
+		// Query the supported rates for this Format.
+		const TConstArrayView<int32> SupportedRates = GetSupportedSampleRates();
+		if (SupportedRates.Num() > 0)
+		{
+			return Algo::IndexOf(SupportedRates, InSampleRate) != INDEX_NONE;
+		}
+
+		// Assume success if no supported rates are defined. 
+		return true;
+	}
+
+	/**
+	 * Returns true if the sample rate is supported by this codec.
+	 * @param InChannelCount The channel count to test for
+	 * @return true or false if supported
+	 */
+	bool IsChannelCountSupported(const int32 InChannelCount) const
+	{
+		// Query the supported counts for this Format.
+		const TConstArrayView<int32> SupportedCounts = GetSupportedChannelCounts();
+		if (SupportedCounts.Num() > 0)
+		{
+			return Algo::IndexOf(SupportedCounts, InChannelCount) != INDEX_NONE;
+		}
+
+		// Assume success if no supported rates are defined. 
+		return true;
+	}
+
+	/**
+	 * Returns the list of supported sample rates
+	 * NOTE: An empty list implies all rates are supported.
+	 * @return List of supported sample rate for this format/codec
+	 */
+ 	virtual TConstArrayView<int32> GetSupportedSampleRates() const { return {}; }
+
+	/**
+	 * Returns the list of supported channel counts
+	 * NOTE: An empty list implies all counts are supported.
+	 * @return List of supported channel counts
+	 */
+	virtual TConstArrayView<int32> GetSupportedChannelCounts() const { return {}; }
+
+	/** Self Registering Constructor */
+	IAudioFormat()
+    {
+    	IModularFeatures::Get().RegisterModularFeature(GetModularFeatureName(), this);
+    }
+    
+	/** Self Registering Virtual destructor. */
+	virtual ~IAudioFormat()
+    {
+    	IModularFeatures::Get().UnregisterModularFeature(GetModularFeatureName(), this);
+    }
 };
+
diff --git a/Engine/Source/Runtime/Engine/Classes/Sound/AudioFormatSettings.h b/Engine/Source/Runtime/Engine/Classes/Sound/AudioFormatSettings.h
--- a/Engine/Source/Runtime/Engine/Classes/Sound/AudioFormatSettings.h
+++ b/Engine/Source/Runtime/Engine/Classes/Sound/AudioFormatSettings.h
@@ -2,18 +2,21 @@
 #pragma once
 
 #include "Containers/Array.h"
+#include "Interfaces/ITargetPlatformSettings.h"
 #include "UObject/NameTypes.h"
 
 class FString;
 class FConfigCacheIni;
 class USoundWave;
+class IAudioFormat;
 
 namespace Audio
 {
 	class FAudioFormatSettings
 	{
 	public:
-		ENGINE_API FAudioFormatSettings(FConfigCacheIni* InConfigSystem, const FString& InConfigFilename, const FString& InPlatformIdentifierForLogging);
+		ENGINE_API FAudioFormatSettings(FConfigCacheIni* InConfigSystem, const FString& InConfigFilename, const FString& IniPlatformName);
+
 		~FAudioFormatSettings() = default;
 
 		ENGINE_API FName GetWaveFormat(const USoundWave* Wave) const;
@@ -23,12 +26,23 @@ namespace Audio
 		FName GetFallbackFormat() const { return FallbackFormat; }
 
 	private:
-		ENGINE_API void ReadConfiguration(FConfigCacheIni*, const FString& InConfigFilename, const FString& InPlatformIdentifierForLogging);
+		ENGINE_API void ReadConfiguration(FConfigCacheIni*, const FString& InConfigFilename);
+
+		struct FPlatformWaveState;
+		friend struct FPlatformWaveState;
+				
+		bool IsFormatAllowed(const FPlatformWaveState& InWave) const;
+		const IAudioFormat* FindFormat(const FName& InFormatName) const;
 
+		FName IniPlatformName;
 		TArray<FName> AllWaveFormats;
 		TArray<FName> WaveFormatModuleHints;
 		FName PlatformFormat;
 		FName PlatformStreamingFormat;
 		FName FallbackFormat;
+		mutable FCriticalSection AudioFormatCacheCs;
+		mutable TMap<FName, IAudioFormat*> AudioFormatCache;
 	};
+
+	
 }
diff --git a/Engine/Source/Runtime/Engine/Classes/Sound/SoundWave.h b/Engine/Source/Runtime/Engine/Classes/Sound/SoundWave.h
--- a/Engine/Source/Runtime/Engine/Classes/Sound/SoundWave.h
+++ b/Engine/Source/Runtime/Engine/Classes/Sound/SoundWave.h
@@ -1179,6 +1179,15 @@ public:
 		ImportedSampleRate = InImportedSampleRate;
 #endif
 	}
+	
+	uint32 GetImportedSampleRate() const
+	{
+#if WITH_EDITORONLY_DATA
+		return ImportedSampleRate;
+#else
+		return 0;
+#endif
+	}
 
 #if WITH_EDITORONLY_DATA
 	ENGINE_API void SetTimecodeInfo(const FSoundWaveTimecodeInfo& InTimecode);
@@ -1399,7 +1408,7 @@ public:
 	/*
 	* Returns a sample rate if there is a specific sample rate override for this platform, -1.0 otherwise.
 	*/
-	ENGINE_API float GetSampleRateForCompressionOverrides(const FPlatformAudioCookOverrides* CompressionOverrides);
+	ENGINE_API float GetSampleRateForCompressionOverrides(const FPlatformAudioCookOverrides* CompressionOverrides) const;
 
 	ENGINE_API void SetError(const TCHAR* InErrorMsg=nullptr);
 	ENGINE_API void ResetError();
diff --git a/Engine/Source/Runtime/Engine/Private/AudioFormatSettings.cpp b/Engine/Source/Runtime/Engine/Private/AudioFormatSettings.cpp
--- a/Engine/Source/Runtime/Engine/Private/AudioFormatSettings.cpp
+++ b/Engine/Source/Runtime/Engine/Private/AudioFormatSettings.cpp
@@ -5,16 +5,20 @@
 #include "Sound/SoundWave.h"
 #include "Misc/ConfigCacheIni.h"
 #include "Audio.h"
+#include "AudioCompressionSettingsUtils.h"
 #include "ISoundWaveCloudStreaming.h"
+#include "Interfaces/IAudioFormat.h"
 
 namespace Audio
 {
-	bool ShouldAllowPlatformSpecificFormats()
+	static bool ShouldAllowHardwareFormats()
 	{
+		/** AudioLink allows other audio engines to take control of the hardware
+		 *  That prevents us creating hardware codecs in most cases so disable them here */
 		static bool IsAudioLinkEnabled = []() -> bool
 		{
-			bool bAvailable = IModularFeatures::Get().IsModularFeatureAvailable(TEXT("AudioLink Factory"));
-			UE_CLOG(bAvailable,LogAudio, Display, TEXT("AudioLink is enabled, disabling platform specific AudioFormats."));
+			const bool bAvailable = IModularFeatures::Get().IsModularFeatureAvailable(TEXT("AudioLink Factory"));
+			UE_CLOG(bAvailable,LogAudio, Display, TEXT("AudioLink is enabled, disabling hardware AudioFormats."));
 			return bAvailable;
 		}();
 		return !IsAudioLinkEnabled;
@@ -45,35 +49,129 @@ namespace Audio
 #endif // WITH_EDITORONLY_DATA
 		return InCurrentFormat;
 	}
-
-	FAudioFormatSettings::FAudioFormatSettings(FConfigCacheIni* InConfigSystem, const FString& InConfigFilename, const FString& InPlatformIdentifierForLogging)
+	
+	FAudioFormatSettings::FAudioFormatSettings(FConfigCacheIni* InConfigSystem, const FString& InConfigFile, const FString& InIniPlatformName)
+		: IniPlatformName(*InIniPlatformName)
 	{
-		ReadConfiguration(InConfigSystem, InConfigFilename, InPlatformIdentifierForLogging);
+		ReadConfiguration(InConfigSystem, InConfigFile); 
 	}
 
-	FName FAudioFormatSettings::GetWaveFormat(const USoundWave* Wave) const
+	// Few things that the platform might define about a wave.
+	struct FAudioFormatSettings::FPlatformWaveState
 	{
-		FName FormatName = Audio::ToName(Wave->GetSoundAssetCompressionType());
-		if (FormatName == Audio::NAME_PLATFORM_SPECIFIC)
+		FName FormatName;
+		FName Name;
+		int32 SampleRate = 0;
+		int32 NumChannels = 0;
+
+		// From a SoundWave.
+		FPlatformWaveState(const USoundWave* InWave,const FAudioFormatSettings* InFormatSettings)
+			: FormatName(ToName(InWave->GetSoundAssetCompressionType()))
+			, Name(InWave->GetName())
+			, SampleRate(InWave->GetImportedSampleRate())
+			, NumChannels(InWave->NumChannels)
 		{
-			if (ShouldAllowPlatformSpecificFormats())
+			// Override sample-rate?
+			if (const FPlatformAudioCookOverrides* CookOverrides = FPlatformCompressionUtilities::GetCookOverrides(*InFormatSettings->IniPlatformName.ToString()) )
 			{
-				if (Wave->IsStreaming())
+				if (const float SampleRateOverride = InWave->GetSampleRateForCompressionOverrides(CookOverrides); SampleRateOverride != -1.f)
 				{
-					FormatName = PlatformStreamingFormat;
-				}
-				else
-				{
-					FormatName = PlatformFormat;
+					SampleRate = SampleRateOverride;
 				}
 			}
-			else
+
+			// Platform Specific? (resolve to this platforms format)
+			if (FormatName == NAME_PLATFORM_SPECIFIC)
 			{
-				FormatName = FallbackFormat;
+				// Convert "PlatformSpecific" into format name, based on our config for this platform.
+				FormatName = InWave->IsStreaming() ?
+					InFormatSettings->PlatformStreamingFormat :
+					InFormatSettings->PlatformFormat;
 			}
 		}
-		FormatName = GetCloudStreamingFormatOverride(FormatName, Wave);
-		return FormatName;
+	};	
+
+	const IAudioFormat* FAudioFormatSettings::FindFormat(const FName& InFormatName) const
+	{
+		// In cache?
+		FScopeLock Lock(&AudioFormatCacheCs);
+		if (const IAudioFormat* const* Found = AudioFormatCache.Find(InFormatName))
+		{
+			return *Found;
+		}
+		
+		// Look it up and cache.
+		const TArray<IAudioFormat*> AllFormats = IModularFeatures::Get().GetModularFeatureImplementations<IAudioFormat>(IAudioFormat::GetModularFeatureName());
+		if (IAudioFormat* const * Found = AllFormats.FindByPredicate
+			([InFormatName](const IAudioFormat* InFormat) -> bool
+			{
+				TArray<FName> SupportedFormats;
+				InFormat->GetSupportedFormats(SupportedFormats);
+				return SupportedFormats.Contains(InFormatName);
+			}))
+		{
+			AudioFormatCache.Add(InFormatName, *Found);
+			return *Found;
+		}
+
+		// Fail.
+		return nullptr;	
+	}
+
+	bool FAudioFormatSettings::IsFormatAllowed(const FPlatformWaveState& InWave) const 
+	{
+		if (const IAudioFormat* Format = FindFormat(InWave.FormatName))
+		{
+			// Platform supported?
+			if (!Format->IsPlatformSupported(IniPlatformName))
+			{
+				UE_LOG(LogAudio, Verbose, TEXT("Wave '%s', format '%s' doesn't support platform '%s'"),
+						*InWave.Name.ToString(), *InWave.FormatName.ToString(), *IniPlatformName.ToString());
+				return false;		
+			}
+			
+			// Sample rate ok? 
+			if (!Format->IsSampleRateSupported(InWave.SampleRate))
+			{
+				UE_LOG(LogAudio, Verbose, TEXT("Wave '%s', format '%s' doesn't support sample-rate: '%d'"),
+					*InWave.Name.ToString(), *InWave.FormatName.ToString(), InWave.SampleRate);
+				return false;	
+			}
+
+			// Channel count ok?
+			if (!Format->IsChannelCountSupported(InWave.NumChannels))
+			{
+				UE_LOG(LogAudio, Verbose, TEXT("Wave '%s', format '%s' doesn't support channel count: '%d'"),
+					*InWave.Name.ToString(), *InWave.FormatName.ToString(), InWave.NumChannels);
+				return false;	
+			}
+
+			// Hardware ok?
+			if (Format->IsHardwareFormat() && !ShouldAllowHardwareFormats())
+			{
+				return false;
+			}
+
+			// Success, passed all tests.
+			return true;
+		}
+
+		// Assume no, if we can't find it registered.
+		return false;
+	}
+
+	FName FAudioFormatSettings::GetWaveFormat(const USoundWave* Wave) const
+	{
+		FPlatformWaveState PlatformWave(Wave, this);
+		
+		// Can we use the one that's defined based on its constraints?
+		if (!IsFormatAllowed(PlatformWave))
+		{
+			PlatformWave.FormatName = FallbackFormat;
+		}
+		
+		PlatformWave.FormatName = GetCloudStreamingFormatOverride(PlatformWave.FormatName, Wave);
+		return PlatformWave.FormatName;
 	}
 
 	void FAudioFormatSettings::GetAllWaveFormats(TArray<FName>& OutFormats) const
@@ -86,7 +184,7 @@ namespace Audio
 		OutHints = WaveFormatModuleHints;
 	}
 
-	void FAudioFormatSettings::ReadConfiguration(FConfigCacheIni* InConfigSystem, const FString& InConfigFilename, const FString& InPlatformIdentifierForLogging)
+	void FAudioFormatSettings::ReadConfiguration(FConfigCacheIni* InConfigSystem, const FString& InConfigFilename)
 	{
 		auto MakePrettyArrayToString = [](const TArray<FName>& InNames) -> FString 
 		{
@@ -178,7 +276,8 @@ namespace Audio
 		}
 
 		UE_LOG(LogAudio, Verbose, TEXT("AudioFormatSettings: TargetName='%s', AllWaveFormats=(%s), Hints=(%s), PlatformFormat='%s', PlatformStreamingFormat='%s', FallbackFormat='%s'"),
-			*InPlatformIdentifierForLogging, *MakePrettyArrayToString(AllWaveFormats), *MakePrettyArrayToString(WaveFormatModuleHints), *PlatformFormat.ToString(), *PlatformStreamingFormat.ToString(), *FallbackFormat.ToString());	
+			*IniPlatformName.ToString(), *MakePrettyArrayToString(AllWaveFormats), *MakePrettyArrayToString(WaveFormatModuleHints), *PlatformFormat.ToString(), *PlatformStreamingFormat.ToString(), *FallbackFormat.ToString());	
 	}
 
+
 }// namespace Audio
diff --git a/Engine/Source/Runtime/Engine/Private/SoundWave.cpp b/Engine/Source/Runtime/Engine/Private/SoundWave.cpp
--- a/Engine/Source/Runtime/Engine/Private/SoundWave.cpp
+++ b/Engine/Source/Runtime/Engine/Private/SoundWave.cpp
@@ -3911,7 +3911,7 @@ float USoundWave::GetSampleRateForCurrentPlatform() const
 	}
 }
 
-float USoundWave::GetSampleRateForCompressionOverrides(const FPlatformAudioCookOverrides* CompressionOverrides)
+float USoundWave::GetSampleRateForCompressionOverrides(const FPlatformAudioCookOverrides* CompressionOverrides) const
 {
 	const float* SampleRatePtr = CompressionOverrides->PlatformSampleRates.Find(SampleRateQuality);
 	if (SampleRatePtr && *SampleRatePtr > 0.0f)
