Exposed tangent options on rich curves to blueprint

diff --git a/Engine/Source/Runtime/Engine/Classes/Curves/RichCurve.h b/Engine/Source/Runtime/Engine/Classes/Curves/RichCurve.h
--- a/Engine/Source/Runtime/Engine/Classes/Curves/RichCurve.h
+++ b/Engine/Source/Runtime/Engine/Classes/Curves/RichCurve.h
@@ -81,15 +81,15 @@ struct FRichCurveKey
 	GENERATED_USTRUCT_BODY()
 
 	/** Interpolation mode between this key and the next */
-	UPROPERTY()
+	UPROPERTY(EditAnywhere, BlueprintReadWrite, Category="Key")
 	TEnumAsByte<ERichCurveInterpMode> InterpMode;
 
 	/** Mode for tangents at this key */
-	UPROPERTY()
+	UPROPERTY(EditAnywhere, BlueprintReadWrite, Category="Key")
 	TEnumAsByte<ERichCurveTangentMode> TangentMode;
 
 	/** If either tangent at this key is 'weighted' */
-	UPROPERTY()
+	UPROPERTY(EditAnywhere, BlueprintReadWrite, Category="Key")
 	TEnumAsByte<ERichCurveTangentWeightMode> TangentWeightMode;
 
 	/** Time at this key */
