Fix cook non-determinism involving old UAtmosphericFogComponent

Static lighting GUID was initialized to random in the USkyAtmosphereComponent ctor and not overwritten
when serializing an old USkyAtmosphereComponent. This GUID does get written to the cooked data,
so it was effectively random every time.

diff --git a/Engine/Source/Runtime/Engine/Classes/Components/SkyAtmosphereComponent.h b/Engine/Source/Runtime/Engine/Classes/Components/SkyAtmosphereComponent.h
--- a/Engine/Source/Runtime/Engine/Classes/Components/SkyAtmosphereComponent.h
+++ b/Engine/Source/Runtime/Engine/Classes/Components/SkyAtmosphereComponent.h
@@ -291,6 +291,12 @@ protected:
 	// When true, this means that this SkyAtmosphere is use as replacement for the deprecated AtmosphericFogComponent as a parent class. 
 	// This is used to adapt the serialisation.
 	bool bIsAtmosphericFog = false;
+	
+	/**
+	 * Sets a dummy, known static lighting GUID to avoid cook non-determinism
+	 * when a map contains the deprecated AtmosphericFogComponent.
+	 */
+	void SetDummyStaticLightingGUIDs();
 };
 
 
diff --git a/Engine/Source/Runtime/Engine/Private/Atmosphere/AtmosphericFogComponent.cpp b/Engine/Source/Runtime/Engine/Private/Atmosphere/AtmosphericFogComponent.cpp
--- a/Engine/Source/Runtime/Engine/Private/Atmosphere/AtmosphericFogComponent.cpp
+++ b/Engine/Source/Runtime/Engine/Private/Atmosphere/AtmosphericFogComponent.cpp
@@ -56,6 +56,14 @@ void UAtmosphericFogComponent::Serialize(FArchive& Ar)
 		Ar << CounterVal;
 
 		SetPositionToMatchDeprecatedAtmosphericFog();
+
+		// Old-style atmospheric fog didn't save a static lighting GUID
+		// but also didn't interact with static baked lighting in the first place.
+		// The USkyAtmosphericFogComponent ctor initializes its GUID to a random value,
+		// which means we end up with cook non-determinims when that GUID doesn't get
+		// replaced on load but does get saved in the exported data on cook. So
+		// we initialize our static lighting GUID to a dummy GUID in this case.
+		SetDummyStaticLightingGUIDs();
 	}
 }
 PRAGMA_ENABLE_DEPRECATION_WARNINGS
diff --git a/Engine/Source/Runtime/Engine/Private/Components/SkyAtmosphereComponent.cpp b/Engine/Source/Runtime/Engine/Private/Components/SkyAtmosphereComponent.cpp
--- a/Engine/Source/Runtime/Engine/Private/Components/SkyAtmosphereComponent.cpp
+++ b/Engine/Source/Runtime/Engine/Private/Components/SkyAtmosphereComponent.cpp
@@ -206,6 +206,12 @@ void USkyAtmosphereComponent::UpdateStaticLightingGUIDs()
 	bStaticLightingBuiltGUID = FGuid::NewGuid();
 }
 
+void USkyAtmosphereComponent::SetDummyStaticLightingGUIDs()
+{
+	// Dummy GUID just to make sure the value is initialized and not random.
+	bStaticLightingBuiltGUID = FGuid(1, 0, 0, 0);
+}
+
 #if WITH_EDITOR
 
 void USkyAtmosphereComponent::CheckForErrors()
@@ -260,7 +266,7 @@ void USkyAtmosphereComponent::PostEditChangeProperty(FPropertyChangedEvent& Prop
 	{
 		if (SkyAtmosphereComponentStaticLightingBuilt(this))
 		{
-			// If we have changed an atmosphere property and the lighyting has already been built, we need to ask for a rebuild by updating the static lighting GUIDs.
+			// If we have changed an atmosphere property and the lighting has already been built, we need to ask for a rebuild by updating the static lighting GUIDs.
 			UpdateStaticLightingGUIDs();
 		}
 
