engine - making USkeleton::GetBlendProfile BlueprintThreadSafe

diff --git a/Engine/Source/Runtime/Engine/Classes/Animation/Skeleton.h b/Engine/Source/Runtime/Engine/Classes/Animation/Skeleton.h
--- a/Engine/Source/Runtime/Engine/Classes/Animation/Skeleton.h
+++ b/Engine/Source/Runtime/Engine/Classes/Animation/Skeleton.h
@@ -507,7 +507,7 @@ public:
 	TArray<TObjectPtr<UBlendProfile>> BlendProfiles;
 
 	/** Get the specified blend profile by name */
-	UFUNCTION(BlueprintPure, Category = Skeleton)
+	UFUNCTION(BlueprintPure, Category = Skeleton, meta = (BlueprintThreadSafe))
 	ENGINE_API UBlendProfile* GetBlendProfile(const FName& InProfileName);
 
 	/** Create a new blend profile with the specified name */
