Disable screenshot tracing in shipping by default

diff --git a/Engine/Plugins/TraceUtilities/Source/TraceUtilities/Private/TraceUtilLibrary.cpp b/Engine/Plugins/TraceUtilities/Source/TraceUtilities/Private/TraceUtilLibrary.cpp
--- a/Engine/Plugins/TraceUtilities/Source/TraceUtilities/Private/TraceUtilLibrary.cpp
+++ b/Engine/Plugins/TraceUtilities/Source/TraceUtilities/Private/TraceUtilLibrary.cpp
@@ -91,5 +91,7 @@ void UTraceUtilLibrary::TraceMarkRegionEnd(const FString& Name)
 
 void UTraceUtilLibrary::TraceScreenshot(const FString& Name, bool bShowUI)
 {
+#if UE_SCREENSHOT_TRACE_ENABLED
 	FTraceScreenshot::RequestScreenshot(Name, bShowUI);
+#endif
 }
diff --git a/Engine/Source/Runtime/Engine/Private/GameViewportClient.cpp b/Engine/Source/Runtime/Engine/Private/GameViewportClient.cpp
--- a/Engine/Source/Runtime/Engine/Private/GameViewportClient.cpp
+++ b/Engine/Source/Runtime/Engine/Private/GameViewportClient.cpp
@@ -2104,6 +2104,7 @@ bool ProcessScreenshotData(TArray<FColorType>& Bitmap, FIntVector Size, TChannel
 			ScreenShotName += ToExtension;
 		}
 
+#if UE_SCREENSHOT_TRACE_ENABLED
 		bool bSuppressWritingToFile = false;
 		if (SHOULD_TRACE_SCREENSHOT())
 		{
@@ -2117,6 +2118,7 @@ bool ProcessScreenshotData(TArray<FColorType>& Bitmap, FIntVector Size, TChannel
 			FImageView Image((const FColorType*)Bitmap.GetData(), Size.X, Size.Y);
 			bIsScreenshotSaved = FImageUtils::SaveImageByExtension(*ScreenShotName, Image);
 		}
+#endif
 	}
 
 	return bIsScreenshotSaved;
@@ -2198,7 +2200,9 @@ bool UGameViewportClient::ProcessScreenShots(FViewport* InViewport)
 		}
 
 		FScreenshotRequest::Reset();
+#if UE_SCREENSHOT_TRACE_ENABLED
 		FTraceScreenshot::Reset();
+#endif
 		FScreenshotRequest::OnScreenshotRequestProcessed().Broadcast();
 
 		// Reeanble screen messages - if we are NOT capturing a movie
diff --git a/Engine/Source/Runtime/Engine/Private/ProfilingDebugging/TraceScreenshot.cpp b/Engine/Source/Runtime/Engine/Private/ProfilingDebugging/TraceScreenshot.cpp
--- a/Engine/Source/Runtime/Engine/Private/ProfilingDebugging/TraceScreenshot.cpp
+++ b/Engine/Source/Runtime/Engine/Private/ProfilingDebugging/TraceScreenshot.cpp
@@ -1,6 +1,8 @@
 // Copyright Epic Games, Inc. All Rights Reserved.
 #include "ProfilingDebugging/TraceScreenshot.h"
 
+#if UE_SCREENSHOT_TRACE_ENABLED
+
 #include "ImageUtils.h"
 #include "Misc/DateTime.h"
 #include "Misc/CString.h"
@@ -184,3 +186,5 @@ void FTraceScreenshot::Reset()
 {
 	bSuppressWritingToFile = false;
 }
+
+#endif // UE_SCREENSHOT_TRACE_ENABLED
diff --git a/Engine/Source/Runtime/Engine/Private/TraceService.cpp b/Engine/Source/Runtime/Engine/Private/TraceService.cpp
--- a/Engine/Source/Runtime/Engine/Private/TraceService.cpp
+++ b/Engine/Source/Runtime/Engine/Private/TraceService.cpp
@@ -30,7 +30,9 @@ private:
 	void OnPause(const FTraceControlPause& Message, const TSharedRef<IMessageContext>& Context);
 	void OnResume(const FTraceControlResume& Message, const TSharedRef<IMessageContext>& Context);
 	void OnBookmark(const FTraceControlBookmark& Message, const TSharedRef<IMessageContext>& Context);
+#if UE_SCREENSHOT_TRACE_ENABLED
 	void OnScreenshot(const FTraceControlScreenshot& Message, const TSharedRef<IMessageContext>& Context);
+#endif
 	void OnSetStatNamedEvents(const FTraceControlSetStatNamedEvents& Message, const TSharedRef<IMessageContext>& Context);
 	
 	static void FillTraceStatusMessage(FTraceControlStatus* Message);
@@ -64,7 +66,9 @@ FTraceServiceImpl::FTraceServiceImpl(const TSharedPtr<IMessageBus>& InBus)
 			.Handling<FTraceControlPause>(this, &FTraceServiceImpl::OnPause)
 			.Handling<FTraceControlResume>(this, &FTraceServiceImpl::OnResume)
 			.Handling<FTraceControlBookmark>(this, &FTraceServiceImpl::OnBookmark)
+#if UE_SCREENSHOT_TRACE_ENABLED
 			.Handling<FTraceControlScreenshot>(this, &FTraceServiceImpl::OnScreenshot)
+#endif
 			.Handling<FTraceControlSetStatNamedEvents>(this, &FTraceServiceImpl::OnSetStatNamedEvents)
 			.Handling<FTraceControlStatusPing>(this, &FTraceServiceImpl::OnStatusPing)
 			.Handling<FTraceControlSettingsPing>(this, &FTraceServiceImpl::OnSettingsPing)
@@ -88,7 +92,9 @@ FTraceServiceImpl::FTraceServiceImpl(const TSharedPtr<IMessageBus>& InBus)
 		MessageEndpoint->Subscribe<FTraceControlPause>();
 		MessageEndpoint->Subscribe<FTraceControlResume>();
 		MessageEndpoint->Subscribe<FTraceControlBookmark>();
+#if UE_SCREENSHOT_TRACE_ENABLED
 		MessageEndpoint->Subscribe<FTraceControlScreenshot>();
+#endif
 		MessageEndpoint->Subscribe<FTraceControlSetStatNamedEvents>();
 	}
 }
@@ -177,10 +183,12 @@ void FTraceServiceImpl::OnBookmark(const FTraceControlBookmark& Message, const T
 	TRACE_BOOKMARK(TEXT("%s"), *Message.Label);
 }
 
+#if UE_SCREENSHOT_TRACE_ENABLED
 void FTraceServiceImpl::OnScreenshot(const FTraceControlScreenshot& Message, const TSharedRef<IMessageContext>& Context)
 {
 	FTraceScreenshot::RequestScreenshot(Message.Name, Message.bShowUI);
 }
+#endif // UE_SCREENSHOT_TRACE_ENABLED
 
 void FTraceServiceImpl::OnSetStatNamedEvents(const FTraceControlSetStatNamedEvents& Message, const TSharedRef<IMessageContext>& Context)
 {
diff --git a/Engine/Source/Runtime/Engine/Public/ProfilingDebugging/TraceScreenshot.h b/Engine/Source/Runtime/Engine/Public/ProfilingDebugging/TraceScreenshot.h
--- a/Engine/Source/Runtime/Engine/Public/ProfilingDebugging/TraceScreenshot.h
+++ b/Engine/Source/Runtime/Engine/Public/ProfilingDebugging/TraceScreenshot.h
@@ -5,6 +5,12 @@
 #include "CoreMinimal.h"
 #include "Engine/World.h"
 
+#if !defined(UE_SCREENSHOT_TRACE_ENABLED)
+	#define UE_SCREENSHOT_TRACE_ENABLED UE_TRACE_ENABLED
+#endif
+
+#if UE_SCREENSHOT_TRACE_ENABLED
+
 class ULevel;
 
 class FTraceScreenshot
@@ -53,3 +59,5 @@ private:
 private:
 	static ENGINE_API bool bSuppressWritingToFile;
 };
+
+#endif // UE_SCREENSHOT_TRACE_ENABLED
\ No newline at end of file
