Substrate - updated node thickness input tooltips.

#rnx

diff --git a/Engine/Source/Runtime/Engine/Public/Materials/MaterialExpressionSubstrate.h b/Engine/Source/Runtime/Engine/Public/Materials/MaterialExpressionSubstrate.h
--- a/Engine/Source/Runtime/Engine/Public/Materials/MaterialExpressionSubstrate.h
+++ b/Engine/Source/Runtime/Engine/Public/Materials/MaterialExpressionSubstrate.h
@@ -954,7 +954,8 @@ class UMaterialExpressionSubstrateVerticalLayering : public UMaterialExpressionS
 	FExpressionInput Base;
 
 	/**
-	 * Thickness of the Top material layer
+	 * Thickness of the Top material layer in centimeter.
+	 * It can be modulated to achieve simple scattering/transmittance variation of the same material.
 	 */
 	UPROPERTY()
 	FExpressionInput Thickness;
@@ -1067,8 +1068,8 @@ class UMaterialExpressionSubstrateTransmittanceToMFP : public UMaterialExpressio
 	FExpressionInput TransmittanceColor;
 
 	/**
-	* The desired thickness in centimeter.
-	* Another use case example: this node output called thickness can be modulated before it is plugged in a slab node, this can be used to achieve simple scattering/transmittance variation of the same material.
+	* Thickness of the layer in centimeter.
+	* Example of use case: this node output called thickness can be modulated before it is plugged in a Vertical Layering node Thickness input. This can be used to achieve simple scattering/transmittance variation of the same material.
 	*/
 	UPROPERTY()
 	FExpressionInput Thickness;
@@ -1173,7 +1174,7 @@ class UMaterialExpressionSubstrateThinFilm : public UMaterialExpressionSubstrate
 	FExpressionInput F90;
 	
 	/**
-	 * Thin film controls the thin film layer coating the current slab. 0 means disabled and 1 means a coating layer of 10 micrometer. (type = float, unitless, default = 0)
+	 * Controls the thickness of the thin film layer coating the current slab. 0 means disabled and 1 means a coating layer of 10 micrometer. (type = float, unitless, default = 0)
 	 */
 	UPROPERTY()
 	FExpressionInput Thickness;
