Substrate - moved maerial expressions into their own cpp.

#rnx

diff --git a/Engine/Source/Runtime/Engine/Private/Materials/MaterialExpressions.cpp b/Engine/Source/Runtime/Engine/Private/Materials/MaterialExpressions.cpp
--- a/Engine/Source/Runtime/Engine/Private/Materials/MaterialExpressions.cpp
+++ b/Engine/Source/Runtime/Engine/Private/Materials/MaterialExpressions.cpp
@@ -36,7 +36,6 @@
 #include "Engine/TextureCubeArray.h"
 #include "Engine/VolumeTexture.h"
 #include "Engine/SubsurfaceProfile.h"
-#include "Engine/SpecularProfile.h"
 #include "Serialization/ObjectWriter.h"
 #include "Serialization/ObjectReader.h"
 #include "VT/RuntimeVirtualTexture.h"
@@ -2397,27 +2396,6 @@ FSubstrateOperator* UMaterialExpression::SubstrateGenerateMaterialTopologyTree(c
 	return nullptr;
 }
 
-static void AssignOperatorIndexIfNotNull(int32& NextOperatorPin, FSubstrateOperator* Operator)
-{
-	NextOperatorPin = Operator ? Operator->Index : INDEX_NONE;
-}
-
-static void CombineFlagForParameterBlending(FSubstrateOperator& DstOp, FSubstrateOperator* OpA, FSubstrateOperator* OpB = nullptr)
-{
-	if (OpA && OpB)
-	{
-		DstOp.CombineFlagsForParameterBlending(*OpA, *OpB);
-	}
-	else if (OpA)
-	{
-		DstOp.CopyFlagsForParameterBlending(*OpA);
-	}
-	else if (OpB)
-	{
-		DstOp.CopyFlagsForParameterBlending(*OpB);
-	}
-}
-
 #endif // WITH_EDITOR
 
 UMaterialExpressionTextureBase::UMaterialExpressionTextureBase(const FObjectInitializer& ObjectInitializer)
@@ -26468,3868 +26446,6 @@ FName UMaterialExpressionSceneDepthWithoutWater::GetInputName(int32 InputIndex)
 
 #endif // WITH_EDITOR
 
-///////////////////////////////////////////////////////////////////////////////
-// Substrate
-
-#if WITH_EDITOR
-static int32 SubstrateBlendNormal(class FMaterialCompiler* Compiler, int32 NormalCodeChunk0, int32 NormalCodeChunk1, int32 MixCodeChunk)
-{
-	int32 SafeMixCodeChunk = Compiler->Saturate(MixCodeChunk);
-	int32 LerpedNormal = Compiler->Lerp(NormalCodeChunk0, NormalCodeChunk1, SafeMixCodeChunk);
-	int32 BlendedNormalCodeChunk = Compiler->Div(LerpedNormal, Compiler->SquareRoot(Compiler->Dot(LerpedNormal, LerpedNormal)));
-	return BlendedNormalCodeChunk;
-}
-#endif // WITH_EDITOR
-
-#if WITH_EDITOR
-
-// Optionnaly cast CodeChunk type to non-LWC type. 
-// Input can be built of WorldPosition data, which would force the derived data to have LWC type 
-// creating issues, as Substrate functions' inputs don't support LWC
-static int32 CastToNonLWCType(class FMaterialCompiler* Compiler, int32 CodeChunk)
-{
-	EMaterialValueType Type = Compiler->GetType(CodeChunk);
-	if (IsLWCType(Type))
-	{
-		Type = MakeNonLWCType(Type);
-		CodeChunk = Compiler->ValidCast(CodeChunk, Type);
-	}
-	return CodeChunk;
-}
-
-// The compilation of an expression can sometimes lead to a INDEX_NONE code chunk when editing material graphs 
-// or when the node is inside a material function, linked to an input pin of the material function and that input is not plugged in to anything.
-// But for normals or tangents, Substrate absolutely need a valid code chunk to de-duplicate when stored in memory. 
-// Also, we want all our nodes to have default, as that is needed when creating BSDF, when registering code chunk representing material topology.
-static int32 CompileWithDefaultCodeChunk(class FMaterialCompiler* Compiler, FExpressionInput& Input, int DefaultCodeChunk, bool* bDefaultIsUsed = nullptr)
-{
-	if (bDefaultIsUsed)
-	{
-		*bDefaultIsUsed = Input.GetTracedInput().Expression == nullptr;
-	}
-	int32 CodeChunk = Input.GetTracedInput().Expression ? Input.Compile(Compiler) : DefaultCodeChunk;
-	if (bDefaultIsUsed)
-	{
-		*bDefaultIsUsed |= CodeChunk == INDEX_NONE;
-	}
-	else
-	{
-		CodeChunk = CastToNonLWCType(Compiler, CodeChunk);
-	}
-	return CodeChunk == INDEX_NONE ? DefaultCodeChunk : CodeChunk;
-}
-static int32 CompileWithDefaultFloat1(class FMaterialCompiler* Compiler, FExpressionInput& Input, float X, bool* bDefaultIsUsed = nullptr)
-{
-	if (bDefaultIsUsed)
-	{
-		*bDefaultIsUsed = Input.GetTracedInput().Expression == nullptr;
-	}
-	int32 CodeChunk = Input.GetTracedInput().Expression ? Input.Compile(Compiler) : Compiler->Constant(X);
-	if (bDefaultIsUsed)
-	{
-		*bDefaultIsUsed |= CodeChunk == INDEX_NONE;
-	}
-	else
-	{
-		CodeChunk = CastToNonLWCType(Compiler, CodeChunk);
-	}
-	return CodeChunk == INDEX_NONE ? Compiler->Constant(X) : CodeChunk;
-}
-static int32 CompileWithDefaultFloat2(class FMaterialCompiler* Compiler, FExpressionInput& Input, float X, float Y, bool* bDefaultIsUsed = nullptr)
-{
-	if (bDefaultIsUsed)
-	{
-		*bDefaultIsUsed = Input.GetTracedInput().Expression == nullptr;
-	}
-	int32 CodeChunk = Input.GetTracedInput().Expression ? Input.Compile(Compiler) : Compiler->Constant2(X, Y);
-	if (bDefaultIsUsed)
-	{
-		*bDefaultIsUsed |= CodeChunk == INDEX_NONE;
-	}
-	else
-	{
-		CodeChunk = CastToNonLWCType(Compiler, CodeChunk);
-	}
-	return CodeChunk == INDEX_NONE ? Compiler->Constant2(X, Y) : CodeChunk;
-}
-static int32 CompileWithDefaultFloat3(class FMaterialCompiler* Compiler, FExpressionInput& Input, float X, float Y, float Z, bool* bDefaultIsUsed = nullptr)
-{
-	if (bDefaultIsUsed)
-	{
-		*bDefaultIsUsed = Input.GetTracedInput().Expression == nullptr;
-	}
-	int32 CodeChunk = Input.GetTracedInput().Expression ? Input.Compile(Compiler) : Compiler->Constant3(X, Y, Z);
-	if (bDefaultIsUsed)
-	{
-		*bDefaultIsUsed |= CodeChunk == INDEX_NONE;
-	}
-	else
-	{
-		CodeChunk = CastToNonLWCType(Compiler, CodeChunk);
-	}
-	return CodeChunk == INDEX_NONE ? Compiler->Constant3(X, Y, Z) : CodeChunk;
-}
-static int32 CompileWithDefaultNormalWS(class FMaterialCompiler* Compiler, FExpressionInput& Input, bool bConvertToRequestedSpace=true)
-{
-	if (Input.GetTracedInput().Expression != nullptr)
-	{
-		int32 NormalCodeChunk = Input.Compile(Compiler);
-
-		if (NormalCodeChunk == INDEX_NONE)
-		{
-			// Nothing is plug in from the linked input, so specify world space normal the BSDF node expects.
-			return Compiler->VertexNormal();
-		}
-		// Transform into world space normal if needed. BSDF nodes always expects world space normal as input.
-		return bConvertToRequestedSpace ? Compiler->TransformNormalFromRequestedBasisToWorld(NormalCodeChunk) : NormalCodeChunk;
-	}
-	// Nothing is plug in on the BSDF node, so specify world space normal the node expects.
-	return Compiler->VertexNormal();
-}
-static int32 CompileWithDefaultTangentWS(class FMaterialCompiler* Compiler, FExpressionInput& Input, bool bConvertToRequestedSpace=true)
-{
-	if (Input.GetTracedInput().Expression != nullptr)
-	{
-		int32 TangentCodeChunk = Input.Compile(Compiler);
-
-		if (TangentCodeChunk == INDEX_NONE)
-		{
-			// Nothing is plug in from the linked input, so specify world space tangent the BSDF node expects.
-			return Compiler->VertexTangent();
-		}
-		// Transform into world space tangent if needed. BSDF nodes always expects world space tangent as input.
-		return bConvertToRequestedSpace ? Compiler->TransformNormalFromRequestedBasisToWorld(TangentCodeChunk) : TangentCodeChunk;
-	}
-	// Nothing is plug in on the BSDF node, so specify world space tangent the node expects.
-	return Compiler->VertexTangent();
-}
-
-#endif // WITH_EDITOR
-
-UMaterialExpressionSubstrateShadingModels::UMaterialExpressionSubstrateShadingModels(const FObjectInitializer& ObjectInitializer)
-	: Super(ObjectInitializer)
-{
-	struct FConstructorStatics
-	{
-		FText NAME_Substrate;
-		FConstructorStatics() : NAME_Substrate(LOCTEXT("Substrate Conversion", "Substrate Conversion")) { }
-	};
-	static FConstructorStatics ConstructorStatics;
-#if WITH_EDITORONLY_DATA
-	MenuCategories.Add(ConstructorStatics.NAME_Substrate);
-#endif
-}
-
-#if WITH_EDITOR
-FExpressionInput* UMaterialExpressionSubstrateShadingModels::GetInput(int32 InputIndex)
-{
-	switch (InputIndex)
-	{
-		case 0: return &BaseColor;
-		case 1: return &Metallic;
-		case 2: return &Specular;
-		case 3: return &Roughness;
-		case 4: return &Anisotropy;
-		case 5: return &EmissiveColor;
-		case 6: return &Normal;
-		case 7: return &Tangent;
-		case 8: return &SubSurfaceColor;
-		case 9: return &ClearCoat;
-		case 10: return &ClearCoatRoughness;
-		case 11: return &Opacity;
-		case 12: return &TransmittanceColor;
-		case 13: return &WaterScatteringCoefficients;
-		case 14: return &WaterAbsorptionCoefficients;
-		case 15: return &WaterPhaseG;
-		case 16: return &ColorScaleBehindWater;
-		case 17: return &ClearCoatNormal;
-		case 18: return &CustomTangent;
-		case 19: return &ThinTranslucentSurfaceCoverage;
-		case 20: return &ShadingModel;
-		default: return nullptr;
-	}
-}
-
-#define LEGACY_DIRECT_ATTRIBUTE_MAPPING(CodeChunkResult, MaterialProperty, Code)	Compiler->PushMaterialAttribute(FMaterialAttributeDefinitionMap::GetID(MaterialProperty));\
-																					int32 CodeChunkResult = Code;\
-																					Compiler->PopMaterialAttribute();
-
-int32 UMaterialExpressionSubstrateShadingModels::Compile(class FMaterialCompiler* Compiler, int32 OutputIndex)
-{
-	LEGACY_DIRECT_ATTRIBUTE_MAPPING(RoughnessCodeChunk, MP_Roughness, CompileWithDefaultFloat1(Compiler, Roughness, 0.5f));
-	LEGACY_DIRECT_ATTRIBUTE_MAPPING(AnisotropyCodeChunk, MP_Anisotropy, CompileWithDefaultFloat1(Compiler, Anisotropy, 0.0f));
-
-	// We also cannot ignore the tangent when using the default Tangent because GetTangentBasis
-	// used in SubstrateGetBSDFSharedBasis cannot be relied on for smooth tangent used for lighting on any mesh.
-	const bool bHasAnisotropy = HasAnisotropy();
-
-	// Regular normal basis
-	LEGACY_DIRECT_ATTRIBUTE_MAPPING(NormalCodeChunk, MP_Normal, CompileWithDefaultNormalWS(Compiler, Normal));
-
-	// When computing NormalCodeChunk, we invoke TransformNormalFromRequestedBasisToWorld which requires input to be float or float3.
-	// Certain material do not respect this requirement. We handle here a simple recovery when source material doesn't have a valid 
-	// normal (e.g., vec2 normal), and avoid crashing the material compilation. The error will still be reported by the compiler up 
-	// to the user, but the compilation will succeed.
-	if (NormalCodeChunk == INDEX_NONE) { NormalCodeChunk = Compiler->VertexNormal(); } 
-
-	int32 TangentCodeChunk = INDEX_NONE;
-	if (bHasAnisotropy)
-	{
-		LEGACY_DIRECT_ATTRIBUTE_MAPPING(TangentCodeChunkTmp, MP_Tangent, CompileWithDefaultTangentWS(Compiler, Normal));
-		TangentCodeChunk = TangentCodeChunkTmp;
-	}
-	const FSubstrateRegisteredSharedLocalBasis NewRegisteredSharedLocalBasis = SubstrateCompilationInfoCreateSharedLocalBasis(Compiler, NormalCodeChunk, TangentCodeChunk);
-	const FString BasisIndexMacro = Compiler->GetSubstrateSharedLocalBasisIndexMacro(NewRegisteredSharedLocalBasis);
-
-	const bool bHasCoatNormal = ClearCoatNormal.IsConnected();
-	// Clear coat normal basis
-	int32 ClearCoat_NormalCodeChunk = INDEX_NONE;
-	int32 ClearCoat_TangentCodeChunk = INDEX_NONE;
-	FString ClearCoat_BasisIndexMacro;
-	FSubstrateRegisteredSharedLocalBasis ClearCoat_NewRegisteredSharedLocalBasis;
-	if (bHasCoatNormal)
-	{
-		ClearCoat_NormalCodeChunk = CompileWithDefaultNormalWS(Compiler, ClearCoatNormal);
-		ClearCoat_TangentCodeChunk = TangentCodeChunk;
-		ClearCoat_NewRegisteredSharedLocalBasis = SubstrateCompilationInfoCreateSharedLocalBasis(Compiler, ClearCoat_NormalCodeChunk, ClearCoat_TangentCodeChunk);
-		ClearCoat_BasisIndexMacro = Compiler->GetSubstrateSharedLocalBasisIndexMacro(ClearCoat_NewRegisteredSharedLocalBasis);
-	}
-	else
-	{
-		ClearCoat_NormalCodeChunk = NormalCodeChunk;
-		ClearCoat_TangentCodeChunk = TangentCodeChunk;
-		ClearCoat_NewRegisteredSharedLocalBasis = NewRegisteredSharedLocalBasis;
-		ClearCoat_BasisIndexMacro = BasisIndexMacro;
-	}
-
-	// Custom tangent. No need to register it as a local basis, as it is only used for eye shading internal conversion
-	int32 CustomTangent_TangentCodeChunk = INDEX_NONE;
-	const bool bHasCustomTangent = CustomTangent.IsConnected();
-	if (bHasCustomTangent)
-	{
-		// Legacy code doesn't do tangent <-> world basis conversion on tangent output, when provided.
-		CustomTangent_TangentCodeChunk = CompileWithDefaultNormalWS(Compiler, CustomTangent, false /*bConvertToRequestedSpace*/);
-	}
-	else
-	{
-		CustomTangent_TangentCodeChunk = NormalCodeChunk;
-	}
-
-	int32 SSSProfileCodeChunk = INDEX_NONE;
-	const bool bHasSSS = HasSSS();
-	if (bHasSSS)
-	{
-		SSSProfileCodeChunk = Compiler->ForceCast(Compiler->ScalarParameter(GetSubsurfaceProfileParameterName(), 1.0f), MCT_Float1);
-	}
-
-	FSubstrateOperator& SubstrateOperator = Compiler->SubstrateCompilationGetOperator(Compiler->SubstrateTreeStackGetPathUniqueId());
-	SubstrateOperator.BSDFRegisteredSharedLocalBasis = NewRegisteredSharedLocalBasis;
-
-	int32 OpacityCodeChunk = INDEX_NONE;
-	if (!Compiler->SubstrateSkipsOpacityEvaluation())
-	{
-		// We evaluate opacity only for shading models and blending mode requiring it.
-		// For instance, a translucent shader reading depth for soft fading should no evaluate opacity when an instance forces an opaque mode.
-		LEGACY_DIRECT_ATTRIBUTE_MAPPING(OpacityCodeChunkTmp, MP_Opacity, CompileWithDefaultFloat1(Compiler, Opacity, 1.0f));
-		OpacityCodeChunk = OpacityCodeChunkTmp;
-	}
-	else
-	{
-		OpacityCodeChunk = Compiler->Constant(1.0f);
-	}
-
-	LEGACY_DIRECT_ATTRIBUTE_MAPPING(EmissiveCodeChunk,			MP_EmissiveColor,	CompileWithDefaultFloat3(Compiler, EmissiveColor, 0.0f, 0.0f, 0.0f));
-
-	LEGACY_DIRECT_ATTRIBUTE_MAPPING(BaseColorCodeChunk,			MP_BaseColor,		CompileWithDefaultFloat3(Compiler, BaseColor, 0.0f, 0.0f, 0.0f));
-	LEGACY_DIRECT_ATTRIBUTE_MAPPING(SpecularChunk,				MP_Specular,		CompileWithDefaultFloat1(Compiler, Specular, 0.5f));
-	LEGACY_DIRECT_ATTRIBUTE_MAPPING(MetallicCodeChunk,			MP_Metallic,		CompileWithDefaultFloat1(Compiler, Metallic, 0.0f));
-
-	LEGACY_DIRECT_ATTRIBUTE_MAPPING(SubSurfaceColorCodeChunk,	MP_SubsurfaceColor,	CompileWithDefaultFloat3(Compiler, SubSurfaceColor, 1.0f, 1.0f, 1.0f));
-
-	LEGACY_DIRECT_ATTRIBUTE_MAPPING(ClearCoatCodeChunk,			MP_CustomData0,		CompileWithDefaultFloat1(Compiler, ClearCoat, 1.0f));
-	LEGACY_DIRECT_ATTRIBUTE_MAPPING(ClearCoatRoughnessCodeChunk,MP_CustomData1,		CompileWithDefaultFloat1(Compiler, ClearCoatRoughness, 0.1f));
-
-	int32 ShadingModelCodeChunk = ShadingModel.IsConnected() ? CompileWithDefaultFloat1(Compiler, ShadingModel, float(MSM_DefaultLit)) : Compiler->Constant(float(ShadingModelOverride));
-	int32 ShadingModelCount = Compiler->GetMaterialShadingModels().CountShadingModels();
-	const bool bHasDynamicShadingModels = ShadingModelCount > 1;
-	int32 OutputCodeChunk = Compiler->SubstrateConversionFromLegacy(
-		bHasDynamicShadingModels,
-		// Metalness workflow
-		BaseColorCodeChunk,
-		SpecularChunk,
-		MetallicCodeChunk,
-		// Roughness
-		RoughnessCodeChunk,
-		AnisotropyCodeChunk,
-		// SSS
-		SubSurfaceColorCodeChunk,
-		SSSProfileCodeChunk != INDEX_NONE ? SSSProfileCodeChunk : Compiler->Constant(0.0f),
-		// Clear Coat / Custom
-		ClearCoatCodeChunk,
-		ClearCoatRoughnessCodeChunk,
-		// Misc
-		EmissiveCodeChunk,
-		OpacityCodeChunk,
-		CompileWithDefaultFloat3(Compiler, TransmittanceColor, 0.5f, 0.5f, 0.5f),
-		CompileWithDefaultFloat1(Compiler, ThinTranslucentSurfaceCoverage, 1.0f),
-		// Water
-		CompileWithDefaultFloat3(Compiler, WaterScatteringCoefficients, 0.0f, 0.0f, 0.0f),
-		CompileWithDefaultFloat3(Compiler, WaterAbsorptionCoefficients, 0.0f, 0.0f, 0.0f),
-		CompileWithDefaultFloat1(Compiler, WaterPhaseG, 0.0f),
-		CompileWithDefaultFloat3(Compiler, ColorScaleBehindWater, 1.0f, 1.0f, 1.0f),
-		// Shading model
-		ShadingModelCodeChunk,
-		NormalCodeChunk,
-		TangentCodeChunk,
-		BasisIndexMacro,
-		ClearCoat_NormalCodeChunk,
-		ClearCoat_TangentCodeChunk,
-		ClearCoat_BasisIndexMacro,
-		CustomTangent_TangentCodeChunk,
-		!SubstrateOperator.bUseParameterBlending || (SubstrateOperator.bUseParameterBlending && SubstrateOperator.bRootOfParameterBlendingSubTree) ? &SubstrateOperator : nullptr);
-
-	return OutputCodeChunk;
-}
-
-
-void UMaterialExpressionSubstrateShadingModels::PostEditChangeProperty(FPropertyChangedEvent& PropertyChangedEvent)
-{
-	if (GraphNode && PropertyChangedEvent.Property != nullptr)
-	{
-		GraphNode->ReconstructNode();
-	}
-	Super::PostEditChangeProperty(PropertyChangedEvent);
-}
-
-void UMaterialExpressionSubstrateShadingModels::GetCaption(TArray<FString>& OutCaptions) const
-{
-	OutCaptions.Add(TEXT("Substrate Shading Models"));
-}
-
-uint32 UMaterialExpressionSubstrateShadingModels::GetOutputType(int32 OutputIndex)
-{
-	return MCT_Substrate;
-}
-
-uint32 UMaterialExpressionSubstrateShadingModels::GetInputType(int32 InputIndex)
-{
-	if (InputIndex == 0)	   return MCT_Float3; // BaseColor
-	else if (InputIndex == 1)  return MCT_Float1; // Metallic
-	else if (InputIndex == 2)  return MCT_Float1; // Specular
-	else if (InputIndex == 3)  return MCT_Float1; // Roughness
-	else if (InputIndex == 4)  return MCT_Float1; // Anisotropy
-	else if (InputIndex == 5)  return MCT_Float3; // EmissiveColor
-	else if (InputIndex == 6)  return MCT_Float3; // Normal
-	else if (InputIndex == 7)  return MCT_Float3; // Tangent
-	else if (InputIndex == 8)  return MCT_Float3; // SubSurfaceColor
-	else if (InputIndex == 9)  return MCT_Float1; // ClearCoat/Custom0
-	else if (InputIndex == 10) return MCT_Float1; // ClearCoatRoughness/Custom1
-	else if (InputIndex == 11) return MCT_Float1; // Opacity
-	else if (InputIndex == 12) return MCT_Float3; // TransmittanceColor
-	else if (InputIndex == 13) return MCT_Float3; // WaterScatteringCoefficients
-	else if (InputIndex == 14) return MCT_Float3; // WaterAbsorptionCoefficients
-	else if (InputIndex == 15) return MCT_Float1; // WaterPhaseG
-	else if (InputIndex == 16) return MCT_Float3; // ColorScaleBehindWater
-	else if (InputIndex == 17) return MCT_Float3; // ClearCoatNormal
-	else if (InputIndex == 18) return MCT_Float3; // CustomTangent
-	else if (InputIndex == 19) return MCT_Float1; // ThinTranslucentSurfaceCoverage
-	else if (InputIndex == 20) return MCT_ShadingModel; // ShadingModel
-	else if (InputIndex == 21) return MCT_ShadingModel; // EMaterialShadingModel with ShowAsInputPin seems to always show at the bottom
-
-	check(false);
-	return MCT_Float1;
-}
-
-FName UMaterialExpressionSubstrateShadingModels::GetInputName(int32 InputIndex) const
-{
-	const bool bShadingModelFromExpression = ShadingModel.IsConnected();
-
-	if (InputIndex == 0)		return TEXT("BaseColor");
-	else if (InputIndex == 1)	return TEXT("Metallic");
-	else if (InputIndex == 2)	return TEXT("Specular");
-	else if (InputIndex == 3)	return TEXT("Roughness");
-	else if (InputIndex == 4)	return TEXT("Anisotropy");
-	else if (InputIndex == 5)	return TEXT("Emissive Color");
-	else if (InputIndex == 6)	return TEXT("Normal");
-	else if (InputIndex == 7)	return TEXT("Tangent");
-	else if (InputIndex == 8)
-	{
-		if (!bShadingModelFromExpression && ShadingModelOverride == MSM_Cloth)
-		{
-			return TEXT("Fuzz Color");
-		}
-		return TEXT("Subsurface Color");
-	}
-	else if (InputIndex == 9)
-	{
-		if (!bShadingModelFromExpression)
-		{
-			if (ShadingModelOverride == MSM_Cloth)
-			{
-				return TEXT("Fuzz Amount");
-			}
-			else if (ShadingModelOverride == MSM_Eye)
-			{
-				return TEXT("Iris Mask");
-			}
-			else if (ShadingModelOverride == MSM_Hair)
-			{
-				return TEXT("Backlit");
-			}
-			else if (ShadingModelOverride == MSM_ClearCoat)
-			{
-				return TEXT("Clear Coat");
-			}
-			return TEXT("Unused");
-		}
-		return TEXT("Custom0");
-	}
-	else if (InputIndex == 10)
-	{
-		if (!bShadingModelFromExpression)
-		{
-			if (ShadingModelOverride == MSM_Eye)
-			{
-				return TEXT("Iris Distance");
-			}
-			else if (ShadingModelOverride == MSM_ClearCoat)
-			{
-				return TEXT("Clear Coat Roughness");
-			}
-			return TEXT("Unused");
-		}
-		return TEXT("Custom1");
-	}
-	else if (InputIndex == 11)	return TEXT("Opacity");
-	else if (InputIndex == 12)	return TEXT("Thin Translucent Transmittance Color");
-	else if (InputIndex == 13)	return TEXT("Water Scattering Coefficients");
-	else if (InputIndex == 14)	return TEXT("Water Absorption Coefficients");
-	else if (InputIndex == 15)	return TEXT("Water Phase G");
-	else if (InputIndex == 16)	return TEXT("Color Scale BehindWater");
-	else if (InputIndex == 17)	return TEXT("Clear Coat Normal");
-	else if (InputIndex == 18)	return TEXT("Custom Tangent");
-	else if (InputIndex == 19)	return TEXT("Thin Translucent Surface Coverage");
-	else if (InputIndex == 20)	return TEXT("Single Shading Model");
-	else if (InputIndex == 21)	return TEXT("Shading Model From Expression");
-	return TEXT("Unknown");
-}
-
-void UMaterialExpressionSubstrateShadingModels::GetConnectorToolTip(int32 InputIndex, int32 OutputIndex, TArray<FString>& OutToolTip)
-{
-	if (OutputIndex == 0)
-	{
-		OutToolTip.Add(TEXT("TT Ouput"));
-		return;
-	}
-	Super::GetConnectorToolTip(InputIndex, INDEX_NONE, OutToolTip);
-}
-
-bool UMaterialExpressionSubstrateShadingModels::IsResultSubstrateMaterial(int32 OutputIndex)
-{
-	return true;
-}
-
-void UMaterialExpressionSubstrateShadingModels::GatherSubstrateMaterialInfo(FSubstrateMaterialInfo& SubstrateMaterialInfo, int32 OutputIndex)
-{	
-	// Track connected input
-	if (BaseColor.IsConnected())			{ SubstrateMaterialInfo.AddPropertyConnected(MP_BaseColor); }
-	if (Metallic.IsConnected())				{ SubstrateMaterialInfo.AddPropertyConnected(MP_Metallic); }
-	if (Specular.IsConnected())				{ SubstrateMaterialInfo.AddPropertyConnected(MP_Specular); }
-	if (Roughness.IsConnected())			{ SubstrateMaterialInfo.AddPropertyConnected(MP_Roughness); }
-	if (Anisotropy.IsConnected())			{ SubstrateMaterialInfo.AddPropertyConnected(MP_Anisotropy); }
-	if (EmissiveColor.IsConnected())		{ SubstrateMaterialInfo.AddPropertyConnected(MP_EmissiveColor); }
-	if (Normal.IsConnected())				{ SubstrateMaterialInfo.AddPropertyConnected(MP_Normal); }
-	if (Tangent.IsConnected())				{ SubstrateMaterialInfo.AddPropertyConnected(MP_Tangent); }
-	if (SubSurfaceColor.IsConnected())		{ SubstrateMaterialInfo.AddPropertyConnected(MP_SubsurfaceColor); }
-	if (ClearCoat.IsConnected())			{ SubstrateMaterialInfo.AddPropertyConnected(MP_CustomData0); }
-	if (ClearCoatRoughness.IsConnected())	{ SubstrateMaterialInfo.AddPropertyConnected(MP_CustomData1); }
-	if (Opacity.IsConnected())				{ SubstrateMaterialInfo.AddPropertyConnected(MP_Opacity); }
-
-	if (ShadingModel.IsConnected())
-	{
-		SubstrateMaterialInfo.AddPropertyConnected(MP_ShadingModel);
-
-		// If the ShadingModel pin is plugged in, we must use a shading model from expression path.
-		SubstrateMaterialInfo.SetShadingModelFromExpression(true);
-	}
-	else
-	{
-		// If the ShadingModel pin is NOT plugged in, we simply use the shading model selected on the root node drop box.
-		if (ShadingModelOverride == MSM_Unlit)				{ SubstrateMaterialInfo.AddShadingModel(ESubstrateShadingModel::SSM_Unlit); }
-		if (ShadingModelOverride == MSM_DefaultLit)			{ SubstrateMaterialInfo.AddShadingModel(ESubstrateShadingModel::SSM_DefaultLit); }
-		if (ShadingModelOverride == MSM_Subsurface)			{ SubstrateMaterialInfo.AddShadingModel(ESubstrateShadingModel::SSM_SubsurfaceWrap); }
-		if (ShadingModelOverride == MSM_PreintegratedSkin)	{ SubstrateMaterialInfo.AddShadingModel(ESubstrateShadingModel::SSM_SubsurfaceWrap); }
-		if (ShadingModelOverride == MSM_ClearCoat)			{ SubstrateMaterialInfo.AddShadingModel(ESubstrateShadingModel::SSM_ClearCoat); }
-		if (ShadingModelOverride == MSM_SubsurfaceProfile)	{ SubstrateMaterialInfo.AddShadingModel(ESubstrateShadingModel::SSM_SubsurfaceProfile); }
-		if (ShadingModelOverride == MSM_TwoSidedFoliage)	{ SubstrateMaterialInfo.AddShadingModel(ESubstrateShadingModel::SSM_SubsurfaceThinTwoSided); }
-		if (ShadingModelOverride == MSM_Hair)				{ SubstrateMaterialInfo.AddShadingModel(ESubstrateShadingModel::SSM_Hair); }
-		if (ShadingModelOverride == MSM_Cloth)				{ SubstrateMaterialInfo.AddShadingModel(ESubstrateShadingModel::SSM_Cloth); }
-		if (ShadingModelOverride == MSM_Eye)				{ SubstrateMaterialInfo.AddShadingModel(ESubstrateShadingModel::SSM_Eye); }
-		if (ShadingModelOverride == MSM_SingleLayerWater)	{ SubstrateMaterialInfo.AddShadingModel(ESubstrateShadingModel::SSM_SingleLayerWater); }
-		if (ShadingModelOverride == MSM_ThinTranslucent)	{ SubstrateMaterialInfo.AddShadingModel(ESubstrateShadingModel::SSM_ThinTranslucent); }
-	}
-
-	if (SubsurfaceProfile)
-	{
-		SubstrateMaterialInfo.AddSubsurfaceProfile(SubsurfaceProfile);
-	}
-}
-
-FSubstrateOperator* UMaterialExpressionSubstrateShadingModels::SubstrateGenerateMaterialTopologyTree(class FMaterialCompiler* Compiler, class UMaterialExpression* Parent, int32 OutputIndex)
-{
-	// Note Thickness has no meaning/usage in the context of SubstrateLegacyConversionNode
-	int32 ThicknessIndex = Compiler->SubstrateThicknessStackGetThicknessIndex();
-
-	auto AddDefaultWorstCase = [&](bool bSSS, bool bFuzz)
-	{
-		FSubstrateOperator& SlabOperator = Compiler->SubstrateCompilationRegisterOperator(SUBSTRATE_OPERATOR_BSDF_LEGACY, Compiler->SubstrateTreeStackGetPathUniqueId(), this, Parent, Compiler->SubstrateTreeStackGetParentPathUniqueId());
-		SlabOperator.BSDFType = SUBSTRATE_BSDF_TYPE_SLAB;
-		SlabOperator.bBSDFHasSSS = bSSS;
-		SlabOperator.bBSDFHasMFPPluggedIn = bSSS;
-		SlabOperator.bBSDFHasFuzz = bFuzz;
-		SlabOperator.bBSDFHasAnisotropy = Anisotropy.IsConnected();
-		SlabOperator.bBSDFWritesEmissive = EmissiveColor.IsConnected();
-		SlabOperator.ThicknessIndex = ThicknessIndex;
-
-		return &SlabOperator;
-	};
-
-	// Get the shading models resulting from the UMaterial::RebuildShadingModelField().
-	FMaterialShadingModelField ShadingModels = Compiler->GetMaterialShadingModels();
-
-	// Logic about shading models and complexity should match UMaterialExpressionSubstrateShadingModels::Compile.
-	const bool bHasShadingModelFromExpression = ShadingModel.IsConnected(); // We keep HasShadingModelFromExpression in case all shading models cannot be safely recovered from material functions.
-	if ((ShadingModels.CountShadingModels() > 1) || bHasShadingModelFromExpression) 
-	{
-		// Special case for unlit only material to get fast path
-		if (ShadingModels.HasOnlyShadingModel(MSM_Unlit))
-		{
-			FSubstrateOperator& Operator = Compiler->SubstrateCompilationRegisterOperator(SUBSTRATE_OPERATOR_BSDF_LEGACY, Compiler->SubstrateTreeStackGetPathUniqueId(), this, Parent, Compiler->SubstrateTreeStackGetParentPathUniqueId());
-			Operator.BSDFType = SUBSTRATE_BSDF_TYPE_UNLIT;
-			Operator.ThicknessIndex = ThicknessIndex;
-			return &Operator;
-		}
-		return AddDefaultWorstCase(true, true);
-	}
-	// else
-	{
-		check(ShadingModels.CountShadingModels() == 1);
-
-		if (ShadingModels.HasShadingModel(MSM_Unlit))
-		{
-			FSubstrateOperator& Operator = Compiler->SubstrateCompilationRegisterOperator(SUBSTRATE_OPERATOR_BSDF_LEGACY, Compiler->SubstrateTreeStackGetPathUniqueId(), this, Parent, Compiler->SubstrateTreeStackGetParentPathUniqueId());
-			Operator.BSDFType = SUBSTRATE_BSDF_TYPE_UNLIT;
-			Operator.ThicknessIndex = ThicknessIndex;
-			Operator.bBSDFWritesEmissive = true;
-			return &Operator;
-		}
-		else if (ShadingModels.HasShadingModel(MSM_DefaultLit))
-		{
-			return AddDefaultWorstCase(false, false);
-		}
-		else if (ShadingModels.HasShadingModel(MSM_ThinTranslucent))
-		{
-			return AddDefaultWorstCase(false, false);
-		}
-		else if (ShadingModels.HasShadingModel(MSM_SubsurfaceProfile))
-		{
-			return AddDefaultWorstCase(true, false);
-		}
-		else if (ShadingModels.HasShadingModel(MSM_Subsurface))
-		{
-			return AddDefaultWorstCase(true, false);
-		}
-		else if (ShadingModels.HasShadingModel(MSM_TwoSidedFoliage))
-		{
-			return AddDefaultWorstCase(true, false);
-		}
-		else if (ShadingModels.HasShadingModel(MSM_PreintegratedSkin))
-		{
-			return AddDefaultWorstCase(true, false);
-		}
-		else if (ShadingModels.HasShadingModel(MSM_Cloth))
-		{
-			return AddDefaultWorstCase(false, true);
-		}
-		else if (ShadingModels.HasShadingModel(MSM_ClearCoat))
-		{
-			FSubstrateOperator& Operator = Compiler->SubstrateCompilationRegisterOperator(SUBSTRATE_OPERATOR_BSDF_LEGACY, Compiler->SubstrateTreeStackGetPathUniqueId(), this, Parent, Compiler->SubstrateTreeStackGetParentPathUniqueId());
-			Operator.BSDFType = SUBSTRATE_BSDF_TYPE_SLAB;
-			Operator.ThicknessIndex = ThicknessIndex;
-			Operator.bBSDFHasSecondRoughnessOrSimpleClearCoat = true;
-			Operator.bBSDFHasAnisotropy = Anisotropy.IsConnected();
-			Operator.bBSDFWritesEmissive = EmissiveColor.IsConnected();
-			return &Operator;
-		}
-		else if (ShadingModels.HasShadingModel(MSM_Hair))
-		{
-			FSubstrateOperator& Operator = Compiler->SubstrateCompilationRegisterOperator(SUBSTRATE_OPERATOR_BSDF_LEGACY, Compiler->SubstrateTreeStackGetPathUniqueId(), this, Parent, Compiler->SubstrateTreeStackGetParentPathUniqueId());
-			Operator.BSDFType = SUBSTRATE_BSDF_TYPE_HAIR;
-			Operator.ThicknessIndex = ThicknessIndex;
-			Operator.bBSDFWritesEmissive = EmissiveColor.IsConnected();
-			return &Operator;
-		}
-		else if (ShadingModels.HasShadingModel(MSM_Eye))
-		{
-			FSubstrateOperator& Operator = Compiler->SubstrateCompilationRegisterOperator(SUBSTRATE_OPERATOR_BSDF_LEGACY, Compiler->SubstrateTreeStackGetPathUniqueId(), this, Parent, Compiler->SubstrateTreeStackGetParentPathUniqueId());
-			Operator.BSDFType = SUBSTRATE_BSDF_TYPE_EYE;
-			Operator.ThicknessIndex = ThicknessIndex;
-			Operator.bBSDFWritesEmissive = EmissiveColor.IsConnected();
-			return &Operator;
-		}
-		else if (ShadingModels.HasShadingModel(MSM_SingleLayerWater))
-		{
-			FSubstrateOperator& Operator = Compiler->SubstrateCompilationRegisterOperator(SUBSTRATE_OPERATOR_BSDF_LEGACY, Compiler->SubstrateTreeStackGetPathUniqueId(), this, Parent, Compiler->SubstrateTreeStackGetParentPathUniqueId());
-			Operator.BSDFType = SUBSTRATE_BSDF_TYPE_SINGLELAYERWATER;
-			Operator.ThicknessIndex = ThicknessIndex;
-			Operator.bBSDFWritesEmissive = EmissiveColor.IsConnected();
-			return &Operator;
-		}
-
-		check(false);
-		static FSubstrateOperator DefaultOperatorOnError;
-		return &DefaultOperatorOnError;
-	}
-}
-
-bool UMaterialExpressionSubstrateShadingModels::HasSSS() const
-{
-	return SubsurfaceProfile != nullptr;
-}
-
-bool UMaterialExpressionSubstrateShadingModels::HasAnisotropy() const
-{
-	return Anisotropy.IsConnected();
-}
-
-#endif // WITH_EDITOR
-
-
-UMaterialExpressionSubstrateBSDF::UMaterialExpressionSubstrateBSDF(const FObjectInitializer& ObjectInitializer)
-	: Super(ObjectInitializer)
-{
-}
-
-#if WITH_EDITOR
-int32 UMaterialExpressionSubstrateBSDF::CompilePreview(class FMaterialCompiler* Compiler, int32 OutputIndex)
-{
-	// Compile the SubstrateData output.
-	int32 SubstrateDataCodeChunk = Compile(Compiler, OutputIndex);
-	// Convert the SubstrateData to a preview color.
-	int32 PreviewCodeChunk = Compiler->SubstrateCompilePreview(SubstrateDataCodeChunk);
-	return PreviewCodeChunk;
-}
-#endif
-
-UMaterialExpressionSubstrateSlabBSDF::UMaterialExpressionSubstrateSlabBSDF(const FObjectInitializer& ObjectInitializer)
-	: Super(ObjectInitializer)
-	, bUseSSSDiffusion(true)
-{
-	struct FConstructorStatics
-	{
-		FText NAME_Substrate;
-		FConstructorStatics() : NAME_Substrate(LOCTEXT("Substrate BSDFs", "Substrate BSDFs")) { }
-	};
-	static FConstructorStatics ConstructorStatics;
-#if WITH_EDITORONLY_DATA
-	MenuCategories.Add(ConstructorStatics.NAME_Substrate);
-#endif
-}
-
-#if WITH_EDITOR
-FName CreateSpecularProfileParameterName(USpecularProfile* InProfile);
-int32 UMaterialExpressionSubstrateSlabBSDF::Compile(class FMaterialCompiler* Compiler, int32 OutputIndex)
-{
-	FGuid PathUniqueId = Compiler->SubstrateTreeStackGetPathUniqueId();
-	FSubstrateOperator& SubstrateOperator = Compiler->SubstrateCompilationGetOperator(PathUniqueId);
-
-	// We also cannot ignore the tangent when using the default Tangent because GetTangentBasis
-	// used in SubstrateGetBSDFSharedBasis cannot be relied on for smooth tangent used for lighting on any mesh.
-	const bool bHasAnisotropy		= SubstrateOperator.bBSDFHasAnisotropy > 0;
-
-	int32 SSSProfileCodeChunk = INDEX_NONE;
-	if (SubstrateOperator.bBSDFHasSSS > 0)
-	{
-		SSSProfileCodeChunk = Compiler->ForceCast(Compiler->ScalarParameter(GetSubsurfaceProfileParameterName(), 1.0f), MCT_Float1);
-	}
-
-	int32 SpecularProfileCodeChunk = INDEX_NONE;
-	if (SubstrateOperator.bBSDFHasSpecularProfile > 0)
-	{
-		const FName SpecularProfileParameterName = CreateSpecularProfileParameterName(SpecularProfile);
-		SpecularProfileCodeChunk = Compiler->ForceCast(Compiler->ScalarParameter(SpecularProfileParameterName, 1.0f), MCT_Float1);
-	}
-
-	const float DefaultSpecular = 0.5f;
-	const float DefaultF0 = DielectricSpecularToF0(DefaultSpecular);
-
-	int32 NormalCodeChunk = CompileWithDefaultNormalWS(Compiler, Normal);
-	int32 TangentCodeChunk = bHasAnisotropy ? CompileWithDefaultTangentWS(Compiler, Tangent) : INDEX_NONE;
-	const FSubstrateRegisteredSharedLocalBasis NewRegisteredSharedLocalBasis = SubstrateCompilationInfoCreateSharedLocalBasis(Compiler, NormalCodeChunk, TangentCodeChunk);
-
-	SubstrateOperator.BSDFRegisteredSharedLocalBasis = NewRegisteredSharedLocalBasis;
-	
-	int32 ThicknesCodeChunk = INDEX_NONE;
-	if (SubstrateOperator.ThicknessIndex != INDEX_NONE)
-	{
-		ThicknesCodeChunk = Compiler->SubstrateThicknessStackGetThicknessCode(SubstrateOperator.ThicknessIndex);
-	}
-	else
-	{
-		// Thickness is not tracked properly, this can happen when opening a material function in editor
-		ThicknesCodeChunk = Compiler->Constant(SUBSTRATE_LAYER_DEFAULT_THICKNESS_CM);
-	}
-	check(ThicknesCodeChunk != INDEX_NONE);
-
-	int32 DiffuseAlbedoCodeChunk		= CompileWithDefaultFloat3(Compiler, DiffuseAlbedo, 0.18f, 0.18f, 0.18f);
-	int32 F0CodeChunk					= CompileWithDefaultFloat3(Compiler, F0, DefaultF0, DefaultF0, DefaultF0);
-	int32 RoughnessCodeChunk			= CompileWithDefaultFloat1(Compiler, Roughness, 0.5f);
-	int32 AnisotropyCodeChunk			= CompileWithDefaultFloat1(Compiler, Anisotropy, 0.0f);
-	int32 F90CodeChunk					= CompileWithDefaultFloat3(Compiler, F90, 1.0f, 1.0f, 1.0f);
-	int32 SSSMFPCodeChunk				= CompileWithDefaultFloat3(Compiler, SSSMFP, 0.0f, 0.0f, 0.0f);
-	int32 SSSMFPScaleCodeChunk			= CompileWithDefaultFloat1(Compiler, SSSMFPScale, 1.0f);
-	int32 SSSPhaseAnisotropyCodeChunk	= CompileWithDefaultFloat1(Compiler, SSSPhaseAnisotropy, 0.0f);
-	int32 SecondRoughnessCodeChunk		= CompileWithDefaultFloat1(Compiler, SecondRoughness, 0.0f);
-	int32 SecondRoughnessWeightCodeChunk= CompileWithDefaultFloat1(Compiler, SecondRoughnessWeight, 0.0f);
-	int32 FuzzAmountCodeChunk			= CompileWithDefaultFloat1(Compiler, FuzzAmount, 0.0f);
-	int32 FuzzColorCodeChunk			= CompileWithDefaultFloat3(Compiler, FuzzColor, 0.0f, 0.0f, 0.0f);
-	int32 FuzzRoughnessCodeChunk		= HasFuzzRoughness() ? CompileWithDefaultFloat1(Compiler, FuzzRoughness, 0.5f) : RoughnessCodeChunk;
-	int32 GlintValueCodeChunk			= CompileWithDefaultFloat1(Compiler, GlintValue, 0.0f);
-	int32 GlintUVCodeChunk				= CompileWithDefaultFloat2(Compiler, GlintUV, 0.0f, 0.0f);
-
-	// Disable some features if requested by the simplification process
-	if (SubstrateOperator.bBSDFHasMFPPluggedIn == 0)
-	{
-		SSSMFPCodeChunk = Compiler->Constant3(0.0f, 0.0f, 0.0f);
-	}
-	if (SubstrateOperator.bBSDFHasEdgeColor == 0)
-	{
-		F90CodeChunk = Compiler->Constant3(1.0f, 1.0f, 1.0f);
-	}
-	if (SubstrateOperator.bBSDFHasFuzz == 0)
-	{
-		FuzzAmountCodeChunk = Compiler->Constant(0.0f);
-	}
-	if (SubstrateOperator.bBSDFHasSecondRoughnessOrSimpleClearCoat == 0)
-	{
-		SecondRoughnessWeightCodeChunk = Compiler->Constant(0.0f);
-	}
-	if (SubstrateOperator.bBSDFHasAnisotropy == 0)
-	{
-		AnisotropyCodeChunk = Compiler->Constant(0.0f);
-	}
-	if (SubstrateOperator.bBSDFHasGlint == 0)
-	{
-		GlintValueCodeChunk = Compiler->Constant(0.0f);
-	}
-	if (SubstrateOperator.bBSDFHasSpecularProfile == 0)
-	{
-		SpecularProfileCodeChunk = INDEX_NONE;
-	}
-
-	int32 OutputCodeChunk = Compiler->SubstrateSlabBSDF(
-		DiffuseAlbedoCodeChunk,
-		F0CodeChunk,
-		F90CodeChunk,
-		RoughnessCodeChunk,
-		AnisotropyCodeChunk,
-		SSSProfileCodeChunk != INDEX_NONE ? SSSProfileCodeChunk : Compiler->Constant(0.0f),
-		SSSMFPCodeChunk,
-		SSSMFPScaleCodeChunk,
-		SSSPhaseAnisotropyCodeChunk,
-		bUseSSSDiffusion ? Compiler->Constant(1.0f) : Compiler->Constant(0.0f),
-		CompileWithDefaultFloat3(Compiler, EmissiveColor, 0.0f, 0.0f, 0.0f),
-		SecondRoughnessCodeChunk,
-		SecondRoughnessWeightCodeChunk,
-		Compiler->Constant(0.0f),										// SecondRoughnessAsSimpleClearCoat
-		FuzzAmountCodeChunk,
-		FuzzColorCodeChunk,
-		FuzzRoughnessCodeChunk,
-		ThicknesCodeChunk,
-		GlintValueCodeChunk,
-		GlintUVCodeChunk,
-		SpecularProfileCodeChunk != INDEX_NONE ? SpecularProfileCodeChunk : Compiler->Constant(0.0f),
-		SubstrateOperator.bIsBottom > 0 ? true : false,
-		NormalCodeChunk,
-		TangentCodeChunk,
-		Compiler->GetSubstrateSharedLocalBasisIndexMacro(NewRegisteredSharedLocalBasis),
-		!SubstrateOperator.bUseParameterBlending || (SubstrateOperator.bUseParameterBlending && SubstrateOperator.bRootOfParameterBlendingSubTree) ? &SubstrateOperator : nullptr);
-
-	return OutputCodeChunk;
-}
-
-#if ENABLE_MATERIAL_LAYER_PROTOTYPE
-int32 UMaterialExpressionSubstrateSlabBSDF::CompileDefaultSlab(FMaterialCompiler* Compiler, FVector3f DiffuseAlbedoOverride)
-{
-	/**
-	 * This creates a default slab, primarily for material function previews where the input may be optional/missing.
-	 */
-	if (!Substrate::IsSubstrateEnabled())
-	{
-		return Compiler->SubstrateCreateAndRegisterNullMaterial();
-	}
-
-	FGuid PathUniqueId = Compiler->SubstrateTreeStackGetPathUniqueId();
-	FSubstrateOperator& SubstrateOperator = Compiler->SubstrateCompilationGetOperator(PathUniqueId);
-
-	int32 SSSProfileCodeChunk = INDEX_NONE;
-	int32 SpecularProfileCodeChunk = INDEX_NONE;	
-	int32 TangentCodeChunk = INDEX_NONE;
-	int32 NormalCodeChunk = Compiler->VertexNormal();
-	int32 ThicknessCodeChunk = Compiler->Constant(SUBSTRATE_LAYER_DEFAULT_THICKNESS_CM);
-
-	const FSubstrateRegisteredSharedLocalBasis NewRegisteredSharedLocalBasis = SubstrateCompilationInfoCreateSharedLocalBasis(Compiler, NormalCodeChunk, TangentCodeChunk);
-	SubstrateOperator.BSDFRegisteredSharedLocalBasis = NewRegisteredSharedLocalBasis;
-
-	int32 DiffuseAlbedoCodeChunk = Compiler->Constant3(DiffuseAlbedoOverride.X, DiffuseAlbedoOverride.Y, DiffuseAlbedoOverride.Z);
-	const float DefaultF0 = 0.04f;
-	int32 F0CodeChunk = Compiler->Constant3(DefaultF0, DefaultF0, DefaultF0);
-	int32 RoughnessCodeChunk = Compiler->Constant(0.5f);
-	int32 AnisotropyCodeChunk = Compiler->Constant(0.0f);
-	int32 F90CodeChunk = Compiler->Constant3(1.0f, 1.0f, 1.0f);
-	int32 SSSMFPCodeChunk = Compiler->Constant3(0.0f, 0.0f, 0.0f);
-	int32 SSSMFPScaleCodeChunk = Compiler->Constant(1.0f);
-	int32 SSSPhaseAnisotropyCodeChunk = Compiler->Constant(0.0f);
-	int32 SecondRoughnessCodeChunk = Compiler->Constant(0.0f);
-	int32 SecondRoughnessWeightCodeChunk = Compiler->Constant(0.0f);
-	int32 FuzzAmountCodeChunk = Compiler->Constant(0.0f);
-	int32 FuzzColorCodeChunk = Compiler->Constant3(0.0f, 0.0f, 0.0f);
-	int32 FuzzRoughnessCodeChunk = Compiler->Constant(0.5f);
-	int32 GlintValueCodeChunk = Compiler->Constant(0.0f);
-	int32 GlintUVCodeChunk = Compiler->Constant2(0.0f, 0.0f);
-
-	return Compiler->SubstrateSlabBSDF(
-		DiffuseAlbedoCodeChunk,
-		F0CodeChunk,
-		F90CodeChunk,
-		RoughnessCodeChunk,
-		AnisotropyCodeChunk,
-		Compiler->Constant(0.0f),
-		SSSMFPCodeChunk,
-		SSSMFPScaleCodeChunk,
-		SSSPhaseAnisotropyCodeChunk,
-		Compiler->Constant(0.0f),
-		Compiler->Constant3(0.0f, 0.0f, 0.0f),
-		SecondRoughnessCodeChunk,
-		SecondRoughnessWeightCodeChunk,
-		Compiler->Constant(0.0f),										// SecondRoughnessAsSimpleClearCoat
-		FuzzAmountCodeChunk,
-		FuzzColorCodeChunk,
-		FuzzRoughnessCodeChunk,
-		ThicknessCodeChunk,
-		GlintValueCodeChunk,
-		GlintUVCodeChunk,
-		Compiler->Constant(0.0f),
-		false,
-		NormalCodeChunk,
-		TangentCodeChunk,
-		Compiler->GetSubstrateSharedLocalBasisIndexMacro(NewRegisteredSharedLocalBasis),
-		(SubstrateOperator.Index != INDEX_NONE && SubstrateOperator.BSDFIndex != INDEX_NONE) && (!SubstrateOperator.bUseParameterBlending || (SubstrateOperator.bUseParameterBlending && SubstrateOperator.bRootOfParameterBlendingSubTree)) ? &SubstrateOperator : nullptr);
-}
-#endif //ENABLE_MATERIAL_LAYER_PROTOTYPE
-
-FSubstrateMaterialComplexity UMaterialExpressionSubstrateSlabBSDF::GetComplexity() const
-{
-	FSubstrateMaterialComplexity Out;
-	Out.bIsComplexSpecial = HasGlint() || HasSpecularProfile();
-	Out.bIsSingle 		  = !HasAnisotropy() && (HasEdgeColor() || HasFuzz() || HasSecondRoughness() || HasMFPPluggedIn() || HasSSS());
-	Out.bIsSimple 		  = !HasAnisotropy();
-
-	// Mask out to only have a single possibility
-	Out.bIsSingle &= !Out.bIsComplexSpecial;
-	Out.bIsSimple &= !Out.bIsComplexSpecial && !Out.bIsSingle;
-
-	return Out;
-}
-
-void UMaterialExpressionSubstrateSlabBSDF::GetCaption(TArray<FString>& OutCaptions) const
-{
-	// The node complexity is manually maintained to match FSubstrateCompilationContext::SubstrateGenerateDerivedMaterialOperatorData and shaders.
-	OutCaptions.Add(TEXT("Substrate Slab BSDF - ") + FSubstrateMaterialComplexity::ToString(GetComplexity().SubstrateMaterialType()));
-}
-
-void UMaterialExpressionSubstrateSlabBSDF::GetExpressionToolTip(TArray<FString>& OutToolTip)
-{
-	OutToolTip.Add(TEXT("Substrate Slab BSDF"));
-	OutToolTip.Add(TEXT("Complexity = ") + FSubstrateMaterialComplexity::ToString(GetComplexity().SubstrateMaterialType()));
-	OutToolTip.Add(TEXT("The complexity represents the cost of the shading path (Lighting, Lumen, SSS) the material will follow:"));
-	OutToolTip.Add(TEXT(" - Simple means the Slab only relies on Diffuse, F0 and Roughness. It will follow a fast shading path."));
-	OutToolTip.Add(TEXT(" - Single means the Slab uses more features such as F90, Fuzz, Second Roughness, MFP or SSS. It will follow a more expenssive shading path."));
-	OutToolTip.Add(TEXT(" - Complex means a Slab uses anisotropic lighting, with any of the previous features."));
-	OutToolTip.Add(TEXT(" - Complex Special means the Slab is using more advanced features such as glints or specular LUT. This is the most expenssive shading path."));
-}
-
-uint32 UMaterialExpressionSubstrateSlabBSDF::GetOutputType(int32 OutputIndex)
-{
-	return MCT_Substrate;
-}
-
-uint32 UMaterialExpressionSubstrateSlabBSDF::GetInputType(int32 InputIndex)
-{
-	if (InputIndex == 0)
-	{
-		return MCT_Float3; // DiffuseAlbedo
-	}
-	else if (InputIndex == 1)
-	{
-		return MCT_Float3; // F0
-	}
-	else if (InputIndex == 2)
-	{
-		return MCT_Float3; // F90
-	}
-	else if (InputIndex == 3)
-	{
-		return MCT_Float1; // Roughness
-	}
-	else if (InputIndex == 4)
-	{
-		return MCT_Float1; // Anisotropy
-	}
-	else if (InputIndex == 5)
-	{
-		return MCT_Float3; // Normal
-	}
-	else if (InputIndex == 6)
-	{
-		return MCT_Float3; // Tangent
-	}
-	else if (InputIndex == 7)
-	{
-		return MCT_Float3; // SSSMFP
-	}
-	else if (InputIndex == 8)
-	{
-		return MCT_Float1; // SSSMFPScale
-	}
-	else if (InputIndex == 9)
-	{
-		return MCT_Float1; // SSSPhaseAniso
-	}
-	else if (InputIndex == 10)
-	{
-		return MCT_Float3; // Emissive Color
-	}
-	else if (InputIndex == 11)
-	{
-		return MCT_Float1; // SecondRoughness
-	}
-	else if (InputIndex == 12)
-	{
-		return MCT_Float1; // SecondRoughnessWeight
-	}
-	else if (InputIndex == 13)
-	{
-		return MCT_Float1; // FuzzRoughness
-	}
-	else if (InputIndex == 14)
-	{
-		return MCT_Float1; // FuzzAmount
-	}
-	else if (InputIndex == 15)
-	{
-		return MCT_Float3; // FuzzColor
-	}
-	else if (InputIndex == 16)
-	{
-		return MCT_Float; // GlintValue
-	}
-	else if (InputIndex == 17)
-	{
-		return MCT_Float2; // GlintUV
-	}
-
-	check(false);
-	return MCT_Float1;
-}
-
-FName UMaterialExpressionSubstrateSlabBSDF::GetInputName(int32 InputIndex) const
-{
-	if (InputIndex == 0)
-	{
-		return TEXT("Diffuse Albedo");
-	}
-	else if (InputIndex == 1)
-	{
-		return TEXT("F0");
-	}
-	else if (InputIndex == 2)
-	{
-		return  TEXT("F90");
-	}
-	else if (InputIndex == 3)
-	{
-		return TEXT("Roughness");
-	}
-	else if (InputIndex == 4)
-	{
-		return TEXT("Anisotropy");
-	}
-	else if (InputIndex == 5)
-	{
-		return TEXT("Normal");
-	}
-	else if (InputIndex == 6)
-	{
-		return TEXT("Tangent");
-	}
-	else if (InputIndex == 7)
-	{
-		return TEXT("SSS MFP");
-	}
-	else if (InputIndex == 8)
-	{
-		return TEXT("SSS MFP Scale");
-	}
-	else if (InputIndex == 9)
-	{
-		return TEXT("SSS Phase Anisotropy");
-	}
-	else if (InputIndex == 10)
-	{
-		return TEXT("Emissive Color");
-	}
-	else if (InputIndex == 11)
-	{
-		return TEXT("Second Roughness");
-	}
-	else if (InputIndex == 12)
-	{
-		return TEXT("Second Roughness Weight");
-	}
-	else if (InputIndex == 13)
-	{
-		return TEXT("Fuzz Roughness");
-	}
-	else if (InputIndex == 14)
-	{
-		return TEXT("Fuzz Amount");
-	}
-	else if (InputIndex == 15)
-	{
-		return TEXT("Fuzz Color");
-	}
-	else if (InputIndex == 16)
-	{
-		return Substrate::IsGlintEnabled() ? TEXT("Glint Density") : TEXT("Glint Density (Disabled)");
-	}
-	else if (InputIndex == 17)
-	{
-		return Substrate::IsGlintEnabled() ? TEXT("Glint UVs") : TEXT("Glint UVs (Disabled)");
-	}
-
-	return TEXT("Unknown");
-}
-
-void UMaterialExpressionSubstrateSlabBSDF::GetConnectorToolTip(int32 InputIndex, int32 OutputIndex, TArray<FString>& OutToolTip)
-{
-	if (OutputIndex == 0)
-	{
-		OutToolTip.Add(TEXT("TT Ouput"));
-		return;
-	}
-
-	Super::GetConnectorToolTip(InputIndex, INDEX_NONE, OutToolTip);
-}
-
-bool UMaterialExpressionSubstrateSlabBSDF::IsResultSubstrateMaterial(int32 OutputIndex)
-{
-	return true;
-}
-
-void UMaterialExpressionSubstrateSlabBSDF::GatherSubstrateMaterialInfo(FSubstrateMaterialInfo& SubstrateMaterialInfo, int32 OutputIndex)
-{
-	// Track connected inputs
-	if (DiffuseAlbedo.IsConnected())		{ SubstrateMaterialInfo.AddPropertyConnected(MP_DiffuseColor); }
-	if (F0.IsConnected())					{ SubstrateMaterialInfo.AddPropertyConnected(MP_SpecularColor); }
-	if (Roughness.IsConnected())			{ SubstrateMaterialInfo.AddPropertyConnected(MP_Roughness); }
-	if (Anisotropy.IsConnected())			{ SubstrateMaterialInfo.AddPropertyConnected(MP_Anisotropy); }
-	if (EmissiveColor.IsConnected())		{ SubstrateMaterialInfo.AddPropertyConnected(MP_EmissiveColor); }
-	if (Normal.IsConnected())				{ SubstrateMaterialInfo.AddPropertyConnected(MP_Normal); }
-	if (Tangent.IsConnected())				{ SubstrateMaterialInfo.AddPropertyConnected(MP_Tangent); }
-	if (SSSMFP.IsConnected())				{ SubstrateMaterialInfo.AddPropertyConnected(MP_SubsurfaceColor); }
-
-	if (HasSSS())
-	{
-		// We still do not know if this is going to be a real SSS node because it is only possible for BSDF at the bottom of the stack. Nevertheless, we take the worst case into account.
-		if (SubsurfaceProfile)
-		{
-			SubstrateMaterialInfo.AddShadingModel(SSM_SubsurfaceProfile);
-			SubstrateMaterialInfo.AddSubsurfaceProfile(SubsurfaceProfile);
-		}
-		else
-		{
-			SubstrateMaterialInfo.AddShadingModel(SSM_SubsurfaceMFP);
-		}
-	}
-	else
-	{
-		SubstrateMaterialInfo.AddShadingModel(SSM_DefaultLit);
-	}
-
-	if (HasSpecularProfile())
-	{
-		SubstrateMaterialInfo.AddSpecularProfile(SpecularProfile);
-	}
-
-	SubstrateMaterialInfo.AddGuid(MaterialExpressionGuid);
-}
-
-FSubstrateOperator* UMaterialExpressionSubstrateSlabBSDF::SubstrateGenerateMaterialTopologyTree(class FMaterialCompiler* Compiler, class UMaterialExpression* Parent, int32 OutputIndex)
-{
-	FSubstrateOperator& SubstrateOperator = Compiler->SubstrateCompilationRegisterOperator(SUBSTRATE_OPERATOR_BSDF, Compiler->SubstrateTreeStackGetPathUniqueId(), this, Parent, Compiler->SubstrateTreeStackGetParentPathUniqueId());
-	SubstrateOperator.BSDFType = SUBSTRATE_BSDF_TYPE_SLAB;
-	SubstrateOperator.bBSDFHasEdgeColor = HasEdgeColor();
-	SubstrateOperator.bBSDFHasFuzz = HasFuzz();
-	SubstrateOperator.bBSDFHasSecondRoughnessOrSimpleClearCoat = HasSecondRoughness();
-	SubstrateOperator.bBSDFHasSSS = HasSSS();
-	SubstrateOperator.bBSDFHasMFPPluggedIn = HasMFPPluggedIn();
-	SubstrateOperator.bBSDFHasAnisotropy = HasAnisotropy();
-	SubstrateOperator.bBSDFHasGlint = HasGlint();
-	SubstrateOperator.bBSDFHasSpecularProfile = HasSpecularProfile();
-	SubstrateOperator.ThicknessIndex = Compiler->SubstrateThicknessStackGetThicknessIndex();
-	SubstrateOperator.bBSDFWritesEmissive = EmissiveColor.IsConnected();
-	return &SubstrateOperator;
-}
-
-bool UMaterialExpressionSubstrateSlabBSDF::HasSSS() const
-{
-	return SubsurfaceProfile != nullptr || SSSMFP.IsConnected();
-}
-
-bool UMaterialExpressionSubstrateSlabBSDF::HasSSSProfile() const
-{
-	return SubsurfaceProfile != nullptr;
-}
-
-bool UMaterialExpressionSubstrateSlabBSDF::HasMFPPluggedIn() const
-{
-	return SSSMFP.IsConnected();
-}
-
-bool UMaterialExpressionSubstrateSlabBSDF::HasEdgeColor() const
-{
-	return F90.IsConnected();
-}
-
-bool UMaterialExpressionSubstrateSlabBSDF::HasFuzz() const
-{
-	return FuzzAmount.IsConnected();
-}
-
-bool UMaterialExpressionSubstrateSlabBSDF::HasFuzzRoughness() const
-{
-	return FuzzRoughness.IsConnected();
-}
-
-bool UMaterialExpressionSubstrateSlabBSDF::HasSecondRoughness() const
-{
-	return SecondRoughnessWeight.IsConnected();
-}
-
-bool UMaterialExpressionSubstrateSlabBSDF::HasAnisotropy() const
-{
-	return Anisotropy.IsConnected();
-}
-
-bool UMaterialExpressionSubstrateSlabBSDF::HasGlint() const
-{
-	// Use IsGlintEnabled to not promote to ComplexSpecial render path if glint rendering is not enabled and avoid registering such possibility at runtime.
-	return GlintValue.IsConnected() && Substrate::IsGlintEnabled();
-}
-
-bool UMaterialExpressionSubstrateSlabBSDF::HasSpecularProfile() const
-{
-	// Use IsSpecularLUTEnabled to not promote to ComplexSpecial render path if glint rendering is not enabled and avoid registering such possibility at runtime.
-	return SpecularProfile != nullptr && Substrate::IsSpecularProfileEnabled();
-}
-
-#endif // WITH_EDITOR
-
-
-
-UMaterialExpressionSubstrateSimpleClearCoatBSDF::UMaterialExpressionSubstrateSimpleClearCoatBSDF(const FObjectInitializer& ObjectInitializer)
-	: Super(ObjectInitializer)
-{
-	struct FConstructorStatics
-	{
-		FText NAME_Substrate;
-		FConstructorStatics() : NAME_Substrate(LOCTEXT("Substrate BSDFs", "Substrate BSDFs")) { }
-	};
-	static FConstructorStatics ConstructorStatics;
-#if WITH_EDITORONLY_DATA
-	MenuCategories.Add(ConstructorStatics.NAME_Substrate);
-#endif
-}
-
-#if WITH_EDITOR
-int32 UMaterialExpressionSubstrateSimpleClearCoatBSDF::Compile(class FMaterialCompiler* Compiler, int32 OutputIndex)
-{
-	const float DefaultSpecular = 0.5f;
-	const float DefaultF0 = DielectricSpecularToF0(DefaultSpecular);
-
-	int32 NormalCodeChunk = CompileWithDefaultNormalWS(Compiler, Normal);
-	const int32 NullTangentCodeChunk = INDEX_NONE;
-	const FSubstrateRegisteredSharedLocalBasis NewRegisteredSharedLocalBasis = SubstrateCompilationInfoCreateSharedLocalBasis(Compiler, NormalCodeChunk, NullTangentCodeChunk);
-
-	FSubstrateOperator& SubstrateOperator = Compiler->SubstrateCompilationGetOperator(Compiler->SubstrateTreeStackGetPathUniqueId());
-	SubstrateOperator.BSDFRegisteredSharedLocalBasis = NewRegisteredSharedLocalBasis;
-
-	int32 ThicknessCodeChunk = Compiler->SubstrateThicknessStackGetThicknessCode(SubstrateOperator.ThicknessIndex);
-	check(ThicknessCodeChunk != INDEX_NONE);
-
-	int32 RoughnessCodeChunk = CompileWithDefaultFloat1(Compiler, Roughness, 0.5f);
-
-	int32 OutputCodeChunk = Compiler->SubstrateSlabBSDF(
-		CompileWithDefaultFloat3(Compiler, DiffuseAlbedo, 0.18f, 0.18f, 0.18f),		// DiffuseAlbedo
-		CompileWithDefaultFloat3(Compiler, F0, DefaultF0, DefaultF0, DefaultF0),	// F0
-		Compiler->Constant3(1.0f, 1.0f, 1.0f),					// F90		
-		RoughnessCodeChunk,										// Roughness
-		Compiler->Constant(0.0f),								// Anisotropy
-		Compiler->Constant(0.0f),								// SSSProfile
-		Compiler->Constant3(0.0f, 0.0f, 0.0f),					// SSSMFP
-		Compiler->Constant(0.0f),								// SSSMFPScale
-		Compiler->Constant(0.0f),								// SSSPhaseAnisotropy
-		Compiler->Constant(0.0f),								// bUseSSSDiffusion
-		CompileWithDefaultFloat3(Compiler, EmissiveColor, 0.0f, 0.0f, 0.0f),
-		CompileWithDefaultFloat1(Compiler, ClearCoatRoughness, 0.1f),
-		CompileWithDefaultFloat1(Compiler, ClearCoatCoverage, 1.0f),
-		Compiler->Constant(1.0f),								// SecondRoughnessAsSimpleClearCoat == true for UMaterialExpressionSubstrateSimpleClearCoatBSDF
-		Compiler->Constant(0.0f),								// FuzzAmount
-		Compiler->Constant3(0.0f, 0.0f, 0.0f),					// FuzzColor
-		RoughnessCodeChunk,										// FuzzRoughness
-		ThicknessCodeChunk,										// Thickness
-		Compiler->Constant(0.0f),								// GlintValue
-		Compiler->Constant2(0.0f, 0.0f),						// GlintUV
-		Compiler->Constant(0.0f),								// SpecularProfile
-		false,													// bIsAtTheBottomOfTopology, always false for SimpleClearCoat
-		NormalCodeChunk,
-		NullTangentCodeChunk,
-		Compiler->GetSubstrateSharedLocalBasisIndexMacro(NewRegisteredSharedLocalBasis),
-		!SubstrateOperator.bUseParameterBlending || (SubstrateOperator.bUseParameterBlending && SubstrateOperator.bRootOfParameterBlendingSubTree) ? &SubstrateOperator : nullptr);
-
-	return OutputCodeChunk;
-}
-
-void UMaterialExpressionSubstrateSimpleClearCoatBSDF::GetCaption(TArray<FString>& OutCaptions) const
-{
-	OutCaptions.Add(TEXT("Substrate Simple Clear Coat"));
-}
-
-uint32 UMaterialExpressionSubstrateSimpleClearCoatBSDF::GetOutputType(int32 OutputIndex)
-{
-	return MCT_Substrate;
-}
-
-uint32 UMaterialExpressionSubstrateSimpleClearCoatBSDF::GetInputType(int32 InputIndex)
-{
-	if (InputIndex == 0)
-	{
-		return MCT_Float3; // DiffuseAlbedo
-	}
-	else if (InputIndex == 1)
-	{
-		return MCT_Float3; // F0
-	}
-	else if (InputIndex == 2)
-	{
-		return MCT_Float1; // Roughness
-	}
-	else if (InputIndex == 3)
-	{
-		return MCT_Float1; // ClearCoatCoverage 
-	}
-	else if (InputIndex == 4)
-	{
-		return MCT_Float1; // ClearCoatRoughness
-	}
-	else if (InputIndex == 5)
-	{
-		return MCT_Float3; // Normal
-	}
-	else if (InputIndex == 6)
-	{
-		return MCT_Float3; // Emissive Color
-	}
-
-	check(false);
-	return MCT_Float1;
-}
-
-FName UMaterialExpressionSubstrateSimpleClearCoatBSDF::GetInputName(int32 InputIndex) const
-{
-	if (InputIndex == 0)
-	{
-		return TEXT("Diffuse Albedo");
-	}
-	else if (InputIndex == 1)
-	{
-		return TEXT("F0");
-	}
-	else if (InputIndex == 2)
-	{
-		return TEXT("Roughness");
-	}
-	else if (InputIndex == 3)
-	{
-		return TEXT("Clear Coat Coverage");
-	}
-	else if (InputIndex == 4)
-	{
-		return TEXT("Clear Coat Roughness");
-	}
-	else if (InputIndex == 5)
-	{
-		return TEXT("Normal");
-	}
-	else if (InputIndex == 6)
-	{
-		return TEXT("Emissive Color");
-	}
-
-	return TEXT("Unknown");
-}
-
-bool UMaterialExpressionSubstrateSimpleClearCoatBSDF::IsResultSubstrateMaterial(int32 OutputIndex)
-{
-	return true;
-}
-
-void UMaterialExpressionSubstrateSimpleClearCoatBSDF::GatherSubstrateMaterialInfo(FSubstrateMaterialInfo& SubstrateMaterialInfo, int32 OutputIndex)
-{
-	// Track connected inputs
-	if (DiffuseAlbedo.IsConnected()) { SubstrateMaterialInfo.AddPropertyConnected(MP_BaseColor); }
-	SubstrateMaterialInfo.AddPropertyConnected(MP_Metallic); // Metallic is always connected with Diffuse/F0 parameterisation
-	if (F0.IsConnected()) { SubstrateMaterialInfo.AddPropertyConnected(MP_Specular); }
-	if (Roughness.IsConnected()) { SubstrateMaterialInfo.AddPropertyConnected(MP_Roughness); }
-	if (Normal.IsConnected()) { SubstrateMaterialInfo.AddPropertyConnected(MP_Normal); }
-	if (EmissiveColor.IsConnected()) { SubstrateMaterialInfo.AddPropertyConnected(MP_EmissiveColor); }
-
-	SubstrateMaterialInfo.AddShadingModel(SSM_DefaultLit);
-	SubstrateMaterialInfo.AddGuid(MaterialExpressionGuid);
-}
-
-FSubstrateOperator* UMaterialExpressionSubstrateSimpleClearCoatBSDF::SubstrateGenerateMaterialTopologyTree(class FMaterialCompiler* Compiler, class UMaterialExpression* Parent, int32 OutputIndex)
-{
-	FSubstrateOperator& SubstrateOperator = Compiler->SubstrateCompilationRegisterOperator(SUBSTRATE_OPERATOR_BSDF, Compiler->SubstrateTreeStackGetPathUniqueId(), this, Parent, Compiler->SubstrateTreeStackGetParentPathUniqueId());
-	SubstrateOperator.BSDFType = SUBSTRATE_BSDF_TYPE_SLAB;
-	SubstrateOperator.bBSDFHasEdgeColor = false;
-	SubstrateOperator.bBSDFHasFuzz = false;
-	SubstrateOperator.bBSDFHasSecondRoughnessOrSimpleClearCoat = true;	// This node explicitly requires simple clear coat
-	SubstrateOperator.bBSDFHasSSS = false;
-	SubstrateOperator.bBSDFHasMFPPluggedIn = false;
-	SubstrateOperator.bBSDFHasAnisotropy = false;
-	SubstrateOperator.bBSDFHasGlint = false;
-	SubstrateOperator.bBSDFHasSpecularProfile = false;
-	SubstrateOperator.ThicknessIndex = Compiler->SubstrateThicknessStackGetThicknessIndex();
-	SubstrateOperator.bBSDFWritesEmissive = EmissiveColor.IsConnected();
-	return &SubstrateOperator;
-}
-#endif // WITH_EDITOR
-
-
-
-UMaterialExpressionSubstrateVolumetricFogCloudBSDF::UMaterialExpressionSubstrateVolumetricFogCloudBSDF(const FObjectInitializer& ObjectInitializer)
-	: Super(ObjectInitializer)
-{
-	struct FConstructorStatics
-	{
-		FText NAME_Substrate;
-		FConstructorStatics() : NAME_Substrate(LOCTEXT("Substrate BSDFs", "Substrate BSDFs")) { }
-	};
-	static FConstructorStatics ConstructorStatics;
-#if WITH_EDITORONLY_DATA
-	MenuCategories.Add(ConstructorStatics.NAME_Substrate);
-#endif
-}
-
-#if WITH_EDITOR
-int32 UMaterialExpressionSubstrateVolumetricFogCloudBSDF::Compile(class FMaterialCompiler* Compiler, int32 OutputIndex)
-{
-	int32 OutputCodeChunk = Compiler->SubstrateVolumetricFogCloudBSDF(
-		CompileWithDefaultFloat3(Compiler, Albedo, 0.0f, 0.0f, 0.0f),
-		CompileWithDefaultFloat3(Compiler, Extinction, 0.0f, 0.0f, 0.0f),
-		CompileWithDefaultFloat3(Compiler, EmissiveColor, 0.0f, 0.0f, 0.0f),
-		CompileWithDefaultFloat1(Compiler, AmbientOcclusion, 1.0f),
-		bEmissiveOnly);
-
-	return OutputCodeChunk;
-}
-
-void UMaterialExpressionSubstrateVolumetricFogCloudBSDF::GetCaption(TArray<FString>& OutCaptions) const
-{
-	FString Caption = TEXT("Substrate Volumetric-Fog-Cloud BSDF");
-	if (bEmissiveOnly)
-	{
-		Caption += TEXT("(Emissive Only)");
-	}
-	OutCaptions.Add(Caption);
-}
-
-uint32 UMaterialExpressionSubstrateVolumetricFogCloudBSDF::GetOutputType(int32 OutputIndex)
-{
-	return MCT_Substrate;
-}
-
-uint32 UMaterialExpressionSubstrateVolumetricFogCloudBSDF::GetInputType(int32 InputIndex)
-{
-	switch (InputIndex)
-	{
-	case 0:
-		return MCT_Float3;
-		break;
-	case 1:
-		return MCT_Float3;
-		break;
-	case 2:
-		return MCT_Float3;
-		break;
-	case 3:
-		return MCT_Float1;
-		break;
-	}
-
-	check(false);
-	return MCT_Float1;
-}
-
-bool UMaterialExpressionSubstrateVolumetricFogCloudBSDF::IsResultSubstrateMaterial(int32 OutputIndex)
-{
-	return true;
-}
-
-void UMaterialExpressionSubstrateVolumetricFogCloudBSDF::GatherSubstrateMaterialInfo(FSubstrateMaterialInfo& SubstrateMaterialInfo, int32 OutputIndex)
-{
-	SubstrateMaterialInfo.AddShadingModel(SSM_VolumetricFogCloud);
-	SubstrateMaterialInfo.AddGuid(MaterialExpressionGuid);
-}
-
-FSubstrateOperator* UMaterialExpressionSubstrateVolumetricFogCloudBSDF::SubstrateGenerateMaterialTopologyTree(class FMaterialCompiler* Compiler, class UMaterialExpression* Parent, int32 OutputIndex)
-{
-	FSubstrateOperator& SubstrateOperator = Compiler->SubstrateCompilationRegisterOperator(SUBSTRATE_OPERATOR_BSDF, Compiler->SubstrateTreeStackGetPathUniqueId(), this, Parent, Compiler->SubstrateTreeStackGetParentPathUniqueId());
-	SubstrateOperator.BSDFType = SUBSTRATE_BSDF_TYPE_VOLUMETRICFOGCLOUD;
-	SubstrateOperator.ThicknessIndex = Compiler->SubstrateThicknessStackGetThicknessIndex();
-	SubstrateOperator.bBSDFWritesEmissive = EmissiveColor.IsConnected();
-	SubstrateOperator.bBSDFWritesAmbientOcclusion = AmbientOcclusion.IsConnected();
-	return &SubstrateOperator;
-}
-#endif // WITH_EDITOR
-
-
-
-UMaterialExpressionSubstrateLightFunction::UMaterialExpressionSubstrateLightFunction(const FObjectInitializer& ObjectInitializer)
-	: Super(ObjectInitializer)
-{
-	struct FConstructorStatics
-	{
-		FText NAME_Substrate;
-		FConstructorStatics() : NAME_Substrate(LOCTEXT("Substrate Extras", "Substrate Extras")) { }
-	};
-	static FConstructorStatics ConstructorStatics;
-#if WITH_EDITORONLY_DATA
-	MenuCategories.Add(ConstructorStatics.NAME_Substrate);
-#endif
-}
-
-#if WITH_EDITOR
-int32 UMaterialExpressionSubstrateLightFunction::Compile(class FMaterialCompiler* Compiler, int32 OutputIndex)
-{
-	FSubstrateOperator& SubstrateOperator = Compiler->SubstrateCompilationGetOperator(Compiler->SubstrateTreeStackGetPathUniqueId());
-	int32 OutputCodeChunk = Compiler->SubstrateUnlitBSDF(
-		CompileWithDefaultFloat3(Compiler, Color, 0.0f, 0.0f, 0.0f),
-		Compiler->Constant(1.0f),				// Opacity / Transmittance is ignored by light functions.
-		Compiler->Constant3(0.0f, 0.0f, 1.0f),	// place holder normal
-		&SubstrateOperator);
-	return OutputCodeChunk;
-}
-
-void UMaterialExpressionSubstrateLightFunction::GetCaption(TArray<FString>& OutCaptions) const
-{
-	OutCaptions.Add(TEXT("Substrate Light Function"));
-}
-
-uint32 UMaterialExpressionSubstrateLightFunction::GetOutputType(int32 OutputIndex)
-{
-	return MCT_Substrate;
-}
-
-uint32 UMaterialExpressionSubstrateLightFunction::GetInputType(int32 InputIndex)
-{
-	switch (InputIndex)
-	{
-	case 0:
-		return MCT_Float3;
-		break;
-	}
-
-	check(false);
-	return MCT_Float1;
-}
-
-bool UMaterialExpressionSubstrateLightFunction::IsResultSubstrateMaterial(int32 OutputIndex)
-{
-	return true;
-}
-
-void UMaterialExpressionSubstrateLightFunction::GatherSubstrateMaterialInfo(FSubstrateMaterialInfo& SubstrateMaterialInfo, int32 OutputIndex)
-{
-	SubstrateMaterialInfo.AddShadingModel(SSM_LightFunction);
-	SubstrateMaterialInfo.AddGuid(MaterialExpressionGuid);
-}
-
-FSubstrateOperator* UMaterialExpressionSubstrateLightFunction::SubstrateGenerateMaterialTopologyTree(class FMaterialCompiler* Compiler, class UMaterialExpression* Parent, int32 OutputIndex)
-{
-	FSubstrateOperator& SubstrateOperator = Compiler->SubstrateCompilationRegisterOperator(SUBSTRATE_OPERATOR_BSDF, Compiler->SubstrateTreeStackGetPathUniqueId(), this, Parent, Compiler->SubstrateTreeStackGetParentPathUniqueId());
-	SubstrateOperator.BSDFType = SUBSTRATE_BSDF_TYPE_UNLIT;
-	SubstrateOperator.SubUsage = SUBSTRATE_OPERATOR_SUBUSAGE_LIGHTFUNCTION;
-	SubstrateOperator.ThicknessIndex = Compiler->SubstrateThicknessStackGetThicknessIndex();
-	return &SubstrateOperator;
-}
-#endif // WITH_EDITOR
-
-
-
-UMaterialExpressionSubstratePostProcess::UMaterialExpressionSubstratePostProcess(const FObjectInitializer& ObjectInitializer)
-	: Super(ObjectInitializer)
-{
-	struct FConstructorStatics
-	{
-		FText NAME_Substrate;
-		FConstructorStatics() : NAME_Substrate(LOCTEXT("Substrate Extras", "Substrate Extras")) { }
-	};
-	static FConstructorStatics ConstructorStatics;
-#if WITH_EDITORONLY_DATA
-	MenuCategories.Add(ConstructorStatics.NAME_Substrate);
-#endif
-}
-
-#if WITH_EDITOR
-int32 UMaterialExpressionSubstratePostProcess::Compile(class FMaterialCompiler* Compiler, int32 OutputIndex)
-{
-	int OpacityCodeChunk = CompileWithDefaultFloat1(Compiler, Opacity, 1.0f);
-	int TransmittanceCodeChunk = Compiler->Saturate(Compiler->Sub(Compiler->Constant(1.0f), OpacityCodeChunk));
-
-	FSubstrateOperator& SubstrateOperator = Compiler->SubstrateCompilationGetOperator(Compiler->SubstrateTreeStackGetPathUniqueId());
-	int32 OutputCodeChunk = Compiler->SubstrateUnlitBSDF(
-		CompileWithDefaultFloat3(Compiler, Color, 0.0f, 0.0f, 0.0f),
-		TransmittanceCodeChunk,
-		Compiler->Constant3(0.0f, 0.0f, 1.0f),	// place holder normal
-		&SubstrateOperator);
-	return OutputCodeChunk;
-}
-
-void UMaterialExpressionSubstratePostProcess::GetCaption(TArray<FString>& OutCaptions) const
-{
-	OutCaptions.Add(TEXT("Substrate Post Process"));
-}
-
-uint32 UMaterialExpressionSubstratePostProcess::GetOutputType(int32 OutputIndex)
-{
-	return MCT_Substrate;
-}
-
-uint32 UMaterialExpressionSubstratePostProcess::GetInputType(int32 InputIndex)
-{
-	switch (InputIndex)
-	{
-	case 0:
-		return MCT_Float3;
-		break;
-	case 1:
-		return MCT_Float1;
-		break;
-	}
-
-	check(false);
-	return MCT_Float1;
-}
-
-bool UMaterialExpressionSubstratePostProcess::IsResultSubstrateMaterial(int32 OutputIndex)
-{
-	return true;
-}
-
-void UMaterialExpressionSubstratePostProcess::GatherSubstrateMaterialInfo(FSubstrateMaterialInfo& SubstrateMaterialInfo, int32 OutputIndex)
-{
-	SubstrateMaterialInfo.AddShadingModel(SSM_PostProcess);
-	SubstrateMaterialInfo.AddGuid(MaterialExpressionGuid);
-}
-
-FSubstrateOperator* UMaterialExpressionSubstratePostProcess::SubstrateGenerateMaterialTopologyTree(class FMaterialCompiler* Compiler, class UMaterialExpression* Parent, int32 OutputIndex)
-{
-	FSubstrateOperator& SubstrateOperator = Compiler->SubstrateCompilationRegisterOperator(SUBSTRATE_OPERATOR_BSDF, Compiler->SubstrateTreeStackGetPathUniqueId(), this, Parent, Compiler->SubstrateTreeStackGetParentPathUniqueId());
-	SubstrateOperator.BSDFType = SUBSTRATE_BSDF_TYPE_UNLIT;
-	SubstrateOperator.SubUsage = SUBSTRATE_OPERATOR_SUBUSAGE_POSTPROCESS;
-	SubstrateOperator.ThicknessIndex = Compiler->SubstrateThicknessStackGetThicknessIndex();
-	return &SubstrateOperator;
-}
-#endif // WITH_EDITOR
-
-
-
-UMaterialExpressionSubstrateUI::UMaterialExpressionSubstrateUI(const FObjectInitializer& ObjectInitializer)
-	: Super(ObjectInitializer)
-{
-	struct FConstructorStatics
-	{
-		FText NAME_Substrate;
-		FConstructorStatics() : NAME_Substrate(LOCTEXT("Substrate Extras", "Substrate Extras")) { }
-	};
-	static FConstructorStatics ConstructorStatics;
-#if WITH_EDITORONLY_DATA
-	MenuCategories.Add(ConstructorStatics.NAME_Substrate);
-#endif
-}
-
-#if WITH_EDITOR
-int32 UMaterialExpressionSubstrateUI::Compile(class FMaterialCompiler* Compiler, int32 OutputIndex)
-{
-	int OpacityCodeChunk = CompileWithDefaultFloat1(Compiler, Opacity, 1.0f);
-
-	FSubstrateOperator& SubstrateOperator = Compiler->SubstrateCompilationGetOperator(Compiler->SubstrateTreeStackGetPathUniqueId());
-	int32 OutputCodeChunk = Compiler->SubstrateUIBSDF(
-		CompileWithDefaultFloat3(Compiler, Color, 0.0f, 0.0f, 0.0f),
-		OpacityCodeChunk,
-		&SubstrateOperator);
-	return OutputCodeChunk;
-}
-
-void UMaterialExpressionSubstrateUI::GetCaption(TArray<FString>& OutCaptions) const
-{
-	OutCaptions.Add(TEXT("Substrate UI"));
-}
-
-uint32 UMaterialExpressionSubstrateUI::GetOutputType(int32 OutputIndex)
-{
-	return MCT_Substrate;
-}
-
-uint32 UMaterialExpressionSubstrateUI::GetInputType(int32 InputIndex)
-{
-	switch (InputIndex)
-	{
-	case 0:
-		return MCT_Float3;
-		break;
-	case 1:
-		return MCT_Float;
-		break;
-	}
-
-	check(false);
-	return MCT_Float1;
-}
-
-bool UMaterialExpressionSubstrateUI::IsResultSubstrateMaterial(int32 OutputIndex)
-{
-	return true;
-}
-
-void UMaterialExpressionSubstrateUI::GatherSubstrateMaterialInfo(FSubstrateMaterialInfo& SubstrateMaterialInfo, int32 OutputIndex)
-{
-	SubstrateMaterialInfo.AddShadingModel(SSM_UI);
-	SubstrateMaterialInfo.AddGuid(MaterialExpressionGuid);
-}
-
-FSubstrateOperator* UMaterialExpressionSubstrateUI::SubstrateGenerateMaterialTopologyTree(class FMaterialCompiler* Compiler, class UMaterialExpression* Parent, int32 OutputIndex)
-{
-	FSubstrateOperator& SubstrateOperator = Compiler->SubstrateCompilationRegisterOperator(SUBSTRATE_OPERATOR_BSDF, Compiler->SubstrateTreeStackGetPathUniqueId(), this, Parent, Compiler->SubstrateTreeStackGetParentPathUniqueId());
-	SubstrateOperator.BSDFType = SUBSTRATE_BSDF_TYPE_UNLIT;
-	SubstrateOperator.SubUsage = SUBSTRATE_OPERATOR_SUBUSAGE_UI;
-	SubstrateOperator.ThicknessIndex = Compiler->SubstrateThicknessStackGetThicknessIndex();
-	return &SubstrateOperator;
-}
-#endif // WITH_EDITOR
-
-
-
-UMaterialExpressionSubstrateConvertToDecal::UMaterialExpressionSubstrateConvertToDecal(const FObjectInitializer& ObjectInitializer)
-	: Super(ObjectInitializer)
-{
-	struct FConstructorStatics
-	{
-		FText NAME_Substrate;
-		FConstructorStatics() : NAME_Substrate(LOCTEXT("Substrate Extras", "Substrate Extras")) { }
-	};
-	static FConstructorStatics ConstructorStatics;
-#if WITH_EDITORONLY_DATA
-	MenuCategories.Add(ConstructorStatics.NAME_Substrate);
-#endif
-}
-
-#if WITH_EDITOR
-int32 UMaterialExpressionSubstrateConvertToDecal::Compile(class FMaterialCompiler* Compiler, int32 OutputIndex)
-{
-	if (!DecalMaterial.GetTracedInput().Expression)
-	{
-		return Compiler->Errorf(TEXT("Missing DecalMaterial input"));
-	}
-
-	int32 CoverageCodeChunk = Coverage.GetTracedInput().Expression ? Coverage.Compile(Compiler) : Compiler->Constant(1.0f);
-	Compiler->SubstrateTreeStackPush(this, 0);
-	int32 DecalMaterialCodeChunk = DecalMaterial.Compile(Compiler);
-	Compiler->SubstrateTreeStackPop();
-
-	int32 OutputCodeChunk = INDEX_NONE;
-	FSubstrateOperator& SubstrateOperator = Compiler->SubstrateCompilationGetOperator(Compiler->SubstrateTreeStackGetPathUniqueId());
-	if (!SubstrateOperator.bUseParameterBlending)
-	{
-		return Compiler->Errorf(TEXT("Substrate Convert To Decal node must receive SubstrateData a parameter blended Substrate material sub tree."));
-	}
-	if (!SubstrateOperator.bRootOfParameterBlendingSubTree)
-	{
-		return Compiler->Errorf(TEXT("Substrate Convert To Decal node must be the root of a parameter blending sub tree: no more Substrate operations can be applied a over its output."));
-	}
-
-	// Propagate the parameter blended normal
-	FSubstrateOperator* Operator = Compiler->SubstrateCompilationGetOperatorFromIndex(SubstrateOperator.LeftIndex);
-	SubstrateOperator.BSDFRegisteredSharedLocalBasis = Operator->BSDFRegisteredSharedLocalBasis;
-
-	OutputCodeChunk = Compiler->SubstrateWeightParameterBlending(
-		DecalMaterialCodeChunk, CoverageCodeChunk,
-		SubstrateOperator.bRootOfParameterBlendingSubTree ? &SubstrateOperator : nullptr);
-
-	return OutputCodeChunk;
-}
-
-void UMaterialExpressionSubstrateConvertToDecal::GetCaption(TArray<FString>& OutCaptions) const
-{
-	OutCaptions.Add(TEXT("Substrate Convert To Decal"));
-}
-
-uint32 UMaterialExpressionSubstrateConvertToDecal::GetOutputType(int32 OutputIndex)
-{
-	return MCT_Substrate;
-}
-
-uint32 UMaterialExpressionSubstrateConvertToDecal::GetInputType(int32 InputIndex)
-{
-	switch (InputIndex)
-	{
-	case 0:
-		return MCT_Substrate;
-		break;
-	case 1:
-		return MCT_Float1;
-		break;
-	}
-
-	check(false);
-	return MCT_Float1;
-}
-
-bool UMaterialExpressionSubstrateConvertToDecal::IsResultSubstrateMaterial(int32 OutputIndex)
-{
-	return true;
-}
-
-void UMaterialExpressionSubstrateConvertToDecal::GatherSubstrateMaterialInfo(FSubstrateMaterialInfo& SubstrateMaterialInfo, int32 OutputIndex)
-{
-	if (!SubstrateMaterialInfo.PushSubstrateTreeStack())
-	{
-		return;
-	}
-
-	FExpressionInput TracedInput = DecalMaterial.GetTracedInput();
-	if (TracedInput.Expression)
-	{
-		TracedInput.Expression->GatherSubstrateMaterialInfo(SubstrateMaterialInfo, TracedInput.OutputIndex);
-	}
-	SubstrateMaterialInfo.AddShadingModel(SSM_Decal);
-	SubstrateMaterialInfo.AddGuid(MaterialExpressionGuid);
-
-	SubstrateMaterialInfo.PopSubstrateTreeStack();
-}
-
-FSubstrateOperator* UMaterialExpressionSubstrateConvertToDecal::SubstrateGenerateMaterialTopologyTree(class FMaterialCompiler* Compiler, class UMaterialExpression* Parent, int32 OutputIndex)
-{
-	const bool bUseParameterBlending = true;
-	FSubstrateOperator& SubstrateOperator = Compiler->SubstrateCompilationRegisterOperator(SUBSTRATE_OPERATOR_WEIGHT, Compiler->SubstrateTreeStackGetPathUniqueId(), this, Parent, Compiler->SubstrateTreeStackGetParentPathUniqueId(), bUseParameterBlending);
-	SubstrateOperator.SubUsage = SUBSTRATE_OPERATOR_SUBUSAGE_DECAL;
-	if (Compiler->GetSubstrateTreeOutOfStackDepthOccurred())
-	{
-		return &SubstrateOperator; // Out ot stack space, return now to fail the compilation
-	}
-
-	FExpressionInput TracedInput = DecalMaterial.GetTracedInput();
-	UMaterialExpression* ChildDecalMaterialExpression = TracedInput.Expression;
-	FSubstrateOperator* OpA = nullptr;
-	if (ChildDecalMaterialExpression)
-	{
-		Compiler->SubstrateTreeStackPush(this, 0);
-		OpA = ChildDecalMaterialExpression->SubstrateGenerateMaterialTopologyTree(Compiler, this, TracedInput.OutputIndex);
-		Compiler->SubstrateTreeStackPop();
-		AssignOperatorIndexIfNotNull(SubstrateOperator.LeftIndex, OpA);
-	}
-	CombineFlagForParameterBlending(SubstrateOperator, OpA);
-
-	return &SubstrateOperator;
-}
-#endif // WITH_EDITOR
-
-
-
-UMaterialExpressionSubstrateUnlitBSDF::UMaterialExpressionSubstrateUnlitBSDF(const FObjectInitializer& ObjectInitializer)
-	: Super(ObjectInitializer)
-{
-	struct FConstructorStatics
-	{
-		FText NAME_Substrate;
-		FConstructorStatics() : NAME_Substrate(LOCTEXT("Substrate BSDFs", "Substrate BSDFs")) { }
-	};
-	static FConstructorStatics ConstructorStatics;
-#if WITH_EDITORONLY_DATA
-	MenuCategories.Add(ConstructorStatics.NAME_Substrate);
-#endif
-}
-
-#if WITH_EDITOR
-int32 UMaterialExpressionSubstrateUnlitBSDF::Compile(class FMaterialCompiler* Compiler, int32 OutputIndex)
-{
-	FSubstrateOperator& SubstrateOperator = Compiler->SubstrateCompilationGetOperator(Compiler->SubstrateTreeStackGetPathUniqueId());
-	int32 OutputCodeChunk = Compiler->SubstrateUnlitBSDF(
-		CompileWithDefaultFloat3(Compiler, EmissiveColor, 0.0f, 0.0f, 0.0f),
-		CompileWithDefaultFloat3(Compiler, TransmittanceColor, 1.0f, 1.0f, 1.0f),
-		CompileWithDefaultNormalWS(Compiler, Normal),
-		&SubstrateOperator);
-	return OutputCodeChunk;
-}
-
-void UMaterialExpressionSubstrateUnlitBSDF::GetCaption(TArray<FString>& OutCaptions) const
-{
-	OutCaptions.Add(TEXT("Substrate Unlit BSDF"));
-}
-
-uint32 UMaterialExpressionSubstrateUnlitBSDF::GetOutputType(int32 OutputIndex)
-{
-	return MCT_Substrate;
-}
-
-uint32 UMaterialExpressionSubstrateUnlitBSDF::GetInputType(int32 InputIndex)
-{
-	switch (InputIndex)
-	{
-	case 0:
-		return MCT_Float3;
-		break;
-	case 1:
-		return MCT_Float3;
-		break;
-	case 2:
-		return MCT_Float3;
-		break;
-	}
-
-	check(false);
-	return MCT_Float1;
-}
-
-bool UMaterialExpressionSubstrateUnlitBSDF::IsResultSubstrateMaterial(int32 OutputIndex)
-{
-	return true;
-}
-
-void UMaterialExpressionSubstrateUnlitBSDF::GatherSubstrateMaterialInfo(FSubstrateMaterialInfo& SubstrateMaterialInfo, int32 OutputIndex)
-{
-	if (EmissiveColor.IsConnected()) { SubstrateMaterialInfo.AddPropertyConnected(MP_EmissiveColor); }
-	SubstrateMaterialInfo.AddShadingModel(SSM_Unlit);
-	SubstrateMaterialInfo.AddGuid(MaterialExpressionGuid);
-}
-
-FSubstrateOperator* UMaterialExpressionSubstrateUnlitBSDF::SubstrateGenerateMaterialTopologyTree(class FMaterialCompiler* Compiler, class UMaterialExpression* Parent, int32 OutputIndex)
-{
-	FSubstrateOperator& SubstrateOperator = Compiler->SubstrateCompilationRegisterOperator(SUBSTRATE_OPERATOR_BSDF, Compiler->SubstrateTreeStackGetPathUniqueId(), this, Parent, Compiler->SubstrateTreeStackGetParentPathUniqueId());
-	SubstrateOperator.BSDFType = SUBSTRATE_BSDF_TYPE_UNLIT;
-	SubstrateOperator.ThicknessIndex = Compiler->SubstrateThicknessStackGetThicknessIndex();
-	SubstrateOperator.bBSDFWritesEmissive = EmissiveColor.IsConnected();
-	return &SubstrateOperator;
-}
-#endif // WITH_EDITOR
-
-
-
-UMaterialExpressionSubstrateHairBSDF::UMaterialExpressionSubstrateHairBSDF(const FObjectInitializer& ObjectInitializer)
-	: Super(ObjectInitializer)
-{
-	struct FConstructorStatics
-	{
-		FText NAME_Substrate;
-		FConstructorStatics() : NAME_Substrate(LOCTEXT("Substrate BSDFs", "Substrate BSDFs")) { }
-	};
-	static FConstructorStatics ConstructorStatics;
-#if WITH_EDITORONLY_DATA
-	MenuCategories.Add(ConstructorStatics.NAME_Substrate);
-#endif
-}
-
-#if WITH_EDITOR
-int32 UMaterialExpressionSubstrateHairBSDF::Compile(class FMaterialCompiler* Compiler, int32 OutputIndex)
-{
-	// For hair, the shared local basis normal in fact represent the tangent
-	int32 TangentCodeChunk = CompileWithDefaultTangentWS(Compiler, Tangent);
-	const FSubstrateRegisteredSharedLocalBasis NewRegisteredSharedLocalBasis = SubstrateCompilationInfoCreateSharedLocalBasis(Compiler, TangentCodeChunk);
-
-	FSubstrateOperator& SubstrateOperator = Compiler->SubstrateCompilationGetOperator(Compiler->SubstrateTreeStackGetPathUniqueId());
-	SubstrateOperator.BSDFRegisteredSharedLocalBasis = NewRegisteredSharedLocalBasis;
-
-	if (SubstrateOperator.bUseParameterBlending)
-	{
-		return Compiler->Errorf(TEXT("Substrate Hair BSDF node cannot be used with parameter blending."));
-	}
-	else if (SubstrateOperator.bRootOfParameterBlendingSubTree)
-	{
-		return Compiler->Errorf(TEXT("Substrate Hair BSDF node cannot be the root of a parameter blending sub tree."));
-	}
-
-	int32 OutputCodeChunk = Compiler->SubstrateHairBSDF(
-		CompileWithDefaultFloat3(Compiler, BaseColor,	0.0f, 0.0f, 0.0f),
-		CompileWithDefaultFloat1(Compiler, Scatter,		0.0f),
-		CompileWithDefaultFloat1(Compiler, Specular,	0.5f),
-		CompileWithDefaultFloat1(Compiler, Roughness,	0.5f),
-		CompileWithDefaultFloat1(Compiler, Backlit,		0.0f),
-		CompileWithDefaultFloat3(Compiler, EmissiveColor,0.0f, 0.0f, 0.0f),
-		TangentCodeChunk,
-		Compiler->GetSubstrateSharedLocalBasisIndexMacro(NewRegisteredSharedLocalBasis),
-		&SubstrateOperator);
-
-	return OutputCodeChunk;
-}
-
-void UMaterialExpressionSubstrateHairBSDF::GetCaption(TArray<FString>& OutCaptions) const
-{
-	OutCaptions.Add(TEXT("Substrate Hair BSDF"));
-}
-
-uint32 UMaterialExpressionSubstrateHairBSDF::GetOutputType(int32 OutputIndex)
-{
-	return MCT_Substrate;
-}
-
-uint32 UMaterialExpressionSubstrateHairBSDF::GetInputType(int32 InputIndex)
-{
-	switch (InputIndex)
-	{
-	case 0:
-		return MCT_Float3;
-		break;
-	case 1:
-		return MCT_Float1;
-		break;
-	case 2:
-		return MCT_Float1;
-		break;
-	case 3:
-		return MCT_Float1;
-		break;
-	case 4:
-		return MCT_Float1;
-		break;
-	case 5:
-		return MCT_Float3;
-		break;
-	case 6:
-		return MCT_Float3;
-		break;
-	}
-
-	check(false);
-	return MCT_Float1;
-}
-
-bool UMaterialExpressionSubstrateHairBSDF::IsResultSubstrateMaterial(int32 OutputIndex)
-{
-	return true;
-}
-
-void UMaterialExpressionSubstrateHairBSDF::GatherSubstrateMaterialInfo(FSubstrateMaterialInfo& SubstrateMaterialInfo, int32 OutputIndex)
-{
-	// Track connected inputs
-	if (BaseColor.IsConnected())		{ SubstrateMaterialInfo.AddPropertyConnected(MP_BaseColor); }
-	if (Specular.IsConnected())			{ SubstrateMaterialInfo.AddPropertyConnected(MP_Specular); }
-	if (Roughness.IsConnected())		{ SubstrateMaterialInfo.AddPropertyConnected(MP_Roughness); }
-	if (EmissiveColor.IsConnected())	{ SubstrateMaterialInfo.AddPropertyConnected(MP_EmissiveColor); }
-	if (Tangent.IsConnected())			{ SubstrateMaterialInfo.AddPropertyConnected(MP_Tangent); }
-
-	SubstrateMaterialInfo.AddShadingModel(SSM_Hair);
-	SubstrateMaterialInfo.AddGuid(MaterialExpressionGuid);
-}
-
-FSubstrateOperator* UMaterialExpressionSubstrateHairBSDF::SubstrateGenerateMaterialTopologyTree(class FMaterialCompiler* Compiler, class UMaterialExpression* Parent, int32 OutputIndex)
-{
-	FSubstrateOperator& SubstrateOperator = Compiler->SubstrateCompilationRegisterOperator(SUBSTRATE_OPERATOR_BSDF, Compiler->SubstrateTreeStackGetPathUniqueId(), this, Parent, Compiler->SubstrateTreeStackGetParentPathUniqueId());
-	SubstrateOperator.BSDFType = SUBSTRATE_BSDF_TYPE_HAIR;
-	SubstrateOperator.ThicknessIndex = Compiler->SubstrateThicknessStackGetThicknessIndex();
-	SubstrateOperator.bBSDFWritesEmissive = EmissiveColor.IsConnected();
-	return &SubstrateOperator;
-}
-#endif // WITH_EDITOR
-
-UMaterialExpressionSubstrateEyeBSDF::UMaterialExpressionSubstrateEyeBSDF(const FObjectInitializer& ObjectInitializer)
-	: Super(ObjectInitializer)
-{
-	struct FConstructorStatics
-	{
-		FText NAME_Substrate;
-		FConstructorStatics() : NAME_Substrate(LOCTEXT("Substrate BSDFs", "Substrate BSDFs")) { }
-	};
-	static FConstructorStatics ConstructorStatics;
-#if WITH_EDITORONLY_DATA
-	MenuCategories.Add(ConstructorStatics.NAME_Substrate);
-#endif
-}
-
-#if WITH_EDITOR
-int32 UMaterialExpressionSubstrateEyeBSDF::Compile(class FMaterialCompiler* Compiler, int32 OutputIndex)
-{
-	int32 CorneaNormalCodeChunk = CompileWithDefaultTangentWS(Compiler, CorneaNormal);
-	const FSubstrateRegisteredSharedLocalBasis NewRegisteredSharedLocalBasis = SubstrateCompilationInfoCreateSharedLocalBasis(Compiler, CorneaNormalCodeChunk);
-
-	int32 SSSProfileCodeChunk = INDEX_NONE;
-	if (SubsurfaceProfile != nullptr)
-	{
-		SSSProfileCodeChunk = Compiler->ForceCast(Compiler->ScalarParameter(GetSubsurfaceProfileParameterName(), 1.0f), MCT_Float1);
-	}
-
-	FSubstrateOperator& SubstrateOperator = Compiler->SubstrateCompilationGetOperator(Compiler->SubstrateTreeStackGetPathUniqueId());
-	SubstrateOperator.BSDFRegisteredSharedLocalBasis = NewRegisteredSharedLocalBasis;
-
-	if (SubstrateOperator.bUseParameterBlending)
-	{
-		return Compiler->Errorf(TEXT("Substrate Eye BSDF node cannot be used with parameter blending."));
-	}
-	else if (SubstrateOperator.bRootOfParameterBlendingSubTree)
-	{
-		return Compiler->Errorf(TEXT("Substrate Eye BSDF node cannot be the root of a parameter blending sub tree."));
-	}
-
-	int32 OutputCodeChunk = Compiler->SubstrateEyeBSDF(
-		CompileWithDefaultFloat3(Compiler, DiffuseColor, 0.0f, 0.0f, 0.0f),
-		CompileWithDefaultFloat1(Compiler, Roughness,	 0.5f),
-		CompileWithDefaultFloat1(Compiler, IrisMask,	 0.0f),
-		CompileWithDefaultFloat1(Compiler, IrisDistance, 0.0f),
-		CompileWithDefaultNormalWS(Compiler, IrisNormal),
-		CompileWithDefaultNormalWS(Compiler, IrisPlaneNormal),
-		SSSProfileCodeChunk != INDEX_NONE ? SSSProfileCodeChunk : Compiler->Constant(0.0f),
-		CompileWithDefaultFloat3(Compiler, EmissiveColor,0.0f, 0.0f, 0.0f),
-		CorneaNormalCodeChunk,
-		Compiler->GetSubstrateSharedLocalBasisIndexMacro(NewRegisteredSharedLocalBasis),
-		&SubstrateOperator);
-
-	return OutputCodeChunk;
-}
-
-void UMaterialExpressionSubstrateEyeBSDF::GetCaption(TArray<FString>& OutCaptions) const
-{
-	OutCaptions.Add(TEXT("Substrate Eye BSDF"));
-}
-
-uint32 UMaterialExpressionSubstrateEyeBSDF::GetOutputType(int32 OutputIndex)
-{
-	return MCT_Substrate;
-}
-
-uint32 UMaterialExpressionSubstrateEyeBSDF::GetInputType(int32 InputIndex)
-{
-	switch (InputIndex)
-	{
-	case 0: return MCT_Float3; // DiffuseColor
-	case 1: return MCT_Float1; // Roughness
-	case 2: return MCT_Float3; // Cornea normal
-	case 3: return MCT_Float3; // IrisNormal
-	case 4: return MCT_Float3; // IrisPlaneNormal
-	case 5: return MCT_Float1; // IrisMask
-	case 6: return MCT_Float1; // IrisDistance
-	case 7: return MCT_Float3; // EmissiveColor
-	}
-
-	check(false);
-	return MCT_Float1;
-}
-
-bool UMaterialExpressionSubstrateEyeBSDF::IsResultSubstrateMaterial(int32 OutputIndex)
-{
-	return true;
-}
-
-void UMaterialExpressionSubstrateEyeBSDF::GatherSubstrateMaterialInfo(FSubstrateMaterialInfo& SubstrateMaterialInfo, int32 OutputIndex)
-{
-	// Track connected inputs
-	if (DiffuseColor.IsConnected())		{ SubstrateMaterialInfo.AddPropertyConnected(MP_BaseColor); }
-	if (Roughness.IsConnected())		{ SubstrateMaterialInfo.AddPropertyConnected(MP_Roughness); }
-	if (CorneaNormal.IsConnected())		{ SubstrateMaterialInfo.AddPropertyConnected(MP_Normal); }
-	if (IrisNormal.IsConnected())		{ SubstrateMaterialInfo.AddPropertyConnected(MP_Tangent); }
-	if (IrisPlaneNormal.IsConnected())	{ SubstrateMaterialInfo.AddPropertyConnected(MP_Tangent); }
-	if (IrisMask.IsConnected())			{ SubstrateMaterialInfo.AddPropertyConnected(MP_CustomData0); }
-	if (IrisDistance.IsConnected())		{ SubstrateMaterialInfo.AddPropertyConnected(MP_CustomData1); }
-	if (EmissiveColor.IsConnected())	{ SubstrateMaterialInfo.AddPropertyConnected(MP_EmissiveColor); }
-	if (SubsurfaceProfile)
-	{
-		SubstrateMaterialInfo.AddSubsurfaceProfile(SubsurfaceProfile);
-	}
-	SubstrateMaterialInfo.AddShadingModel(SSM_Eye);
-	SubstrateMaterialInfo.AddGuid(MaterialExpressionGuid);
-}
-
-FSubstrateOperator* UMaterialExpressionSubstrateEyeBSDF::SubstrateGenerateMaterialTopologyTree(class FMaterialCompiler* Compiler, class UMaterialExpression* Parent, int32 OutputIndex)
-{
-	FSubstrateOperator& SubstrateOperator = Compiler->SubstrateCompilationRegisterOperator(SUBSTRATE_OPERATOR_BSDF, Compiler->SubstrateTreeStackGetPathUniqueId(), this, Parent, Compiler->SubstrateTreeStackGetParentPathUniqueId());
-	SubstrateOperator.BSDFType = SUBSTRATE_BSDF_TYPE_EYE;
-	SubstrateOperator.ThicknessIndex = Compiler->SubstrateThicknessStackGetThicknessIndex();
-	SubstrateOperator.bBSDFWritesEmissive = EmissiveColor.IsConnected();
-	return &SubstrateOperator;
-}
-#endif // WITH_EDITOR
-
-
-UMaterialExpressionSubstrateSingleLayerWaterBSDF::UMaterialExpressionSubstrateSingleLayerWaterBSDF(const FObjectInitializer& ObjectInitializer)
-	: Super(ObjectInitializer)
-{
-	struct FConstructorStatics
-	{
-		FText NAME_Substrate;
-		FConstructorStatics() : NAME_Substrate(LOCTEXT("Substrate BSDFs", "Substrate BSDFs")) { }
-	};
-	static FConstructorStatics ConstructorStatics;
-#if WITH_EDITORONLY_DATA
-	MenuCategories.Add(ConstructorStatics.NAME_Substrate);
-#endif
-}
-
-#if WITH_EDITOR
-int32 UMaterialExpressionSubstrateSingleLayerWaterBSDF::Compile(class FMaterialCompiler* Compiler, int32 OutputIndex)
-{
-	int32 NormalCodeChunk = CompileWithDefaultNormalWS(Compiler, Normal);
-	const FSubstrateRegisteredSharedLocalBasis NewRegisteredSharedLocalBasis = SubstrateCompilationInfoCreateSharedLocalBasis(Compiler, NormalCodeChunk);
-
-	FSubstrateOperator& SubstrateOperator = Compiler->SubstrateCompilationGetOperator(Compiler->SubstrateTreeStackGetPathUniqueId());
-	SubstrateOperator.BSDFRegisteredSharedLocalBasis = NewRegisteredSharedLocalBasis;
-
-	if (SubstrateOperator.bUseParameterBlending)
-	{
-		return Compiler->Errorf(TEXT("Substrate SingleLayerWater BSDF node cannot be used with parameter blending."));
-	}
-	else if (SubstrateOperator.bRootOfParameterBlendingSubTree)
-	{
-		return Compiler->Errorf(TEXT("Substrate SingleLayerWater BSDF node cannot be the root of a parameter blending sub tree."));
-	}
-
-	int32 OutputCodeChunk = Compiler->SubstrateSingleLayerWaterBSDF(
-		CompileWithDefaultFloat3(Compiler, BaseColor, 0.0f, 0.0f, 0.0f),
-		CompileWithDefaultFloat1(Compiler, Metallic, 0.0f),
-		CompileWithDefaultFloat1(Compiler, Specular, 0.5f),
-		CompileWithDefaultFloat1(Compiler, Roughness, 0.5f),
-		CompileWithDefaultFloat3(Compiler, EmissiveColor, 0.0f, 0.0f, 0.0f),
-		CompileWithDefaultFloat1(Compiler, TopMaterialOpacity, 0.0f),
-		CompileWithDefaultFloat3(Compiler, WaterAlbedo, 0.0f, 0.0f, 0.0f),
-		CompileWithDefaultFloat3(Compiler, WaterExtinction, 0.0f, 0.0f, 0.0f),
-		CompileWithDefaultFloat1(Compiler, WaterPhaseG, 0.0f),
-		CompileWithDefaultFloat3(Compiler, ColorScaleBehindWater, 1.0f, 1.0f, 1.0f),
-		NormalCodeChunk,
-		Compiler->GetSubstrateSharedLocalBasisIndexMacro(NewRegisteredSharedLocalBasis),
-		&SubstrateOperator);
-
-	return OutputCodeChunk;
-}
-
-void UMaterialExpressionSubstrateSingleLayerWaterBSDF::GetCaption(TArray<FString>& OutCaptions) const
-{
-	OutCaptions.Add(TEXT("Substrate Single Layer Water BSDF"));
-}
-
-uint32 UMaterialExpressionSubstrateSingleLayerWaterBSDF::GetOutputType(int32 OutputIndex)
-{
-	return MCT_Substrate;
-}
-
-uint32 UMaterialExpressionSubstrateSingleLayerWaterBSDF::GetInputType(int32 InputIndex)
-{
-	switch (InputIndex)
-	{
-	case 0:
-		return MCT_Float3; // BaseColor
-		break;
-	case 1:
-		return MCT_Float1; // Metallic
-		break;
-	case 2:
-		return MCT_Float1; // Specular
-		break;
-	case 3:
-		return MCT_Float1; // Roughness
-		break;
-	case 4:
-		return MCT_Float3; // Normal
-		break;
-	case 5:
-		return MCT_Float3; // Emissive Color
-		break;
-	case 6:
-		return MCT_Float1; // TopMaterialOpacity
-		break;
-	case 7:
-		return MCT_Float3; // WaterAlbedo
-		break;
-	case 8:
-		return MCT_Float3; // WaterExtinction
-		break;
-	case 9:
-		return MCT_Float1; // WaterPhaseG
-		break;
-	case 10:
-		return MCT_Float3; // ColorScaleBehindWater
-		break;
-	}
-	
-	check(false);
-	return MCT_Float1;
-}
-
-bool UMaterialExpressionSubstrateSingleLayerWaterBSDF::IsResultSubstrateMaterial(int32 OutputIndex)
-{
-	return true;
-}
-
-void UMaterialExpressionSubstrateSingleLayerWaterBSDF::GatherSubstrateMaterialInfo(FSubstrateMaterialInfo& SubstrateMaterialInfo, int32 OutputIndex)
-{
-	// Track connected inputs
-	if (BaseColor.IsConnected())		{ SubstrateMaterialInfo.AddPropertyConnected(MP_BaseColor); }
-	if (Metallic.IsConnected())			{ SubstrateMaterialInfo.AddPropertyConnected(MP_Metallic); }
-	if (Specular.IsConnected())			{ SubstrateMaterialInfo.AddPropertyConnected(MP_Specular); }
-	if (Roughness.IsConnected())		{ SubstrateMaterialInfo.AddPropertyConnected(MP_Roughness); }
-	if (EmissiveColor.IsConnected())	{ SubstrateMaterialInfo.AddPropertyConnected(MP_EmissiveColor); }
-	if (Normal.IsConnected())			{ SubstrateMaterialInfo.AddPropertyConnected(MP_Normal); }
-
-	SubstrateMaterialInfo.AddShadingModel(SSM_SingleLayerWater);
-	SubstrateMaterialInfo.AddGuid(MaterialExpressionGuid);
-}
-
-FSubstrateOperator* UMaterialExpressionSubstrateSingleLayerWaterBSDF::SubstrateGenerateMaterialTopologyTree(class FMaterialCompiler* Compiler, class UMaterialExpression* Parent, int32 OutputIndex)
-{
-	FSubstrateOperator& SubstrateOperator = Compiler->SubstrateCompilationRegisterOperator(SUBSTRATE_OPERATOR_BSDF, Compiler->SubstrateTreeStackGetPathUniqueId(), this, Parent, Compiler->SubstrateTreeStackGetParentPathUniqueId());
-	SubstrateOperator.BSDFType = SUBSTRATE_BSDF_TYPE_SINGLELAYERWATER;
-	SubstrateOperator.ThicknessIndex = Compiler->SubstrateThicknessStackGetThicknessIndex();
-	SubstrateOperator.bBSDFWritesEmissive = EmissiveColor.IsConnected();
-	return &SubstrateOperator;
-}
-#endif // WITH_EDITOR
-
-
-
-UMaterialExpressionSubstrateHorizontalMixing::UMaterialExpressionSubstrateHorizontalMixing(const FObjectInitializer& ObjectInitializer)
-	: Super(ObjectInitializer)
-	, bUseParameterBlending(false)
-{
-	struct FConstructorStatics
-	{
-		FText NAME_Substrate;
-		FConstructorStatics() : NAME_Substrate(LOCTEXT("Substrate Ops", "Substrate Operators")) { }
-	};
-	static FConstructorStatics ConstructorStatics;
-#if WITH_EDITORONLY_DATA
-	MenuCategories.Add(ConstructorStatics.NAME_Substrate);
-#endif
-}
-
-#if WITH_EDITOR
-int32 UMaterialExpressionSubstrateHorizontalMixing::Compile(class FMaterialCompiler* Compiler, int32 OutputIndex)
-{
-	if (!Foreground.GetTracedInput().Expression)
-	{
-		return Compiler->Errorf(TEXT("Missing Foreground input"));
-	}
-	if (!Background.GetTracedInput().Expression)
-	{
-		return Compiler->Errorf(TEXT("Missing Background input"));
-	}
-
-	Compiler->SubstrateTreeStackPush(this, 0);
-	int32 BackgroundCodeChunk = Background.Compile(Compiler);
-	Compiler->SubstrateTreeStackPop();
-	Compiler->SubstrateTreeStackPush(this, 1);
-	int32 ForegroundCodeChunk = Foreground.Compile(Compiler);
-	Compiler->SubstrateTreeStackPop();
-
-	const int32 HorizontalMixCodeChunk = CompileWithDefaultFloat1(Compiler, Mix, 0.5f);
-
-	int32 OutputCodeChunk = INDEX_NONE;
-	FSubstrateOperator& SubstrateOperator = Compiler->SubstrateCompilationGetOperator(Compiler->SubstrateTreeStackGetPathUniqueId());
-	if (SubstrateOperator.bUseParameterBlending)
-	{
-		if (ForegroundCodeChunk == INDEX_NONE)
-		{
-			return Compiler->Errorf(TEXT("Foreground input graphs could not be evaluated for parameter blending."));
-		}
-		if (BackgroundCodeChunk == INDEX_NONE)
-		{
-			return Compiler->Errorf(TEXT("Background input graphs could not be evaluated for parameter blending."));
-		}
-		const int32 NormalMixCodeChunk = Compiler->SubstrateHorizontalMixingParameterBlendingBSDFCoverageToNormalMixCodeChunk(BackgroundCodeChunk, ForegroundCodeChunk, HorizontalMixCodeChunk);
-
-		FSubstrateOperator* BackgroundBSDFOperator = Compiler->SubstrateCompilationGetOperatorFromIndex(SubstrateOperator.LeftIndex);
-		FSubstrateOperator* ForegroundBSDFOperator = Compiler->SubstrateCompilationGetOperatorFromIndex(SubstrateOperator.RightIndex);
-		if (!BackgroundBSDFOperator || !ForegroundBSDFOperator)
-		{
-			return Compiler->Errorf(TEXT("Missing input on horizontal blending node."));
-		}
-
-		// Compute the new Normal and Tangent resulting from the blending using code chunk
-		const int32 NewNormalCodeChunk = SubstrateBlendNormal(Compiler, BackgroundBSDFOperator->BSDFRegisteredSharedLocalBasis.NormalCodeChunk, ForegroundBSDFOperator->BSDFRegisteredSharedLocalBasis.NormalCodeChunk, NormalMixCodeChunk);
-		// The tangent is optional so we treat it differently if INDEX_NONE is specified
-		int32 NewTangentCodeChunk = INDEX_NONE;
-		if (ForegroundBSDFOperator->BSDFRegisteredSharedLocalBasis.TangentCodeChunk != INDEX_NONE && BackgroundBSDFOperator->BSDFRegisteredSharedLocalBasis.TangentCodeChunk != INDEX_NONE)
-		{
-			NewTangentCodeChunk = SubstrateBlendNormal(Compiler, BackgroundBSDFOperator->BSDFRegisteredSharedLocalBasis.TangentCodeChunk, ForegroundBSDFOperator->BSDFRegisteredSharedLocalBasis.TangentCodeChunk, NormalMixCodeChunk);
-		}
-		else if (ForegroundBSDFOperator->BSDFRegisteredSharedLocalBasis.TangentCodeChunk != INDEX_NONE)
-		{
-			NewTangentCodeChunk = ForegroundBSDFOperator->BSDFRegisteredSharedLocalBasis.TangentCodeChunk;
-		}
-		else if (BackgroundBSDFOperator->BSDFRegisteredSharedLocalBasis.TangentCodeChunk != INDEX_NONE)
-		{
-			NewTangentCodeChunk = BackgroundBSDFOperator->BSDFRegisteredSharedLocalBasis.TangentCodeChunk;
-		}
-		const FSubstrateRegisteredSharedLocalBasis NewRegisteredSharedLocalBasis = SubstrateCompilationInfoCreateSharedLocalBasis(Compiler, NewNormalCodeChunk, NewTangentCodeChunk);
-
-		OutputCodeChunk = Compiler->SubstrateHorizontalMixingParameterBlending(
-			BackgroundCodeChunk, ForegroundCodeChunk, HorizontalMixCodeChunk, NormalMixCodeChunk, Compiler->GetSubstrateSharedLocalBasisIndexMacro(NewRegisteredSharedLocalBasis),
-			SubstrateOperator.bRootOfParameterBlendingSubTree ? &SubstrateOperator : nullptr);
-
-		// Propagate the parameter blended normal
-		SubstrateOperator.BSDFRegisteredSharedLocalBasis = NewRegisteredSharedLocalBasis;
-	}
-	else
-	{
-		OutputCodeChunk = Compiler->SubstrateHorizontalMixing(
-			BackgroundCodeChunk,
-			ForegroundCodeChunk,
-			HorizontalMixCodeChunk,
-			SubstrateOperator.Index,
-			SubstrateOperator.MaxDistanceFromLeaves);
-	}
-
-	return OutputCodeChunk;
-}
-
-void UMaterialExpressionSubstrateHorizontalMixing::GetCaption(TArray<FString>& OutCaptions) const
-{
-	if (bUseParameterBlending)
-	{
-		OutCaptions.Add(TEXT("Substrate Horizontal Blend (Parameter Blend)"));
-	}
-	else
-	{
-		OutCaptions.Add(TEXT("Substrate Horizontal Blend"));
-	}
-}
-
-uint32 UMaterialExpressionSubstrateHorizontalMixing::GetOutputType(int32 OutputIndex)
-{
-	return MCT_Substrate;
-}
-
-uint32 UMaterialExpressionSubstrateHorizontalMixing::GetInputType(int32 InputIndex)
-{
-	return InputIndex == 2 ? MCT_Float1 : MCT_Substrate;
-}
-
-bool UMaterialExpressionSubstrateHorizontalMixing::IsResultSubstrateMaterial(int32 OutputIndex)
-{
-	return true;
-}
-
-void UMaterialExpressionSubstrateHorizontalMixing::GatherSubstrateMaterialInfo(FSubstrateMaterialInfo& SubstrateMaterialInfo, int32 OutputIndex)
-{
-	if (!SubstrateMaterialInfo.PushSubstrateTreeStack())
-	{
-		return;
-	}
-
-	FExpressionInput TracedInputA = Foreground.GetTracedInput();
-	FExpressionInput TracedInputB = Background.GetTracedInput();
-	if (TracedInputA.Expression)
-	{
-		TracedInputA.Expression->GatherSubstrateMaterialInfo(SubstrateMaterialInfo, TracedInputA.OutputIndex);
-	}
-	if (TracedInputB.Expression)
-	{
-		TracedInputB.Expression->GatherSubstrateMaterialInfo(SubstrateMaterialInfo, TracedInputB.OutputIndex);
-	}
-	SubstrateMaterialInfo.AddGuid(MaterialExpressionGuid);
-
-	SubstrateMaterialInfo.PopSubstrateTreeStack();
-}
-
-FSubstrateOperator* UMaterialExpressionSubstrateHorizontalMixing::SubstrateGenerateMaterialTopologyTree(class FMaterialCompiler* Compiler, class UMaterialExpression* Parent, int32 OutputIndex)
-{
-	FSubstrateOperator& SubstrateOperator = Compiler->SubstrateCompilationRegisterOperator(SUBSTRATE_OPERATOR_HORIZONTAL, Compiler->SubstrateTreeStackGetPathUniqueId(), this, Parent, Compiler->SubstrateTreeStackGetParentPathUniqueId(), bUseParameterBlending);
-	if (Compiler->GetSubstrateTreeOutOfStackDepthOccurred())
-	{
-		return &SubstrateOperator; // Out ot stack space, return now to fail the compilation
-	}
-
-	FExpressionInput TracedInputA = Background.GetTracedInput();
-	FExpressionInput TracedInputB = Foreground.GetTracedInput();
-	UMaterialExpression* ChildAExpression = TracedInputA.Expression;
-	UMaterialExpression* ChildBExpression = TracedInputB.Expression;
-	FSubstrateOperator* OpA = nullptr;
-	FSubstrateOperator* OpB = nullptr;
-	if (ChildAExpression)
-	{
-		Compiler->SubstrateTreeStackPush(this, 0);
-		OpA = ChildAExpression->SubstrateGenerateMaterialTopologyTree(Compiler, this, TracedInputA.OutputIndex);
-		Compiler->SubstrateTreeStackPop();
-		AssignOperatorIndexIfNotNull(SubstrateOperator.LeftIndex, OpA);
-	}
-	if (ChildBExpression)
-	{
-		Compiler->SubstrateTreeStackPush(this, 1);
-		OpB = ChildBExpression->SubstrateGenerateMaterialTopologyTree(Compiler, this, TracedInputB.OutputIndex);
-		Compiler->SubstrateTreeStackPop();
-		AssignOperatorIndexIfNotNull(SubstrateOperator.RightIndex, OpB);
-	}
-	CombineFlagForParameterBlending(SubstrateOperator, OpA, OpB);
-
-	return &SubstrateOperator;
-}
-#endif // WITH_EDITOR
-
-
-
-UMaterialExpressionSubstrateVerticalLayering::UMaterialExpressionSubstrateVerticalLayering(const FObjectInitializer& ObjectInitializer)
-	: Super(ObjectInitializer)
-	, bUseParameterBlending(false)
-{
-	struct FConstructorStatics
-	{
-		FText NAME_Substrate;
-		FConstructorStatics() : NAME_Substrate(LOCTEXT("Substrate Ops", "Substrate Operators")) { }
-	};
-	static FConstructorStatics ConstructorStatics;
-#if WITH_EDITORONLY_DATA
-	MenuCategories.Add(ConstructorStatics.NAME_Substrate);
-#endif
-}
-
-#if WITH_EDITOR
-int32 UMaterialExpressionSubstrateVerticalLayering::Compile(class FMaterialCompiler* Compiler, int32 OutputIndex)
-{
-	if (!Top.GetTracedInput().Expression)
-	{
-		return Compiler->Errorf(TEXT("Missing Top input"));
-	}
-	if (!Base.GetTracedInput().Expression)
-	{
-		return Compiler->Errorf(TEXT("Missing Base input"));
-	}
-
-	Compiler->SubstrateTreeStackPush(this, 0);
-	int32 TopCodeChunk = Top.Compile(Compiler);
-	Compiler->SubstrateTreeStackPop();
-	Compiler->SubstrateTreeStackPush(this, 1);
-	int32 BaseCodeChunk = Base.Compile(Compiler);
-	Compiler->SubstrateTreeStackPop();
-	Compiler->SubstrateTreeStackPush(this, 2);
-	int32 ThicknessCodeChunk = Thickness.Compile(Compiler);
-	Compiler->SubstrateTreeStackPop();
-
-	int32 OutputCodeChunk = INDEX_NONE;
-	FSubstrateOperator& SubstrateOperator = Compiler->SubstrateCompilationGetOperator(Compiler->SubstrateTreeStackGetPathUniqueId());
-	if (SubstrateOperator.bUseParameterBlending)
-	{
-		FSubstrateOperator* TopBSDFOperator = Compiler->SubstrateCompilationGetOperatorFromIndex(SubstrateOperator.LeftIndex);
-		FSubstrateOperator* BaseBSDFOperator = Compiler->SubstrateCompilationGetOperatorFromIndex(SubstrateOperator.RightIndex);
-		if (!TopBSDFOperator || !BaseBSDFOperator)
-		{
-			return Compiler->Errorf(TEXT("Missing input on vertical layering node."));
-		}
-		if (TopCodeChunk == INDEX_NONE)
-		{
-			return Compiler->Errorf(TEXT("Top input graph could not be evaluated for parameter blending."));
-		}
-		if (BaseCodeChunk == INDEX_NONE)
-		{
-			return Compiler->Errorf(TEXT("Base input graph could not be evaluated for parameter blending."));
-		}
-
-		const int32 TopNormalMixCodeChunk = Compiler->SubstrateVerticalLayeringParameterBlendingBSDFCoverageToNormalMixCodeChunk(TopCodeChunk);
-
-		// Compute the new Normal and Tangent resulting from the blending using code chunk
-		const int32 NewNormalCodeChunk = SubstrateBlendNormal(Compiler, BaseBSDFOperator->BSDFRegisteredSharedLocalBasis.NormalCodeChunk, TopBSDFOperator->BSDFRegisteredSharedLocalBasis.NormalCodeChunk, TopNormalMixCodeChunk);
-		// The tangent is optional so we treat it differently if INDEX_NONE is specified
-		int32 NewTangentCodeChunk = INDEX_NONE;
-		if (TopBSDFOperator->BSDFRegisteredSharedLocalBasis.TangentCodeChunk != INDEX_NONE && BaseBSDFOperator->BSDFRegisteredSharedLocalBasis.TangentCodeChunk != INDEX_NONE)
-		{
-			NewTangentCodeChunk = SubstrateBlendNormal(Compiler, BaseBSDFOperator->BSDFRegisteredSharedLocalBasis.TangentCodeChunk, TopBSDFOperator->BSDFRegisteredSharedLocalBasis.TangentCodeChunk, TopNormalMixCodeChunk);
-		}
-		else if (TopBSDFOperator->BSDFRegisteredSharedLocalBasis.TangentCodeChunk != INDEX_NONE)
-		{
-			NewTangentCodeChunk = TopBSDFOperator->BSDFRegisteredSharedLocalBasis.TangentCodeChunk;
-		}
-		else if (BaseBSDFOperator->BSDFRegisteredSharedLocalBasis.TangentCodeChunk != INDEX_NONE)
-		{
-			NewTangentCodeChunk = BaseBSDFOperator->BSDFRegisteredSharedLocalBasis.TangentCodeChunk;
-		}
-		const FSubstrateRegisteredSharedLocalBasis NewRegisteredSharedLocalBasis = SubstrateCompilationInfoCreateSharedLocalBasis(Compiler, NewNormalCodeChunk, NewTangentCodeChunk);
-
-		OutputCodeChunk = Compiler->SubstrateVerticalLayeringParameterBlending(
-			TopCodeChunk, BaseCodeChunk, ThicknessCodeChunk, Compiler->GetSubstrateSharedLocalBasisIndexMacro(NewRegisteredSharedLocalBasis), TopBSDFOperator->BSDFRegisteredSharedLocalBasis.NormalCodeChunk,
-			SubstrateOperator.bRootOfParameterBlendingSubTree ? &SubstrateOperator : nullptr);
-
-		// Propagate the parameter blended normal
-		SubstrateOperator.BSDFRegisteredSharedLocalBasis = NewRegisteredSharedLocalBasis;
-	}
-	else
-	{
-		OutputCodeChunk = Compiler->SubstrateVerticalLayering(TopCodeChunk, BaseCodeChunk, ThicknessCodeChunk, SubstrateOperator.Index, SubstrateOperator.MaxDistanceFromLeaves);
-	}
-
-
-	return OutputCodeChunk;
-}
-
-void UMaterialExpressionSubstrateVerticalLayering::GetCaption(TArray<FString>& OutCaptions) const
-{
-	if (bUseParameterBlending)
-	{
-		OutCaptions.Add(TEXT("Substrate Vertical Layer (Parameter Blend)"));
-	}
-	else
-	{
-		OutCaptions.Add(TEXT("Substrate Vertical Layer"));
-	}
-}
-
-FName UMaterialExpressionSubstrateVerticalLayering::GetInputName(int32 InputIndex) const
-{
-	if (InputIndex == 0)
-	{
-		return TEXT("Top");
-	}
-	else if (InputIndex == 1)
-	{
-		return TEXT("Bottom");
-	}
-	else if (InputIndex == 2)
-	{
-		return  TEXT("Top Thickness");
-	}
-
-	return TEXT("Unknown");
-}
-
-uint32 UMaterialExpressionSubstrateVerticalLayering::GetOutputType(int32 OutputIndex)
-{
-	return MCT_Substrate;
-}
-
-uint32 UMaterialExpressionSubstrateVerticalLayering::GetInputType(int32 InputIndex)
-{
-	if (InputIndex == 2)
-	{
-		return MCT_Float;
-	}
-	return MCT_Substrate;
-}
-
-bool UMaterialExpressionSubstrateVerticalLayering::IsResultSubstrateMaterial(int32 OutputIndex)
-{
-	return true;
-}
-
-void UMaterialExpressionSubstrateVerticalLayering::GatherSubstrateMaterialInfo(FSubstrateMaterialInfo& SubstrateMaterialInfo, int32 OutputIndex)
-{
-	if (!SubstrateMaterialInfo.PushSubstrateTreeStack())
-	{
-		return;
-	}
-
-	FExpressionInput TracedInputTop  = Top.GetTracedInput();
-	FExpressionInput TracedInputBase = Base.GetTracedInput();
-	if (TracedInputTop.Expression)
-	{
-		TracedInputTop.Expression->GatherSubstrateMaterialInfo(SubstrateMaterialInfo, TracedInputTop.OutputIndex);
-	}
-	if (TracedInputBase.Expression)
-	{
-		TracedInputBase.Expression->GatherSubstrateMaterialInfo(SubstrateMaterialInfo, TracedInputBase.OutputIndex);
-	}
-	SubstrateMaterialInfo.AddGuid(MaterialExpressionGuid);
-
-	SubstrateMaterialInfo.PopSubstrateTreeStack();
-}
-
-FSubstrateOperator* UMaterialExpressionSubstrateVerticalLayering::SubstrateGenerateMaterialTopologyTree(class FMaterialCompiler* Compiler, class UMaterialExpression* Parent, int32 OutputIndex)
-{
-	FSubstrateOperator& SubstrateOperator = Compiler->SubstrateCompilationRegisterOperator(SUBSTRATE_OPERATOR_VERTICAL, Compiler->SubstrateTreeStackGetPathUniqueId(), this, Parent, Compiler->SubstrateTreeStackGetParentPathUniqueId(), bUseParameterBlending);
-	if (Compiler->GetSubstrateTreeOutOfStackDepthOccurred())
-	{
-		return &SubstrateOperator; // Out ot stack space, return now to fail the compilation
-	}
-
-	FExpressionInput TracedInputTop  = Top.GetTracedInput();
-	FExpressionInput TracedInputBase = Base.GetTracedInput();
-	UMaterialExpression* ChildAExpression = TracedInputTop.Expression;
-	UMaterialExpression* ChildBExpression = TracedInputBase.Expression;
-	FSubstrateOperator* OpA = nullptr;
-	FSubstrateOperator* OpB = nullptr;
-
-	// Top - Use the vertical operator thickness
-	if (ChildAExpression)
-	{
-		Compiler->SubstrateThicknessStackPush(this, &Thickness);
-		Compiler->SubstrateTreeStackPush(this, 0);
-		OpA = ChildAExpression->SubstrateGenerateMaterialTopologyTree(Compiler, this, TracedInputTop.OutputIndex);
-		Compiler->SubstrateTreeStackPop();
-		Compiler->SubstrateThicknessStackPop();
-		AssignOperatorIndexIfNotNull(SubstrateOperator.LeftIndex, OpA);
-	}
-	// Bottom - Use the propagated thickness from parent
-	if (ChildBExpression)
-	{
-		Compiler->SubstrateTreeStackPush(this, 1);
-		OpB = ChildBExpression->SubstrateGenerateMaterialTopologyTree(Compiler, this, TracedInputBase.OutputIndex);
-		Compiler->SubstrateTreeStackPop();
-		AssignOperatorIndexIfNotNull(SubstrateOperator.RightIndex, OpB);
-	}
-	CombineFlagForParameterBlending(SubstrateOperator, OpA, OpB);
-
-	return &SubstrateOperator;
-}
-#endif // WITH_EDITOR
-
-
-
-UMaterialExpressionSubstrateAdd::UMaterialExpressionSubstrateAdd(const FObjectInitializer& ObjectInitializer)
-	: Super(ObjectInitializer)
-	, bUseParameterBlending(false)
-{
-	struct FConstructorStatics
-	{
-		FText NAME_Substrate;
-		FConstructorStatics() : NAME_Substrate(LOCTEXT("Substrate Ops", "Substrate Operators")) { }
-	};
-	static FConstructorStatics ConstructorStatics;
-#if WITH_EDITORONLY_DATA
-	MenuCategories.Add(ConstructorStatics.NAME_Substrate);
-#endif
-}
-
-#if WITH_EDITOR
-int32 UMaterialExpressionSubstrateAdd::Compile(class FMaterialCompiler* Compiler, int32 OutputIndex)
-{
-	if (!A.GetTracedInput().Expression)
-	{
-		return Compiler->Errorf(TEXT("Missing A input"));
-	}
-	if (!B.GetTracedInput().Expression)
-	{
-		return Compiler->Errorf(TEXT("Missing B input"));
-	}
-
-	Compiler->SubstrateTreeStackPush(this, 0);
-	int32 ACodeChunk = A.Compile(Compiler);
-	Compiler->SubstrateTreeStackPop();
-	Compiler->SubstrateTreeStackPush(this, 1);
-	int32 BCodeChunk = B.Compile(Compiler);
-	Compiler->SubstrateTreeStackPop();
-
-	int32 OutputCodeChunk = INDEX_NONE;
-	FSubstrateOperator& SubstrateOperator = Compiler->SubstrateCompilationGetOperator(Compiler->SubstrateTreeStackGetPathUniqueId());
-	if (SubstrateOperator.bUseParameterBlending)
-	{
-		FSubstrateOperator* ABSDFOperator = Compiler->SubstrateCompilationGetOperatorFromIndex(SubstrateOperator.LeftIndex);
-		FSubstrateOperator* BBSDFOperator = Compiler->SubstrateCompilationGetOperatorFromIndex(SubstrateOperator.RightIndex);
-		if (!ABSDFOperator || !BBSDFOperator)
-		{
-			return Compiler->Errorf(TEXT("Missing input on add node."));
-		}
-		if (ACodeChunk == INDEX_NONE)
-		{
-			return Compiler->Errorf(TEXT("A input graph could not be evaluated for parameter blending."));
-		}
-		if (BCodeChunk == INDEX_NONE)
-		{
-			return Compiler->Errorf(TEXT("B input graph could not be evaluated for parameter blending."));
-		}
-
-		const int32 ANormalMixCodeChunk = Compiler->SubstrateAddParameterBlendingBSDFCoverageToNormalMixCodeChunk(ACodeChunk, BCodeChunk);
-
-		// Compute the new Normal and Tangent resulting from the blending using code chunk
-		const int32 NewNormalCodeChunk = SubstrateBlendNormal(Compiler, BBSDFOperator->BSDFRegisteredSharedLocalBasis.NormalCodeChunk, ABSDFOperator->BSDFRegisteredSharedLocalBasis.NormalCodeChunk, ANormalMixCodeChunk);
-		// The tangent is optional so we treat it differently if INDEX_NONE is specified
-		int32 NewTangentCodeChunk = INDEX_NONE;
-		if (ABSDFOperator->BSDFRegisteredSharedLocalBasis.TangentCodeChunk != INDEX_NONE && BBSDFOperator->BSDFRegisteredSharedLocalBasis.TangentCodeChunk != INDEX_NONE)
-		{
-			NewTangentCodeChunk = SubstrateBlendNormal(Compiler, BBSDFOperator->BSDFRegisteredSharedLocalBasis.TangentCodeChunk, ABSDFOperator->BSDFRegisteredSharedLocalBasis.TangentCodeChunk, ANormalMixCodeChunk);
-		}
-		else if (ABSDFOperator->BSDFRegisteredSharedLocalBasis.TangentCodeChunk != INDEX_NONE)
-		{
-			NewTangentCodeChunk = ABSDFOperator->BSDFRegisteredSharedLocalBasis.TangentCodeChunk;
-		}
-		else if (BBSDFOperator->BSDFRegisteredSharedLocalBasis.TangentCodeChunk != INDEX_NONE)
-		{
-			NewTangentCodeChunk = BBSDFOperator->BSDFRegisteredSharedLocalBasis.TangentCodeChunk;
-		}
-		const FSubstrateRegisteredSharedLocalBasis NewRegisteredSharedLocalBasis = SubstrateCompilationInfoCreateSharedLocalBasis(Compiler, NewNormalCodeChunk, NewTangentCodeChunk);
-
-		OutputCodeChunk = Compiler->SubstrateAddParameterBlending(
-			ACodeChunk, BCodeChunk, ANormalMixCodeChunk, Compiler->GetSubstrateSharedLocalBasisIndexMacro(NewRegisteredSharedLocalBasis),
-			SubstrateOperator.bRootOfParameterBlendingSubTree ? &SubstrateOperator : nullptr);
-
-		// Propagate the parameter blended normal
-		SubstrateOperator.BSDFRegisteredSharedLocalBasis = NewRegisteredSharedLocalBasis;
-	}
-	else
-	{
-		OutputCodeChunk = Compiler->SubstrateAdd(ACodeChunk, BCodeChunk, SubstrateOperator.Index, SubstrateOperator.MaxDistanceFromLeaves);
-	}
-
-	return OutputCodeChunk;
-}
-
-void UMaterialExpressionSubstrateAdd::GetCaption(TArray<FString>& OutCaptions) const
-{
-	if (bUseParameterBlending)
-	{
-		OutCaptions.Add(TEXT("Substrate Add (Parameter Blend)"));
-	}
-	else
-	{
-		OutCaptions.Add(TEXT("Substrate Add"));
-	}
-}
-
-uint32 UMaterialExpressionSubstrateAdd::GetOutputType(int32 OutputIndex)
-{
-	return MCT_Substrate;
-}
-
-uint32 UMaterialExpressionSubstrateAdd::GetInputType(int32 InputIndex)
-{
-	return MCT_Substrate;
-}
-
-bool UMaterialExpressionSubstrateAdd::IsResultSubstrateMaterial(int32 OutputIndex)
-{
-	return true;
-}
-
-void UMaterialExpressionSubstrateAdd::GatherSubstrateMaterialInfo(FSubstrateMaterialInfo& SubstrateMaterialInfo, int32 OutputIndex)
-{
-	if (!SubstrateMaterialInfo.PushSubstrateTreeStack())
-	{
-		return;
-	}
-
-	FExpressionInput TracedInputA = A.GetTracedInput();
-	FExpressionInput TracedInputB = B.GetTracedInput();
-	if (TracedInputA.Expression)
-	{
-		TracedInputA.Expression->GatherSubstrateMaterialInfo(SubstrateMaterialInfo, TracedInputA.OutputIndex);
-	}
-	if (TracedInputB.Expression)
-	{
-		TracedInputB.Expression->GatherSubstrateMaterialInfo(SubstrateMaterialInfo, TracedInputB.OutputIndex);
-	}
-	SubstrateMaterialInfo.AddGuid(MaterialExpressionGuid);
-
-	SubstrateMaterialInfo.PopSubstrateTreeStack();
-}
-
-FSubstrateOperator* UMaterialExpressionSubstrateAdd::SubstrateGenerateMaterialTopologyTree(class FMaterialCompiler* Compiler, class UMaterialExpression* Parent, int32 OutputIndex)
-{
-	FSubstrateOperator& SubstrateOperator = Compiler->SubstrateCompilationRegisterOperator(SUBSTRATE_OPERATOR_ADD, Compiler->SubstrateTreeStackGetPathUniqueId(), this, Parent, Compiler->SubstrateTreeStackGetParentPathUniqueId(), bUseParameterBlending);
-	if (Compiler->GetSubstrateTreeOutOfStackDepthOccurred())
-	{
-		return &SubstrateOperator; // Out ot stack space, return now to fail the compilation
-	}
-
-	FExpressionInput TracedInputA = A.GetTracedInput();
-	FExpressionInput TracedInputB = B.GetTracedInput();
-	UMaterialExpression* ChildAExpression = TracedInputA.Expression;
-	UMaterialExpression* ChildBExpression = TracedInputB.Expression;
-	FSubstrateOperator* OpA = nullptr;
-	FSubstrateOperator* OpB = nullptr;
-	if (ChildAExpression)
-	{
-		Compiler->SubstrateTreeStackPush(this, 0);
-		OpA = ChildAExpression->SubstrateGenerateMaterialTopologyTree(Compiler, this, TracedInputA.OutputIndex);
-		Compiler->SubstrateTreeStackPop();
-		AssignOperatorIndexIfNotNull(SubstrateOperator.LeftIndex, OpA);
-	}
-	if (ChildBExpression)
-	{
-		Compiler->SubstrateTreeStackPush(this, 1);
-		OpB = ChildBExpression->SubstrateGenerateMaterialTopologyTree(Compiler, this, TracedInputB.OutputIndex);
-		Compiler->SubstrateTreeStackPop();
-		AssignOperatorIndexIfNotNull(SubstrateOperator.RightIndex, OpB);
-	}
-	CombineFlagForParameterBlending(SubstrateOperator, OpA, OpB);
-
-	return &SubstrateOperator;
-}
-#endif // WITH_EDITOR
-
-
-
-UMaterialExpressionSubstrateWeight::UMaterialExpressionSubstrateWeight(const FObjectInitializer& ObjectInitializer)
-	: Super(ObjectInitializer)
-{
-	struct FConstructorStatics
-	{
-		FText NAME_Substrate;
-		FConstructorStatics() : NAME_Substrate(LOCTEXT("Substrate Ops", "Substrate Operators")) { }
-	};
-	static FConstructorStatics ConstructorStatics;
-#if WITH_EDITORONLY_DATA
-	MenuCategories.Add(ConstructorStatics.NAME_Substrate);
-#endif
-}
-
-#if WITH_EDITOR
-int32 UMaterialExpressionSubstrateWeight::Compile(class FMaterialCompiler* Compiler, int32 OutputIndex)
-{
-	if (!A.GetTracedInput().Expression)
-	{
-		return Compiler->Errorf(TEXT("Missing A input"));
-	}
-
-	Compiler->SubstrateTreeStackPush(this, 0);
-	int32 ACodeChunk = A.Compile(Compiler);
-	Compiler->SubstrateTreeStackPop();
-	int32 WeightCodeChunk = Weight.GetTracedInput().Expression ? Weight.Compile(Compiler) : Compiler->Constant(1.0f);
-
-	int32 OutputCodeChunk = INDEX_NONE;
-	FSubstrateOperator& SubstrateOperator = Compiler->SubstrateCompilationGetOperator(Compiler->SubstrateTreeStackGetPathUniqueId());
-	if (SubstrateOperator.bUseParameterBlending)
-	{
-		// Propagate the parameter blended normal
-		FSubstrateOperator* Operator = Compiler->SubstrateCompilationGetOperatorFromIndex(SubstrateOperator.LeftIndex);
-		if (!Operator)
-		{
-			return Compiler->Errorf(TEXT("Missing input on weight node."));
-		}
-		if (ACodeChunk == INDEX_NONE)
-		{
-			return Compiler->Errorf(TEXT("A input graph could not be evaluated for parameter blending."));
-		}
-		if (WeightCodeChunk == INDEX_NONE)
-		{
-			return Compiler->Errorf(TEXT("Weight input graph could not be evaluated for parameter blending."));
-		}
-
-		OutputCodeChunk = Compiler->SubstrateWeightParameterBlending(
-			ACodeChunk, WeightCodeChunk, 
-			SubstrateOperator.bRootOfParameterBlendingSubTree ? &SubstrateOperator : nullptr);
-
-		// Propagate the parameter blended normal
-		SubstrateOperator.BSDFRegisteredSharedLocalBasis = Operator->BSDFRegisteredSharedLocalBasis;
-	}
-	else
-	{
-		OutputCodeChunk = Compiler->SubstrateWeight(ACodeChunk, WeightCodeChunk, SubstrateOperator.Index, SubstrateOperator.MaxDistanceFromLeaves);
-	}
-
-	return OutputCodeChunk;
-}
-
-void UMaterialExpressionSubstrateWeight::GetCaption(TArray<FString>& OutCaptions) const
-{
-	OutCaptions.Add(TEXT("Substrate Coverage Weight"));
-}
-
-uint32 UMaterialExpressionSubstrateWeight::GetOutputType(int32 OutputIndex)
-{
-	return MCT_Substrate;
-}
-
-uint32 UMaterialExpressionSubstrateWeight::GetInputType(int32 InputIndex)
-{
-	if (InputIndex == 0)
-	{
-		return MCT_Substrate;
-	}
-	return MCT_Float1;
-}
-
-bool UMaterialExpressionSubstrateWeight::IsResultSubstrateMaterial(int32 OutputIndex)
-{
-	return true;
-}
-
-void UMaterialExpressionSubstrateWeight::GatherSubstrateMaterialInfo(FSubstrateMaterialInfo& SubstrateMaterialInfo, int32 OutputIndex)
-{
-	if (!SubstrateMaterialInfo.PushSubstrateTreeStack())
-	{
-		return;
-	}
-
-	FExpressionInput TracedInputA = A.GetTracedInput();
-	if (TracedInputA.Expression)
-	{
-		TracedInputA.Expression->GatherSubstrateMaterialInfo(SubstrateMaterialInfo, TracedInputA.OutputIndex);
-	}
-	SubstrateMaterialInfo.AddGuid(MaterialExpressionGuid);
-
-	SubstrateMaterialInfo.PopSubstrateTreeStack();
-}
-
-FSubstrateOperator* UMaterialExpressionSubstrateWeight::SubstrateGenerateMaterialTopologyTree(class FMaterialCompiler* Compiler, class UMaterialExpression* Parent, int32 OutputIndex)
-{
-	FSubstrateOperator& SubstrateOperator = Compiler->SubstrateCompilationRegisterOperator(SUBSTRATE_OPERATOR_WEIGHT, Compiler->SubstrateTreeStackGetPathUniqueId(), this, Parent, Compiler->SubstrateTreeStackGetParentPathUniqueId());
-	if (Compiler->GetSubstrateTreeOutOfStackDepthOccurred())
-	{
-		return &SubstrateOperator; // Out ot stack space, return now to fail the compilation
-	}
-
-	FExpressionInput TracedInputA = A.GetTracedInput();
-	UMaterialExpression* ChildAExpression = TracedInputA.Expression;
-	FSubstrateOperator* OpA = nullptr;
-	if (ChildAExpression)
-	{
-		Compiler->SubstrateTreeStackPush(this, 0);
-		OpA = ChildAExpression->SubstrateGenerateMaterialTopologyTree(Compiler, this, TracedInputA.OutputIndex);
-		Compiler->SubstrateTreeStackPop();
-		AssignOperatorIndexIfNotNull(SubstrateOperator.LeftIndex, OpA);
-	}
-	CombineFlagForParameterBlending(SubstrateOperator, OpA);
-
-	return &SubstrateOperator;
-}
-#endif // WITH_EDITOR
-
-UMaterialExpressionSubstrateUtilityBase::UMaterialExpressionSubstrateUtilityBase(const FObjectInitializer& ObjectInitializer)
-	: Super(ObjectInitializer)
-{
-}
-
-UMaterialExpressionSubstrateTransmittanceToMFP::UMaterialExpressionSubstrateTransmittanceToMFP(const FObjectInitializer& ObjectInitializer)
-	: Super(ObjectInitializer)
-{
-	struct FConstructorStatics
-	{
-		FText NAME_Substrate;
-		FConstructorStatics() : NAME_Substrate(LOCTEXT("Substrate Helpers", "Substrate Helpers")) { }
-	};
-	static FConstructorStatics ConstructorStatics;
-#if WITH_EDITORONLY_DATA
-	MenuCategories.Add(ConstructorStatics.NAME_Substrate);
-
-	bShowOutputNameOnPin = true;
-
-	Outputs.Reset();
-	Outputs.Add(FExpressionOutput(TEXT("MFP")));
-	Outputs.Add(FExpressionOutput(TEXT("Thickness")));
-#endif
-}
-
-#if WITH_EDITOR
-int32 UMaterialExpressionSubstrateTransmittanceToMFP::Compile(class FMaterialCompiler* Compiler, int32 OutputIndex)
-{
-	int32 TransmittanceColorCodeChunk = TransmittanceColor.GetTracedInput().Expression ? TransmittanceColor.Compile(Compiler) : Compiler->Constant(0.5f);
-	int32 ThicknessCodeChunk = Thickness.GetTracedInput().Expression ? Thickness.Compile(Compiler) : Compiler->Constant(SUBSTRATE_LAYER_DEFAULT_THICKNESS_CM);
-	if (TransmittanceColorCodeChunk == INDEX_NONE)
-	{
-		return Compiler->Errorf(TEXT("TransmittanceColor input graph could not be evaluated for TransmittanceToMFP."));
-	}
-	if (ThicknessCodeChunk == INDEX_NONE)
-	{
-		return Compiler->Errorf(TEXT("ThicknessCodeChunk input graph could not be evaluated for TransmittanceToMFP."));
-	}
-	return Compiler->SubstrateTransmittanceToMFP(
-		TransmittanceColorCodeChunk,
-		ThicknessCodeChunk,
-		OutputIndex);
-}
-
-void UMaterialExpressionSubstrateTransmittanceToMFP::GetCaption(TArray<FString>& OutCaptions) const
-{
-	OutCaptions.Add(TEXT("Substrate Transmittance-To-MeanFreePath"));
-}
-
-uint32 UMaterialExpressionSubstrateTransmittanceToMFP::GetOutputType(int32 OutputIndex)
-{
-	switch (OutputIndex)
-	{
-	case 0:
-		return MCT_Float3; // MFP
-		break;
-	case 1:
-		return MCT_Float1; // Thickness
-		break;
-	}
-
-	check(false);
-	return MCT_Float1;
-}
-
-uint32 UMaterialExpressionSubstrateTransmittanceToMFP::GetInputType(int32 InputIndex)
-{
-	switch (InputIndex)
-	{
-	case 0:
-		return MCT_Float3; // Transmittance
-		break;
-	case 1:
-		return MCT_Float1; // Thickness
-		break;
-	}
-
-	check(false);
-	return MCT_Float1;
-}
-void UMaterialExpressionSubstrateTransmittanceToMFP::GetConnectorToolTip(int32 InputIndex, int32 OutputIndex, TArray<FString>& OutToolTip)
-{
-	if (InputIndex != INDEX_NONE)
-	{
-		switch (InputIndex)
-		{
-		case 0:
-			ConvertToMultilineToolTip(TEXT("The colored transmittance for a view perpendicular to the surface. The transmittance for other view orientations will automatically be deduced according to surface thickness."), 80, OutToolTip);
-			break;
-		case 1:
-			ConvertToMultilineToolTip(TEXT("The desired thickness in centimeter. This can be set lower than 0.1mm(= 0.01cm) to enable the Thin lighting model on the slab node for instance. Another use case example: this node output called thickness can be modulated before it is plugged in a slab node.this can be used to achieve simple scattering/transmittance variation of the same material."), 80, OutToolTip);
-			break;
-		}
-	}
-	else if (OutputIndex != INDEX_NONE)
-	{
-		switch (OutputIndex)
-		{
-		case 0:
-			ConvertToMultilineToolTip(TEXT("The Mean Free Path defining the participating media constituting the slab of material (unit = centimeters)."), 80, OutToolTip);
-			break;
-		case 1:
-			ConvertToMultilineToolTip(TEXT("The thickness of the slab of material (unit = centimeters)."), 80, OutToolTip);
-			break;
-		}
-	}
-}
-
-void UMaterialExpressionSubstrateTransmittanceToMFP::GetExpressionToolTip(TArray<FString>& OutToolTip)
-{
-	ConvertToMultilineToolTip(TEXT("Convert a transmittance color corresponding to a slab of participating media viewed perpendicularly to its surface.\nThis node directly maps to the Slab BSDF input."), 80, OutToolTip);
-
-}
-#endif // WITH_EDITOR
-
-UMaterialExpressionSubstrateMetalnessToDiffuseAlbedoF0::UMaterialExpressionSubstrateMetalnessToDiffuseAlbedoF0(const FObjectInitializer& ObjectInitializer)
-	: Super(ObjectInitializer)
-{
-	struct FConstructorStatics
-	{
-		FText NAME_Substrate;
-		FConstructorStatics() : NAME_Substrate(LOCTEXT("Substrate Helpers", "Substrate Helpers")) { }
-	};
-	static FConstructorStatics ConstructorStatics;
-#if WITH_EDITORONLY_DATA
-	MenuCategories.Add(ConstructorStatics.NAME_Substrate);
-
-	bShowOutputNameOnPin = true;
-
-	Outputs.Reset();
-	Outputs.Add(FExpressionOutput(TEXT("DiffuseAlbedo")));
-	Outputs.Add(FExpressionOutput(TEXT("F0")));
-#endif
-}
-
-#if WITH_EDITOR
-int32 UMaterialExpressionSubstrateMetalnessToDiffuseAlbedoF0::Compile(class FMaterialCompiler* Compiler, int32 OutputIndex)
-{
-	int32 BaseColorCodeChunk = BaseColor.GetTracedInput().Expression ? BaseColor.Compile(Compiler) : Compiler->Constant(0.18f);
-	int32 SpecularCodeChunk = Specular.GetTracedInput().Expression	 ? Specular.Compile(Compiler) : Compiler->Constant(0.5f);
-	int32 MetallicCodeChunk = Metallic.GetTracedInput().Expression ? Metallic.Compile(Compiler) : Compiler->Constant(0.0f);
-	if (BaseColorCodeChunk == INDEX_NONE)
-	{
-		return Compiler->Errorf(TEXT("BaseColor input graph could not be evaluated for MetalnessToDiffuseAlbedoF0."));
-	}
-	if (SpecularCodeChunk == INDEX_NONE)
-	{
-		return Compiler->Errorf(TEXT("Specular input graph could not be evaluated for MetalnessToDiffuseAlbedoF0."));
-	}
-	if (MetallicCodeChunk == INDEX_NONE)
-	{
-		return Compiler->Errorf(TEXT("Metallic input graph could not be evaluated for MetalnessToDiffuseAlbedoF0."));
-	}
-	return Compiler->SubstrateMetalnessToDiffuseAlbedoF0(
-		BaseColorCodeChunk,
-		SpecularCodeChunk,
-		MetallicCodeChunk,
-		OutputIndex);
-}
-
-void UMaterialExpressionSubstrateMetalnessToDiffuseAlbedoF0::GetCaption(TArray<FString>& OutCaptions) const
-{
-	OutCaptions.Add(TEXT("Substrate Metalness-To-DiffuseAlbedo-F0"));
-}
-
-uint32 UMaterialExpressionSubstrateMetalnessToDiffuseAlbedoF0::GetOutputType(int32 OutputIndex)
-{
-	switch (OutputIndex)
-	{
-	case 0:
-		return MCT_Float3; // Diffuse Albedo
-		break;
-	case 1:
-		return MCT_Float3; // F0
-		break;
-	}
-
-	check(false);
-	return MCT_Float1;
-}
-
-uint32 UMaterialExpressionSubstrateMetalnessToDiffuseAlbedoF0::GetInputType(int32 InputIndex)
-{
-	if (InputIndex == 0) { return MCT_Float3; }
-	if (InputIndex == 1) { return MCT_Float1; }
-	if (InputIndex == 2) { return MCT_Float1; }
-	return MCT_Float1;
-}
-
-void UMaterialExpressionSubstrateMetalnessToDiffuseAlbedoF0::GetConnectorToolTip(int32 InputIndex, int32 OutputIndex, TArray<FString>& OutToolTip)
-{
-	switch (OutputIndex)
-	{
-		case 1: ConvertToMultilineToolTip(TEXT("Defines the overall color of the Material. (type = float3, unit = unitless, defaults to 0.18)"), 80, OutToolTip); break;
-		case 2: ConvertToMultilineToolTip(TEXT("Controls how \"metal-like\" your surface looks like. 0 means dielectric, 1 means conductor (type = float, unit = unitless, defaults to 0)"), 80, OutToolTip); break;
-		case 3: ConvertToMultilineToolTip(TEXT("Used to scale the current amount of specularity on non-metallic surfaces and is a value between 0 and 1 (type = float, unit = unitless, defaults to plastic 0.5)"), 80, OutToolTip); break;
-	}
-}
-
-void UMaterialExpressionSubstrateMetalnessToDiffuseAlbedoF0::GetExpressionToolTip(TArray<FString>& OutToolTip)
-{
-	ConvertToMultilineToolTip(TEXT("Convert a metalness parameterization (BaseColor/Specular/Metallic) into DiffuseAlbedo/F0 parameterization."), 80, OutToolTip);
-
-}
-#endif // WITH_EDITOR
-
-
-
-UMaterialExpressionSubstrateHazinessToSecondaryRoughness::UMaterialExpressionSubstrateHazinessToSecondaryRoughness(const FObjectInitializer& ObjectInitializer)
-	: Super(ObjectInitializer)
-{
-	struct FConstructorStatics
-	{
-		FText NAME_Substrate;
-		FConstructorStatics() : NAME_Substrate(LOCTEXT("Substrate Helpers", "Substrate Helpers")) { }
-	};
-	static FConstructorStatics ConstructorStatics;
-#if WITH_EDITORONLY_DATA
-	MenuCategories.Add(ConstructorStatics.NAME_Substrate);
-
-	bShowOutputNameOnPin = true;
-
-	Outputs.Reset();
-	Outputs.Add(FExpressionOutput(TEXT("Second Roughness")));
-	Outputs.Add(FExpressionOutput(TEXT("Second Roughness Weight")));
-#endif
-}
-
-#if WITH_EDITOR
-int32 UMaterialExpressionSubstrateHazinessToSecondaryRoughness::Compile(class FMaterialCompiler* Compiler, int32 OutputIndex)
-{
-	int32 BaseRoughnessCodeChunk= BaseRoughness.GetTracedInput().Expression ? BaseRoughness.Compile(Compiler) : Compiler->Constant(0.1f);
-	int32 HazinessCodeChunk		= Haziness.GetTracedInput().Expression      ? Haziness.Compile(Compiler)      : Compiler->Constant(0.5f);
-	if (BaseRoughnessCodeChunk == INDEX_NONE)
-	{
-		return Compiler->Errorf(TEXT("BaseRoughness input graph could not be evaluated for HazinessToSecondaryRoughness."));
-	}
-	if (HazinessCodeChunk == INDEX_NONE)
-	{
-		return Compiler->Errorf(TEXT("Haziness input graph could not be evaluated for HazinessToSecondaryRoughness."));
-	}
-	return Compiler->SubstrateHazinessToSecondaryRoughness(
-		BaseRoughnessCodeChunk,
-		HazinessCodeChunk,
-		OutputIndex);
-}
-
-void UMaterialExpressionSubstrateHazinessToSecondaryRoughness::GetCaption(TArray<FString>& OutCaptions) const
-{
-	OutCaptions.Add(TEXT("Substrate Haziness-To-Secondary-Roughness"));
-}
-
-uint32 UMaterialExpressionSubstrateHazinessToSecondaryRoughness::GetOutputType(int32 OutputIndex)
-{
-	switch (OutputIndex)
-	{
-	case 0:
-		return MCT_Float1; // Second Roughness
-		break;
-	case 1:
-		return MCT_Float1; // Second Roughness Weight
-		break;
-	}
-
-	check(false);
-	return MCT_Float1;
-}
-
-uint32 UMaterialExpressionSubstrateHazinessToSecondaryRoughness::GetInputType(int32 InputIndex)
-{
-	switch (InputIndex)
-	{
-	case 0:
-		return MCT_Float1; // BaseRoughness
-		break;
-	case 1:
-		return MCT_Float1; // Haziness
-		break;
-	}
-
-	check(false);
-	return MCT_Float1;
-}
-void UMaterialExpressionSubstrateHazinessToSecondaryRoughness::GetConnectorToolTip(int32 InputIndex, int32 OutputIndex, TArray<FString>& OutToolTip)
-{
-	if (InputIndex != INDEX_NONE)
-	{
-		switch (InputIndex)
-		{
-		case 0:
-			ConvertToMultilineToolTip(TEXT("The base roughness of the surface. It represented the smoothest part of the reflection."), 80, OutToolTip);
-			break;
-		case 1:
-			ConvertToMultilineToolTip(TEXT("Haziness represent the amount of irregularity of the surface. A high value will lead to a second rough specular lobe causing the surface too look `milky`."), 80, OutToolTip);
-			break;
-		}
-	}
-	else if (OutputIndex != INDEX_NONE)
-	{
-		switch (OutputIndex)
-		{
-		case 0:
-			ConvertToMultilineToolTip(TEXT("The roughness of the second lobe."), 80, OutToolTip);
-			break;
-		case 1:
-			ConvertToMultilineToolTip(TEXT("The weight of the secondary specular lobe, while the primary specular lobe will have a weight of (1 - SecondRoughnessWeight)."), 80, OutToolTip);
-			break;
-		}
-	}
-}
-
-void UMaterialExpressionSubstrateHazinessToSecondaryRoughness::GetExpressionToolTip(TArray<FString>& OutToolTip)
-{
-	ConvertToMultilineToolTip(TEXT("Compute a second specular lobe roughness from a base surface roughness and haziness. This parameterisation ensure that the haziness makes physically and is perceptually easy to author."), 80, OutToolTip);
-
-}
-#endif // WITH_EDITOR
-
-
-
-UMaterialExpressionSubstrateThinFilm::UMaterialExpressionSubstrateThinFilm(const FObjectInitializer& ObjectInitializer)
-	: Super(ObjectInitializer)
-{
-	struct FConstructorStatics
-	{
-		FText NAME_Substrate;
-		FConstructorStatics() : NAME_Substrate(LOCTEXT("Substrate Helpers", "Substrate Helpers")) { }
-	};
-	static FConstructorStatics ConstructorStatics;
-#if WITH_EDITORONLY_DATA
-	MenuCategories.Add(ConstructorStatics.NAME_Substrate);
-#endif
-
-	bShowOutputNameOnPin = true;
-
-	Outputs.Reset();
-	Outputs.Add(FExpressionOutput(TEXT("Specular Color")));
-	Outputs.Add(FExpressionOutput(TEXT("Edge Specular Color")));
-}
-
-#if WITH_EDITOR
-int32 UMaterialExpressionSubstrateThinFilm::Compile(class FMaterialCompiler* Compiler, int32 OutputIndex)
-{
-	int32 NormalCodeChunk		= CompileWithDefaultNormalWS(Compiler, Normal);
-
-	int32 F0CodeChunk			= F0.GetTracedInput().Expression		? F0.Compile(Compiler)			: Compiler->Constant3(0.04f, 0.04f, 0.04f);
-	int32 F90CodeChunk			= F90.GetTracedInput().Expression		? F90.Compile(Compiler)			: Compiler->Constant3(1.0f, 1.0f, 1.0f);
-
-	int32 ThicknessCodeChunk	= Thickness.GetTracedInput().Expression	? Thickness.Compile(Compiler)	: Compiler->Constant(1.0f);
-	int32 IORCodeChunk			= IOR.GetTracedInput().Expression		? IOR.Compile(Compiler)			: Compiler->Constant(1.44f);
-
-	if (NormalCodeChunk == INDEX_NONE)
-	{
-		return Compiler->Errorf(TEXT("NormalCode input graph could not be evaluated for ThinFilm."));
-	}
-	if (F0CodeChunk == INDEX_NONE)
-	{
-		return Compiler->Errorf(TEXT("F0 input graph could not be evaluated for ThinFilm."));
-	}
-	if (F90CodeChunk == INDEX_NONE)
-	{
-		return Compiler->Errorf(TEXT("F90 input graph could not be evaluated for ThinFilm."));
-	}
-	if (ThicknessCodeChunk == INDEX_NONE)
-	{
-		return Compiler->Errorf(TEXT("Thickness input graph could not be evaluated for ThinFilm."));
-	}
-	if (IORCodeChunk == INDEX_NONE)
-	{
-		return Compiler->Errorf(TEXT("IOR input graph could not be evaluated for ThinFilm."));
-	}
-
-	return Compiler->SubstrateThinFilm(NormalCodeChunk, F0CodeChunk, F90CodeChunk, ThicknessCodeChunk, IORCodeChunk, OutputIndex);
-}
-
-void UMaterialExpressionSubstrateThinFilm::GetCaption(TArray<FString>& OutCaptions) const
-{
-	OutCaptions.Add(TEXT("Substrate Thin-Film"));
-}
-
-uint32 UMaterialExpressionSubstrateThinFilm::GetOutputType(int32 OutputIndex)
-{
-	switch (OutputIndex)
-	{
-	case 0:
-		return MCT_Float3; // F0
-		break;
-	case 1:
-		return MCT_Float3; // F90
-		break;
-	}
-
-	check(false);
-	return MCT_Float1;
-}
-
-uint32 UMaterialExpressionSubstrateThinFilm::GetInputType(int32 InputIndex)
-{
-	if (InputIndex == 0) { return MCT_Float3; } // Normal
-	if (InputIndex == 1) { return MCT_Float3; } // F0
-	if (InputIndex == 2) { return MCT_Float3; } // F90
-	if (InputIndex == 3) { return MCT_Float1; } // Thickness
-	if (InputIndex == 4) { return MCT_Float1; } // IOR
-
-	check(false);
-	return MCT_Float1;
-}
-void UMaterialExpressionSubstrateThinFilm::GetConnectorToolTip(int32 InputIndex, int32 OutputIndex, TArray<FString>& OutToolTip)
-{
-	if (OutputIndex != INDEX_NONE)
-	{
-		switch (OutputIndex)
-		{
-		case 0:
-			ConvertToMultilineToolTip(TEXT("F0 accounting for thin film interferences. This is percentage of light reflected as specular from a surface when the view is perpendicular to the surface. (type = float3, unit = unitless, defaults to plastic 0.04)"), 80, OutToolTip);
-			break;
-		case 1:
-			ConvertToMultilineToolTip(TEXT("F90 accounting for thin film interferences. the percentage of light reflected as specular from a surface when the view is tangent to the surface. (type = float3, unit = unitless, defaults to 1.0f)."), 80, OutToolTip);
-			break;
-		}
-		return;
-	}
-
-	// Else use the default input tooltip
-	Super::GetConnectorToolTip(InputIndex, OutputIndex, OutToolTip);
-}
-
-void UMaterialExpressionSubstrateThinFilm::GetExpressionToolTip(TArray<FString>& OutToolTip)
-{
-	ConvertToMultilineToolTip(TEXT("Compute the resulting material specular parameter F0 and F90 according to input surface properties as well as the thin film parameters."), 80, OutToolTip);
-}
-#endif // WITH_EDITOR
-
-// Return a conservative list of connected material attribute inputs
-#if WITH_EDITOR
-static uint64 GetConnectedMaterialAttributesInputs(TConstArrayView<TObjectPtr<UMaterialExpression>> Expressions)
-{
-	uint64 Out = 0ull;
-	for (const UMaterialExpression* Expression : Expressions)
-	{
-		if (Expression)
-		{
-			if (Expression->IsA(UMaterialExpressionSetMaterialAttributes::StaticClass()))
-			{
-				const UMaterialExpressionSetMaterialAttributes* Attr = Cast<UMaterialExpressionSetMaterialAttributes>(Expression);
-				Out |= Attr->GetConnectedInputs();
-			}
-			else if (Expression->IsA(UMaterialExpressionMakeMaterialAttributes::StaticClass()))
-			{
-				const UMaterialExpressionMakeMaterialAttributes* Attr = Cast<UMaterialExpressionMakeMaterialAttributes>(Expression);
-				Out |= Attr->GetConnectedInputs();
-			}
-			else if (Expression->IsA(UMaterialExpressionMaterialFunctionCall::StaticClass()))
-			{
-				const UMaterialExpressionMaterialFunctionCall* Attr = Cast<UMaterialExpressionMaterialFunctionCall>(Expression);
-				if (Attr->MaterialFunction)
-				{
-					Out |= GetConnectedMaterialAttributesInputs(Attr->MaterialFunction->GetExpressions());
-				}
-			}
-		}
-	}
-	return Out;
-}
-
-static uint64 GetConnectedMaterialAttributesInputs(const UMaterial* InMaterial)
-{
-	if (!InMaterial) return 0;
-	return GetConnectedMaterialAttributesInputs(InMaterial->GetExpressions());
-}
-
-static bool IsCustomMaterialAttributeInputConnected(uint64 InCache, FGuid InProperty)
-{
-	static TArray<FMaterialCustomOutputAttributeDefintion> CustomAttributeList;
-	if (CustomAttributeList.Num() == 0)
-	{
-		FMaterialAttributeDefinitionMap::GetCustomAttributeList(CustomAttributeList);
-	}
-
-	const EMaterialProperty Property = FMaterialAttributeDefinitionMap::GetProperty(InProperty);
-	if (Property == MP_CustomOutput)
-	{
-		if (const FMaterialCustomOutputAttributeDefintion* Attribute = FMaterialAttributeDefinitionMap::GetCustomAttribute(InProperty))
-		{
-			const uint32 CustomIndex = CustomAttributeList.IndexOfByPredicate([Attribute](const FMaterialCustomOutputAttributeDefintion& A) { return A.AttributeID == Attribute->AttributeID; });
-			if (CustomIndex != INDEX_NONE)
-			{
-				return !!(InCache & (1ull << (MP_MAX + CustomIndex)));
-			}
-		}
-	}
-	return false;
-}
-
-static bool IsMaterialAttributeInputConnected(uint64 InCache, EMaterialProperty InProperty)
-{
-	return !!(InCache & (1ull << uint64(InProperty)));
-}
-#endif
-
-UMaterialExpressionSubstrateConvertMaterialAttributes::UMaterialExpressionSubstrateConvertMaterialAttributes(const FObjectInitializer& ObjectInitializer)
-: Super(ObjectInitializer)
-{
-	struct FConstructorStatics
-	{
-		FText NAME_Strata;
-		FConstructorStatics() : NAME_Strata(LOCTEXT("Substrate Conversion", "Substrate Conversion")) { }
-	};
-	static FConstructorStatics ConstructorStatics;
-#if WITH_EDITORONLY_DATA
-	MenuCategories.Add(ConstructorStatics.NAME_Strata);
-#endif
-
-#if WITH_EDITORONLY_DATA
-	Outputs.Reset();
-	Outputs.Add(FExpressionOutput(TEXT(""))); // Substrate
-	bShowOutputNameOnPin = true;
-#endif
-}
-
-#if WITH_EDITOR
-FExpressionInput* UMaterialExpressionSubstrateConvertMaterialAttributes::GetInput(int32 InputIndex)
-{
-	switch (InputIndex)
-	{
-		case 0: return &MaterialAttributes;
-		case 1: return &WaterScatteringCoefficients;
-		case 2: return &WaterAbsorptionCoefficients;
-		case 3: return &WaterPhaseG;
-		case 4: return &ColorScaleBehindWater;
-		default: return nullptr;
-	}
-}
-
-int32 UMaterialExpressionSubstrateConvertMaterialAttributes::Compile(class FMaterialCompiler* Compiler, int32 OutputIndex)
-{
-	static const FGuid ClearCoatBottomNormalGuid 			= FMaterialAttributeDefinitionMap::GetCustomAttributeID(TEXT("ClearCoatBottomNormal"));
-	static const FGuid CustomEyeTangentGuid 				= FMaterialAttributeDefinitionMap::GetCustomAttributeID(TEXT("CustomEyeTangent"));
-	static const FGuid TransmittanceColorGuid 				= FMaterialAttributeDefinitionMap::GetCustomAttributeID(TEXT("TransmittanceColor"));
-	static const FGuid ThinTranslucentSurfaceCoverageGuid	= FMaterialAttributeDefinitionMap::GetCustomAttributeID(TEXT("ThinTranslucentSurfaceCoverage"));
-
-	if (OutputIndex != 0)
-	{
-		return Compiler->Error(TEXT("Output pin index error"));
-	}
-		
-	// We also cannot ignore the tangent when using the default Tangent because GetTangentBasis
-	// used in SubstrateGetBSDFSharedBasis cannot be relied on for smooth tangent used for lighting on any mesh.
-
-	const uint64 Cached = GetConnectedMaterialAttributesInputs(Material);
-	const bool bHasAnisotropy = IsMaterialAttributeInputConnected(Cached, MP_Anisotropy);
-
-	// Regular normal basis
-	int32 NormalCodeChunk = Compiler->TransformNormalFromRequestedBasisToWorld(MaterialAttributes.CompileWithDefault(Compiler, FMaterialAttributeDefinitionMap::GetID(MP_Normal)));
-
-	// When computing NormalCodeChunk, we invoke TransformNormalFromRequestedBasisToWorld which requires input to be float or float3.
-	// Certain material do not respect this requirement. We handle here a simple recovery when source material doesn't have a valid 
-	// normal (e.g., vec2 normal), and avoid crashing the material compilation. The error will still be reported by the compiler up 
-	// to the user, but the compilation will succeed.
-	if (NormalCodeChunk == INDEX_NONE) { NormalCodeChunk = Compiler->VertexNormal(); } 
-
-	int32 TangentCodeChunk = bHasAnisotropy ? Compiler->TransformNormalFromRequestedBasisToWorld(MaterialAttributes.CompileWithDefault(Compiler, FMaterialAttributeDefinitionMap::GetID(MP_Tangent))) : INDEX_NONE;
-	const FSubstrateRegisteredSharedLocalBasis NewRegisteredSharedLocalBasis = SubstrateCompilationInfoCreateSharedLocalBasis(Compiler, NormalCodeChunk, TangentCodeChunk);
-	const FString BasisIndexMacro = Compiler->GetSubstrateSharedLocalBasisIndexMacro(NewRegisteredSharedLocalBasis);
-
-	const bool bHasCoatNormal = IsCustomMaterialAttributeInputConnected(Cached, ClearCoatBottomNormalGuid);
-	// Clear coat normal basis
-	int32 ClearCoat_NormalCodeChunk = INDEX_NONE;
-	int32 ClearCoat_TangentCodeChunk = INDEX_NONE;
-	FString ClearCoat_BasisIndexMacro;
-	FSubstrateRegisteredSharedLocalBasis ClearCoat_NewRegisteredSharedLocalBasis;
-	if (bHasCoatNormal)
-	{
-		ClearCoat_NormalCodeChunk = MaterialAttributes.CompileWithDefault(Compiler, ClearCoatBottomNormalGuid);
-		ClearCoat_TangentCodeChunk = TangentCodeChunk;
-		ClearCoat_NewRegisteredSharedLocalBasis = SubstrateCompilationInfoCreateSharedLocalBasis(Compiler, ClearCoat_NormalCodeChunk, ClearCoat_TangentCodeChunk);
-		ClearCoat_BasisIndexMacro = Compiler->GetSubstrateSharedLocalBasisIndexMacro(ClearCoat_NewRegisteredSharedLocalBasis);
-	}
-	else
-	{
-		ClearCoat_NormalCodeChunk = NormalCodeChunk;
-		ClearCoat_TangentCodeChunk = TangentCodeChunk;
-		ClearCoat_NewRegisteredSharedLocalBasis = NewRegisteredSharedLocalBasis;
-		ClearCoat_BasisIndexMacro = BasisIndexMacro;
-	}
-
-	// Custom tangent. No need to register it as a local basis, as it is only used for eye shading internal conversion
-	int32 CustomTangent_TangentCodeChunk = INDEX_NONE;
-	const bool bHasCustomTangent = IsCustomMaterialAttributeInputConnected(Cached, CustomEyeTangentGuid);
-	if (bHasCustomTangent)
-	{
-		// Legacy code doesn't do tangent <-> world basis conversion on tangent output, when provided.
-		CustomTangent_TangentCodeChunk = MaterialAttributes.CompileWithDefault(Compiler, CustomEyeTangentGuid); // CompileWithDefaultNormalWS(Compiler, CustomTangent, false /*bConvertToRequestedSpace*/);
-		if (CustomTangent_TangentCodeChunk == INDEX_NONE)
-		{
-			// Nothing is plug in from the linked input, so specify world space normal the BSDF node expects.
-			CustomTangent_TangentCodeChunk = Compiler->VertexNormal();
-		}
-	}
-	else
-	{
-		CustomTangent_TangentCodeChunk = NormalCodeChunk;
-	}
-
-	// SSS profile
-	// Need to handle this by looking at the material instead of the node?
-	int32 SSSProfileCodeChunk = INDEX_NONE;
-	const bool bHasSSS = HasSSS();
-	if (bHasSSS)
-	{
-		SSSProfileCodeChunk = Compiler->ForceCast(Compiler->ScalarParameter(GetSubsurfaceProfileParameterName(), 1.0f), MCT_Float1);
-	}
-	SSSProfileCodeChunk = SSSProfileCodeChunk != INDEX_NONE ? SSSProfileCodeChunk : Compiler->Constant(0.0f);
-
-	FSubstrateOperator& SubstrateOperator = Compiler->SubstrateCompilationGetOperator(Compiler->SubstrateTreeStackGetPathUniqueId());
-	SubstrateOperator.BSDFRegisteredSharedLocalBasis = NewRegisteredSharedLocalBasis;
-
-	// Opacity
-	int32 OpacityCodeChunk = INDEX_NONE;
-	if (!Compiler->SubstrateSkipsOpacityEvaluation())
-	{
-		// We evaluate opacity only for shading models and blending mode requiring it.
-		// For instance, a translucent shader reading depth for soft fading should no evaluate opacity when an instance forces an opaque mode.
-		OpacityCodeChunk = MaterialAttributes.CompileWithDefault(Compiler, FMaterialAttributeDefinitionMap::GetID(MP_Opacity));
-	}
-	else
-	{
-		OpacityCodeChunk = Compiler->Constant(1.0f);
-	}
-
-	// Transmittance Color
-	const bool bHasTransmittanceColor = IsCustomMaterialAttributeInputConnected(Cached, TransmittanceColorGuid);
-	int32 TransmittanceColorChunk = INDEX_NONE;
-	if (bHasTransmittanceColor)
-	{
-		TransmittanceColorChunk = MaterialAttributes.CompileWithDefault(Compiler, TransmittanceColorGuid);
-	}
-	else
-	{
-		TransmittanceColorChunk = Compiler->Constant3(0.5f, 0.5f, 0.5f);
-	}
-	// Thin Translucent Surface Coverage
-	const bool bHasThinTranslucentSurfaceCoverage = IsCustomMaterialAttributeInputConnected(Cached, ThinTranslucentSurfaceCoverageGuid);
-	int32 ThinTranslucentSurfaceCoverageChunk = INDEX_NONE;
-	if (bHasThinTranslucentSurfaceCoverage)
-	{
-		ThinTranslucentSurfaceCoverageChunk = MaterialAttributes.CompileWithDefault(Compiler, ThinTranslucentSurfaceCoverageGuid);
-	}
-	else
-	{
-		ThinTranslucentSurfaceCoverageChunk = Compiler->Constant(1.0f);
-	}
-
-	int32 ShadingModelCodeChunk = MaterialAttributes.CompileWithDefault(Compiler, FMaterialAttributeDefinitionMap::GetID(MP_ShadingModel));
-	const bool bHasShadingModelExpression = IsMaterialAttributeInputConnected(Cached, MP_ShadingModel) 
-											|| ShadingModelOverride == MSM_FromMaterialExpression; // In this case, rely on the default compilation to return DefaultLit.
-	if (!bHasShadingModelExpression)
-	{
-		ShadingModelCodeChunk = Compiler->Constant(float(ShadingModelOverride));
-	}
-	int32 ShadingModelCount = Compiler->GetMaterialShadingModels().CountShadingModels();
-	const bool bHasDynamicShadingModels = ShadingModelCount > 1;
-	int32 OutputCodeChunk = Compiler->SubstrateConversionFromLegacy(
-		bHasDynamicShadingModels,
-		// Metalness workflow
-		MaterialAttributes.CompileWithDefault(Compiler, FMaterialAttributeDefinitionMap::GetID(MP_BaseColor)),
-		MaterialAttributes.CompileWithDefault(Compiler, FMaterialAttributeDefinitionMap::GetID(MP_Specular)),
-		MaterialAttributes.CompileWithDefault(Compiler, FMaterialAttributeDefinitionMap::GetID(MP_Metallic)),
-		// Roughness
-		MaterialAttributes.CompileWithDefault(Compiler, FMaterialAttributeDefinitionMap::GetID(MP_Roughness)),
-		MaterialAttributes.CompileWithDefault(Compiler, FMaterialAttributeDefinitionMap::GetID(MP_Anisotropy)),
-		// SSS
-		MaterialAttributes.CompileWithDefault(Compiler, FMaterialAttributeDefinitionMap::GetID(MP_SubsurfaceColor)),
-		SSSProfileCodeChunk,
-		// Clear Coat / Custom
-		MaterialAttributes.CompileWithDefault(Compiler, FMaterialAttributeDefinitionMap::GetID(MP_CustomData0)),// Clear coat
-		MaterialAttributes.CompileWithDefault(Compiler, FMaterialAttributeDefinitionMap::GetID(MP_CustomData1)),// Clear coat roughness
-		// Misc
-		MaterialAttributes.CompileWithDefault(Compiler, FMaterialAttributeDefinitionMap::GetID(MP_EmissiveColor)),
-		OpacityCodeChunk,
-		TransmittanceColorChunk,
-		ThinTranslucentSurfaceCoverageChunk,
-		// Water
-		CompileWithDefaultFloat3(Compiler, WaterScatteringCoefficients, 0.0f, 0.0f, 0.0f),
-		CompileWithDefaultFloat3(Compiler, WaterAbsorptionCoefficients, 0.0f, 0.0f, 0.0f),
-		CompileWithDefaultFloat1(Compiler, WaterPhaseG, 0.0f),
-		CompileWithDefaultFloat3(Compiler, ColorScaleBehindWater, 1.0f, 1.0f, 1.0f),
-		// Shading model
-		ShadingModelCodeChunk,
-		NormalCodeChunk,
-		TangentCodeChunk,
-		BasisIndexMacro,
-		ClearCoat_NormalCodeChunk,
-		ClearCoat_TangentCodeChunk,
-		ClearCoat_BasisIndexMacro,
-		CustomTangent_TangentCodeChunk,
-		!SubstrateOperator.bUseParameterBlending || (SubstrateOperator.bUseParameterBlending && SubstrateOperator.bRootOfParameterBlendingSubTree) ? &SubstrateOperator : nullptr);
-
-	return OutputCodeChunk;
-}
-
-void UMaterialExpressionSubstrateConvertMaterialAttributes::PostEditChangeProperty(FPropertyChangedEvent& PropertyChangedEvent)
-{
-	if (GraphNode && PropertyChangedEvent.Property != nullptr)
-	{
-		GraphNode->ReconstructNode();
-	}
-	Super::PostEditChangeProperty(PropertyChangedEvent);
-}
-
-void UMaterialExpressionSubstrateConvertMaterialAttributes::GetCaption(TArray<FString>& OutCaptions) const
-{
-	OutCaptions.Add(TEXT("Substrate Convert Material Attributes"));
-}
-
-uint32 UMaterialExpressionSubstrateConvertMaterialAttributes::GetOutputType(int32 OutputIndex)
-{
-	switch (OutputIndex)
-	{
-		case 0 : return MCT_Substrate;
-	}
-	check(false);
-	return MCT_Float1;
-}
-
-uint32 UMaterialExpressionSubstrateConvertMaterialAttributes::GetInputType(int32 InputIndex)
-{
-	if (InputIndex == 0)	  return MCT_MaterialAttributes; // MaterialAttributes
-	else if (InputIndex == 1) return MCT_Float3; // WaterScatteringCoefficients
-	else if (InputIndex == 2) return MCT_Float3; // WaterAbsorptionCoefficients
-	else if (InputIndex == 3) return MCT_Float1; // WaterPhaseG
-	else if (InputIndex == 4) return MCT_Float3; // ColorScaleBehindWater
-	else if (InputIndex == 5) return MCT_ShadingModel; // ShadingModelOverride (as it uses 'ShowAsInputPin' metadata)
-	
-	check(false);
-	return MCT_Float1;
-}
-
-FName UMaterialExpressionSubstrateConvertMaterialAttributes::GetInputName(int32 InputIndex) const
-{
-	if (InputIndex == 0)		return TEXT("Attributes");
-	else if (InputIndex == 1)	return TEXT("Water Scattering Coefficients (Water)");
-	else if (InputIndex == 2)	return TEXT("Water Absorption Coefficients (Water)");
-	else if (InputIndex == 3)	return TEXT("Water Phase G (Water)");
-	else if (InputIndex == 4)	return TEXT("Color Scale BehindWater (Water)");
-	else if (InputIndex == 5)	return TEXT("Shading Model From Expression");
-	return NAME_None;
-}
-
-void UMaterialExpressionSubstrateConvertMaterialAttributes::GetConnectorToolTip(int32 InputIndex, int32 OutputIndex, TArray<FString>& OutToolTip)
-{
-	switch (OutputIndex)
-	{
-		case 0: OutToolTip.Add(TEXT("TT Out Substrate Data")); break;
-	}
-	Super::GetConnectorToolTip(InputIndex, INDEX_NONE, OutToolTip);
-}
-
-bool UMaterialExpressionSubstrateConvertMaterialAttributes::IsResultSubstrateMaterial(int32 OutputIndex)
-{
-	return OutputIndex == 0;
-}
-
-bool UMaterialExpressionSubstrateConvertMaterialAttributes::IsResultMaterialAttributes(int32 OutputIndex)
-{
-	return false;
-}
-
-void UMaterialExpressionSubstrateConvertMaterialAttributes::GatherSubstrateMaterialInfo(FSubstrateMaterialInfo& SubstrateMaterialInfo, int32 OutputIndex)
-{
-	const uint64 Cached = GetConnectedMaterialAttributesInputs(Material);
-
-	if (IsMaterialAttributeInputConnected(Cached, MP_BaseColor)) 		{ SubstrateMaterialInfo.AddPropertyConnected(MP_BaseColor); }
-	if (IsMaterialAttributeInputConnected(Cached, MP_Metallic)) 		{ SubstrateMaterialInfo.AddPropertyConnected(MP_Metallic); }
-	if (IsMaterialAttributeInputConnected(Cached, MP_Specular)) 		{ SubstrateMaterialInfo.AddPropertyConnected(MP_Specular); }
-	if (IsMaterialAttributeInputConnected(Cached, MP_Roughness)) 		{ SubstrateMaterialInfo.AddPropertyConnected(MP_Roughness); }
-	if (IsMaterialAttributeInputConnected(Cached, MP_Anisotropy)) 		{ SubstrateMaterialInfo.AddPropertyConnected(MP_Anisotropy); }
-	if (IsMaterialAttributeInputConnected(Cached, MP_EmissiveColor)) 	{ SubstrateMaterialInfo.AddPropertyConnected(MP_EmissiveColor); }
-	if (IsMaterialAttributeInputConnected(Cached, MP_Normal)) 			{ SubstrateMaterialInfo.AddPropertyConnected(MP_Normal); }
-	if (IsMaterialAttributeInputConnected(Cached, MP_Tangent)) 			{ SubstrateMaterialInfo.AddPropertyConnected(MP_Tangent); }
-	if (IsMaterialAttributeInputConnected(Cached, MP_SubsurfaceColor)) 	{ SubstrateMaterialInfo.AddPropertyConnected(MP_SubsurfaceColor); }
-	if (IsMaterialAttributeInputConnected(Cached, MP_CustomData0)) 		{ SubstrateMaterialInfo.AddPropertyConnected(MP_CustomData0); }
-	if (IsMaterialAttributeInputConnected(Cached, MP_CustomData1)) 		{ SubstrateMaterialInfo.AddPropertyConnected(MP_CustomData1); }
-	if (IsMaterialAttributeInputConnected(Cached, MP_Opacity)) 			{ SubstrateMaterialInfo.AddPropertyConnected(MP_Opacity); }
-
-	if (IsMaterialAttributeInputConnected(Cached, MP_ShadingModel) || ShadingModelOverride == MSM_FromMaterialExpression)
-	{
-		SubstrateMaterialInfo.AddPropertyConnected(MP_ShadingModel);
-
-		// If the ShadingModel pin is plugged in, we must use a shading model from expression path.
-		SubstrateMaterialInfo.SetShadingModelFromExpression(true);
-	}
-	else
-	{
-		// If the ShadingModel pin is NOT plugged in, we simply use the shading model selected on the root node drop box.
-		if (ShadingModelOverride == MSM_Unlit)					{ SubstrateMaterialInfo.AddShadingModel(ESubstrateShadingModel::SSM_Unlit); }
-		if (ShadingModelOverride == MSM_DefaultLit)				{ SubstrateMaterialInfo.AddShadingModel(ESubstrateShadingModel::SSM_DefaultLit); }
-		if (ShadingModelOverride == MSM_Subsurface)				{ SubstrateMaterialInfo.AddShadingModel(ESubstrateShadingModel::SSM_SubsurfaceWrap); }
-		if (ShadingModelOverride == MSM_PreintegratedSkin)		{ SubstrateMaterialInfo.AddShadingModel(ESubstrateShadingModel::SSM_SubsurfaceWrap); }
-		if (ShadingModelOverride == MSM_ClearCoat)				{ SubstrateMaterialInfo.AddShadingModel(ESubstrateShadingModel::SSM_ClearCoat); }
-		if (ShadingModelOverride == MSM_SubsurfaceProfile)		{ SubstrateMaterialInfo.AddShadingModel(ESubstrateShadingModel::SSM_SubsurfaceProfile); }
-		if (ShadingModelOverride == MSM_TwoSidedFoliage)		{ SubstrateMaterialInfo.AddShadingModel(ESubstrateShadingModel::SSM_SubsurfaceThinTwoSided); }
-		if (ShadingModelOverride == MSM_Hair)					{ SubstrateMaterialInfo.AddShadingModel(ESubstrateShadingModel::SSM_Hair); }
-		if (ShadingModelOverride == MSM_Cloth)					{ SubstrateMaterialInfo.AddShadingModel(ESubstrateShadingModel::SSM_Cloth); }
-		if (ShadingModelOverride == MSM_Eye)					{ SubstrateMaterialInfo.AddShadingModel(ESubstrateShadingModel::SSM_Eye); }
-		if (ShadingModelOverride == MSM_SingleLayerWater)		{ SubstrateMaterialInfo.AddShadingModel(ESubstrateShadingModel::SSM_SingleLayerWater); }
-		if (ShadingModelOverride == MSM_ThinTranslucent)		{ SubstrateMaterialInfo.AddShadingModel(ESubstrateShadingModel::SSM_ThinTranslucent); }
-	}
-
-	if (SubsurfaceProfile)
-	{
-		SubstrateMaterialInfo.AddSubsurfaceProfile(SubsurfaceProfile);
-	}
-	SubstrateMaterialInfo.AddGuid(MaterialExpressionGuid);
-}
-
-FSubstrateOperator* UMaterialExpressionSubstrateConvertMaterialAttributes::SubstrateGenerateMaterialTopologyTree(class FMaterialCompiler* Compiler, class UMaterialExpression* Parent, int32 OutputIndex)
-{
-	const uint64 Cached = GetConnectedMaterialAttributesInputs(Material);
-
-	// Note Thickness has no meaning/usage in the context of SubstrateLegacyConversionNode
-	int32 ThicknessIndex = Compiler->SubstrateThicknessStackGetThicknessIndex();
-
-	const bool bHasAnisotropy = IsMaterialAttributeInputConnected(Cached, MP_Anisotropy);
-
-	auto AddDefaultWorstCase = [&](bool bSSS, bool bFuzz)
-	{
-
-		FSubstrateOperator& SlabOperator = Compiler->SubstrateCompilationRegisterOperator(SUBSTRATE_OPERATOR_BSDF_LEGACY, Compiler->SubstrateTreeStackGetPathUniqueId(), this, Parent, Compiler->SubstrateTreeStackGetParentPathUniqueId());
-		SlabOperator.BSDFType = SUBSTRATE_BSDF_TYPE_SLAB;
-		SlabOperator.bBSDFHasSSS = bSSS;
-		SlabOperator.bBSDFHasMFPPluggedIn = bSSS;
-		SlabOperator.bBSDFHasFuzz = bFuzz;
-		SlabOperator.bBSDFHasAnisotropy = bHasAnisotropy;
-		SlabOperator.ThicknessIndex = ThicknessIndex;
-
-		return &SlabOperator;
-	};
-
-	// Get the shading models resulting from the UMaterial::RebuildShadingModelField().
-	FMaterialShadingModelField ShadingModels = Compiler->GetMaterialShadingModels();
-
-	// Logic about shading models and complexity should match UMaterialExpressionSubstrateConvertMaterialAttributes::Compile.
-	const bool bHasShadingModelFromExpression = IsMaterialAttributeInputConnected(Cached, MP_ShadingModel) || ShadingModelOverride == MSM_FromMaterialExpression; // We keep HasShadingModelFromExpression in case all shading models cannot be safely recovered from material functions.
-	if ((ShadingModels.CountShadingModels() > 1) || bHasShadingModelFromExpression) 
-	{
-		return AddDefaultWorstCase(true, true);
-	}
-	else
-	{
-		check(ShadingModels.CountShadingModels() == 1);
-
-		if (ShadingModels.HasShadingModel(MSM_Unlit))
-		{
-			FSubstrateOperator& Operator = Compiler->SubstrateCompilationRegisterOperator(SUBSTRATE_OPERATOR_BSDF_LEGACY, Compiler->SubstrateTreeStackGetPathUniqueId(), this, Parent, Compiler->SubstrateTreeStackGetParentPathUniqueId());
-			Operator.BSDFType = SUBSTRATE_BSDF_TYPE_UNLIT;
-			Operator.ThicknessIndex = ThicknessIndex;
-			return &Operator;
-		}
-		else if (ShadingModels.HasShadingModel(MSM_DefaultLit))
-		{
-			return AddDefaultWorstCase(false, false);
-		}
-		else if (ShadingModels.HasShadingModel(MSM_ThinTranslucent))
-		{
-			return AddDefaultWorstCase(false, false);
-		}
-		else if (ShadingModels.HasShadingModel(MSM_SubsurfaceProfile))
-		{
-			return AddDefaultWorstCase(true, false);
-		}
-		else if (ShadingModels.HasShadingModel(MSM_Subsurface))
-		{
-			return AddDefaultWorstCase(true, false);
-		}
-		else if (ShadingModels.HasShadingModel(MSM_TwoSidedFoliage))
-		{
-			return AddDefaultWorstCase(true, false);
-		}
-		else if (ShadingModels.HasShadingModel(MSM_PreintegratedSkin))
-		{
-			return AddDefaultWorstCase(true, false);
-		}
-		else if (ShadingModels.HasShadingModel(MSM_Cloth))
-		{
-			return AddDefaultWorstCase(false, true);
-		}
-		else if (ShadingModels.HasShadingModel(MSM_ClearCoat))
-		{
-			FSubstrateOperator& Operator = Compiler->SubstrateCompilationRegisterOperator(SUBSTRATE_OPERATOR_BSDF_LEGACY, Compiler->SubstrateTreeStackGetPathUniqueId(), this, Parent, Compiler->SubstrateTreeStackGetParentPathUniqueId());
-			Operator.BSDFType = SUBSTRATE_BSDF_TYPE_SLAB;
-			Operator.ThicknessIndex = ThicknessIndex;
-			Operator.bBSDFHasSecondRoughnessOrSimpleClearCoat = true;
-			Operator.bBSDFHasAnisotropy = bHasAnisotropy;
-			return &Operator;
-		}
-		else if (ShadingModels.HasShadingModel(MSM_Hair))
-		{
-			FSubstrateOperator& Operator = Compiler->SubstrateCompilationRegisterOperator(SUBSTRATE_OPERATOR_BSDF_LEGACY, Compiler->SubstrateTreeStackGetPathUniqueId(), this, Parent, Compiler->SubstrateTreeStackGetParentPathUniqueId());
-			Operator.BSDFType = SUBSTRATE_BSDF_TYPE_HAIR;
-			Operator.ThicknessIndex = ThicknessIndex;
-			return &Operator;
-		}
-		else if (ShadingModels.HasShadingModel(MSM_Eye))
-		{
-			FSubstrateOperator& Operator = Compiler->SubstrateCompilationRegisterOperator(SUBSTRATE_OPERATOR_BSDF_LEGACY, Compiler->SubstrateTreeStackGetPathUniqueId(), this, Parent, Compiler->SubstrateTreeStackGetParentPathUniqueId());
-			Operator.BSDFType = SUBSTRATE_BSDF_TYPE_EYE;
-			Operator.ThicknessIndex = ThicknessIndex;
-			return &Operator;
-		}
-		else if (ShadingModels.HasShadingModel(MSM_SingleLayerWater))
-		{
-			FSubstrateOperator& Operator = Compiler->SubstrateCompilationRegisterOperator(SUBSTRATE_OPERATOR_BSDF_LEGACY, Compiler->SubstrateTreeStackGetPathUniqueId(), this, Parent, Compiler->SubstrateTreeStackGetParentPathUniqueId());
-			Operator.BSDFType = SUBSTRATE_BSDF_TYPE_SINGLELAYERWATER;
-			Operator.ThicknessIndex = ThicknessIndex;
-			Operator.bBSDFWritesEmissive = MaterialAttributes.IsConnected(MP_EmissiveColor);
-			return &Operator;
-		}
-
-		check(false);
-		static FSubstrateOperator DefaultOperatorOnError;
-		return &DefaultOperatorOnError;
-	}
-}
-
-bool UMaterialExpressionSubstrateConvertMaterialAttributes::HasSSS() const
-{
-	return SubsurfaceProfile != nullptr;
-}
-
-#endif // WITH_EDITOR
-
 UMaterialExpressionExecBegin::UMaterialExpressionExecBegin(const FObjectInitializer& ObjectInitializer)
 	: Super(ObjectInitializer)
 {
