Substrate - moved maerial expressions into their own cpp.

diff --git a/Engine/Source/Runtime/Engine/Public/Materials/MaterialExpressionSubstrate.h b/Engine/Source/Runtime/Engine/Public/Materials/MaterialExpressionSubstrate.h
--- a/Engine/Source/Runtime/Engine/Public/Materials/MaterialExpressionSubstrate.h
+++ b/Engine/Source/Runtime/Engine/Public/Materials/MaterialExpressionSubstrate.h
@@ -9,6 +9,20 @@
 
 class FMaterialCompiler;
 struct FSubstrateMaterialComplexity;
+struct FSubstrateOperator;
+
+
+
+///////////////////////////////////////////////////////////////////////////////
+// Functions
+
+#if WITH_EDITOR
+
+void AssignOperatorIndexIfNotNull(int32& NextOperatorPin, FSubstrateOperator* Operator);
+
+void CombineFlagForParameterBlending(FSubstrateOperator& DstOp, FSubstrateOperator* OpA, FSubstrateOperator* OpB = nullptr);
+
+#endif
 
 /**
  * Compile a special blend function for Substrate when blending material attribute
@@ -22,6 +36,7 @@ struct FSubstrateMaterialComplexity;
 extern int32 CompileSubstrateBlendFunction(FMaterialCompiler* Compiler, const int32 A, const int32 B, const int32 Alpha);
 
 
+
 ///////////////////////////////////////////////////////////////////////////////
 // BSDF nodes
 
@@ -891,6 +906,8 @@ class UMaterialExpressionSubstrateConvertMaterialAttributes : public UMaterialEx
 	//~ End UMaterialExpression Interface
 };
 
+
+
 ///////////////////////////////////////////////////////////////////////////////
 // Operator nodes
 
@@ -1047,6 +1064,7 @@ class UMaterialExpressionSubstrateWeight : public UMaterialExpressionSubstrateBS
 };
 
 
+
 ///////////////////////////////////////////////////////////////////////////////
 // Utilities
 
