[shaders] shadermap serialization cleanup
- unify shadermap serialization params into a context object
- deprecate shader code inlining into shader maps; this is no longer needed internally (and should not be used by licensees) and removing it will simplify forthcoming changes related to how shader/shadermap DDC data is structured

diff --git a/Engine/Plugins/Compositing/OpenColorIO/Source/OpenColorIO/Private/OpenColorIOShaderMap.cpp b/Engine/Plugins/Compositing/OpenColorIO/Source/OpenColorIO/Private/OpenColorIOShaderMap.cpp
--- a/Engine/Plugins/Compositing/OpenColorIO/Source/OpenColorIO/Private/OpenColorIOShaderMap.cpp
+++ b/Engine/Plugins/Compositing/OpenColorIO/Source/OpenColorIO/Private/OpenColorIOShaderMap.cpp
@@ -621,12 +621,12 @@ FOpenColorIOShaderMap::~FOpenColorIOShaderMap()
 	AllOpenColorIOShaderMaps.RemoveSwap(this);
 }
 
-bool FOpenColorIOShaderMap::Serialize(FArchive& Ar, bool bInlineShaderResources)
+bool FOpenColorIOShaderMap::Serialize(FArchive& Ar)
 {
 	// Note: This is saved to the DDC, not into packages (except when cooked)
 	// Backwards compatibility therefore will not work based on the version of Ar
 	// Instead, just bump OPENCOLORIO_DERIVEDDATA_VER
-	return Super::Serialize(Ar, bInlineShaderResources, false);
+	return Super::Serialize(Ar, {});
 }
 
 #if WITH_EDITOR
diff --git a/Engine/Plugins/Compositing/OpenColorIO/Source/OpenColorIO/Public/OpenColorIOShared.h b/Engine/Plugins/Compositing/OpenColorIO/Source/OpenColorIO/Public/OpenColorIOShared.h
--- a/Engine/Plugins/Compositing/OpenColorIO/Source/OpenColorIO/Public/OpenColorIOShared.h
+++ b/Engine/Plugins/Compositing/OpenColorIO/Source/OpenColorIO/Public/OpenColorIOShared.h
@@ -229,7 +229,7 @@ public:
 #endif // WITH_EDITOR
 
 	/** Serializes the shader map. */
-	bool Serialize(FArchive& Ar, bool bInlineShaderResources = true);
+	bool Serialize(FArchive& Ar);
 
 #if WITH_EDITOR
 	/** Saves this shader map to the derived data cache. */
diff --git a/Engine/Plugins/FX/Niagara/Source/NiagaraEditor/Private/NiagaraCompilationTasks.cpp b/Engine/Plugins/FX/Niagara/Source/NiagaraEditor/Private/NiagaraCompilationTasks.cpp
--- a/Engine/Plugins/FX/Niagara/Source/NiagaraEditor/Private/NiagaraCompilationTasks.cpp
+++ b/Engine/Plugins/FX/Niagara/Source/NiagaraEditor/Private/NiagaraCompilationTasks.cpp
@@ -275,7 +275,7 @@ namespace NiagaraCompilationTasksImpl
 
 		TArray<uint8> BinaryData;
 		FMemoryWriter Ar(BinaryData, true);
-		ShaderMap->Serialize(Ar);
+		ShaderMap->Serialize(Ar, FShaderMapBase::FSerializationContext{});
 
 		if (!BinaryData.IsEmpty() && !Ar.IsError())
 		{
@@ -294,7 +294,7 @@ namespace NiagaraCompilationTasksImpl
 
 		FMemoryReaderView Ar(SharedBuffer.GetView(), true);
 		FNiagaraShaderMapRef ShaderMap = new FNiagaraShaderMap(FNiagaraShaderMap::WorkerThread);
-		ShaderMap->Serialize(Ar);
+		ShaderMap->Serialize(Ar, FShaderMapBase::FSerializationContext{});
 
 		if (Ar.IsError())
 		{
diff --git a/Engine/Plugins/FX/Niagara/Source/NiagaraShader/Private/NiagaraShader.cpp b/Engine/Plugins/FX/Niagara/Source/NiagaraShader/Private/NiagaraShader.cpp
--- a/Engine/Plugins/FX/Niagara/Source/NiagaraShader/Private/NiagaraShader.cpp
+++ b/Engine/Plugins/FX/Niagara/Source/NiagaraShader/Private/NiagaraShader.cpp
@@ -548,7 +548,7 @@ void FNiagaraShaderMap::LoadFromDerivedDataCache(const FNiagaraShaderScript* Scr
 				FMemoryReader Ar(CachedData, true);
 
 				// Deserialize from the cached data
-				if (InOutShaderMap->Serialize(Ar))
+				if (InOutShaderMap->Serialize(Ar, FShaderMapBase::FSerializationContext{}))
 				{
 					check(InOutShaderMap->GetShaderMapId() == ShaderMapId);
 
@@ -578,7 +578,7 @@ void FNiagaraShaderMap::SaveToDerivedDataCache(const FNiagaraShaderScript* Scrip
 	COOK_STAT(auto Timer = NiagaraShaderCookStats::UsageStats.TimeSyncWork());
 	TArray<uint8> SaveData;
 	FMemoryWriter Ar(SaveData, true);
-	Serialize(Ar);
+	Serialize(Ar, FShaderMapBase::FSerializationContext{});
 
 	GetDerivedDataCacheRef().Put(*GetNiagaraShaderMapKeyString(GetContent()->ShaderMapId, GetShaderPlatform()), SaveData, Script ? Script->GetFriendlyName() : TEXT(""));
 	COOK_STAT(Timer.AddMiss(SaveData.Num()));
@@ -1057,12 +1057,12 @@ FNiagaraShaderMap::~FNiagaraShaderMap()
 	check(!bRegistered);
 }
 
-bool FNiagaraShaderMap::Serialize(FArchive& Ar, bool bInlineShaderResources, bool bLoadedByCookedMaterial)
+bool FNiagaraShaderMap::Serialize(FArchive& Ar, const FShaderMapBase::FSerializationContext& Ctx)
 {
 	// Note: This is saved to the DDC, not into packages (except when cooked)
 	// Backwards compatibility therefore will not work based on the version of Ar
 	// Instead, just bump NIAGARASHADERMAP_DERIVEDDATA_VER
-	return Super::Serialize(Ar, bInlineShaderResources, bLoadedByCookedMaterial);
+	return Super::Serialize(Ar, Ctx);
 }
 
 #if WITH_EDITOR
diff --git a/Engine/Plugins/FX/Niagara/Source/NiagaraShader/Private/NiagaraShared.cpp b/Engine/Plugins/FX/Niagara/Source/NiagaraShader/Private/NiagaraShared.cpp
--- a/Engine/Plugins/FX/Niagara/Source/NiagaraShader/Private/NiagaraShared.cpp
+++ b/Engine/Plugins/FX/Niagara/Source/NiagaraShader/Private/NiagaraShared.cpp
@@ -344,7 +344,7 @@ void FNiagaraShaderScript::SerializeShaderMap(FArchive& Ar)
 				{
 					GameThreadShaderMap->AssociateWithAsset(BaseVMScript->GetOutermost()->GetFName());
 				}
-				GameThreadShaderMap->Serialize(Ar);
+				GameThreadShaderMap->Serialize(Ar, FShaderMapBase::FSerializationContext{});
 			}
 			//else if (GameThreadShaderMap != nullptr && !GameThreadShaderMap->CompiledSuccessfully())
 			//{
@@ -361,7 +361,7 @@ void FNiagaraShaderScript::SerializeShaderMap(FArchive& Ar)
 			if (bValid)
 			{
 				FNiagaraShaderMapRef LoadedShaderMap = new FNiagaraShaderMap();
-				bool bLoaded = LoadedShaderMap->Serialize(Ar, true, true);
+				bool bLoaded = LoadedShaderMap->Serialize(Ar, FShaderMapBase::FSerializationContext{ /*bLoadedByCookedMaterial*/ true });
 
 				// Toss the loaded shader data if this is a server only instance
 				//@todo - don't cook it in the first place
diff --git a/Engine/Plugins/FX/Niagara/Source/NiagaraShader/Public/NiagaraShared.h b/Engine/Plugins/FX/Niagara/Source/NiagaraShader/Public/NiagaraShared.h
--- a/Engine/Plugins/FX/Niagara/Source/NiagaraShader/Public/NiagaraShared.h
+++ b/Engine/Plugins/FX/Niagara/Source/NiagaraShader/Public/NiagaraShared.h
@@ -599,7 +599,13 @@ public:
 #endif // WITH_EDITOR
 
 	/** Serializes the shader map. */
-	NIAGARASHADER_API bool Serialize(FArchive& Ar, bool bInlineShaderResources = true, bool bLoadedByCookedMaterial = false);
+	UE_DEPRECATED(5.5, "Use overload taking a FShaderMapBase::FSerializationContext")
+	NIAGARASHADER_API bool Serialize(FArchive& Ar, bool bInlineShaderResources = true, bool bLoadedByCookedMaterial = false)
+	{
+		return Serialize(Ar, FShaderMapBase::FSerializationContext{ bLoadedByCookedMaterial });
+	}
+
+	NIAGARASHADER_API bool Serialize(FArchive& Ar, const FShaderMapBase::FSerializationContext& Ctx);
 
 #if WITH_EDITOR
 	/** Saves this shader map to the derived data cache. */
diff --git a/Engine/Plugins/Runtime/ComputeFramework/Source/ComputeFramework/Private/ComputeFramework/ComputeKernelShared.h b/Engine/Plugins/Runtime/ComputeFramework/Source/ComputeFramework/Private/ComputeFramework/ComputeKernelShared.h
--- a/Engine/Plugins/Runtime/ComputeFramework/Source/ComputeFramework/Private/ComputeFramework/ComputeKernelShared.h
+++ b/Engine/Plugins/Runtime/ComputeFramework/Source/ComputeFramework/Private/ComputeFramework/ComputeKernelShared.h
@@ -211,7 +211,7 @@ public:
 #endif
 
 	/** Serializes the shader map. */
-	bool Serialize(FArchive& Ar, bool bInlineShaderResources = true);
+	bool Serialize(FArchive& Ar);
 
 #if WITH_EDITOR
 	/** Saves this shader map to the derived data cache. */
diff --git a/Engine/Plugins/Runtime/ComputeFramework/Source/ComputeFramework/Private/ComputeKernelShaderMap.cpp b/Engine/Plugins/Runtime/ComputeFramework/Source/ComputeFramework/Private/ComputeKernelShaderMap.cpp
--- a/Engine/Plugins/Runtime/ComputeFramework/Source/ComputeFramework/Private/ComputeKernelShaderMap.cpp
+++ b/Engine/Plugins/Runtime/ComputeFramework/Source/ComputeFramework/Private/ComputeKernelShaderMap.cpp
@@ -666,12 +666,12 @@ FComputeKernelShaderMap::~FComputeKernelShaderMap()
 	AllComputeKernelShaderMaps.RemoveSwap(this);
 }
 
-bool FComputeKernelShaderMap::Serialize(FArchive& Ar, bool bInlineShaderResources)
+bool FComputeKernelShaderMap::Serialize(FArchive& Ar)
 {
 	// Note: This is saved to the DDC, not into packages (except when cooked)
 	// Backwards compatibility therefore will not work based on the version of Ar
 	// Instead, just bump COMPUTEKERNEL_DERIVEDDATA_VER
-	return Super::Serialize(Ar, bInlineShaderResources, false);
+	return Super::Serialize(Ar, {});
 }
 
 #if WITH_EDITOR
diff --git a/Engine/Source/Runtime/Engine/Private/Materials/MaterialShader.cpp b/Engine/Source/Runtime/Engine/Private/Materials/MaterialShader.cpp
--- a/Engine/Source/Runtime/Engine/Private/Materials/MaterialShader.cpp
+++ b/Engine/Source/Runtime/Engine/Private/Materials/MaterialShader.cpp
@@ -1867,7 +1867,7 @@ TSharedRef<FMaterialShaderMap::FAsyncLoadContext> FMaterialShaderMap::BeginLoadF
 				FMemoryReaderView Ar(CachedData, /*bIsPersistent*/ true);
 
 				// Deserialize from the cached data
-				ShaderMap->Serialize(Ar);
+				ShaderMap->Serialize(Ar, FShaderMapBase::FSerializationContext{});
 
 				check(Material != nullptr);
 
@@ -2004,7 +2004,7 @@ void FMaterialShaderMap::SaveToDerivedDataCache(const FMaterialShaderParameters&
 	COOK_STAT(auto Timer = MaterialShaderCookStats::UsageStats.TimeSyncWork());
 	TArray64<uint8> SaveData;
 	FMemoryWriter64 Ar(SaveData, true);
-	Serialize(Ar);
+	Serialize(Ar, FShaderMapBase::FSerializationContext{});
 
 	TRACE_COUNTER_ADD(Shaders_FMaterialShaderMapDDCBytesSent, SaveData.Num());
 	COOK_STAT(Timer.AddMiss(SaveData.Num()));
@@ -2063,7 +2063,7 @@ void FMaterialShaderMap::SaveForRemoteRecompile(FArchive& Ar, const TMap<FString
 			{
 				uint8 bIsValid = 1;
 				Ar << bIsValid;
-				ShaderMap->Serialize(Ar, false);
+				ShaderMap->Serialize(Ar, FShaderMapBase::FSerializationContext{});
 			}
 			else
 			{
@@ -2115,7 +2115,7 @@ void FMaterialShaderMap::LoadForRemoteRecompile(FArchive& Ar, EShaderPlatform Sh
 				TRefCountPtr<FMaterialShaderMap> ShaderMap = new FMaterialShaderMap();
 
 				// serialize the id and the material shader map
-				ShaderMap->Serialize(Ar, false);
+				ShaderMap->Serialize(Ar, FShaderMapBase::FSerializationContext{});
 
 				LoadedShaderMapsDictionary.Add(ShaderMap->GetShaderMapId(), ShaderMap);
 				MaterialShaderMapData.LoadedShaderMapsIds.Add(ShaderMap->GetShaderMapId());
@@ -3604,14 +3604,14 @@ FMaterialShaderMap* FMaterialShaderMap::GetFinalizedClone() const
 }
 #endif // WITH_EDITOR
 
-bool FMaterialShaderMap::Serialize(FArchive& Ar, bool bInlineShaderResources, bool bLoadedByCookedMaterial, bool bInlineShaderCode, const FName& SerializingAsset)
+bool FMaterialShaderMap::Serialize(FArchive& Ar, const FShaderMapBase::FSerializationContext& Ctx)
 {
 	SCOPED_LOADTIMER(FMaterialShaderMap_Serialize);
 	// Note: This is saved to the DDC, not into packages (except when cooked)
 	// Backwards compatibility therefore will not work based on the version of Ar
 	// Instead, just bump MATERIALSHADERMAP_DERIVEDDATA_VER
-	ShaderMapId.Serialize(Ar, bLoadedByCookedMaterial);
-	bool bSerialized = Super::Serialize(Ar, bInlineShaderResources, bLoadedByCookedMaterial, bInlineShaderCode, SerializingAsset);
+	ShaderMapId.Serialize(Ar, Ctx.bLoadedByCookedMaterial);
+	bool bSerialized = Super::Serialize(Ar, Ctx);
 #if STATS
 	// This is an unsavory hack to repair STAT_Shaders_NumShadersLoaded not being calculated right in the superclass because the Content class isn't allowed to have virtual functions atm.
 	// A better way is tracked in UE-127112
diff --git a/Engine/Source/Runtime/Engine/Private/Materials/MaterialShared.cpp b/Engine/Source/Runtime/Engine/Private/Materials/MaterialShared.cpp
--- a/Engine/Source/Runtime/Engine/Private/Materials/MaterialShared.cpp
+++ b/Engine/Source/Runtime/Engine/Private/Materials/MaterialShared.cpp
@@ -1540,17 +1540,7 @@ void FMaterial::SerializeInlineShaderMap(FArchive& Ar, const FName& SerializingA
 
 			if (bValid)
 			{
-				// do not put shader code of certain materials into the library
-				bool bEnableInliningWorkaround = false;
-
-				checkf(GConfig, TEXT("We expect GConfig to exist at this point"));
-				FString SettingName(TEXT("bEnableInliningWorkaround_"));
-				SettingName += Ar.CookingTarget()->IniPlatformName();
-
-				GConfig->GetBool(TEXT("ShaderCodeLibrary"), *SettingName, bEnableInliningWorkaround, GEngineIni);
-
-				bool bInlineShaderCode = bEnableInliningWorkaround && ShouldInlineShaderCode();
-				GameThreadShaderMap->Serialize(Ar, true, false, bInlineShaderCode);
+				GameThreadShaderMap->Serialize(Ar, FShaderMapBase::FSerializationContext{});
 			}
 			else
 			{
@@ -1572,7 +1562,7 @@ void FMaterial::SerializeInlineShaderMap(FArchive& Ar, const FName& SerializingA
 			if (bValid)
 			{
 				TRefCountPtr<FMaterialShaderMap> LoadedShaderMap = new FMaterialShaderMap();
-				if (LoadedShaderMap->Serialize(Ar, true, bCooked && Ar.IsLoading(), false, SerializingAsset))
+				if (LoadedShaderMap->Serialize(Ar, FShaderMapBase::FSerializationContext{ bCooked && Ar.IsLoading(), SerializingAsset }))
 				{
 					GameThreadShaderMap = MoveTemp(LoadedShaderMap);
 					GameThreadShaderMap->GetResource()->SetOwnerName(GetOwnerFName());
@@ -2346,49 +2336,6 @@ FName FMaterialResource::GetAssetPath() const
 	return OutermostName;
 }
 
-bool FMaterialResource::ShouldInlineShaderCode() const
-{
-	if (IsSpecialEngineMaterial() || IsDefaultMaterial())
-	{
-		UE_LOG(LogMaterial, Display, TEXT("%s: shader code is inlined because the workaround is enabled and it's a special or default material"), *GetFriendlyName());
-		return true;
-	}
-
-	// Check the material name against those configured to be enabled.
-	// For the cooker commandlet, this check could be cached, but due to concerns of that cache possibly getting stale for edge cases like COTF and general work-aroundness of this
-	// function, let's check the configs every time.
-
-	FString OutermostName;
-	if (MaterialInstance)
-	{
-		OutermostName = MaterialInstance->GetOutermost()->GetName();
-	}
-	else if (Material)
-	{
-		OutermostName = Material->GetOutermost()->GetName();
-	}
-
-	bool bNeedsToBeInlined = false;
-	const FConfigSection* ShaderLibrarySec = GConfig->GetSection(TEXT("ShaderCodeLibrary"), false, GEngineIni);
-	if (ShaderLibrarySec)
-	{
-		TArray<FString> ConfiguredMaterials;
-		ShaderLibrarySec->MultiFind(TEXT("MaterialToInline"), ConfiguredMaterials);
-
-		for(const FString& ConfiguredMaterial : ConfiguredMaterials)
-		{
-			if (ConfiguredMaterial == OutermostName)
-			{
-				bNeedsToBeInlined = true;
-				break;
-			}
-		}
-	}
-
-	UE_LOG(LogMaterial, Display, TEXT("%s (package %s): shader code is %s to be inlined as a workaround"), *GetFriendlyName(), *OutermostName, bNeedsToBeInlined ? TEXT("configured") : TEXT("NOT configured"));
-	return bNeedsToBeInlined;
-}
-
 bool FMaterialResource::IsUsingControlFlow() const
 {
 	if (Material)
diff --git a/Engine/Source/Runtime/Engine/Public/MaterialShared.h b/Engine/Source/Runtime/Engine/Public/MaterialShared.h
--- a/Engine/Source/Runtime/Engine/Public/MaterialShared.h
+++ b/Engine/Source/Runtime/Engine/Public/MaterialShared.h
@@ -1539,7 +1539,14 @@ public:
 	ENGINE_API void Release();
 
 	/** Serializes the shader map. */
-	bool Serialize(FArchive& Ar, bool bInlineShaderResources=true, bool bLoadedByCookedMaterial=false, bool bInlineShaderCode=false, const FName& SerializingAsset = NAME_None);
+	UE_DEPRECATED(5.5, "Inlining of shader code into shadermaps is no longer supported. If you rely on this please reach out to the UE rendering team.")
+	// note: FSerializationContext param in the below overload can be defaulted once this is removed, to avoid a bunch of explicit instantiations of FSerializationContext in various callsites
+	bool Serialize(FArchive& Ar, bool bInlineShaderResources = true, bool bLoadedByCookedMaterial = false, bool bInlineShaderCode = false, const FName& SerializingAsset = NAME_None)
+	{
+		return Serialize(Ar, { bLoadedByCookedMaterial, SerializingAsset });
+	}
+
+	bool Serialize(FArchive& Ar, const FShaderMapBase::FSerializationContext& Ctx);
 
 #if WITH_EDITOR
 	/** Saves this shader map to the derived data cache. */
@@ -2558,7 +2565,7 @@ public:
 	/** Allows to associate the shader resources with the asset for load order. */
 	virtual FName GetAssetPath() const { return NAME_None; };
 
-	/** Some materials may be loaded early - before the shader library - and need their code inlined */
+	UE_DEPRECATED(5.5, "ShouldInlineShaderCode is no longer used when packaging shader maps. If you rely on this for some reason please reach out to the UE rendering team.")
 	virtual bool ShouldInlineShaderCode() const { return false; }
 #endif // WITH_EDITOR
 
@@ -3021,7 +3028,8 @@ public:
 	ENGINE_API virtual void NotifyCompilationFinished() override;
 	/** Allows to associate the shader resources with the asset for load order. */
 	ENGINE_API virtual FName GetAssetPath() const override;
-	ENGINE_API virtual bool ShouldInlineShaderCode() const override;
+	UE_DEPRECATED(5.5, "ShouldInlineShaderCode is no longer used when packaging shader maps. If you rely on this for some reason please reach out to the UE rendering team.")
+	ENGINE_API virtual bool ShouldInlineShaderCode() const override { return false; }
 	ENGINE_API virtual bool IsUsingControlFlow() const override;
 	ENGINE_API virtual bool IsUsingNewHLSLGenerator() const override;
 	ENGINE_API virtual bool CheckInValidStateForCompilation(FMaterialCompiler* Compiler) const override;
diff --git a/Engine/Source/Runtime/RenderCore/Private/GlobalShader.cpp b/Engine/Source/Runtime/RenderCore/Private/GlobalShader.cpp
--- a/Engine/Source/Runtime/RenderCore/Private/GlobalShader.cpp
+++ b/Engine/Source/Runtime/RenderCore/Private/GlobalShader.cpp
@@ -548,7 +548,7 @@ FGlobalShaderMapSection* FGlobalShaderMapSection::CreateFromArchive(FArchive& Ar
 
 bool FGlobalShaderMapSection::Serialize(FArchive& Ar)
 {
-	return Super::Serialize(Ar, true, false);
+	return Super::Serialize(Ar, FShaderMapBase::FSerializationContext());
 }
 
 TShaderRef<FShader> FGlobalShaderMapSection::GetShader(FShaderType* ShaderType, int32 PermutationId) const
diff --git a/Engine/Source/Runtime/RenderCore/Private/ShaderMap.cpp b/Engine/Source/Runtime/RenderCore/Private/ShaderMap.cpp
--- a/Engine/Source/Runtime/RenderCore/Private/ShaderMap.cpp
+++ b/Engine/Source/Runtime/RenderCore/Private/ShaderMap.cpp
@@ -235,7 +235,7 @@ void FShaderMapBase::UnfreezeContent()
 
 #define CHECK_SHADERMAP_DEPENDENCIES (WITH_EDITOR || !(UE_BUILD_SHIPPING || UE_BUILD_TEST))
 
-bool FShaderMapBase::Serialize(FArchive& Ar, bool bInlineShaderResources, bool bLoadedByCookedMaterial, bool bInlineShaderCode, const FName& SerializingAsset)
+bool FShaderMapBase::Serialize(FArchive& Ar, const FSerializationContext& Ctx)
 {
 	LLM_SCOPE(ELLMTag::Shaders);
 	if (Ar.IsSaving())
@@ -262,7 +262,7 @@ bool FShaderMapBase::Serialize(FArchive& Ar, bool bInlineShaderResources, bool b
 
 		bool bShareCode = false;
 #if WITH_EDITOR
-		bShareCode = !bInlineShaderCode && FShaderLibraryCooker::IsShaderLibraryEnabled() && Ar.IsCooking();
+		bShareCode = FShaderLibraryCooker::IsShaderLibraryEnabled() && Ar.IsCooking();
 #endif // WITH_EDITOR
 		Ar << bShareCode;
 #if WITH_EDITOR
@@ -291,7 +291,7 @@ bool FShaderMapBase::Serialize(FArchive& Ar, bool bInlineShaderResources, bool b
 		else
 #endif // WITH_EDITOR
 		{
-			Code->Serialize(Ar, bLoadedByCookedMaterial);
+			Code->Serialize(Ar, Ctx.bLoadedByCookedMaterial);
 		}
 	}
 	else
@@ -324,14 +324,14 @@ bool FShaderMapBase::Serialize(FArchive& Ar, bool bInlineShaderResources, bool b
 				{
 					UE_LOG(LogShaders, Error, TEXT("Missing shader resource for hash '%s' for shader platform '%s' in the shader library while serializing asset %s"), *ResourceHash.ToString(),
 						*LexToString(ShaderPlatform),
-						*SerializingAsset.ToString());
+						*Ctx.SerializingAsset.ToString());
 				}
 			}
 		}
 		else
 		{
 			Code = new FShaderMapResourceCode();
-			Code->Serialize(Ar, bLoadedByCookedMaterial);
+			Code->Serialize(Ar, Ctx.bLoadedByCookedMaterial);
 			Resource = new FShaderMapResource_InlineCode(ShaderPlatform, Code);
 		}
 
diff --git a/Engine/Source/Runtime/RenderCore/Public/Shader.h b/Engine/Source/Runtime/RenderCore/Public/Shader.h
--- a/Engine/Source/Runtime/RenderCore/Public/Shader.h
+++ b/Engine/Source/Runtime/RenderCore/Public/Shader.h
@@ -2375,7 +2375,18 @@ public:
 
 	RENDERCORE_API void FinalizeContent();
 	RENDERCORE_API void UnfreezeContent();
-	RENDERCORE_API bool Serialize(FArchive& Ar, bool bInlineShaderResources, bool bLoadedByCookedMaterial, bool bInlineShaderCode=false, const FName& SerializingAsset = NAME_None);
+	UE_DEPRECATED(5.5, "Inlining of shader code into shadermaps is no longer supported. If you rely on this please reach out to the UE rendering team.")
+	RENDERCORE_API bool Serialize(FArchive& Ar, bool bInlineShaderResources, bool bLoadedByCookedMaterial, bool bInlineShaderCode = false, const FName& SerializingAsset = NAME_None)
+	{
+		return Serialize(Ar, { bLoadedByCookedMaterial, SerializingAsset });
+	}
+
+	struct FSerializationContext
+	{
+		bool bLoadedByCookedMaterial = false;
+		FName SerializingAsset = NAME_None;
+	};
+	RENDERCORE_API bool Serialize(FArchive& Ar, const FSerializationContext& Ctx);
 
 	EShaderPermutationFlags GetPermutationFlags() const
 	{
