When changing the render static attribute we need to recall PrecachePSO since it can change the VertexFactory when computing the PSOs (see bCPUSkin in USkinnedMeshComponent::CollectPSOPrecacheData).

#tests PC, replayruns

diff --git a/Engine/Source/Runtime/Engine/Private/Components/SkinnedMeshComponent.cpp b/Engine/Source/Runtime/Engine/Private/Components/SkinnedMeshComponent.cpp
--- a/Engine/Source/Runtime/Engine/Private/Components/SkinnedMeshComponent.cpp
+++ b/Engine/Source/Runtime/Engine/Private/Components/SkinnedMeshComponent.cpp
@@ -5120,6 +5120,7 @@ void USkinnedMeshComponent::SetRenderStatic(bool bNewValue)
 	if (bRenderStatic != bNewValue)
 	{
 		bRenderStatic = bNewValue;
+		PrecachePSOs();
 		MarkRenderStateDirty();
 	}
 }
