Add support for explicitly setting mesh paint texture resolution.
This can be set on the static mesh or overridden on the static mesh component.
It defaults to unset (0) which indicates that we want to use the automatically generated size.
The value is only used when creating mesh paint textures, and changing it does not resize existing mesh paint textures.
We will add a tool in the mesh paint mode to "fix"/resize mesh paint textures if they don't match the current settings.

diff --git a/Engine/Plugins/MeshPainting/Source/MeshPaintEditorMode/Private/MeshPaintModeHelpers.cpp b/Engine/Plugins/MeshPainting/Source/MeshPaintEditorMode/Private/MeshPaintModeHelpers.cpp
--- a/Engine/Plugins/MeshPainting/Source/MeshPaintEditorMode/Private/MeshPaintModeHelpers.cpp
+++ b/Engine/Plugins/MeshPainting/Source/MeshPaintEditorMode/Private/MeshPaintModeHelpers.cpp
@@ -326,7 +326,7 @@ void UMeshPaintModeSubsystem::ImportMeshPaintTextureFromVertexColors(UMeshCompon
 			return InstanceMeshLODInfo->OverrideVertexColors->VertexColor(Index);
 		});
 
-	const int32 TextureSize = MeshPaintVirtualTexture::GetDefaultTextureSize(DynamicMesh.VertexCount());
+	const int32 TextureSize = StaticMeshComponent->GetMeshPaintTextureResolution();
 
 	const UE::Geometry::FDynamicMeshAABBTree3 DetailSpatial(&DynamicMesh);
 	UE::Geometry::FMeshBakerDynamicMeshSampler DetailSampler(&DynamicMesh, &DetailSpatial);
diff --git a/Engine/Plugins/MeshPainting/Source/MeshPaintingToolset/Private/MeshPaintHelpers.cpp b/Engine/Plugins/MeshPainting/Source/MeshPaintingToolset/Private/MeshPaintHelpers.cpp
--- a/Engine/Plugins/MeshPainting/Source/MeshPaintingToolset/Private/MeshPaintHelpers.cpp
+++ b/Engine/Plugins/MeshPainting/Source/MeshPaintingToolset/Private/MeshPaintHelpers.cpp
@@ -85,7 +85,7 @@ void UMeshPaintingSubsystem::CreateComponentMeshPaintTexture(UStaticMeshComponen
 	{
 		StaticMeshComponent->Modify();
 
-		const uint32 TextureSize = MeshPaintVirtualTexture::GetDefaultTextureSize(StaticMeshComponent->GetStaticMesh()->GetNumVertices(0));
+		const uint32 TextureSize = StaticMeshComponent->GetMeshPaintTextureResolution();
 		const uint32 TextureNumMips = FMath::FloorLog2(TextureSize) + 1;
 
 		UMeshPaintVirtualTexture* NewTexture = NewObject<UMeshPaintVirtualTexture>(StaticMeshComponent->GetOutermost());
diff --git a/Engine/Source/Runtime/Engine/Classes/Components/StaticMeshComponent.h b/Engine/Source/Runtime/Engine/Classes/Components/StaticMeshComponent.h
--- a/Engine/Source/Runtime/Engine/Classes/Components/StaticMeshComponent.h
+++ b/Engine/Source/Runtime/Engine/Classes/Components/StaticMeshComponent.h
@@ -316,21 +316,29 @@ public:
 	 * If false, vertex color mesh painting is disabled on this instance. 
 	 * This may be set to false by blueprint functions that override vertex colors in construction script.
 	 */
-	UPROPERTY(EditAnywhere, AdvancedDisplay, BlueprintReadWrite, Category = "Mesh Painting")
+	UPROPERTY(EditAnywhere, BlueprintReadWrite, Category = "Mesh Painting")
 	uint8 bEnableVertexColorMeshPainting : 1;
 
 	/** If false, texture color mesh painting is disabled on this instance. */
-	UPROPERTY(EditAnywhere, AdvancedDisplay, BlueprintReadWrite, Category = "Mesh Painting")
+	UPROPERTY(EditAnywhere, BlueprintReadWrite, Category = "Mesh Painting")
 	uint8 bEnableTextureColorMeshPainting : 1;
 
 	/** Whether to override the MeshPaintTextureCoordinateIndex set on the static mesh. */
 	UPROPERTY(EditAnywhere, Category = "Mesh Painting", meta=(InlineEditConditionToggle))
 	uint8 bOverrideMeshPaintTextureCoordinateIndex : 1;
 
-	/** The overriden coordinate index to use when texture color painting on this mesh. */
-	UPROPERTY(EditAnywhere, Category = "Mesh Painting", meta=(UIMin = "0", UIMax = "3", editcondition = "bOverrideMeshPaintTextureCoordinateIndex"))
+	/** Whether to override the MeshPaintTextureCoordinateIndex set on the static mesh. */
+	UPROPERTY(EditAnywhere, Category = "Mesh Painting", meta=(InlineEditConditionToggle))
+	uint8 bOverrideMeshPaintTextureResolution : 1;
+
+	/** The overriden coordinate index to use when texture color painting on this component. */
+	UPROPERTY(EditAnywhere, AdvancedDisplay, Category = "Mesh Painting", meta=(UIMin = "0", UIMax = "3", editcondition = "bOverrideMeshPaintTextureCoordinateIndex"))
 	int32 OverriddenMeshPaintTextureCoordinateIndex;
 
+	/** The overriden resolution of texture color mesh paint textures on this component. */
+	UPROPERTY(EditAnywhere, AdvancedDisplay, Category = "Mesh Painting", meta = (UIMin = "0", UIMax = "4096", ClampMax = "4096", editcondition = "bOverrideMeshPaintTextureResolution"))
+	int32 OverriddenMeshPaintTextureResolution;
+
 	/** Light map resolution to use on this component, used if bOverrideLightMapRes is true and there is a valid StaticMesh. */
 	UPROPERTY(EditAnywhere, BlueprintReadOnly, Category=Lighting, meta=(ClampMax = 4096, editcondition="bOverrideLightMapRes") )
 	int32 OverriddenLightMapRes;
@@ -757,6 +765,8 @@ public:
 	ENGINE_API bool CanMeshPaintVertexColors() const;
 	/* Returns true if mesh texture color painting is supported on this component. */
 	ENGINE_API bool CanMeshPaintTextureColors() const;
+	/** Returns the size to use when creating a mesh paint texture on this component. */
+	ENGINE_API int32 GetMeshPaintTextureResolution() const;
 
 private:
 	/** Initializes the resources used by the static mesh component. */
diff --git a/Engine/Source/Runtime/Engine/Classes/Engine/StaticMesh.h b/Engine/Source/Runtime/Engine/Classes/Engine/StaticMesh.h
--- a/Engine/Source/Runtime/Engine/Classes/Engine/StaticMesh.h
+++ b/Engine/Source/Runtime/Engine/Classes/Engine/StaticMesh.h
@@ -1022,9 +1022,17 @@ public:
 	ENGINE_API bool CanMeshPaintTextureColors() const;
 
 	/** The default coordinate index to use when texture color painting on this mesh. */
-	UPROPERTY(EditAnywhere, AdvancedDisplay, Category = StaticMesh, meta=(UIMin = "0", UIMax = "3"))
+	UPROPERTY(EditAnywhere, AdvancedDisplay, Category = StaticMesh, meta = (UIMin = "0", UIMax = "3"))
 	int32 MeshPaintTextureCoordinateIndex;
 
+	/** 
+	 * The resolution of texture color mesh paint textures on this mesh.
+	 * The final size will be rounded up to a power of 2 and a multiple of the "Mesh Paint Tile Size" project setting.
+	 * A default value of 0 will auto calculate the size using the "Mesh paint texels per vertex" project setting.
+	 */
+	UPROPERTY(EditAnywhere, AdvancedDisplay, Category = StaticMesh, meta = (UIMin = "0", UIMax = "4096", ClampMax = "4096"))
+	int32 MeshPaintTextureResolution;
+
 	/** Useful for reducing self shadowing from distance field methods when using world position offset to animate the mesh's vertices. */
 	UPROPERTY(EditAnywhere, AdvancedDisplay, Category = StaticMesh)
 	float DistanceFieldSelfShadowBias;
diff --git a/Engine/Source/Runtime/Engine/Classes/VT/MeshPaintVirtualTexture.h b/Engine/Source/Runtime/Engine/Classes/VT/MeshPaintVirtualTexture.h
--- a/Engine/Source/Runtime/Engine/Classes/VT/MeshPaintVirtualTexture.h
+++ b/Engine/Source/Runtime/Engine/Classes/VT/MeshPaintVirtualTexture.h
@@ -44,6 +44,8 @@ namespace MeshPaintVirtualTexture
 	/** Returns the fallback color to use for unmapped virtual textures. Use white (same as default vertex color). */
 	inline uint32 GetDefaultFallbackColor() { return 0xFFFFFFFF; }
 
+	/** Returns the passed in size after it is rounded up meet any size constraints. */
+	ENGINE_API uint32 GetAlignedTextureSize(int32 InSize);
 	/** Returns the default texture size to use for a mesh based on the number of vertices. */
 	ENGINE_API uint32 GetDefaultTextureSize(int32 InNumVertices);
 
diff --git a/Engine/Source/Runtime/Engine/Private/Components/StaticMeshComponent.cpp b/Engine/Source/Runtime/Engine/Private/Components/StaticMeshComponent.cpp
--- a/Engine/Source/Runtime/Engine/Private/Components/StaticMeshComponent.cpp
+++ b/Engine/Source/Runtime/Engine/Private/Components/StaticMeshComponent.cpp
@@ -3204,6 +3204,27 @@ int32 UStaticMeshComponent::GetMeshPaintTextureCoordinateIndex() const
 	return 0;
 }
 
+int32 UStaticMeshComponent::GetMeshPaintTextureResolution() const
+{
+	int32 TextureResolution = 0;
+	if (bOverrideMeshPaintTextureResolution)
+	{
+		TextureResolution = OverriddenMeshPaintTextureResolution;
+	}
+	else if (StaticMesh != nullptr)
+	{
+		TextureResolution = StaticMesh->MeshPaintTextureResolution;
+	}
+	
+	if (TextureResolution > 0)
+	{
+		return MeshPaintVirtualTexture::GetAlignedTextureSize(TextureResolution);
+	}
+
+	const int32 NumVertices = StaticMesh != nullptr ? StaticMesh->GetNumVertices(0) : 0;
+	return MeshPaintVirtualTexture::GetDefaultTextureSize(NumVertices);
+}
+
 bool UStaticMeshComponent::CanMeshPaintVertexColors() const
 {
 	if (!bSupportMeshPainting || !bEnableVertexColorMeshPainting)
diff --git a/Engine/Source/Runtime/Engine/Private/VT/MeshPaintVirtualTexture.cpp b/Engine/Source/Runtime/Engine/Private/VT/MeshPaintVirtualTexture.cpp
--- a/Engine/Source/Runtime/Engine/Private/VT/MeshPaintVirtualTexture.cpp
+++ b/Engine/Source/Runtime/Engine/Private/VT/MeshPaintVirtualTexture.cpp
@@ -88,12 +88,19 @@ namespace MeshPaintVirtualTexture
 		return FVirtualTextureBuildSettings::ClampAndAlignTileBorderSize(CVarMeshPaintVirtualTextureTileBorderSize.GetValueOnAnyThread());
 	}
 
+	uint32 GetAlignedTextureSize(int32 InSize)
+	{
+		// Must be power of 2 aligned.
+		const uint32 SizePow2Aligned = FMath::RoundUpToPowerOfTwo((uint32)FMath::Max(InSize, 0));
+		// Must be at least the size of a tile (which is also power of two aligned).
+		return FMath::Max(GetTileSize(), SizePow2Aligned);
+	}
+
 	uint32 GetDefaultTextureSize(int32 InNumVertices)
 	{
-		const int32 NumTexels = InNumVertices * CVarMeshPaintVirtualTextureTexelsPerVertex.GetValueOnGameThread();
-		const uint32 TextureSize = (uint32)FMath::Sqrt((float)NumTexels);
-		const uint32 TextureSizePow2Aligned = FMath::RoundUpToPowerOfTwo(TextureSize);
-		return FMath::Max(GetTileSize(), TextureSizePow2Aligned);
+		const int32 NumTexels = FMath::Max(InNumVertices * CVarMeshPaintVirtualTextureTexelsPerVertex.GetValueOnGameThread(), 0);
+		const int32 TextureSize = FMath::CeilToInt(FMath::Sqrt((float)NumTexels));
+		return GetAlignedTextureSize(TextureSize);
 	}
 
 	/** 
