[shaders] strip material DDC key log entries as they are huge and unwieldy; make all paths which dump DDC keys dump to files in the ShaderDebugInfo folder instead (automatically published with build artifacts)

diff --git a/Engine/Source/Runtime/Engine/Private/Materials/MaterialShader.cpp b/Engine/Source/Runtime/Engine/Private/Materials/MaterialShader.cpp
--- a/Engine/Source/Runtime/Engine/Private/Materials/MaterialShader.cpp
+++ b/Engine/Source/Runtime/Engine/Private/Materials/MaterialShader.cpp
@@ -1924,29 +1924,14 @@ TSharedRef<FMaterialShaderMap::FAsyncLoadContext> FMaterialShaderMap::BeginLoadF
 			FCacheKey CacheKey = GetMaterialShaderMapKey(Result->DataKey);
 			OutDDCKeyDesc = LexToString(CacheKey.Hash);
 
-			if (UNLIKELY(ShouldDumpShaderDDCKeys()) || (Material->IsDefaultMaterial() && Material->GetMaterialDomain() == EMaterialDomain::MD_Surface))
+			FString DebugDDCKeyAsset = CVarShaderCompilerDebugDDCKeyAsset.GetValueOnAnyThread();
+			if (UNLIKELY(ShouldDumpShaderDDCKeys()) ||
+				(Material->IsDefaultMaterial() && Material->GetMaterialDomain() == EMaterialDomain::MD_Surface) ||
+				UNLIKELY(!DebugDDCKeyAsset.IsEmpty() && Material->GetAssetName().Contains(DebugDDCKeyAsset)))
 			{	
 				DumpShaderDDCKeyToFile(InPlatform, ShaderMapId.LayoutParams.WithEditorOnly(), *Material->GetDebugGroupName(), Result->DataKey);
 			}
 
-			if (Material->IsDefaultMaterial() && Material->GetMaterialDomain() == EMaterialDomain::MD_Surface)
-			{
-				FString SpecialEngineDDCKey = Result->DataKey;
-				SpecialEngineDDCKey.RemoveFromEnd(TEXT("\n"));
-				UE_LOG(LogMaterial, Display, TEXT("%s-%s-%s: %s"), *Material->GetAssetName(), *LexToString(ShaderMapId.FeatureLevel), *LexToString(ShaderMapId.QualityLevel), *SpecialEngineDDCKey);
-			}
-
-			if (UNLIKELY(!CVarShaderCompilerDebugDDCKeyAsset.GetValueOnAnyThread().IsEmpty()))
-			{
-				FString DebugDDCKeyAsset = CVarShaderCompilerDebugDDCKeyAsset.GetValueOnAnyThread();
-				if (Material->GetAssetName().Contains(DebugDDCKeyAsset))
-				{
-					FString DataKey = Result->DataKey;
-					DataKey.RemoveFromEnd(TEXT("\n"));
-					UE_LOG(LogMaterial, Display, TEXT("%s-%s-%s: %s"), *Material->GetAssetName(), *LexToString(ShaderMapId.FeatureLevel), *LexToString(ShaderMapId.QualityLevel), *DataKey);
-				}
-			}
-
 			bool bCheckCache = true;
 
 			// If NoShaderDDC then don't check for a material the first time we encounter it to simulate
