Console: Copy text from console (history) - with partial line support

diff --git a/Engine/Source/Runtime/Engine/Classes/Engine/Console.h b/Engine/Source/Runtime/Engine/Classes/Engine/Console.h
--- a/Engine/Source/Runtime/Engine/Classes/Engine/Console.h
+++ b/Engine/Source/Runtime/Engine/Classes/Engine/Console.h
@@ -14,6 +14,9 @@
 class FViewport;
 class SWidget;
 struct FAutoCompleteCommand;
+class FCanvasTextItem;
+class FCanvasTileItem;
+class UFont;
 
 /**
  * Node for storing an auto-complete tree based on each char in the command.
@@ -140,16 +143,16 @@ class UConsole
 	// Scrollback region selection mouse and state
 	struct ScrollbackSelection
 	{
-		FVector2D MousePosDown = FVector2D(0, 0);
-		FVector2D MousePosUp = FVector2D(0, 0);
-		int32 MousePosX = 0;
-		int32 MousePosY = 0;
+		FVector2f MousePosDown = FVector2f(0.0f, 0.0f);
+		FVector2f MousePosUp = FVector2f(0.0f, 0.0f);
+		FVector2f MousePos = FVector2f(0.0f, 0.0f);
 		bool bActive = false;
 		bool bCapture = false;
 		bool bMade = false;
 		float Offset = 0.0f;
 	} Selection;
 
+	float TextW = 15.0f;
 	float TextH = 15.0f;
 
 	ENGINE_API ~UConsole();
@@ -328,13 +331,13 @@ class UConsole
 	DECLARE_MULTICAST_DELEGATE_OneParam(FRegisterConsoleAutoCompleteEntries, TArray<FAutoCompleteCommand>&);
 	static ENGINE_API FRegisterConsoleAutoCompleteEntries RegisterConsoleAutoCompleteEntries;
 
-	/** Deletate for when the console is activated or deactivated */
+	/** Delegate for when the console is activated or deactivated */
 	DECLARE_MULTICAST_DELEGATE_OneParam(FOnConsoleActivationStateChanged, bool);
 	static ENGINE_API FOnConsoleActivationStateChanged OnConsoleActivationStateChanged;
 
 private:
 
-	ENGINE_API bool InputKey_InputLine(FInputDeviceId DeviceId, FKey Key, EInputEvent Event, float AmountDepressed = 1.f, bool bGamepad = false);
+	bool InputKey_InputLine(FInputDeviceId DeviceId, FKey Key, EInputEvent Event, float AmountDepressed = 1.f, bool bGamepad = false);
 
 	// interface FOutputDevice
 	ENGINE_API virtual void Serialize( const TCHAR* V, ELogVerbosity::Type Verbosity, const class FName& Category ) override;
@@ -345,21 +348,31 @@ private:
 	 * Prints a single line of text to the console.
 	 * @param Text - A line of text to display on the console.
 	 */
-	ENGINE_API void OutputTextLine(const FString& Text);
+	void OutputTextLine(const FString& Text);
 
-	ENGINE_API void PostRender_InputLine(class UCanvas* Canvas, FIntPoint UserInputLinePos);
+	void PostRender_InputLine(class UCanvas* Canvas, FIntPoint UserInputLinePos);
 
-	ENGINE_API void SetAutoCompleteFromHistory();
+	void SetAutoCompleteFromHistory();
 
-	ENGINE_API void SetInputLineFromAutoComplete();
+	void SetInputLineFromAutoComplete();
 
-	ENGINE_API void UpdatePrecompletedInputLine();
+	void UpdatePrecompletedInputLine();
 
-	ENGINE_API void NormalizeHistoryBuffer();
+	void NormalizeHistoryBuffer();
+
+	// Scrollback region selection utilities
+	void DrawLine(
+		UCanvas* Canvas, UFont* Font, FCanvasTileItem& TileBlack, FCanvasTileItem& TileWhite,
+		FCanvasTextItem ConsoleText, const FString& Line, float PenX, float PenY, float PosLeft,
+		TArray<FString>& SelectedLines);
+	void GetPosInTextLine(UCanvas* Canvas, UFont* Font, const FString& Line, float MousePosX, float& PosInLine, int32& PosInString);
+	float GetPosTextLineEnd(UCanvas* Canvas, UFont* Font, const FString& Line);
 
 	// Console settings from BaseInput.ini
 	const UConsoleSettings* ConsoleSettings;
 
 	// Widget that was focused before the console was opened (focus will be restored to this if it's valid after the console closes)
 	TWeakPtr<SWidget> PreviousFocusedWidget;
+
+	const bool bDPIAwareStringMeasurement = true;
 };
diff --git a/Engine/Source/Runtime/Engine/Private/UserInterface/Console.cpp b/Engine/Source/Runtime/Engine/Private/UserInterface/Console.cpp
--- a/Engine/Source/Runtime/Engine/Private/UserInterface/Console.cpp
+++ b/Engine/Source/Runtime/Engine/Private/UserInterface/Console.cpp
@@ -1224,16 +1224,16 @@ void UConsole::PostRender_Console_Typing(UCanvas* Canvas)
 {
 	float ClipX = Canvas->ClipX;
 	float ClipY = Canvas->ClipY;
-	float LeftPos = 0;
+	float PosLeft = 0;
 
 	if (CVarCustomConsolePosEnabled.GetValueOnAnyThread())
 	{
-		LeftPos = (float)CVarConsoleXPos.GetValueOnAnyThread();
+		PosLeft = (float)CVarConsoleXPos.GetValueOnAnyThread();
 		float BottomOffset = (float)CVarConsoleYPos.GetValueOnAnyThread();
 		ClipY = ClipY - BottomOffset;
 	}
 
-	PostRender_InputLine(Canvas, FIntPoint(LeftPos, ClipY));
+	PostRender_InputLine(Canvas, FIntPoint(PosLeft, ClipY));
 }
 
 void UConsole::BeginState_Typing(FName PreviousStateName)
@@ -1284,7 +1284,7 @@ bool UConsole::InputKey_Open(FInputDeviceId DeviceId, FKey Key, EInputEvent Even
 	{
 		if (Event == IE_Pressed)
 		{
-			Selection.MousePosDown = FVector2D(Selection.MousePosX, Selection.MousePosY);
+			Selection.MousePosDown = Selection.MousePos;
 			Selection.bActive = true;
 			Selection.bMade = false;
 			Selection.Offset = 0.0f;
@@ -1292,7 +1292,7 @@ bool UConsole::InputKey_Open(FInputDeviceId DeviceId, FKey Key, EInputEvent Even
 		}
 		else if (Event == IE_Released)
 		{
-			Selection.MousePosUp = FVector2D(Selection.MousePosX, Selection.MousePosY);
+			Selection.MousePosUp = Selection.MousePos;
 			Selection.bCapture = true;
 			Selection.bMade = true;
 			return true;
@@ -1340,12 +1340,12 @@ void UConsole::PostRender_Console_Open(UCanvas* Canvas)
 
 	// shrink for TVs
 	float ClipX = Canvas->ClipX;
-	float TopPos = 0;
-	float LeftPos = 0;
+	float PosTop = 0;
+	float PosLeft = 0;
 
 	if (CVarCustomConsolePosEnabled.GetValueOnAnyThread())
 	{
-		LeftPos = (float)CVarConsoleXPos.GetValueOnAnyThread();
+		PosLeft = (float)CVarConsoleXPos.GetValueOnAnyThread();
 		float BottomOffset = (float)CVarConsoleYPos.GetValueOnAnyThread();
 		Height = Canvas->ClipY - BottomOffset;
 	}
@@ -1353,12 +1353,10 @@ void UConsole::PostRender_Console_Open(UCanvas* Canvas)
 	UFont* Font = GEngine->GetSmallFont();
 
 	const float DPIScale = Canvas->GetDPIScale();
-	const bool bDPIAwareStringMeasurement = true;
 
 	// determine the height of the text
-	float xl;
-	Canvas->StrLen(Font, TEXT("M"), xl, TextH, bDPIAwareStringMeasurement);
-	xl /= DPIScale;
+	Canvas->StrLen(Font, TEXT("M"), TextW, TextH, bDPIAwareStringMeasurement);
+	TextW /= DPIScale;
 	TextH /= DPIScale;
 
 	float y = Height - TextH;
@@ -1370,10 +1368,10 @@ void UConsole::PostRender_Console_Open(UCanvas* Canvas)
 		BackgroundColorBlack.A = ConsoleSettings->BackgroundOpacityPercentage / 100.0f;
 		BackgroundColorWhite.A = ConsoleSettings->BackgroundOpacityPercentage / 100.0f;
 
-		FCanvasTextItem ConsoleText(FVector2D(LeftPos, TopPos + Height - 5 - TextH), FText::FromString(TEXT("")), Font, ConsoleSettings->InputColor);
+		FCanvasTextItem ConsoleText(FVector2D(PosLeft, PosTop + Height - 5 - TextH), FText::FromString(TEXT("")), Font, ConsoleSettings->InputColor);
 
-		FCanvasTileItem ConsoleTileBlack(FVector2D(LeftPos, 0.0f), DefaultTexture_Black->GetResource(), FVector2D(ClipX, Height + TopPos - TextH), BackgroundColorBlack);
-		FCanvasTileItem ConsoleTileWhite(FVector2D(LeftPos, 0.0f), DefaultTexture_White->GetResource(), FVector2D(ClipX, Height + TopPos - TextH), BackgroundColorWhite);
+		FCanvasTileItem ConsoleTileBlack(FVector2D(PosLeft, 0.0f), DefaultTexture_Black->GetResource(), FVector2D(ClipX, Height + PosTop - TextH), BackgroundColorBlack);
+		FCanvasTileItem ConsoleTileWhite(FVector2D(PosLeft, 0.0f), DefaultTexture_White->GetResource(), FVector2D(ClipX, Height + PosTop - TextH), BackgroundColorWhite);
 		ConsoleTileBlack.BlendMode = SE_BLEND_AlphaBlend;		// Preserve alpha to allow single-pass composite
 		ConsoleTileWhite.BlendMode = SE_BLEND_AlphaBlend;		// Preserve alpha to allow single-pass composite
 
@@ -1388,8 +1386,8 @@ void UConsole::PostRender_Console_Open(UCanvas* Canvas)
 			float PenX;
 			float PenY;
 			float PenZ = 0.1f;
-			PenX = LeftPos;
-			PenY = TopPos + y;
+			PenX = PosLeft;
+			PenY = PosTop + y;
 
 			// adjust the location for any word wrapping due to long text lines
 			if (idx < Scrollback.Num())
@@ -1401,65 +1399,11 @@ void UConsole::PostRender_Console_Open(UCanvas* Canvas)
 				if (ScrollLineYL > TextH)
 				{
 					y -= (ScrollLineYL - TextH);
-					PenX = LeftPos;
-					PenY = TopPos + y;
-				}
-
-				bool bLineSelected = false;
-				if (Selection.bActive)
-				{
-					// Check if selection range overlaps current line
-					float x0 = PenY;
-					float x1 = PenY + TextH;
-					float y0 = Selection.MousePosDown.Y;
-					float y1 = Selection.bMade ? Selection.MousePosUp.Y : Selection.MousePosY;
-
-					// Adjust for new lines added since selection was initiated
-					y0 -= Selection.Offset;
-					y1 -= Selection.Offset;
-
-					// Swap to keep range in order, to allow drag selecting down or up
-					if (y0 > y1)
-					{
-						float tmp = y0;
-						y0 = y1;
-						y1 = tmp;
-					}
-
-					auto Overlaps = [](float x0, float x1, float y0, float y1)
-					{
-						return (x0 <= y1) && (x1 >= y0);
-					};
-
-					bLineSelected = Overlaps(x0, x1, y0, y1);
+					PenX = PosLeft;
+					PenY = PosTop + y;
 				}
 
-				if (bLineSelected)
-				{
-					// Selected line (black text on white background)
-					ConsoleTileWhite.Position = FVector2D(LeftPos, PenY);
-					ConsoleTileWhite.Size = FVector2D(ClipX, TextH);
-					Canvas->DrawItem(ConsoleTileWhite);
-
-					ConsoleText.SetColor(FLinearColor::Black);
-
-					if (Selection.bCapture)
-					{
-						SelectedLines.Add(Scrollback[idx]);
-					}
-				}
-				else
-				{
-					// Unselected line (white text on black background)
-					ConsoleTileBlack.Position = FVector2D(LeftPos, PenY);
-					ConsoleTileBlack.Size = FVector2D(ClipX, TextH);
-					Canvas->DrawItem(ConsoleTileBlack);
-
-					ConsoleText.SetColor(FLinearColor::White);
-				}
-
-				ConsoleText.Text = FText::FromString(Scrollback[idx]);
-				Canvas->DrawItem(ConsoleText, PenX, PenY);
+				DrawLine(Canvas, Font, ConsoleTileBlack, ConsoleTileWhite, ConsoleText, Scrollback[idx], PenX, PenY, PosLeft, SelectedLines);
 			}
 			idx--;
 			y -= TextH;
@@ -1470,7 +1414,10 @@ void UConsole::PostRender_Console_Open(UCanvas* Canvas)
 			Algo::Reverse(SelectedLines);
 			const FString SelectedText = FString::Join(SelectedLines, LINE_TERMINATOR);
 
-			FPlatformApplicationMisc::ClipboardCopy(*SelectedText);
+			if (SelectedText.Len())
+			{
+				FPlatformApplicationMisc::ClipboardCopy(*SelectedText);
+			}
 
 			Selection.bCapture = false;		// Finished copy to clipboard when mouse is released
 		}
@@ -1482,13 +1429,318 @@ void UConsole::PostRender_Console_Open(UCanvas* Canvas)
 		FLinearColor BackgroundColor = ConsoleDefs::AutocompleteBackgroundColor.ReinterpretAsLinear();
 		BackgroundColor.A = ConsoleSettings->BackgroundOpacityPercentage / 100.0f;
 
-		FCanvasTileItem ConsoleTile(FVector2D(LeftPos, 0.0f), DefaultTexture_Black->GetResource(), FVector2D(ClipX, y + TextH), BackgroundColor);
+		FCanvasTileItem ConsoleTile(FVector2D(PosLeft, 0.0f), DefaultTexture_Black->GetResource(), FVector2D(ClipX, y + TextH), BackgroundColor);
 		ConsoleTile.BlendMode = SE_BLEND_AlphaBlend;			// Preserve alpha to allow single-pass composite
 
 		Canvas->DrawItem(ConsoleTile);
 	}
 
-	PostRender_InputLine(Canvas, FIntPoint(LeftPos, TopPos + Height + 6));
+	PostRender_InputLine(Canvas, FIntPoint(PosLeft, PosTop + Height + 6));
+}
+
+void UConsole::DrawLine(
+	UCanvas* Canvas, UFont* Font, FCanvasTileItem& TileBlack, FCanvasTileItem& TileWhite,
+	FCanvasTextItem ConsoleText, const FString& Line, float PenX, float PenY, float PosLeft,
+	TArray<FString>& SelectedLines)
+{
+	float ClipX = Canvas->ClipX;
+	bool bLineSelected = false;				// Partial line is selected
+	bool bLineSelectedWhole = false;		// Whole line is selected
+	bool bLineUnselectedWhole = false;		// Whole line is unselected
+	bool bLineSelectedDown = false;			// Partial line selected, on mouse down (selection begin)
+	bool bLineSelectedUp = false;			// Partial line selected, on mouse up (selection end)
+	bool bLineSelectTopToBottom = true;		// Drag selecting down or up
+	float PosInLineDown = 0.0f;				// Partial line selected (pixels)
+	float PosInLineUp = 0.0f;				// Partial line selected (pixels)
+	int32 PosInStringDown = 0;				// Partial line selected (characters)
+	int32 PosInStringUp = 0;				// Partial line selected (characters)
+
+	if (Selection.bActive)
+	{
+		// Check if selection range overlaps current line
+		float PenY0 = PenY;
+		float PenY1 = PenY + TextH;
+		float PosDownY0 = Selection.MousePosDown.Y;
+		float PosUpY1 = Selection.bMade ? Selection.MousePosUp.Y : Selection.MousePos.Y;
+
+		// Adjust for new lines added since selection was initiated
+		PosDownY0 -= Selection.Offset;
+		PosUpY1 -= Selection.Offset;
+
+		bLineSelectedDown = (PosDownY0 >= PenY0) && (PosDownY0 <= PenY1);
+		if (bLineSelectedDown)
+		{
+			// Partial line selection (mouse down)
+			GetPosInTextLine(Canvas, Font, Line, Selection.MousePosDown.X, PosInLineDown, PosInStringDown);
+		}
+
+		bLineSelectedUp = (PosUpY1 >= PenY0) && (PosUpY1 <= PenY1);
+		if (bLineSelectedUp)
+		{
+			// Partial line selection (mouse up)
+			float MousePosUpX = Selection.bMade ? Selection.MousePosUp.X : Selection.MousePos.X;
+			GetPosInTextLine(Canvas, Font, Line, MousePosUpX, PosInLineUp, PosInStringUp);
+		}
+
+		if (PosDownY0 > PosUpY1)
+		{
+			// Swap to keep Y-range in order, to allow drag selecting down or up
+			Swap(PosDownY0, PosUpY1);
+
+			bLineSelectTopToBottom = false;
+			bLineSelectedDown = !bLineSelectedDown;
+			bLineSelectedUp = !bLineSelectedUp;
+		}
+
+		auto Overlaps = [](float PenY0, float PenY1, float PosDownY0, float PosUpY1)
+		{
+			return (PosUpY1 >= PenY0) && (PosDownY0 <= PenY1);
+		};
+
+		bLineSelected = Overlaps(PenY0, PenY1, PosDownY0, PosUpY1);
+		bLineSelectedWhole = (PosUpY1 > PenY1) && (PosDownY0 < PenY0);
+		bLineUnselectedWhole = bLineSelectedUp && (PosInLineUp <= 0.0f);
+	}
+
+	if (!bLineUnselectedWhole && !bLineSelectedWhole && bLineSelectedDown && bLineSelectedUp)
+	{
+		// Draw three segments: Black White Black
+
+		if (PosInLineDown > PosInLineUp)
+		{
+			Swap(PosInLineDown, PosInLineUp);
+			Swap(PosInStringDown, PosInStringUp);
+		}
+
+		{
+			// white text on black background (from left to start of selected text)
+			TileBlack.Position = FVector2D(PosLeft, PenY);
+			TileBlack.Size = FVector2D(PosInLineDown, TextH);
+			Canvas->DrawItem(TileBlack);
+
+			ConsoleText.SetColor(FLinearColor::White);
+			FString Segment = Line.Mid(0, PosInStringDown + 1);
+
+			ConsoleText.Text = FText::FromString(Segment);
+			Canvas->DrawItem(ConsoleText, PenX, PenY);
+		}
+		{
+			// black text on white background (from start of selected text)
+			TileWhite.Position = FVector2D(PosInLineDown, PenY);
+			TileWhite.Size = FVector2D(PosInLineUp - PosInLineDown, TextH);
+			Canvas->DrawItem(TileWhite);
+
+			ConsoleText.SetColor(FLinearColor::Black);
+			int32 PosCharDown = (PosInStringDown > 0) ? PosInStringDown + 1 : 0;
+			FString Segment = Line.Mid(PosCharDown, PosInStringUp - PosCharDown + 1);
+
+			ConsoleText.Text = FText::FromString(Segment);
+			Canvas->DrawItem(ConsoleText, PenX + PosInLineDown, PenY);
+		}
+		{
+			// white text on black background (from start of unselected text)
+			TileBlack.Position = FVector2D(PosInLineUp, PenY);
+			TileBlack.Size = FVector2D(ClipX - PosInLineUp, TextH);
+			Canvas->DrawItem(TileBlack);
+
+			ConsoleText.SetColor(FLinearColor::White);
+			int32 PosCharUp = (PosInLineUp > 0) ? PosInStringUp + 1 : 0;
+			FString Segment = Line.Mid(PosCharUp, Line.Len());
+
+			ConsoleText.Text = FText::FromString(Segment);
+			Canvas->DrawItem(ConsoleText, PenX + PosInLineUp, PenY);
+		}
+		if (Selection.bCapture)
+		{
+			int32 PosCharDown = PosInStringDown + 1;
+			FString Segment = Line.Mid(PosCharDown, PosInStringUp - PosCharDown + 1);
+			SelectedLines.Add(Segment);
+		}
+	}
+	else if (!bLineUnselectedWhole && !bLineSelectedWhole && bLineSelectedDown)
+	{
+		// Draw three segments: Black White Black
+
+		float LineWidth = GetPosTextLineEnd(Canvas, Font, Line);
+
+		if (!bLineSelectTopToBottom)
+		{
+			Swap(PosInLineDown, PosInLineUp);
+			Swap(PosInStringDown, PosInStringUp);
+		}
+
+		{
+			// white text on black background (from left to start of selected text)
+			TileBlack.Position = FVector2D(PosLeft, PenY);
+			TileBlack.Size = FVector2D(PosInLineDown, TextH);
+			Canvas->DrawItem(TileBlack);
+
+			ConsoleText.SetColor(FLinearColor::White);
+			FString Segment = Line.Mid(0, PosInStringDown + 1);
+
+			ConsoleText.Text = FText::FromString(Segment);
+			Canvas->DrawItem(ConsoleText, PenX, PenY);
+		}
+		{
+			// black text on white background (from start of selected text)
+			TileWhite.Position = FVector2D(PosInLineDown, PenY);
+			TileWhite.Size = FVector2D(LineWidth - PosInLineDown, TextH);
+			Canvas->DrawItem(TileWhite);
+
+			ConsoleText.SetColor(FLinearColor::Black);
+			int32 PosCharDown = (PosInStringDown > 0) ? PosInStringDown + 1 : 0;
+			FString Segment = Line.Mid(PosCharDown, Line.Len());
+
+			ConsoleText.Text = FText::FromString(Segment);
+			Canvas->DrawItem(ConsoleText, PenX + PosInLineDown, PenY);
+		}
+		{
+			// remaining black background (from end of text to end of line)
+			TileBlack.Position = FVector2D(LineWidth, PenY);
+			TileBlack.Size = FVector2D(ClipX - LineWidth, TextH);
+			Canvas->DrawItem(TileBlack);
+		}
+		if (Selection.bCapture)
+		{
+			int32 PosCharDown = (PosInStringDown > 0) ? PosInStringDown + 1 : 0;
+			FString Segment = Line.Mid(PosCharDown, Line.Len());
+			SelectedLines.Add(Segment);
+		}
+	}
+	else if (!bLineSelectedWhole && bLineSelectedUp)
+	{
+		// Draw two segments: White Black
+
+		if (!bLineSelectTopToBottom)
+		{
+			Swap(PosInLineDown, PosInLineUp);
+			Swap(PosInStringDown, PosInStringUp);
+		}
+
+		{
+			// black text on white background (from left to start of unselected text)
+			TileWhite.Position = FVector2D(PosLeft, PenY);
+			TileWhite.Size = FVector2D(PosInLineUp, TextH);
+			Canvas->DrawItem(TileWhite);
+
+			ConsoleText.SetColor(FLinearColor::Black);
+			int32 PosCharUp = (PosInLineUp > 0.0f) ? PosInStringUp + 1 : 0;
+			FString Segment = Line.Mid(0, PosCharUp);
+			ConsoleText.Text = FText::FromString(Segment);
+			Canvas->DrawItem(ConsoleText, PenX, PenY);
+		}
+		{
+			// white text on black background (from start of unselected text)
+			TileBlack.Position = FVector2D(PosInLineUp, PenY);
+			TileBlack.Size = FVector2D(ClipX - PosInLineUp, TextH);
+			Canvas->DrawItem(TileBlack);
+
+			ConsoleText.SetColor(FLinearColor::White);
+			int32 PosCharUp = (PosInLineUp > 0.0f) ? PosInStringUp + 1 : 0;
+			FString Segment = Line.Mid(PosCharUp, Line.Len());
+
+			ConsoleText.Text = FText::FromString(Segment);
+			Canvas->DrawItem(ConsoleText, PenX + PosInLineUp, PenY);
+		}
+		if (Selection.bCapture)
+		{
+			int32 PosCharUp = (PosInLineUp > 0.0f) ? PosInStringUp + 1 : 0;
+			FString Segment = Line.Mid(0, PosCharUp);
+			if (!Segment.IsEmpty())
+			{
+				SelectedLines.Add(Segment);
+			}
+		}
+	}
+	else if (bLineSelected)
+	{
+		// Draw two segments: White Black
+
+		float LineWidth = GetPosTextLineEnd(Canvas, Font, Line);
+
+		{
+			// black text on white background (from left to end of text)
+			TileWhite.Position = FVector2D(PosLeft, PenY);
+			TileWhite.Size = FVector2D(LineWidth, TextH);
+			Canvas->DrawItem(TileWhite);
+
+			ConsoleText.SetColor(FLinearColor::Black);
+			FString Segment = Line.Mid(0, Line.Len());
+
+			ConsoleText.Text = FText::FromString(Segment);
+			Canvas->DrawItem(ConsoleText, PenX, PenY);
+		}
+		{
+			// remaining black background (from end of text to end of line)
+			TileBlack.Position = FVector2D(LineWidth, PenY);
+			TileBlack.Size = FVector2D(ClipX - LineWidth, TextH);
+			Canvas->DrawItem(TileBlack);
+		}
+		if (Selection.bCapture)
+		{
+			SelectedLines.Add(Line);
+		}
+	}
+	else
+	{
+		// Draw two segments: Black Black
+
+		float LineWidth = GetPosTextLineEnd(Canvas, Font, Line);
+
+		{
+			TileBlack.Position = FVector2D(PosLeft, PenY);
+			TileBlack.Size = FVector2D(LineWidth, TextH);
+			Canvas->DrawItem(TileBlack);
+
+			ConsoleText.SetColor(FLinearColor::White);
+			FString Segment = Line.Mid(0, Line.Len());
+
+			ConsoleText.Text = FText::FromString(Segment);
+			Canvas->DrawItem(ConsoleText, PenX, PenY);
+		}
+		{
+			// remaining black background (from end of text to end of line)
+			TileBlack.Position = FVector2D(LineWidth, PenY);
+			TileBlack.Size = FVector2D(ClipX - LineWidth, TextH);
+			Canvas->DrawItem(TileBlack);
+		}
+	}
+}
+
+void UConsole::GetPosInTextLine(UCanvas* Canvas, UFont* Font, const FString& Line, float MousePosX, float& PosInLine, int32& PosInString)
+{
+	// add up character widths to selection position in X, since we are not using a monospace font
+	float CharPos = 0.0f;
+	for (int32 i = 0; i < Line.Len(); ++i)
+	{
+		FString LineChar1 = FString::Printf(TEXT("%c"), Line[i]);
+		float CharW1, CharH1;
+		Canvas->StrLen(Font, LineChar1, CharW1, CharH1, bDPIAwareStringMeasurement);
+		CharPos += CharW1;
+
+		if (CharPos > MousePosX)
+		{
+			break;
+		}
+		PosInLine = CharPos;
+		PosInString = i;
+	}
+}
+
+float UConsole::GetPosTextLineEnd(UCanvas* Canvas, UFont* Font, const FString& Line)
+{
+	// add up character widths of whole line in X, since we are not using a monospace font
+	float CharPos = 0.0f;
+	for (int32 i = 0; i < Line.Len(); ++i)
+	{
+		FString LineChar1 = FString::Printf(TEXT("%c"), Line[i]);
+		float CharW1, CharH1;
+		Canvas->StrLen(Font, LineChar1, CharW1, CharH1, bDPIAwareStringMeasurement);
+
+		CharPos += CharW1;
+	}
+
+	return CharPos;
 }
 
 void UConsole::BeginState_Open(FName PreviousStateName)
@@ -1564,14 +1816,12 @@ bool UConsole::InputKey(FInputDeviceId DeviceId, FKey Key, EInputEvent Event, fl
 
 void UConsole::MouseMove(FViewport* Viewport, int32 X, int32 Y)
 {
-	Selection.MousePosX = X;
-	Selection.MousePosY = Y;
+	Selection.MousePos = FVector2f(X, Y);
 }
 
 void UConsole::CapturedMouseMove(FViewport* InViewport, int32 X, int32 Y)
 {
-	Selection.MousePosX = X;
-	Selection.MousePosY = Y;
+	Selection.MousePos = FVector2f(X, Y);
 }
 
 void UConsole::PostRender_Console(UCanvas* Canvas)
@@ -1603,7 +1853,6 @@ void UConsole::PostRender_InputLine(UCanvas* Canvas, FIntPoint UserInputLinePos)
 	const FString PrecompletedInputText = !PrecompletedInputLine.IsEmpty() ? PrecompletedInputLine.RightChop(TypedStr.Len()) : FString();
 
 	const float DPIScale = Canvas->GetDPIScale();
-	const bool bDPIAwareStringMeasurement = true;
 
 	// use the smallest font
 	UFont* Font = GEngine->GetSmallFont();
