Retroactive fix for materials with alpha holdout blend mode failing to cook.

diff --git a/Engine/Source/Runtime/Engine/Private/Materials/HLSLMaterialTranslator.cpp b/Engine/Source/Runtime/Engine/Private/Materials/HLSLMaterialTranslator.cpp
--- a/Engine/Source/Runtime/Engine/Private/Materials/HLSLMaterialTranslator.cpp
+++ b/Engine/Source/Runtime/Engine/Private/Materials/HLSLMaterialTranslator.cpp
@@ -1874,10 +1874,16 @@ void FHLSLMaterialTranslator::TranslateMaterial()
 			Errorf(TEXT("Alpha Holdout blend mode must use unlit shading model."));
 		}
 
-		if (Material->IsTranslucencyAfterDOFEnabled() || Material->IsTranslucencyAfterMotionBlurEnabled())
+		const int32 FnMainVersion = Material->GetMaterialInterface()->GetLinkerCustomVersion(FFortniteMainBranchObjectVersion::GUID);
+
+		// Note: As a retroactive fix for content failing to cook, we use the closest fn main version change to constrain the error only to recently created materials.
+		if (FnMainVersion > FFortniteMainBranchObjectVersion::LandscapeTargetLayersInLandscapeActor)
 		{
-			// We must write into the alpha channel of the SceneColor now for later translucents rendered after DOF to be able to be affected by the alpha.
-			Errorf(TEXT("Alpha Holdout blend mode must use BeforeDOF under Translucency / Advanced / Translucency Pass.")); 
+			if (Material->IsTranslucencyAfterDOFEnabled() || Material->IsTranslucencyAfterMotionBlurEnabled())
+			{
+				// We must write into the alpha channel of the SceneColor now for later translucents rendered after DOF to be able to be affected by the alpha.
+				Errorf(TEXT("Alpha Holdout blend mode must use BeforeDOF under Translucency / Advanced / Translucency Pass."));
+			}
 		}
 	}
 
