Fixed linker errors on some console platform when building modular builds. Fixes involves moving ctor/dtor and friends from header to cpp and dll import/export

diff --git a/Engine/Source/Runtime/AnimGraphRuntime/Private/AnimNodes/AnimNode_Mirror.cpp b/Engine/Source/Runtime/AnimGraphRuntime/Private/AnimNodes/AnimNode_Mirror.cpp
--- a/Engine/Source/Runtime/AnimGraphRuntime/Private/AnimNodes/AnimNode_Mirror.cpp
+++ b/Engine/Source/Runtime/AnimGraphRuntime/Private/AnimNodes/AnimNode_Mirror.cpp
@@ -19,6 +19,9 @@ FAnimNode_MirrorBase::FAnimNode_MirrorBase()
 {
 }
 
+FAnimNode_MirrorBase::FAnimNode_MirrorBase(const FAnimNode_MirrorBase&) = default;
+FAnimNode_MirrorBase::~FAnimNode_MirrorBase() = default;
+
 UMirrorDataTable* FAnimNode_MirrorBase::GetMirrorDataTable() const
 {
 	return nullptr;
@@ -235,9 +238,8 @@ void FAnimNode_MirrorBase::GatherDebugData(FNodeDebugData& DebugData)
 	Source.GatherDebugData(DebugData);
 }
 
-FAnimNode_Mirror::FAnimNode_Mirror()
-{
-}
+FAnimNode_Mirror::FAnimNode_Mirror() = default;
+FAnimNode_Mirror::~FAnimNode_Mirror() = default;
 
 UMirrorDataTable* FAnimNode_Mirror::GetMirrorDataTable() const 
 { 
@@ -378,9 +380,9 @@ bool FAnimNode_Mirror::SetResetChildOnMirrorStateChange(bool bInResetChildOnMirr
 }
 
 
-FAnimNode_Mirror_Standalone::FAnimNode_Mirror_Standalone()
-{
-}
+FAnimNode_Mirror_Standalone::FAnimNode_Mirror_Standalone() = default;
+FAnimNode_Mirror_Standalone::FAnimNode_Mirror_Standalone(const FAnimNode_Mirror_Standalone&) = default;
+FAnimNode_Mirror_Standalone::~FAnimNode_Mirror_Standalone() = default;
 
 UMirrorDataTable* FAnimNode_Mirror_Standalone::GetMirrorDataTable() const
 {
diff --git a/Engine/Source/Runtime/AnimGraphRuntime/Public/AnimNodes/AnimNode_Mirror.h b/Engine/Source/Runtime/AnimGraphRuntime/Public/AnimNodes/AnimNode_Mirror.h
--- a/Engine/Source/Runtime/AnimGraphRuntime/Public/AnimNodes/AnimNode_Mirror.h
+++ b/Engine/Source/Runtime/AnimGraphRuntime/Public/AnimNodes/AnimNode_Mirror.h
@@ -14,6 +14,8 @@ struct FAnimNode_MirrorBase : public FAnimNode_Base
 	GENERATED_BODY()
 public:
 	ANIMGRAPHRUNTIME_API FAnimNode_MirrorBase(); 
+	ANIMGRAPHRUNTIME_API FAnimNode_MirrorBase(const FAnimNode_MirrorBase&); 
+	ANIMGRAPHRUNTIME_API ~FAnimNode_MirrorBase(); 
 
 	ANIMGRAPHRUNTIME_API virtual void Initialize_AnyThread(const FAnimationInitializeContext& Context) override;
 	ANIMGRAPHRUNTIME_API virtual void CacheBones_AnyThread(const FAnimationCacheBonesContext& Context) override;
@@ -98,6 +100,7 @@ struct FAnimNode_Mirror : public FAnimNode_MirrorBase
 
 public:
 	ANIMGRAPHRUNTIME_API FAnimNode_Mirror();
+	ANIMGRAPHRUNTIME_API ~FAnimNode_Mirror();
 
 	ANIMGRAPHRUNTIME_API virtual UMirrorDataTable* GetMirrorDataTable() const override;
 	ANIMGRAPHRUNTIME_API virtual bool SetMirrorDataTable(UMirrorDataTable* MirrorTable) override;
@@ -153,6 +156,8 @@ struct FAnimNode_Mirror_Standalone : public FAnimNode_MirrorBase
 
 public:
 	ANIMGRAPHRUNTIME_API FAnimNode_Mirror_Standalone();
+	ANIMGRAPHRUNTIME_API FAnimNode_Mirror_Standalone(const FAnimNode_Mirror_Standalone&);
+	ANIMGRAPHRUNTIME_API ~FAnimNode_Mirror_Standalone();
 
 	ANIMGRAPHRUNTIME_API virtual UMirrorDataTable* GetMirrorDataTable() const override;
 	ANIMGRAPHRUNTIME_API virtual bool SetMirrorDataTable(UMirrorDataTable* MirrorTable) override;
diff --git a/Engine/Source/Runtime/CoreUObject/Private/Serialization/ArchiveObjectCrc32.cpp b/Engine/Source/Runtime/CoreUObject/Private/Serialization/ArchiveObjectCrc32.cpp
--- a/Engine/Source/Runtime/CoreUObject/Private/Serialization/ArchiveObjectCrc32.cpp
+++ b/Engine/Source/Runtime/CoreUObject/Private/Serialization/ArchiveObjectCrc32.cpp
@@ -26,6 +26,8 @@ FArchiveObjectCrc32::FArchiveObjectCrc32()
 	SetIsSaving(true);
 }
 
+FArchiveObjectCrc32::~FArchiveObjectCrc32() = default;
+
 void FArchiveObjectCrc32::Serialize(void* Data, int64 Length)
 {
 	MemoryWriter.Serialize(Data, Length);
diff --git a/Engine/Source/Runtime/CoreUObject/Public/Serialization/ArchiveObjectCrc32.h b/Engine/Source/Runtime/CoreUObject/Public/Serialization/ArchiveObjectCrc32.h
--- a/Engine/Source/Runtime/CoreUObject/Public/Serialization/ArchiveObjectCrc32.h
+++ b/Engine/Source/Runtime/CoreUObject/Public/Serialization/ArchiveObjectCrc32.h
@@ -27,6 +27,7 @@ public:
 	* Default constructor.
 	*/
 	COREUOBJECT_API FArchiveObjectCrc32();
+	COREUOBJECT_API ~FArchiveObjectCrc32();
 
 	//~ Begin FArchive Interface
 	COREUOBJECT_API virtual void Serialize(void* Data, int64 Length);
diff --git a/Engine/Source/Runtime/Engine/Classes/Engine/DataTable.h b/Engine/Source/Runtime/Engine/Classes/Engine/DataTable.h
--- a/Engine/Source/Runtime/Engine/Classes/Engine/DataTable.h
+++ b/Engine/Source/Runtime/Engine/Classes/Engine/DataTable.h
@@ -70,8 +70,6 @@ class UDataTable
 {
 	GENERATED_UCLASS_BODY()
 
-	virtual ~UDataTable() {};
-
 	DECLARE_MULTICAST_DELEGATE(FOnDataTableChanged);
 	DECLARE_MULTICAST_DELEGATE(FOnDataTableImport);
 	
diff --git a/Engine/Source/Runtime/Engine/Classes/Engine/SplineMeshComponentDescriptor.h b/Engine/Source/Runtime/Engine/Classes/Engine/SplineMeshComponentDescriptor.h
--- a/Engine/Source/Runtime/Engine/Classes/Engine/SplineMeshComponentDescriptor.h
+++ b/Engine/Source/Runtime/Engine/Classes/Engine/SplineMeshComponentDescriptor.h
@@ -22,8 +22,8 @@ struct FSplineMeshComponentDescriptorBase
 	GENERATED_BODY()
 
 	ENGINE_API FSplineMeshComponentDescriptorBase();
-	explicit FSplineMeshComponentDescriptorBase(ENoInit) {}
-	virtual ~FSplineMeshComponentDescriptorBase() {}
+	ENGINE_API explicit FSplineMeshComponentDescriptorBase(ENoInit);
+	ENGINE_API virtual ~FSplineMeshComponentDescriptorBase();
 
 	ENGINE_API USplineMeshComponent* CreateComponent(UObject* Outer, FName Name = NAME_None, EObjectFlags ObjectFlags = EObjectFlags::RF_NoFlags) const;
 
@@ -204,6 +204,7 @@ struct FSplineMeshComponentDescriptor : public FSplineMeshComponentDescriptorBas
 
 	ENGINE_API FSplineMeshComponentDescriptor();
 	ENGINE_API explicit FSplineMeshComponentDescriptor(const FSoftSplineMeshComponentDescriptor& Other);
+	ENGINE_API ~FSplineMeshComponentDescriptor();
 	static ENGINE_API FSplineMeshComponentDescriptor CreateFrom(const TSubclassOf<UStaticMeshComponent>& ComponentClass);
 
 	ENGINE_API virtual void InitFrom(const UStaticMeshComponent* Component, bool bInitBodyInstance = true) override;
@@ -239,6 +240,7 @@ struct FSoftSplineMeshComponentDescriptor : public FSplineMeshComponentDescripto
 
 	ENGINE_API FSoftSplineMeshComponentDescriptor();
 	ENGINE_API explicit FSoftSplineMeshComponentDescriptor(const FSplineMeshComponentDescriptor& Other);
+	ENGINE_API ~FSoftSplineMeshComponentDescriptor();
 	static ENGINE_API FSoftSplineMeshComponentDescriptor CreateFrom(const TSubclassOf<UStaticMeshComponent>& ComponentClass);
 
 	ENGINE_API virtual void InitFrom(const UStaticMeshComponent* Component, bool bInitBodyInstance = true) override;
diff --git a/Engine/Source/Runtime/Engine/Private/ActiveSound.cpp b/Engine/Source/Runtime/Engine/Private/ActiveSound.cpp
--- a/Engine/Source/Runtime/Engine/Private/ActiveSound.cpp
+++ b/Engine/Source/Runtime/Engine/Private/ActiveSound.cpp
@@ -143,6 +143,8 @@ FActiveSound::FActiveSound()
 	}
 }
 
+FActiveSound::FActiveSound(const FActiveSound&) = default;
+
 FActiveSound::~FActiveSound()
 {
 	ensureMsgf(WaveInstances.Num() == 0, TEXT("Destroyed an active sound that had active wave instances."));
diff --git a/Engine/Source/Runtime/Engine/Private/ISMPartition/ISMComponentDescriptor.cpp b/Engine/Source/Runtime/Engine/Private/ISMPartition/ISMComponentDescriptor.cpp
--- a/Engine/Source/Runtime/Engine/Private/ISMPartition/ISMComponentDescriptor.cpp
+++ b/Engine/Source/Runtime/Engine/Private/ISMPartition/ISMComponentDescriptor.cpp
@@ -21,6 +21,10 @@ FISMComponentDescriptorBase::FISMComponentDescriptorBase()
 	InitFrom(UHierarchicalInstancedStaticMeshComponent::StaticClass()->GetDefaultObject<UHierarchicalInstancedStaticMeshComponent>());
 }
 
+FISMComponentDescriptorBase::FISMComponentDescriptorBase(ENoInit) {}
+FISMComponentDescriptorBase::FISMComponentDescriptorBase(const FISMComponentDescriptorBase&) {}
+FISMComponentDescriptorBase::~FISMComponentDescriptorBase() = default;
+
 FISMComponentDescriptor::FISMComponentDescriptor()
 	: FISMComponentDescriptorBase(NoInit)
 {
@@ -38,6 +42,9 @@ FISMComponentDescriptor::FISMComponentDescriptor(const FSoftISMComponentDescript
 	Hash = Other.Hash;
 }
 
+FISMComponentDescriptor::FISMComponentDescriptor(const FISMComponentDescriptor&) = default;
+FISMComponentDescriptor::~FISMComponentDescriptor() = default;
+
 FSoftISMComponentDescriptor::FSoftISMComponentDescriptor()
 	: FISMComponentDescriptorBase(NoInit)
 {
@@ -55,6 +62,9 @@ FSoftISMComponentDescriptor::FSoftISMComponentDescriptor(const FISMComponentDesc
 	Hash = Other.Hash;
 }
 
+FSoftISMComponentDescriptor::FSoftISMComponentDescriptor(const FSoftISMComponentDescriptor& Other) = default;
+FSoftISMComponentDescriptor::~FSoftISMComponentDescriptor() = default;
+
 FISMComponentDescriptor FISMComponentDescriptor::CreateFrom(const TSubclassOf<UStaticMeshComponent>& From)
 {
 	FISMComponentDescriptor ComponentDescriptor;
diff --git a/Engine/Source/Runtime/Engine/Private/Rendering/SkeletalMeshHalfEdgeBuffer.cpp b/Engine/Source/Runtime/Engine/Private/Rendering/SkeletalMeshHalfEdgeBuffer.cpp
--- a/Engine/Source/Runtime/Engine/Private/Rendering/SkeletalMeshHalfEdgeBuffer.cpp
+++ b/Engine/Source/Runtime/Engine/Private/Rendering/SkeletalMeshHalfEdgeBuffer.cpp
@@ -215,6 +215,9 @@ void SkeletalMeshHalfEdgeUtility::BuildHalfEdgeBuffers(const FSkeletalMeshLODRen
 	}	
 }
 
+FSkeletalMeshHalfEdgeBuffer::FSkeletalMeshHalfEdgeBuffer() = default;
+FSkeletalMeshHalfEdgeBuffer::~FSkeletalMeshHalfEdgeBuffer() = default;
+
 void FSkeletalMeshHalfEdgeBuffer::Init(const FSkeletalMeshLODRenderData& InLodRenderData)
 {
 	check(VertexToEdgeData.Num() == 0);
diff --git a/Engine/Source/Runtime/Engine/Private/SplineMeshComponentDescriptor.cpp b/Engine/Source/Runtime/Engine/Private/SplineMeshComponentDescriptor.cpp
--- a/Engine/Source/Runtime/Engine/Private/SplineMeshComponentDescriptor.cpp
+++ b/Engine/Source/Runtime/Engine/Private/SplineMeshComponentDescriptor.cpp
@@ -20,6 +20,9 @@ FSplineMeshComponentDescriptorBase::FSplineMeshComponentDescriptorBase()
 	InitFrom(USplineMeshComponent::StaticClass()->GetDefaultObject<USplineMeshComponent>());
 }
 
+FSplineMeshComponentDescriptorBase::FSplineMeshComponentDescriptorBase(ENoInit) {}
+FSplineMeshComponentDescriptorBase::~FSplineMeshComponentDescriptorBase() = default;
+
 FSplineMeshComponentDescriptor::FSplineMeshComponentDescriptor()
 	: FSplineMeshComponentDescriptorBase(NoInit)
 {
@@ -37,6 +40,8 @@ FSplineMeshComponentDescriptor::FSplineMeshComponentDescriptor(const FSoftSpline
 	Hash = Other.Hash;
 }
 
+FSplineMeshComponentDescriptor::~FSplineMeshComponentDescriptor() = default;
+
 FSoftSplineMeshComponentDescriptor::FSoftSplineMeshComponentDescriptor()
 	: FSplineMeshComponentDescriptorBase(NoInit)
 {
@@ -54,6 +59,8 @@ FSoftSplineMeshComponentDescriptor::FSoftSplineMeshComponentDescriptor(const FSp
 	Hash = Other.Hash;
 }
 
+FSoftSplineMeshComponentDescriptor::~FSoftSplineMeshComponentDescriptor() = default;
+
 FSplineMeshComponentDescriptor FSplineMeshComponentDescriptor::CreateFrom(const TSubclassOf<UStaticMeshComponent>& From)
 {
 	FSplineMeshComponentDescriptor ComponentDescriptor;
diff --git a/Engine/Source/Runtime/Engine/Public/ActiveSound.h b/Engine/Source/Runtime/Engine/Public/ActiveSound.h
--- a/Engine/Source/Runtime/Engine/Public/ActiveSound.h
+++ b/Engine/Source/Runtime/Engine/Public/ActiveSound.h
@@ -277,6 +277,7 @@ struct FActiveSound : public ISoundModulatable
 public:
 
 	ENGINE_API FActiveSound();
+	ENGINE_API FActiveSound(const FActiveSound&);
 	ENGINE_API ~FActiveSound();
 
 	static ENGINE_API FActiveSound* CreateVirtualCopy(const FActiveSound& ActiveSoundToCopy, FAudioDevice& AudioDevice);
diff --git a/Engine/Source/Runtime/Engine/Public/ISMPartition/ISMComponentDescriptor.h b/Engine/Source/Runtime/Engine/Public/ISMPartition/ISMComponentDescriptor.h
--- a/Engine/Source/Runtime/Engine/Public/ISMPartition/ISMComponentDescriptor.h
+++ b/Engine/Source/Runtime/Engine/Public/ISMPartition/ISMComponentDescriptor.h
@@ -20,8 +20,9 @@ struct FISMComponentDescriptorBase
 	GENERATED_BODY()
 
 	ENGINE_API FISMComponentDescriptorBase();
-	explicit FISMComponentDescriptorBase(ENoInit) {}
-	virtual ~FISMComponentDescriptorBase() {}
+	ENGINE_API explicit FISMComponentDescriptorBase(ENoInit);
+	ENGINE_API FISMComponentDescriptorBase(const FISMComponentDescriptorBase&);
+	ENGINE_API virtual ~FISMComponentDescriptorBase();
 
 	ENGINE_API UInstancedStaticMeshComponent* CreateComponent(UObject* Outer, FName Name = NAME_None, EObjectFlags ObjectFlags = EObjectFlags::RF_NoFlags) const;
 
@@ -216,6 +217,8 @@ struct FISMComponentDescriptor : public FISMComponentDescriptorBase
 	GENERATED_BODY()
 
 	ENGINE_API FISMComponentDescriptor();
+	ENGINE_API FISMComponentDescriptor(const FISMComponentDescriptor&);
+	ENGINE_API ~FISMComponentDescriptor();
 	ENGINE_API explicit FISMComponentDescriptor(const FSoftISMComponentDescriptor& Other);
 	static ENGINE_API FISMComponentDescriptor CreateFrom(const TSubclassOf<UStaticMeshComponent>& ComponentClass);
 
@@ -251,7 +254,9 @@ struct FSoftISMComponentDescriptor : public FISMComponentDescriptorBase
 	GENERATED_BODY()
 
 	ENGINE_API FSoftISMComponentDescriptor();
+	ENGINE_API ~FSoftISMComponentDescriptor();
 	ENGINE_API explicit FSoftISMComponentDescriptor(const FISMComponentDescriptor& Other);
+	ENGINE_API explicit FSoftISMComponentDescriptor(const FSoftISMComponentDescriptor& Other);
 	static ENGINE_API FSoftISMComponentDescriptor CreateFrom(const TSubclassOf<UStaticMeshComponent>& ComponentClass);
 
 	ENGINE_API virtual void InitFrom(const UStaticMeshComponent* Component, bool bInitBodyInstance = true) override;
diff --git a/Engine/Source/Runtime/Engine/Public/Rendering/SkeletalMeshHalfEdgeBuffer.h b/Engine/Source/Runtime/Engine/Public/Rendering/SkeletalMeshHalfEdgeBuffer.h
--- a/Engine/Source/Runtime/Engine/Public/Rendering/SkeletalMeshHalfEdgeBuffer.h
+++ b/Engine/Source/Runtime/Engine/Public/Rendering/SkeletalMeshHalfEdgeBuffer.h
@@ -60,6 +60,9 @@ public:
 		return EdgeToTwinEdgeBufferSRV;
 	}
 
+	ENGINE_API FSkeletalMeshHalfEdgeBuffer();
+	ENGINE_API ~FSkeletalMeshHalfEdgeBuffer();
+
 private:
 	uint32 GetMinBufferSize() const;
 	
diff --git a/Engine/Source/Runtime/MediaAssets/Private/Assets/MediaPlayer.cpp b/Engine/Source/Runtime/MediaAssets/Private/Assets/MediaPlayer.cpp
--- a/Engine/Source/Runtime/MediaAssets/Private/Assets/MediaPlayer.cpp
+++ b/Engine/Source/Runtime/MediaAssets/Private/Assets/MediaPlayer.cpp
@@ -1405,6 +1405,8 @@ FMediaPlayerProxy::FMediaPlayerProxy(UMediaPlayer* Player)
 	TypeHash = GetTypeHash(Player);
 }
 
+FMediaPlayerProxy::FMediaPlayerProxy(const FMediaPlayerProxy& Other) = default;
+
 FMediaPlayerProxy::~FMediaPlayerProxy()
 {
 }
diff --git a/Engine/Source/Runtime/MediaAssets/Public/MediaPlayer.h b/Engine/Source/Runtime/MediaAssets/Public/MediaPlayer.h
--- a/Engine/Source/Runtime/MediaAssets/Public/MediaPlayer.h
+++ b/Engine/Source/Runtime/MediaAssets/Public/MediaPlayer.h
@@ -145,7 +145,7 @@ public:
 
 	MEDIAASSETS_API explicit FMediaPlayerProxy(UMediaPlayer* Player);
 
-	FMediaPlayerProxy(const FMediaPlayerProxy& Other) = default;
+	MEDIAASSETS_API FMediaPlayerProxy(const FMediaPlayerProxy& Other);
 
 	MEDIAASSETS_API ~FMediaPlayerProxy();
 
diff --git a/Engine/Source/Runtime/SignalProcessing/Private/Osc.cpp b/Engine/Source/Runtime/SignalProcessing/Private/Osc.cpp
--- a/Engine/Source/Runtime/SignalProcessing/Private/Osc.cpp
+++ b/Engine/Source/Runtime/SignalProcessing/Private/Osc.cpp
@@ -25,9 +25,8 @@ namespace Audio
 	{
 	}
 
-	IOscBase::~IOscBase()
-	{
-	}
+	IOscBase::IOscBase(const IOscBase&) = default;
+	IOscBase::~IOscBase() = default;
 
 	void IOscBase::Init(const float InSampleRate, const int32 InVoiceId, FModulationMatrix* InMatrix, const int32 ModMatrixStage)
 	{
diff --git a/Engine/Source/Runtime/SignalProcessing/Public/DSP/Osc.h b/Engine/Source/Runtime/SignalProcessing/Public/DSP/Osc.h
--- a/Engine/Source/Runtime/SignalProcessing/Public/DSP/Osc.h
+++ b/Engine/Source/Runtime/SignalProcessing/Public/DSP/Osc.h
@@ -53,6 +53,7 @@ namespace Audio
 	{
 	public:
 		SIGNALPROCESSING_API IOscBase();
+		SIGNALPROCESSING_API IOscBase(const IOscBase&);
 		SIGNALPROCESSING_API virtual ~IOscBase();
 
 		// Initializes the oscillator
diff --git a/Engine/Source/Runtime/Slate/Private/Framework/Application/AnalogCursor.cpp b/Engine/Source/Runtime/Slate/Private/Framework/Application/AnalogCursor.cpp
--- a/Engine/Source/Runtime/Slate/Private/Framework/Application/AnalogCursor.cpp
+++ b/Engine/Source/Runtime/Slate/Private/Framework/Application/AnalogCursor.cpp
@@ -23,6 +23,9 @@ FAnalogCursor::FAnalogCursor()
 	AnalogValues[ static_cast< uint8 >( EAnalogStick::Right ) ] = FVector2D::ZeroVector;
 }
 
+FAnalogCursor::FAnalogCursor(const FAnalogCursor&) = default;
+FAnalogCursor::~FAnalogCursor() = default;
+
 void FAnalogCursor::Tick(const float DeltaTime, FSlateApplication& SlateApp, TSharedRef<ICursor>)
 {
 	if (TSharedPtr<FSlateUser> SlateUser = SlateApp.GetUser(GetOwnerUserIndex()))
diff --git a/Engine/Source/Runtime/Slate/Private/Widgets/Layout/SConstraintCanvas.cpp b/Engine/Source/Runtime/Slate/Private/Widgets/Layout/SConstraintCanvas.cpp
--- a/Engine/Source/Runtime/Slate/Private/Widgets/Layout/SConstraintCanvas.cpp
+++ b/Engine/Source/Runtime/Slate/Private/Widgets/Layout/SConstraintCanvas.cpp
@@ -107,6 +107,8 @@ SConstraintCanvas::SConstraintCanvas()
 	bCanSupportFocus = false;
 }
 
+SConstraintCanvas::~SConstraintCanvas() = default;
+
 void SConstraintCanvas::Construct( const SConstraintCanvas::FArguments& InArgs )
 {
 	// Sort the children based on ZOrder.
diff --git a/Engine/Source/Runtime/Slate/Public/Framework/Application/AnalogCursor.h b/Engine/Source/Runtime/Slate/Public/Framework/Application/AnalogCursor.h
--- a/Engine/Source/Runtime/Slate/Public/Framework/Application/AnalogCursor.h
+++ b/Engine/Source/Runtime/Slate/Public/Framework/Application/AnalogCursor.h
@@ -36,10 +36,10 @@ class FAnalogCursor : public IInputProcessor, public TSharedFromThis<FAnalogCurs
 {
 public:
 	SLATE_API FAnalogCursor();
+	SLATE_API FAnalogCursor(const FAnalogCursor&);
 
 	/** Dtor */
-	virtual ~FAnalogCursor()
-	{}
+	SLATE_API virtual ~FAnalogCursor();
 
 	SLATE_API virtual void Tick(const float DeltaTime, FSlateApplication& SlateApp, TSharedRef<ICursor> Cursor) override;
 
diff --git a/Engine/Source/Runtime/Slate/Public/Widgets/Layout/SConstraintCanvas.h b/Engine/Source/Runtime/Slate/Public/Widgets/Layout/SConstraintCanvas.h
--- a/Engine/Source/Runtime/Slate/Public/Widgets/Layout/SConstraintCanvas.h
+++ b/Engine/Source/Runtime/Slate/Public/Widgets/Layout/SConstraintCanvas.h
@@ -123,6 +123,7 @@ public:
 	SLATE_END_ARGS()
 
 	SLATE_API SConstraintCanvas();
+	SLATE_API ~SConstraintCanvas();
 
 	/**
 	 * Construct this widget
diff --git a/Engine/Source/Runtime/Sockets/Private/SocketSubsystem.cpp b/Engine/Source/Runtime/Sockets/Private/SocketSubsystem.cpp
--- a/Engine/Source/Runtime/Sockets/Private/SocketSubsystem.cpp
+++ b/Engine/Source/Runtime/Sockets/Private/SocketSubsystem.cpp
@@ -77,6 +77,9 @@ static IModuleInterface* LoadSubsystemModule(const FString& SubsystemName)
 #endif
 }
 
+ISocketSubsystem::ISocketSubsystem() = default;
+ISocketSubsystem::~ISocketSubsystem() = default;
+
 FUniqueSocket ISocketSubsystem::CreateUniqueSocket(const FName& SocketType, const FString& SocketDescription, bool bForceUDP)
 {
 	return FUniqueSocket(CreateSocket(SocketType, SocketDescription, bForceUDP), FSocketDeleter(this));
diff --git a/Engine/Source/Runtime/Sockets/Private/Sockets.cpp b/Engine/Source/Runtime/Sockets/Private/Sockets.cpp
--- a/Engine/Source/Runtime/Sockets/Private/Sockets.cpp
+++ b/Engine/Source/Runtime/Sockets/Private/Sockets.cpp
@@ -29,6 +29,21 @@ void FRecvMulti::CountBytes(FArchive& Ar) const
 // FSocket stats implementation
 //
 
+FSocket::FSocket() :
+	SocketType(SOCKTYPE_Unknown),
+	SocketDescription(TEXT("")),
+	SocketProtocol(NAME_None)
+{ }
+
+FSocket::FSocket(ESocketType InSocketType, const FString& InSocketDescription, const FName& InSocketProtocol) :
+	SocketType(InSocketType),
+	SocketDescription(InSocketDescription),
+	SocketProtocol(InSocketProtocol)
+{ }
+
+/** Virtual destructor. */
+FSocket::~FSocket() = default;
+
 bool FSocket::SendTo(const uint8* Data, int32 Count, int32& BytesSent, const FInternetAddr& Destination)
 {
 //	NETWORK_PROFILER(GNetworkProfiler.TrackSocketSendTo(this,Data,BytesSent,Destination));
diff --git a/Engine/Source/Runtime/Sockets/Public/SocketSubsystem.h b/Engine/Source/Runtime/Sockets/Public/SocketSubsystem.h
--- a/Engine/Source/Runtime/Sockets/Public/SocketSubsystem.h
+++ b/Engine/Source/Runtime/Sockets/Public/SocketSubsystem.h
@@ -69,7 +69,8 @@ public:
 	static SOCKETS_API void ShutdownAllSystems();
 
 
-	virtual ~ISocketSubsystem() { }
+	SOCKETS_API ISocketSubsystem();
+	SOCKETS_API virtual ~ISocketSubsystem();
 
 	/**
 	 * Does per platform initialization of the sockets library
diff --git a/Engine/Source/Runtime/Sockets/Public/Sockets.h b/Engine/Source/Runtime/Sockets/Public/Sockets.h
--- a/Engine/Source/Runtime/Sockets/Public/Sockets.h
+++ b/Engine/Source/Runtime/Sockets/Public/Sockets.h
@@ -31,11 +31,7 @@ protected:
 public:
 
 	/** Default constructor. */
-	inline FSocket() :
-		SocketType(SOCKTYPE_Unknown),
-		SocketDescription(TEXT("")),
-		SocketProtocol(NAME_None)
-	{ }
+	SOCKETS_API FSocket();
 
 	/**
 	 * Specifies the type of socket being created
@@ -76,15 +72,10 @@ public:
 	 * @param InSocketDescription The debug description of the socket
 	 * @param InSocketProtocol the protocol stack this socket should be created on.
 	 */
-	inline FSocket(ESocketType InSocketType, const FString& InSocketDescription, const FName& InSocketProtocol) :
-		SocketType(InSocketType),
-		SocketDescription(InSocketDescription),
-		SocketProtocol(InSocketProtocol)
-	{ }
+	SOCKETS_API FSocket(ESocketType InSocketType, const FString& InSocketDescription, const FName& InSocketProtocol);
 
 	/** Virtual destructor. */
-	virtual ~FSocket()
-	{ }
+	SOCKETS_API virtual ~FSocket();
 
 	/**
 	 * Shuts down the socket, making it unusable for reads and/or writes. This does not close the socket!
