Track Iris client stats when using RepGraph to help compare performance of Iris and RepGraph.

diff --git a/Engine/Source/Runtime/Engine/Private/DataChannel.cpp b/Engine/Source/Runtime/Engine/Private/DataChannel.cpp
--- a/Engine/Source/Runtime/Engine/Private/DataChannel.cpp
+++ b/Engine/Source/Runtime/Engine/Private/DataChannel.cpp
@@ -14,6 +14,7 @@
 #if UE_WITH_IRIS
 #include "Iris/ReplicationSystem/ReplicationSystem.h"
 #endif
+#include "Iris/Core/IrisProfiler.h"
 #include "Misc/MemStack.h"
 #include "Misc/ScopeExit.h"
 #include "Net/Core/Trace/Private/NetTraceInternal.h"
@@ -3130,6 +3131,13 @@ void UActorChannel::ProcessBunch( FInBunch & Bunch )
 		AActor* NewChannelActor = NULL;
 		bSpawnedNewActor = Connection->PackageMap->SerializeNewActor(Bunch, this, NewChannelActor);
 
+#if IRIS_CLIENT_PROFILER_ENABLE
+		if (bSpawnedNewActor)
+		{
+			UE::Net::FClientProfiler::RecordObjectCreate(NewChannelActor->GetFName(), false);
+		}
+#endif
+
 		// We are unsynchronized. Instead of crashing, let's try to recover.
 		if (!IsValid(NewChannelActor))
 		{
@@ -4893,6 +4901,10 @@ UObject* UActorChannel::ReadContentBlockHeader(FInBunch& Bunch, bool& bObjectDel
 			LLM_SCOPE_BYTAG(GuidCache);
 			Connection->Driver->GuidCache->ImportedNetGuids.Add( NetGUID );
 		}
+
+#if IRIS_CLIENT_PROFILER_ENABLE
+		UE::Net::FClientProfiler::RecordObjectCreate(SubObj->GetFName(), true);
+#endif
 	}
 
 	return SubObj;
diff --git a/Engine/Source/Runtime/Engine/Private/DataReplication.cpp b/Engine/Source/Runtime/Engine/Private/DataReplication.cpp
--- a/Engine/Source/Runtime/Engine/Private/DataReplication.cpp
+++ b/Engine/Source/Runtime/Engine/Private/DataReplication.cpp
@@ -25,6 +25,7 @@
 #include "HAL/LowLevelMemStats.h"
 #include "Net/Core/PushModel/Types/PushModelPerNetDriverState.h"
 #include "Net/RPCDoSDetection.h"
+#include "Iris/Core/IrisProfiler.h"
 
 DECLARE_LLM_MEMORY_STAT(TEXT("NetObjReplicator"), STAT_NetObjReplicatorLLM, STATGROUP_LLMFULL);
 LLM_DEFINE_TAG(NetObjReplicator, NAME_None, TEXT("Networking"), GET_STATFNAME(STAT_NetObjReplicatorLLM), GET_STATFNAME(STAT_NetworkingSummaryLLM));
@@ -1370,6 +1371,10 @@ bool FObjectReplicator::ReceivedRPC(FNetBitReader& Reader, const FReplicationFla
 			// Reset errors from replay driver
 			RPC_ResetLastFailedReason();
 
+#if IRIS_CLIENT_PROFILER_ENABLE
+			UE::Net::FClientProfiler::RecordRPC(Function->GetFName());
+#endif
+
 			{
 				UE::Net::FScopedNetContextRPC CallingRPC;
 				// Call the function.
diff --git a/Engine/Source/Runtime/Engine/Private/RepLayout.cpp b/Engine/Source/Runtime/Engine/Private/RepLayout.cpp
--- a/Engine/Source/Runtime/Engine/Private/RepLayout.cpp
+++ b/Engine/Source/Runtime/Engine/Private/RepLayout.cpp
@@ -27,6 +27,7 @@
 #include "UObject/EnumProperty.h"
 #if UE_WITH_IRIS
 #include "Iris/IrisConfig.h"
+#include "Iris/Core/IrisProfiler.h"
 #include "Net/Iris/ReplicationSystem/ReplicationSystemUtil.h"
 #endif // UE_WITH_IRIS
 
@@ -4655,6 +4656,10 @@ void FRepLayout::CallRepNotifies(FReceivingRepState* RepState, UObject* Object)
 		const FRepParentCmd& Parent = Parents[RepProperty->RepIndex];
 		const int32 NumParms = RepNotifyFunc->NumParms;
 
+#if IRIS_CLIENT_PROFILER_ENABLE
+		UE::Net::FClientProfiler::RecordRepNotify(RepNotifyFunc->GetFName());
+#endif
+
 		switch (NumParms)
 		{
 			case 0:
