Allow  FActorInstanceGuid::SetLevelInstanceGuid to be called many times on the same Level as long as we pass the same ptr, owner ptr and GUID

diff --git a/Engine/Source/Runtime/Engine/Private/WorldPartition/ActorInstanceGuids.cpp b/Engine/Source/Runtime/Engine/Private/WorldPartition/ActorInstanceGuids.cpp
--- a/Engine/Source/Runtime/Engine/Private/WorldPartition/ActorInstanceGuids.cpp
+++ b/Engine/Source/Runtime/Engine/Private/WorldPartition/ActorInstanceGuids.cpp
@@ -30,8 +30,24 @@ void FActorInstanceGuid::ReleaseLevelInstanceGuid(ULevel* Level)
 
 void FActorInstanceGuid::SetLevelInstanceGuid(ULevel* Level, ULevel* OwnerLevel, const FGuid& Guid, const FGuid& ResolvedGuid)
 {
-	// double registration is often an order issue (LevelInstance GUID tried to be accessed before FLevelInstanceActorImpl::OnLevelInstanceLoaded was called)
-	check(!GLevelInstanceGuids.GetAnnotationMap().Contains(Level));
+	auto NewOrIdenticalRegistration = [Level, OwnerLevel, &Guid] () -> bool
+	{
+		FLevelInstanceGuid LevelGuids = GLevelInstanceGuids.GetAnnotation(Level);
+		
+		if (!LevelGuids.IsDefault())
+		{
+			return  (LevelGuids.Level == Level) &&
+					(LevelGuids.OwnerLevel == OwnerLevel) &&
+					(LevelGuids.LevelInstanceGuid == Guid);
+		}
+
+		return true;
+	};
+	
+	// Double registration is often an order issue (LevelInstance GUID tried to be accessed before FLevelInstanceActorImpl::OnLevelInstanceLoaded was called) 
+	// But there are also some paths where we end-up reregistering the same Level* with the same GUID when Levels are reused so we allow it
+	// if the same values are passed. 
+	ensure(NewOrIdenticalRegistration());
 
 	FLevelInstanceGuid	LIGuid;
 
