Removing UpdateAudioPluginSettingsObjectCache in favour of AddStructReferencedObjects on the structs in question

#rnx

diff --git a/Engine/Source/Runtime/Engine/Classes/Sound/SoundAttenuation.h b/Engine/Source/Runtime/Engine/Classes/Sound/SoundAttenuation.h
--- a/Engine/Source/Runtime/Engine/Classes/Sound/SoundAttenuation.h
+++ b/Engine/Source/Runtime/Engine/Classes/Sound/SoundAttenuation.h
@@ -95,8 +95,20 @@ struct FSoundAttenuationPluginSettings
 	/** Settings to use with source data override audio plugin. These are defined by the plugin creator. Not all audio plugins utilize this feature. This is an array so multiple plugins can have settings. */
 	UPROPERTY(EditAnywhere, BlueprintReadWrite, Category = AttenuationSourceDataOverride, meta = (DisplayName = "Source Data Override Plugin Settings"))
 	TArray<TObjectPtr<USourceDataOverridePluginSourceSettingsBase>> SourceDataOverridePluginSettingsArray;
+
+	void AddStructReferencedObjects(FReferenceCollector& Collector);	
+};
+
+template<>
+struct TStructOpsTypeTraits<FSoundAttenuationPluginSettings> : public TStructOpsTypeTraitsBase2<FSoundAttenuationPluginSettings>
+{
+	enum
+	{
+		WithAddStructReferencedObjects = true
+	};
 };
 
+
 // Defines how to speaker map the sound when using the non-spatialized radius feature
 UENUM(BlueprintType)
 enum class ENonSpatializedRadiusSpeakerMapMode : uint8
@@ -406,18 +418,22 @@ struct FSoundAttenuationSettings : public FBaseAttenuationSettings
 	ENGINE_API float GetFocusPriorityScale(const struct FGlobalFocusSettings& FocusSettings, float FocusFactor) const;
 	ENGINE_API float GetFocusAttenuation(const struct FGlobalFocusSettings& FocusSettings, float FocusFactor) const;
 	ENGINE_API float GetFocusDistanceScale(const struct FGlobalFocusSettings& FocusSettings, float FocusFactor) const;
+
+	void AddStructReferencedObjects(FReferenceCollector& Collector);
 };
 
-#if WITH_EDITORONLY_DATA
 template<>
 struct TStructOpsTypeTraits<FSoundAttenuationSettings> : public TStructOpsTypeTraitsBase2<FSoundAttenuationSettings>
 {
 	enum 
 	{
+#if WITH_EDITORONLY_DATA
 		WithPostSerialize = true,
+#endif
+		WithAddStructReferencedObjects = true
 	};
 };
-#endif
+
 
 /** 
  * Defines how a sound changes volume with distance to the listener
diff --git a/Engine/Source/Runtime/Engine/Private/AudioDevice.cpp b/Engine/Source/Runtime/Engine/Private/AudioDevice.cpp
--- a/Engine/Source/Runtime/Engine/Private/AudioDevice.cpp
+++ b/Engine/Source/Runtime/Engine/Private/AudioDevice.cpp
@@ -563,8 +563,6 @@ bool FAudioDevice::Init(Audio::FDeviceId InDeviceID, int32 InMaxSources, int32 I
 	{
 		LLM_SCOPE(ELLMTag::AudioMixerPlugins);
 
-	// Cache any plugin settings objects we have loaded
-	UpdateAudioPluginSettingsObjectCache();
 
 	//Get the requested default spatialization plugin and set it up.
 	IAudioSpatializationFactory* SpatializationPluginFactory = AudioPluginUtilities::GetDesiredSpatializationPlugin();
@@ -1047,34 +1045,6 @@ void FAudioDevice::CountBytes(FArchive& Ar)
 	SoundMixModifiers.CountBytes(Ar);
 }
 
-void FAudioDevice::UpdateAudioPluginSettingsObjectCache()
-{
-	TRACE_CPUPROFILER_EVENT_SCOPE(FAudioDevice_UpdatePluginSettingsObjectCache);
-
-	PluginSettingsObjects.Reset();
-
-	// Make sure we don't GC 3rd party plugin settings since these live on FSoundAttenuationSettings, which may not live in UObject graph due to overrides.
-	// There shouldn't be many of these objects (on the order of 10s not 100s) so if we find any loaded, don't let GC get them.
-	for (TObjectIterator<USpatializationPluginSourceSettingsBase> It; It; ++It)
-	{
-		PluginSettingsObjects.Add(*It);
-	}
-
-	for (TObjectIterator<UOcclusionPluginSourceSettingsBase> It; It; ++It)
-	{
-		PluginSettingsObjects.Add(*It);
-	}
-
-	for (TObjectIterator<UReverbPluginSourceSettingsBase> It; It; ++It)
-	{
-		PluginSettingsObjects.Add(*It);
-	}
-
-	for (TObjectIterator<USourceDataOverridePluginSourceSettingsBase> It; It; ++It)
-	{
-		PluginSettingsObjects.Add(*It);
-	}
-}
 
 void FAudioDevice::AddReferencedObjects(FReferenceCollector& Collector)
 {
@@ -1112,9 +1082,6 @@ void FAudioDevice::AddReferencedObjects(FReferenceCollector& Collector)
 
 	// Make sure we don't try to delete any sound waves which may have in-flight decodes
 	Collector.AddReferencedObjects(ReferencedSoundWaves);
-
-	// Loop through the cached plugin settings objects and add to the collector
-	Collector.AddReferencedObjects(PluginSettingsObjects);
 }
 
 void FAudioDevice::ResetInterpolation()
@@ -4667,7 +4634,6 @@ void FAudioDevice::Update(bool bGameTicking)
 		PrimaryVolume *= FApp::GetVolumeMultiplier();
 	}
 
-	UpdateAudioPluginSettingsObjectCache();
 
 	{
 		TRACE_CPUPROFILER_EVENT_SCOPE(FAudioDevice_UpdateDeviceTiming);
diff --git a/Engine/Source/Runtime/Engine/Private/SoundAttenuation.cpp b/Engine/Source/Runtime/Engine/Private/SoundAttenuation.cpp
--- a/Engine/Source/Runtime/Engine/Private/SoundAttenuation.cpp
+++ b/Engine/Source/Runtime/Engine/Private/SoundAttenuation.cpp
@@ -229,6 +229,18 @@ void FSoundAttenuationSettings::CollectAttenuationShapesForVisualization(TMultiM
 		FBaseAttenuationSettings::CollectAttenuationShapesForVisualization(ShapeDetailsMap);
 	}
 }
+void FSoundAttenuationSettings::AddStructReferencedObjects(FReferenceCollector& Collector)
+{
+	Collector.AddStableReference(&AudioLinkSettingsOverride);
+}
+
+void FSoundAttenuationPluginSettings::AddStructReferencedObjects(FReferenceCollector& Collector)
+{
+	Collector.AddStableReferenceArray(&SpatializationPluginSettingsArray);
+	Collector.AddStableReferenceArray(&OcclusionPluginSettingsArray);
+	Collector.AddStableReferenceArray(&ReverbPluginSettingsArray);
+	Collector.AddStableReferenceArray(&SourceDataOverridePluginSettingsArray);
+}
 
 USoundAttenuation::USoundAttenuation(const FObjectInitializer& ObjectInitializer)
 	: Super(ObjectInitializer)
diff --git a/Engine/Source/Runtime/Engine/Public/AudioDevice.h b/Engine/Source/Runtime/Engine/Public/AudioDevice.h
--- a/Engine/Source/Runtime/Engine/Public/AudioDevice.h
+++ b/Engine/Source/Runtime/Engine/Public/AudioDevice.h
@@ -513,7 +513,6 @@ private:
 	void GetSoundClassInfo(TMap<FName, FAudioClassInfo>& AudioClassInfos);
 #endif
 
-	void UpdateAudioPluginSettingsObjectCache();
 
 public:
 
@@ -2145,9 +2144,6 @@ private:
 	/** Map of sound mix sound class overrides. Will override any sound class effects for any sound mixes */
 	TMap<USoundMix*, FSoundMixClassOverrideMap> SoundMixClassEffectOverrides;
 
-	/** Cached array of plugin settings objects currently loaded. This is stored so we can add it in AddReferencedObjects. */
-	TArray<TObjectPtr<UObject>> PluginSettingsObjects;
-
 protected:
 	/** Interface to audio effects processing */
 	FAudioEffectsManager* Effects;
