Cloth - Destroy the ClothingSimulation object when bAllowClothActors is false to fix issues where is is changed in Blueprints while the character is being rendered.

#tests PIE Quickplay
#rnx

diff --git a/Engine/Source/Runtime/Engine/Private/SkeletalMeshComponentPhysics.cpp b/Engine/Source/Runtime/Engine/Private/SkeletalMeshComponentPhysics.cpp
--- a/Engine/Source/Runtime/Engine/Private/SkeletalMeshComponentPhysics.cpp
+++ b/Engine/Source/Runtime/Engine/Private/SkeletalMeshComponentPhysics.cpp
@@ -2734,12 +2734,12 @@ void USkeletalMeshComponent::RecreateClothingActors()
 	ReleaseAllClothingResources();
 
 	USkeletalMesh* SkelMesh = GetSkeletalMeshAsset();
-	if(!bAllowClothActors || SkelMesh == nullptr || !IsRegistered())
+	if (SkelMesh == nullptr || !IsRegistered())
 	{
 		return;
 	}
 
-	if(CVarEnableClothPhysics.GetValueOnGameThread() && (SkelMesh->GetMeshClothingAssets().Num() > 0))
+	if (bAllowClothActors && CVarEnableClothPhysics.GetValueOnGameThread() && (SkelMesh->GetMeshClothingAssets().Num() > 0))
 	{
 		UClass* SimFactoryClass = *ClothingSimulationFactory;
 		if (SimFactoryClass)
@@ -3765,7 +3765,7 @@ const TMap<int32, FClothSimulData>& USkeletalMeshComponent::GetCurrentClothingDa
 		return SEmptyClothSimulationData;
 	}
 
-	return CurrentSimulationData;
+	return ClothingSimulation ? CurrentSimulationData : SEmptyClothSimulationData;
 }
 
 const TMap<int32, FClothSimulData>& USkeletalMeshComponent::GetCurrentClothingData_AnyThread() const
@@ -3773,7 +3773,7 @@ const TMap<int32, FClothSimulData>& USkeletalMeshComponent::GetCurrentClothingDa
 	// This is called during EndOfFrameUpdates, usually in a parallel-for loop. We need to be sure that
 	// the cloth task (if there is one) is complete, but it cannpt be waited for here. See OnPreEndOfFrameUpdateSync
 	// which is called just before EOF updates and is where we would have waited for the cloth task.
-	if (!IsValidRef(ParallelClothTask) || ParallelClothTask->IsComplete())
+	if (ClothingSimulation && (!IsValidRef(ParallelClothTask) || ParallelClothTask->IsComplete()))
 	{
 		return CurrentSimulationData;
 	}
