Updating Gravity Group setup, allowing gravity groups to be properly set via Push/Pull, and allowing the BodyInstance to offer an API for modification.

diff --git a/Engine/Source/Runtime/Engine/Classes/PhysicsEngine/BodyInstance.h b/Engine/Source/Runtime/Engine/Classes/PhysicsEngine/BodyInstance.h
--- a/Engine/Source/Runtime/Engine/Classes/PhysicsEngine/BodyInstance.h
+++ b/Engine/Source/Runtime/Engine/Classes/PhysicsEngine/BodyInstance.h
@@ -573,6 +573,10 @@ public:
 	UPROPERTY(EditAnywhere, AdvancedDisplay, BlueprintReadWrite, Category = Physics, meta = (ClampMin = "0.001", UIMin = "0.001"))
 	float MassScale;
 
+	/** What gravity group the BI should use, which determines rate of acceleration */
+	UPROPERTY(EditAnywhere,BlueprintReadOnly,Category = Physics, meta = (ClampMin = "0", UIMin = "0", UIMax = "7"))
+	uint8 GravityGroupIndex;
+
 	/** Per-instance scaling of inertia (bigger number means  it'll be harder to rotate) */
 	UPROPERTY(EditAnywhere, AdvancedDisplay, BlueprintReadWrite, Category = Physics)
 	FVector InertiaTensorScale;
@@ -931,6 +935,8 @@ public:
 	ENGINE_API void SetInstanceNotifyRBCollision(bool bNewNotifyCollision);
 	/** Enables/disables whether this body is affected by gravity. */
 	ENGINE_API void SetEnableGravity(bool bGravityEnabled);
+	/** Sets the gravity group index, which determines acceleration when gravity is enabled. */
+	ENGINE_API void SetGravityGroupIndex(int32 NewGravityGroupIndex);
 	/** Enables/disables whether this body, when kinematic, is updated from the simulation rather than when setting the kinematic target. */
 	ENGINE_API void SetUpdateKinematicFromSimulation(bool bUpdateKinematicFromSimulation);
 	/** Enables/disables contact modification */
diff --git a/Engine/Source/Runtime/Engine/Private/PhysicsEngine/BodyInstance.cpp b/Engine/Source/Runtime/Engine/Private/PhysicsEngine/BodyInstance.cpp
--- a/Engine/Source/Runtime/Engine/Private/PhysicsEngine/BodyInstance.cpp
+++ b/Engine/Source/Runtime/Engine/Private/PhysicsEngine/BodyInstance.cpp
@@ -3673,6 +3673,24 @@ void FBodyInstance::SetEnableGravity(bool bInGravityEnabled)
 	}
 }
 
+void FBodyInstance::SetGravityGroupIndex(int32 NewGravityGroupIndex)
+{
+	if (GravityGroupIndex >= 0)
+	{
+		GravityGroupIndex = NewGravityGroupIndex;
+
+		{
+			FPhysicsCommand::ExecuteWrite(ActorHandle, [&](const FPhysicsActorHandle& Actor)
+				{
+					if (FPhysicsInterface::IsRigidBody(Actor))
+					{
+						FPhysicsInterface::SetGravityGroupIndex_AssumesLocked(Actor, GravityGroupIndex);
+					}
+				});
+		}
+	}
+}
+
 void FBodyInstance::SetUpdateKinematicFromSimulation(bool bInUpdateKinematicFromSimulation)
 {
 	if (bUpdateKinematicFromSimulation != bInUpdateKinematicFromSimulation)
diff --git a/Engine/Source/Runtime/Engine/Private/PhysicsEngine/ImmediatePhysics/ImmediatePhysicsChaos/ImmediatePhysicsActorHandle_Chaos.cpp b/Engine/Source/Runtime/Engine/Private/PhysicsEngine/ImmediatePhysics/ImmediatePhysicsChaos/ImmediatePhysicsActorHandle_Chaos.cpp
--- a/Engine/Source/Runtime/Engine/Private/PhysicsEngine/ImmediatePhysics/ImmediatePhysicsChaos/ImmediatePhysicsActorHandle_Chaos.cpp
+++ b/Engine/Source/Runtime/Engine/Private/PhysicsEngine/ImmediatePhysics/ImmediatePhysicsChaos/ImmediatePhysicsActorHandle_Chaos.cpp
@@ -259,6 +259,7 @@ namespace ImmediatePhysics_Chaos
 					Dynamic->SetLinearEtherDrag(BodyInstance->LinearDamping);
 					Dynamic->SetAngularEtherDrag(BodyInstance->AngularDamping);
 					Dynamic->SetGravityEnabled(BodyInstance->bEnableGravity);
+					Dynamic->SetGravityGroupIndex(BodyInstance->GravityGroupIndex);
 					Dynamic->SetUpdateKinematicFromSimulation(BodyInstance->bUpdateKinematicFromSimulation);
 				}
 				Dynamic->SetDisabled(true);
diff --git a/Engine/Source/Runtime/Experimental/Chaos/Public/Chaos/ParticleDirtyFlags.h b/Engine/Source/Runtime/Experimental/Chaos/Public/Chaos/ParticleDirtyFlags.h
--- a/Engine/Source/Runtime/Experimental/Chaos/Public/Chaos/ParticleDirtyFlags.h
+++ b/Engine/Source/Runtime/Experimental/Chaos/Public/Chaos/ParticleDirtyFlags.h
@@ -533,6 +533,8 @@ public:
 
 	bool GravityEnabled() const { return MControlFlags.GetGravityEnabled(); }
 	void SetGravityEnabled(bool bInGravity){ MControlFlags.SetGravityEnabled(bInGravity); }
+	int32 GravityGroupIndex() const { return MControlFlags.GetGravityGroupIndex(); }
+	void SetGravityGroupIndex(int32 InIndex) { MControlFlags.SetGravityGroupIndex(InIndex); }
 
 	bool UpdateKinematicFromSimulation() const { return MControlFlags.GetUpdateKinematicFromSimulation(); }
 	void SetUpdateKinematicFromSimulation(bool bUpdateKinematicFromSimulation) { MControlFlags.SetUpdateKinematicFromSimulation(bUpdateKinematicFromSimulation); }
diff --git a/Engine/Source/Runtime/Experimental/Chaos/Public/Chaos/ParticleHandle.h b/Engine/Source/Runtime/Experimental/Chaos/Public/Chaos/ParticleHandle.h
--- a/Engine/Source/Runtime/Experimental/Chaos/Public/Chaos/ParticleHandle.h
+++ b/Engine/Source/Runtime/Experimental/Chaos/Public/Chaos/ParticleHandle.h
@@ -125,6 +125,7 @@ void PBDRigidParticleDefaultConstruct(FConcrete& Concrete, const FPBDRigidPartic
 	Concrete.SetLinearEtherDrag(0.f);
 	Concrete.SetAngularEtherDrag(0.f);
 	Concrete.SetGravityEnabled(Params.bGravityEnabled);
+	Concrete.SetGravityGroupIndex(0);
 	Concrete.SetCCDEnabled(Params.bCCDEnabled);
 	Concrete.SetMACDEnabled(false);
 	Concrete.SetDisabled(Params.bDisabled);
@@ -3350,6 +3351,12 @@ public:
 	{
 		MMiscData.Modify(true, MDirtyFlags, Proxy, [bInEnabled](auto& Data) { Data.SetGravityEnabled(bInEnabled); });
 	}
+
+	int32 GravityGroupIndex() const { return MMiscData.Read().GravityGroupIndex(); }
+	void SetGravityGroupIndex(int32 NewIndex)
+	{
+		MMiscData.Modify(true, MDirtyFlags, Proxy, [NewIndex](auto& Data) { Data.SetGravityGroupIndex(NewIndex); });
+	}
 	
 	bool UpdateKinematicFromSimulation() const { return MMiscData.Read().UpdateKinematicFromSimulation(); }
 	void SetUpdateKinematicFromSimulation(const bool bUpdateKinematicFromSimulation)
diff --git a/Engine/Source/Runtime/Experimental/Chaos/Public/Chaos/RigidParticles.h b/Engine/Source/Runtime/Experimental/Chaos/Public/Chaos/RigidParticles.h
--- a/Engine/Source/Runtime/Experimental/Chaos/Public/Chaos/RigidParticles.h
+++ b/Engine/Source/Runtime/Experimental/Chaos/Public/Chaos/RigidParticles.h
@@ -65,7 +65,7 @@ struct FRigidParticleCoreData
 {
 	int32 CollisionGroup;							// 4 bytes
 	uint32 CollisionConstraintFlags;				// 4 bytes
-	FRigidParticleControlFlags ControlFlags;		// 1 byte
+	FRigidParticleControlFlags ControlFlags;		// 2 bytes
 	FRigidParticleTransientFlags TransientFlags;	// 1 byte
 	EObjectStateType ObjectState;					// 1 byte
 	EObjectStateType PreObjectState;				// 1 byte
diff --git a/Engine/Source/Runtime/Experimental/Chaos/Public/PhysicsProxy/SingleParticlePhysicsProxy.h b/Engine/Source/Runtime/Experimental/Chaos/Public/PhysicsProxy/SingleParticlePhysicsProxy.h
--- a/Engine/Source/Runtime/Experimental/Chaos/Public/PhysicsProxy/SingleParticlePhysicsProxy.h
+++ b/Engine/Source/Runtime/Experimental/Chaos/Public/PhysicsProxy/SingleParticlePhysicsProxy.h
@@ -177,6 +177,8 @@ protected:
 	FParticleHandle* Handle;
 	FPhysicsObjectUniquePtr Reference;
 
+	int32 GravityGroupIndex;
+
 private:
 	TUniquePtr<FProxyInterpolationBase> InterpolationData;
 
@@ -409,6 +411,30 @@ public:
 		});
 	}
 
+	int32 GravityGroupIndex() const
+	{
+		return Read([](auto* Particle)
+		{
+			if (auto Rigid = Particle->CastToRigidParticle())
+			{
+				return Rigid->GravityGroupIndex();
+			}
+
+			return -1;
+		});
+	}
+
+	void SetGravityGroupIndex(const uint32 InGravityGroupIndex)
+	{
+		Write([InGravityGroupIndex](auto* Particle)
+			{
+				if (auto Rigid = Particle->CastToRigidParticle())
+				{
+					return Rigid->SetGravityGroupIndex(InGravityGroupIndex);
+				}
+			});
+	}
+
 	bool UpdateKinematicFromSimulation() const
 	{
 		return Read([](auto* Particle)
diff --git a/Engine/Source/Runtime/PhysicsCore/Private/ChaosEngineInterface.cpp b/Engine/Source/Runtime/PhysicsCore/Private/ChaosEngineInterface.cpp
--- a/Engine/Source/Runtime/PhysicsCore/Private/ChaosEngineInterface.cpp
+++ b/Engine/Source/Runtime/PhysicsCore/Private/ChaosEngineInterface.cpp
@@ -1164,6 +1164,16 @@ void FChaosEngineInterface::SetGravityEnabled_AssumesLocked(const FPhysicsActorH
 	InActorReference->GetGameThreadAPI().SetGravityEnabled(bEnabled);
 }
 
+PHYSICSCORE_API int32 FChaosEngineInterface::GetGravityGroupIndex_AssumesLocked(const FPhysicsActorHandle& InActorReference)
+{
+	return InActorReference->GetGameThreadAPI().GravityGroupIndex();
+}
+
+void FChaosEngineInterface::SetGravityGroupIndex_AssumesLocked(const FPhysicsActorHandle& InActorReference, uint32 Index)
+{
+	InActorReference->GetGameThreadAPI().SetGravityGroupIndex(Index);
+}
+
 bool FChaosEngineInterface::GetUpdateKinematicFromSimulation_AssumesLocked(const FPhysicsActorHandle& InActorReference)
 {
 	return InActorReference->GetGameThreadAPI().UpdateKinematicFromSimulation();
diff --git a/Engine/Source/Runtime/PhysicsCore/Public/Chaos/ChaosEngineInterface.h b/Engine/Source/Runtime/PhysicsCore/Public/Chaos/ChaosEngineInterface.h
--- a/Engine/Source/Runtime/PhysicsCore/Public/Chaos/ChaosEngineInterface.h
+++ b/Engine/Source/Runtime/PhysicsCore/Public/Chaos/ChaosEngineInterface.h
@@ -458,6 +458,8 @@ public:
 
 	static PHYSICSCORE_API bool IsGravityEnabled_AssumesLocked(const FPhysicsActorHandle& InActorReference);
 	static PHYSICSCORE_API void SetGravityEnabled_AssumesLocked(const FPhysicsActorHandle& InActorReference,bool bEnabled);
+	static PHYSICSCORE_API int32 GetGravityGroupIndex_AssumesLocked(const FPhysicsActorHandle& InActorReference);
+	static PHYSICSCORE_API void SetGravityGroupIndex_AssumesLocked(const FPhysicsActorHandle& InActorReference, uint32 Index);
 
 	static PHYSICSCORE_API bool GetUpdateKinematicFromSimulation_AssumesLocked(const FPhysicsActorHandle& InActorReference);
 	static PHYSICSCORE_API void SetUpdateKinematicFromSimulation_AssumesLocked(const FPhysicsActorHandle& InActorReference, bool bUpdateKinematicFromSimulation);
