Modifying gravity group settor function in BodyInstance.cpp

diff --git a/Engine/Source/Runtime/Engine/Private/PhysicsEngine/BodyInstance.cpp b/Engine/Source/Runtime/Engine/Private/PhysicsEngine/BodyInstance.cpp
--- a/Engine/Source/Runtime/Engine/Private/PhysicsEngine/BodyInstance.cpp
+++ b/Engine/Source/Runtime/Engine/Private/PhysicsEngine/BodyInstance.cpp
@@ -3675,7 +3675,7 @@ void FBodyInstance::SetEnableGravity(bool bInGravityEnabled)
 
 void FBodyInstance::SetGravityGroupIndex(int32 NewGravityGroupIndex)
 {
-	if (GravityGroupIndex >= 0)
+	if (NewGravityGroupIndex >= 0)
 	{
 		GravityGroupIndex = NewGravityGroupIndex;
 
