Fix broken copy constructor on FISMComponentDescriptorBase which prevented PCG from generating ISMC/HISMC components correctly

Breakage: 35712392

#rnx

diff --git a/Engine/Source/Runtime/Engine/Private/ISMPartition/ISMComponentDescriptor.cpp b/Engine/Source/Runtime/Engine/Private/ISMPartition/ISMComponentDescriptor.cpp
--- a/Engine/Source/Runtime/Engine/Private/ISMPartition/ISMComponentDescriptor.cpp
+++ b/Engine/Source/Runtime/Engine/Private/ISMPartition/ISMComponentDescriptor.cpp
@@ -22,7 +22,7 @@ FISMComponentDescriptorBase::FISMComponentDescriptorBase()
 }
 
 FISMComponentDescriptorBase::FISMComponentDescriptorBase(ENoInit) {}
-FISMComponentDescriptorBase::FISMComponentDescriptorBase(const FISMComponentDescriptorBase&) {}
+FISMComponentDescriptorBase::FISMComponentDescriptorBase(const FISMComponentDescriptorBase&) = default;
 FISMComponentDescriptorBase::~FISMComponentDescriptorBase() = default;
 
 FISMComponentDescriptor::FISMComponentDescriptor()
