FActorIterator fixes:
- Fix when iterating over specific actor class, which could return actors not part of any level when unloading world areas.
- Fix when iterating over the base actor class on an uninitialized world, which would return no actors, skipping over the persistent level.

diff --git a/Engine/Source/Runtime/Engine/Public/EngineUtils.h b/Engine/Source/Runtime/Engine/Public/EngineUtils.h
--- a/Engine/Source/Runtime/Engine/Public/EngineUtils.h
+++ b/Engine/Source/Runtime/Engine/Public/EngineUtils.h
@@ -184,41 +184,53 @@ public:
 		check(CurrentWorld);
 
 #if WITH_EDITOR
-		// In the editor, you are more likely to have many worlds in memory at once.
-		// As an optimization to avoid iterating over many actors that are not in the world we are asking for,
-		// if the filter class is AActor, just use the actors that are in the world you asked for.
-		// This could be useful in runtime code as well if there are many worlds in memory, but for now we will leave
-		// it in editor code.
-		if (InClass == AActor::StaticClass())
+		if (InClass != AActor::StaticClass())
 		{
-			// First determine the number of actors in the world to reduce reallocations when we append them to the array below.
-			int32 NumActors = 0;
-			for (ULevel* Level : InWorld->GetLevels())
+			TArray<UObject*> ObjectsOfClass;
+			GetObjectsOfClass(InClass, ObjectsOfClass, true, RF_ClassDefaultObject, EInternalObjectFlags::Garbage);
+
+			if (ObjectsOfClass.IsEmpty())
 			{
-				if (Level)
-				{
-					NumActors += Level->Actors.Num();
-				}
+				return;
 			}
+		}
 
-			// Presize the array
-			ObjectArray.Reserve(NumActors);
+		// In the editor, you are more likely to have many worlds in memory at once.
+		// As an optimization to avoid iterating over many actors that are not in the world we are asking for,
+		// just use the actors that are in the world you asked for (persistent level + streaming levels basically).
+		// This could be useful in runtime code as well if there are many worlds in memory, but for now we will
+		// leave it in editor code.
 
-			// Fill the array
-			for (ULevel* Level : InWorld->GetLevels())
+		// First determine the number of actors in the world to reduce reallocations when we append them to the array below.
+		int32 MaxActors = InWorld->PersistentLevel ? InWorld->PersistentLevel->Actors.Num() : 0;
+		for (ULevel* Level : InWorld->GetLevels())
+		{
+			if (Level && Level != InWorld->PersistentLevel)
 			{
-				if (Level)
-				{
-					ObjectArray.Append(Level->Actors);
-				}
+				MaxActors += Level->Actors.Num();
 			}
 		}
-		else
-#endif // WITH_EDITOR
+		ObjectArray.Reserve(MaxActors);
+
+		// Add persistent level actors, we need to do this explicitly to handle worlds that are not initialized yet, as
+		// InWorld->GetLevels() would return an empty level list even if the world's persistent level contains actors.
+		if (InWorld->PersistentLevel)
+		{
+			ObjectArray.Append(InWorld->PersistentLevel->Actors);
+		}
+
+		// Add streaming levels actors
+		for (ULevel* Level : InWorld->GetLevels())
 		{
-			constexpr EObjectFlags ExcludeFlags = RF_ClassDefaultObject;
-			GetObjectsOfClass(InClass, ObjectArray, true, ExcludeFlags, EInternalObjectFlags::Garbage);
+			if (Level && Level != InWorld->PersistentLevel)
+			{
+				ObjectArray.Append(Level->Actors);
+			}
 		}
+#else // WITH_EDITOR
+		constexpr EObjectFlags ExcludeFlags = RF_ClassDefaultObject;
+		GetObjectsOfClass(InClass, ObjectArray, true, ExcludeFlags, EInternalObjectFlags::Garbage);
+#endif
 
 		const auto ActorSpawnedDelegate = FOnActorSpawned::FDelegate::CreateRaw(this, &FActorIteratorState::OnActorSpawned);
 		ActorSpawnedDelegateHandle = CurrentWorld->AddOnActorSpawnedHandler(ActorSpawnedDelegate);
@@ -400,6 +412,13 @@ protected:
 			return false;
 		}
 
+#if WITH_EDITOR
+		// In editor, we don't filter actors by class when creating the iterator.
+		if (State->DesiredClass != AActor::StaticClass() && !Actor->IsA(State->DesiredClass))
+		{
+			return false;
+		}
+#endif
 		return true;
 	}
 
@@ -973,4 +992,4 @@ namespace VirtualTextureUtils
 {
 	/** Function that will test if the passed in texture is VT. If so, print to Messagelog that the property does not support VT textures */
 	ENGINE_API void CheckAndReportInvalidUsage(const UObject* Owner, const FName& PropertyName, const UTexture* Texture);
-}
+}
\ No newline at end of file
