[World Partition Class Descriptors] Fix typo in comments.

diff --git a/Engine/Source/Runtime/Engine/Private/WorldPartition/WorldPartitionClassDescRegistry.cpp b/Engine/Source/Runtime/Engine/Private/WorldPartition/WorldPartitionClassDescRegistry.cpp
--- a/Engine/Source/Runtime/Engine/Private/WorldPartition/WorldPartitionClassDescRegistry.cpp
+++ b/Engine/Source/Runtime/Engine/Private/WorldPartition/WorldPartitionClassDescRegistry.cpp
@@ -86,7 +86,7 @@ static FAutoConsoleCommand DumpClassDescs(
 
 /*
  * We have to deal with some different realities between non-cooked and cooked editor builds: for non-cooked editor builds, UBlueprint are considered assets
- * and will be the opnly thing visible from the asset registry standpoint. On the other hand, for cooked editor builds, UBlueprint are not considered assets
+ * and will be the only thing visible from the asset registry standpoint. On the other hand, for cooked editor builds, UBlueprint are not considered assets
  * while their corresponding UBlueprintGeneratedClass will be. Also, we expect redirectors only in non-cooked editor builds.
  */	
 static FTopLevelAssetPath GetAssetDataClassNameForBlueprint(const FString& InAssetDataClassName)
