Substrate - Fixed specular asset profile duplication preventing them to work at the same time in the material graph.

#rnx

diff --git a/Engine/Source/Runtime/Engine/Classes/Engine/SpecularProfile.h b/Engine/Source/Runtime/Engine/Classes/Engine/SpecularProfile.h
--- a/Engine/Source/Runtime/Engine/Classes/Engine/SpecularProfile.h
+++ b/Engine/Source/Runtime/Engine/Classes/Engine/SpecularProfile.h
@@ -81,6 +81,7 @@ class USpecularProfile : public UObject
 	//~ Begin UObject Interface
 	virtual void BeginDestroy();
 	virtual void PostEditChangeProperty(struct FPropertyChangedEvent& PropertyChangedEvent);
+	virtual void PostDuplicate(EDuplicateMode::Type DuplicateMode);
 	//~ End UObject Interface
 };
 
diff --git a/Engine/Source/Runtime/Engine/Private/Rendering/SpecularProfile.cpp b/Engine/Source/Runtime/Engine/Private/Rendering/SpecularProfile.cpp
--- a/Engine/Source/Runtime/Engine/Private/Rendering/SpecularProfile.cpp
+++ b/Engine/Source/Runtime/Engine/Private/Rendering/SpecularProfile.cpp
@@ -544,6 +544,14 @@ void USpecularProfile::PostEditChangeProperty(struct FPropertyChangedEvent& Prop
 	});
 }
 
+void USpecularProfile::PostDuplicate(EDuplicateMode::Type DuplicateMode)
+{
+	Super::PostDuplicate(DuplicateMode);
+
+	// When a Specular Profile asset is duplicated/copied pasted (e.g. from the asset browser), we want the guid to be regenerated.
+	Guid = FGuid::NewGuid();
+}
+
 ///////////////////////////////////////////////////////////////////////////////////////////////////
 // Public API
 
