Add project setting for maximum size of mesh paint textures.

diff --git a/Engine/Source/Runtime/Engine/Classes/Engine/RendererSettings.h b/Engine/Source/Runtime/Engine/Classes/Engine/RendererSettings.h
--- a/Engine/Source/Runtime/Engine/Classes/Engine/RendererSettings.h
+++ b/Engine/Source/Runtime/Engine/Classes/Engine/RendererSettings.h
@@ -466,6 +466,12 @@ class URendererSettings : public UDeveloperSettings
 		ToolTip = "Default ratio of texels to vertices when creating a mesh paint texture for a mesh."))
 	int32 MeshPaintDefaultTexelsPerVertex;
 
+	UPROPERTY(config, EditAnywhere, Category = VirtualTextures, meta = (
+		EditCondition = "bVirtualTextures",
+		ConsoleVariable = "r.MeshPaintVirtualTexture.MaxTextureSize", DisplayName = "Mesh paint maximum texture size",
+		ToolTip = "Maximum size in pixels for mesh paint virtual textures. Will be rounded to next power-of-2."))
+	int32 MeshPaintVirtualTextureMaxTextureSize;
+
 	UPROPERTY(config, EditAnywhere, Category = "Runtime Virtual Textures|Enable Material Types", meta = (
 		EditCondition = "bVirtualTextures",
 		ConsoleVariable = "r.vt.rvt.EnableBaseColor", DisplayName = "Enable 'Base Color'",
diff --git a/Engine/Source/Runtime/Engine/Private/VT/MeshPaintVirtualTexture.cpp b/Engine/Source/Runtime/Engine/Private/VT/MeshPaintVirtualTexture.cpp
--- a/Engine/Source/Runtime/Engine/Private/VT/MeshPaintVirtualTexture.cpp
+++ b/Engine/Source/Runtime/Engine/Private/VT/MeshPaintVirtualTexture.cpp
@@ -47,6 +47,12 @@ static TAutoConsoleVariable<int32> CVarMeshPaintVirtualTextureTexelsPerVertex(
 	TEXT("Default ratio of texels to vertices when creating a texture for a mesh"),
 	ECVF_Default);
 
+static TAutoConsoleVariable<int32> CVarMeshPaintVirtualTextureMaxTextureSize(
+	TEXT("r.MeshPaintVirtualTexture.MaxTextureSize"),
+	4096,
+	TEXT("Maximum allowed size for mesh paint textures"),
+	ECVF_Default);
+
 namespace MeshPaintVirtualTexture
 {
 	bool IsSupported(EShaderPlatform InShaderPlatform)
@@ -88,12 +94,18 @@ namespace MeshPaintVirtualTexture
 		return FVirtualTextureBuildSettings::ClampAndAlignTileBorderSize(CVarMeshPaintVirtualTextureTileBorderSize.GetValueOnAnyThread());
 	}
 
+	static uint32 GetMaximumTextureSize()
+	{
+		// Maximum texture size needs to be at least one tile and power of two aligned.
+		return FMath::RoundUpToPowerOfTwo(FMath::Max((int32)GetTileSize(), CVarMeshPaintVirtualTextureMaxTextureSize.GetValueOnGameThread()));
+	}
+
 	uint32 GetAlignedTextureSize(int32 InSize)
 	{
 		// Must be power of 2 aligned.
 		const uint32 SizePow2Aligned = FMath::RoundUpToPowerOfTwo((uint32)FMath::Max(InSize, 0));
-		// Must be at least the size of a tile (which is also power of two aligned).
-		return FMath::Max(GetTileSize(), SizePow2Aligned);
+		// Must be at least the size of a tile and clamped to the maximum (both of which are also power of two aligned).
+		return FMath::Clamp(SizePow2Aligned, GetTileSize(), GetMaximumTextureSize());
 	}
 
 	uint32 GetDefaultTextureSize(int32 InNumVertices)
