[Material Editor] New Material Translator push.

- Adds support for Material Function calls. Editing of material function themselves is  still unsupported.
- Renamed ArithmeticType to PrimitiveType (Arithmetic are ints, floats, etc, bools aren't)
- Added optimizations to binary operator emission. It will now decompose the operation between vector types and emit one operation per component only when it brings a benefit, that is when at least the operation between the i-th component of  lhs and rhs values could be folded.
- Various other fixes and improvements.

diff --git a/Engine/Source/Runtime/Engine/Private/Materials/MaterialExpressionsIR.cpp b/Engine/Source/Runtime/Engine/Private/Materials/MaterialExpressionsIR.cpp
--- a/Engine/Source/Runtime/Engine/Private/Materials/MaterialExpressionsIR.cpp
+++ b/Engine/Source/Runtime/Engine/Private/Materials/MaterialExpressionsIR.cpp
@@ -296,25 +296,25 @@ void UMaterialExpressionConstant4Vector::Build(MIR::FEmitter& Emitter)
 
 void UMaterialExpressionAppendVector::Build(MIR::FEmitter& Emitter)
 {
-	MIR::FValue* AValue = Emitter.GetArithmetic(&A);
-	MIR::FValue* BValue = Emitter.TryGet(&B) ? Emitter.GetArithmetic(&B) : nullptr;
+	MIR::FValue* AValue = Emitter.GetPrimitive(&A);
+	MIR::FValue* BValue = Emitter.TryGet(&B) ? Emitter.GetPrimitive(&B) : nullptr;
 
 	if (Emitter.IsInvalid())
 	{
 		return;
 	}
 
-	MIR::FArithmeticTypePtr AType = AValue->Type->AsArithmetic();
+	MIR::FPrimitiveTypePtr AType = AValue->Type->AsPrimitive();
 	if (AType->IsMatrix())
 	{
 		Emitter.Error(TEXT("Input A does not have scalar or vector type."));
 		return;
 	}
 
-	MIR::FArithmeticTypePtr BType = nullptr;
+	MIR::FPrimitiveTypePtr BType = nullptr;
 	if (BValue)
 	{
-		BType = BValue->Type->AsArithmetic();
+		BType = BValue->Type->AsPrimitive();
 		if (BType->IsMatrix()) 
 		{
 			Emitter.Error(TEXT("Input B does not have scalar or vector type."));
@@ -332,7 +332,7 @@ void UMaterialExpressionAppendVector::Build(MIR::FEmitter& Emitter)
 	check(Dimensions >= 2 && Dimensions <= 4);
 
 	// Construct the output vector type.
-	MIR::FArithmeticTypePtr ResultType = MIR::FArithmeticType::GetVector(MIR::SK_Float, Dimensions);
+	MIR::FPrimitiveTypePtr ResultType = MIR::FPrimitiveType::GetVector(MIR::SK_Float, Dimensions);
 
 	// Set up each output vector component 
 	MIR::FValue* Components[4] = { nullptr, nullptr, nullptr, nullptr };
@@ -378,9 +378,9 @@ static void BuildBinaryArithmeticOperator(
 	FExpressionInput* RhsInput,
 	float RhsConst)
 {
-	// Default inputs to their relative constants if disconnected, then get each input after checking it has arithmetic type.
-	MIR::FValue* LhsValue = Emitter.DefaultTo(LhsInput, LhsConst).GetArithmetic(LhsInput);
-	MIR::FValue* RhsValue = Emitter.DefaultTo(RhsInput, RhsConst).GetArithmetic(RhsInput);
+	// Default inputs to their relative constants if disconnected, then get each input after checking it has primitive type.
+	MIR::FValue* LhsValue = Emitter.DefaultTo(LhsInput, LhsConst).GetPrimitive(LhsInput);
+	MIR::FValue* RhsValue = Emitter.DefaultTo(RhsInput, RhsConst).GetPrimitive(RhsInput);
 
 	if (Emitter.IsInvalid())
 	{
@@ -390,7 +390,7 @@ static void BuildBinaryArithmeticOperator(
 	// Determine operation input/output type by looking at the first connected input and picking float1 otherwise.
 	MIR::FTypePtr ResultType = LhsValue ? LhsValue->Type
 		: RhsValue ? RhsValue->Type
-		: MIR::FArithmeticType::GetScalar(MIR::SK_Float);
+		: MIR::FPrimitiveType::GetScalar(MIR::SK_Float);
 
 	// Convert operand values to determined result type. 
 	LhsValue = Emitter.EmitConstruct(ResultType, LhsValue);
@@ -440,17 +440,17 @@ void UMaterialExpressionIf::Build(MIR::FEmitter& Emitter)
 	// Get input values and check their types are what we expect.
 	MIR::FValue* AValue = Emitter.GetScalar(&A);
 	MIR::FValue* BValue = Emitter.GetScalar(&B);
-	MIR::FValue* AGreaterThanBValue = Emitter.GetArithmetic(&AGreaterThanB);
-	MIR::FValue* AEqualsBValue = Emitter.GetArithmetic(&AEqualsB);
-	MIR::FValue* ALessThanBValue = Emitter.GetArithmetic(&ALessThanB);
+	MIR::FValue* AGreaterThanBValue = Emitter.GetPrimitive(&AGreaterThanB);
+	MIR::FValue* AEqualsBValue = Emitter.GetPrimitive(&AEqualsB);
+	MIR::FValue* ALessThanBValue = Emitter.GetPrimitive(&ALessThanB);
 
 	if (Emitter.IsInvalid())
 	{
 		return;
 	}
 
-	// Get the arithmetic common type between the conditional arguments (e.g. if inputs are int and float, it will return float).
-	MIR::FArithmeticTypePtr ConditionArgsType = Emitter.GetCommonArithmeticType(AValue->Type->AsArithmetic(), BValue->Type->AsArithmetic());
+	// Get the primitive common type between the conditional arguments (e.g. if inputs are int and float, it will return float).
+	MIR::FPrimitiveTypePtr ConditionArgsType = Emitter.GetCommonPrimitiveType(AValue->Type->AsPrimitive(), BValue->Type->AsPrimitive());
 
 	if (Emitter.IsInvalid())
 	{
@@ -461,9 +461,9 @@ void UMaterialExpressionIf::Build(MIR::FEmitter& Emitter)
 	AValue = Emitter.EmitConstruct(ConditionArgsType, AValue);
 	BValue = Emitter.EmitConstruct(ConditionArgsType, BValue);
 
-	// Now determine the output type by taking the common arithmetic type between result values.
-	MIR::FArithmeticTypePtr OutputType = Emitter.GetCommonArithmeticType(AGreaterThanBValue->Type->AsArithmetic(), AEqualsBValue->Type->AsArithmetic());
-	OutputType = Emitter.GetCommonArithmeticType(OutputType, ALessThanBValue->Type->AsArithmetic());
+	// Now determine the output type by taking the common primitive type between result values.
+	MIR::FPrimitiveTypePtr OutputType = Emitter.GetCommonPrimitiveType(AGreaterThanBValue->Type->AsPrimitive(), AEqualsBValue->Type->AsPrimitive());
+	OutputType = Emitter.GetCommonPrimitiveType(OutputType, ALessThanBValue->Type->AsPrimitive());
 
 	if (Emitter.IsInvalid())
 	{
@@ -491,11 +491,10 @@ void UMaterialExpressionIf::Build(MIR::FEmitter& Emitter)
 	Emitter.Put(GetOutput(0), OutputValue);
 }
 
-
 void UMaterialExpressionTextureSample::Build(UE::MIR::FEmitter& Emitter)
 {
 	UTexture* InputTexture = Texture;
-	if (MIR::FValue* InputTextureValue = Emitter.Get(&TextureObject))
+	if (MIR::FValue* InputTextureValue = Emitter.TryGet(&TextureObject))
 	{
 		InputTexture = InputTextureValue->GetTexture();
 		if (!InputTexture)
@@ -510,7 +509,7 @@ void UMaterialExpressionTextureSample::Build(UE::MIR::FEmitter& Emitter)
 		return;
 	}
 
-	MIR::FValue* TexCoordsValue = Emitter.Get(&Coordinates);
+	MIR::FValue* TexCoordsValue = Emitter.TryGet(&Coordinates);
 	if (!TexCoordsValue)
 	{
 		TexCoordsValue = Emitter.GetExternalInput(MIR::TexCoordIndexToExternalInput(ConstCoordinate));
diff --git a/Engine/Source/Runtime/Engine/Private/Materials/MaterialIR.cpp b/Engine/Source/Runtime/Engine/Private/Materials/MaterialIR.cpp
--- a/Engine/Source/Runtime/Engine/Private/Materials/MaterialIR.cpp
+++ b/Engine/Source/Runtime/Engine/Private/Materials/MaterialIR.cpp
@@ -39,6 +39,14 @@ const FInstruction* FValue::AsInstruction() const
 
 bool FValue::Equals(const FValue* Other) const
 {
+	// Trivial case, pointers match.
+	if (this == Other)
+	{
+		return true;
+	}
+
+	// Todo: we are guaranteeing that the IR values are unique. Are we sure we need this?
+
 	// If kinds are different the two values are surely different.
 	if (Kind != Other->Kind)
 	{
@@ -144,6 +152,57 @@ UTexture* FValue::GetTexture()
 	return nullptr;
 }
 
+bool FValue::IsConstantTrue() const
+{
+	const MIR::FConstant* Constant = As<MIR::FConstant>();
+	return Constant && Constant->IsBool() && Constant->Boolean == true;
+}
+
+bool FValue::IsConstantFalse() const
+{
+	const MIR::FConstant* Constant = As<MIR::FConstant>();
+	return Constant && Constant->IsBool() && Constant->Boolean == false;
+}
+
+bool FValue::IsConstantZero() const
+{
+	const MIR::FConstant* Constant = As<MIR::FConstant>();
+	if (!Constant)
+	{
+		return false;
+	}
+
+	return (Constant->IsInteger() && Constant->Integer == 0)
+		|| (Constant->IsFloat() && Constant->Float == 0.0f);
+}
+
+bool FValue::IsConstantOne() const
+{
+	const MIR::FConstant* Constant = As<MIR::FConstant>();
+	if (!Constant)
+	{
+		return false;
+	}
+
+	return (Constant->IsInteger() && Constant->Integer == 1)
+		|| (Constant->IsFloat() && Constant->Float == 1.0f);
+}
+
+bool FConstant::IsBool() const
+{
+	return Type->IsBoolScalar();
+}
+
+bool FConstant::IsInteger() const
+{
+	return Type == FPrimitiveType::GetInt1();
+}
+
+bool FConstant::IsFloat() const
+{
+	return Type == FPrimitiveType::GetFloat1();
+}
+
 const TCHAR* ExternalInputToString(EExternalInput Input)
 {
 	switch (Input)
@@ -187,7 +246,7 @@ FTypePtr GetExternalInputType(EExternalInput Id)
 	int IdIndex = (int)Id;
 	if (IdIndex >= (int)EExternalInput::TexCoord0 && IdIndex <= (int)EExternalInput::TexCoord7_Ddy)
 	{
-		return FArithmeticType::GetFloat2();
+		return FPrimitiveType::GetFloat2();
 	}
 
 	UE_MIR_UNREACHABLE();
@@ -245,11 +304,11 @@ TArrayView<FValue* const> FDimensional::GetComponents() const
 
 TArrayView<FValue*> FDimensional::GetComponents()
 {
-	FArithmeticTypePtr ArithmeticType = Type->AsArithmetic();
-	check(ArithmeticType);
+	FPrimitiveTypePtr PrimitiveType = Type->AsPrimitive();
+	check(PrimitiveType);
 
 	FValue** Ptr = (FValue**)static_cast<TDimensional<1>*>(this)->Components;
-	return { Ptr, ArithmeticType->NumRows };
+	return { Ptr, PrimitiveType->NumRows };
 }
 
 bool FDimensional::AreComponentsConstant() const
diff --git a/Engine/Source/Runtime/Engine/Private/Materials/MaterialIRDebug.cpp b/Engine/Source/Runtime/Engine/Private/Materials/MaterialIRDebug.cpp
new file mode 100644
--- /dev/null
+++ b/Engine/Source/Runtime/Engine/Private/Materials/MaterialIRDebug.cpp
@@ -0,0 +1,186 @@
+// Copyright Epic Games, Inc. All Rights Reserved.
+
+#include "Materials/MaterialIRDebug.h"
+#include "Materials/MaterialIRModule.h"
+#include "Materials/MaterialIR.h"
+#include "Materials/MaterialIRTypes.h"
+#include "Materials/MaterialAttributeDefinitionMap.h"
+#include "Misc/Paths.h"
+#include "Misc/FileHelper.h"
+
+#if WITH_EDITOR
+
+static TAutoConsoleVariable<bool> CVarDumpMaterialIRUseGraph(
+	TEXT("r.Material.Translator.DumpUseGraph"),
+	true,
+	TEXT("Whether the material translator should emit the Module IR 'Uses' graph in Graphviz Dot syntax (to 'MaterialIRDumpGraph.dot')."),
+	ECVF_RenderThreadSafe);
+
+
+static TAutoConsoleVariable<bool> CVarDumpMaterialIRUseGraph_EnableNext(
+	TEXT("r.Material.Translator.DumpUseGraphOpts.EnableSuccessors"),
+	false,
+	TEXT("Whether the Material Module IR 'Uses' graph should also display 'Instruction Next' edges."),
+	ECVF_RenderThreadSafe);
+
+
+namespace UE::MIR
+{
+
+static void DumpValueInfo(const MIR::FValue* Value, FString& Out)
+{
+	if (auto Constant = Value->As<MIR::FConstant>())
+	{
+		switch (Constant->Type->AsPrimitive()->ScalarKind)
+		{
+			case MIR::SK_Bool:	Out.Append(Constant->Boolean ? TEXT("true") : TEXT("false")); break;
+			case MIR::SK_Int:	Out.Appendf(TEXT("%lld"), Constant->Integer); break;
+			case MIR::SK_Float: Out.Appendf(TEXT("%f"), Constant->Float); break;
+			default: UE_MIR_UNREACHABLE();
+		}
+	}
+	else if (auto ExternalInput = Value->As<MIR::FExternalInput>())
+	{
+		Out.Append(MIR::ExternalInputToString(ExternalInput->Id));
+	}
+	else if (auto SetMaterailOutput = Value->As<MIR::FSetMaterialOutput>())
+	{
+		const FString& PropertyName = (SetMaterailOutput->Property == MP_SubsurfaceColor)
+			? TEXT("Subsurface")
+			: FMaterialAttributeDefinitionMap::GetAttributeName(SetMaterailOutput->Property);
+		
+		Out.Append(PropertyName);
+	}
+	else if (auto Subscript = Value->As<MIR::FSubscript>())
+	{
+		if (Subscript->Arg->Type->AsVector())
+		{
+			static const TCHAR* Suffix[] = { TEXT(".x"), TEXT(".y"), TEXT(".z"), TEXT(".w") };
+			check(Subscript->Index < 4); 
+			Out.Append(Suffix[Subscript->Index]);
+		}
+		else
+		{
+			Out.Appendf(TEXT("Index: %d"), Subscript->Index);
+		}
+	}
+	else if (auto BinaryOperator = Value->As<MIR::FBinaryOperator>())
+	{
+		Out.Append(MIR::BinaryOperatorToString(BinaryOperator->Operator));
+	}
+}				
+
+static void DumpUseInfo(const MIR::FValue* Value, const MIR::FValue* Use, int UseIndex, FString& Out)
+{
+	if (auto Dimensional = Value->As<MIR::FDimensional>())
+	{
+		if (Dimensional->Type->AsPrimitive()->IsVector())
+		{
+			check(UseIndex < 4);
+			Out.AppendChar(TEXT("xyzw")[UseIndex]);
+		}
+		else
+		{
+			Out.AppendInt(UseIndex);
+		}
+	}
+	else if (auto If = Value->As<MIR::FBranch>())
+	{
+		static const TCHAR* Uses[] = { TEXT("condition"), TEXT("true"), TEXT("false") };
+		Out.Append(Uses[UseIndex]);
+	}
+	else if (auto BinaryOperator = Value->As<MIR::FBinaryOperator>())
+	{
+		static const TCHAR* Uses[] = { TEXT("lhs"), TEXT("rhs") };
+		Out.Append(Uses[UseIndex]);
+	}
+}
+
+void DebugDumpIRUseGraph(const FMaterialIRModule& Module)
+{
+	if (!CVarDumpMaterialIRUseGraph.GetValueOnAnyThread())
+	{
+		return;
+	}
+
+	FString Content;
+	TSet<const MIR::FValue*> Crawled;
+	TArray<const MIR::FValue*> ValueStack;
+
+	Content.Appendf(TEXT(
+		"digraph G {\n\n"
+		"rankdir=LR\n"
+		"node [shape=box,fontname=\"Consolas\"]\n"
+		"edge [fontname=\"Consolas\"]\n\n"
+	));
+
+	for (const UE::MIR::FSetMaterialOutput* Output : Module.GetOutputs())
+	{
+		ValueStack.Push(Output);
+	}
+
+	bool DumpInstructionSequence = CVarDumpMaterialIRUseGraph_EnableNext.GetValueOnAnyThread();
+
+	while (!ValueStack.IsEmpty())
+	{
+		const MIR::FValue* Value = ValueStack.Pop();
+
+		// Begin the node declaration
+		Content.Appendf(TEXT("\"%p\" [label=< <b>%s</b>  (%s) <br/> "),
+						Value,
+						MIR::ValueKindToString(Value->Kind),
+						Value->Type ? Value->Type->GetSpelling().GetData() : TEXT("???"));
+
+		DumpValueInfo(Value, Content);
+
+		// End the node declaration
+		Content.Append(TEXT(">]\n"));
+
+		const MIR::FInstruction* Instr = Value->AsInstruction();
+		if (DumpInstructionSequence && Instr && Instr->Next)
+		{
+			Content.Appendf(TEXT("\"%p\" -> \"%p\" [color=\"red\"]\n"), Instr, Instr->Next);
+		}
+
+		int UseIndex = -1;
+		for (const MIR::FValue* Use : Value->GetUses())
+		{
+			++UseIndex;
+
+			if (!Use)
+			{
+				continue;
+			}
+			
+			Content.Appendf(TEXT("\"%p\" -> \"%p\" [label=\""), Value, Use);
+
+			DumpUseInfo(Value, Use, UseIndex, Content);
+
+			Content.Appendf(TEXT("\"]\n"));
+
+			if (!Crawled.Contains(Use))
+			{
+				Crawled.Add(Use);
+				ValueStack.Push(Use);
+			}
+
+			if (DumpInstructionSequence && Instr)
+			{
+				const MIR::FInstruction* UseInstr = Use->AsInstruction();
+				if (UseInstr && UseInstr->Block != Instr->Block)
+				{
+					Content.Appendf(TEXT("\"%p\" -> \"%p\" [color=\"red\", style=\"dashed\"]\n"), UseInstr, Instr);
+				}
+			}
+		}
+	}
+
+	Content.Appendf(TEXT("\n}\n"));
+	
+	FString FilePath = FPaths::Combine(FPaths::ProjectLogDir(), TEXT("MaterialIRDumpGraph.dot"));
+	FFileHelper::SaveStringToFile(Content, *FilePath);
+}
+
+} // namespace UE::MIR
+
+#endif // #if WITH_EDITOR
diff --git a/Engine/Source/Runtime/Engine/Private/Materials/MaterialIREmitter.cpp b/Engine/Source/Runtime/Engine/Private/Materials/MaterialIREmitter.cpp
--- a/Engine/Source/Runtime/Engine/Private/Materials/MaterialIREmitter.cpp
+++ b/Engine/Source/Runtime/Engine/Private/Materials/MaterialIREmitter.cpp
@@ -3,6 +3,7 @@
 #include "Materials/MaterialIREmitter.h"
 #include "Materials/MaterialIRModule.h"
 #include "Materials/MaterialIRModuleBuilder.h"
+#include "Materials/MaterialIRInternal.h"
 #include "Shader/ShaderTypes.h"
 #include "MaterialShared.h"
 #include "MaterialExpressionIO.h"
@@ -93,7 +94,7 @@ struct FEmitter::FPrivate
 };
 
 // Creates a new `FDimensional` value of specified `Type` and returns it.
-static FDimensional* NewDimensionalValue(FEmitter* Emitter, FArithmeticTypePtr Type)
+static FDimensional* NewDimensionalValue(FEmitter* Emitter, FPrimitiveTypePtr Type)
 {
 	check(!Type->IsScalar());
 
@@ -175,7 +176,7 @@ static FName SlowFindInputName(UMaterialExpression* Expression, const FExpressio
 
 // FEmitter API
 
-FEmitter::FEmitter(FMaterialIRModuleBuilder* InBuilder, UMaterial* InMaterial, FMaterialIRModule* InModule)
+FEmitter::FEmitter(FMaterialIRModuleBuilderImpl* InBuilder, UMaterial* InMaterial, FMaterialIRModule* InModule)
 {
 	Builder = InBuilder;
 	Material = InMaterial;
@@ -184,8 +185,8 @@ FEmitter::FEmitter(FMaterialIRModuleBuilder* InBuilder, UMaterial* InMaterial, F
 
 FValue*  FEmitter::TryGet(const FExpressionInput* Input)
 {
-	FValue** Value = Builder->InputValues.Find(Input);
-	return Value ? *Value : nullptr;
+	FValue* Value = MIR::Internal::GetInputValue(Builder, Input);
+	return Value ? Value : nullptr;
 }
 
 FValue* FEmitter::Get(const FExpressionInput* Input)
@@ -208,7 +209,7 @@ FEmitter& FEmitter::Put(int OutputIndex, FValue* Value)
 FEmitter& FEmitter::Put(const FExpressionOutput* Output, FValue* Value)
 {
 	check(Output);
-	Builder->OutputValues.Add(Output, Value);
+	MIR::Internal::SetOutputValue(Builder, Output, Value);
 	return *this;
 }
 
@@ -221,7 +222,7 @@ FEmitter& FEmitter::DefaultTo(const FExpressionInput* Input, TFloat Float)
 {
 	if (!Input->IsConnected())
 	{
-		Builder->InputValues.Add(Input, EmitConstantFloat1(Float));
+		Internal::SetInputValue(Builder, Input, EmitConstantFloat1(Float));
 	}
 
 	return *this;
@@ -236,20 +237,20 @@ FValue* FEmitter::GetFloat(const FExpressionInput* Input)
 
 FValue* FEmitter::GetScalar(const FExpressionInput* Input)
 {
-	FValue* Value = GetArithmetic(Input);
+	FValue* Value = GetPrimitive(Input);
 	if (!Value)
 	{
 		return nullptr;
 	}
 
-	FArithmeticTypePtr ScalarType = Value->Type->AsArithmetic()->ToScalar();
+	FPrimitiveTypePtr ScalarType = Value->Type->AsPrimitive()->ToScalar();
 	return EmitConstruct(ScalarType, Value);
 }
 
-FValue* FEmitter::GetArithmetic(const FExpressionInput* Input)
+FValue* FEmitter::GetPrimitive(const FExpressionInput* Input)
 {
 	FValue* Value = Get(Input);
-	CheckInputTypeIs(Input, Value, TK_Arithmetic);
+	CheckInputTypeIs(Input, Value, TK_Primitive);
 	return Value;
 }
 
@@ -277,8 +278,8 @@ bool FEmitter::CheckValueValid(const FValue* Value)
 
 void FEmitter::CheckInputIsScalar(const FExpressionInput* Input, FValue* InputValue)
 {
-	FArithmeticTypePtr ArithmeticType = InputValue->Type->AsArithmetic();
-	if (!ArithmeticType || !ArithmeticType->IsScalar())
+	FPrimitiveTypePtr PrimitiveType = InputValue->Type->AsPrimitive();
+	if (!PrimitiveType || !PrimitiveType->IsScalar())
 	{
 		Errorf(TEXT("Input '%s' expected to be a scalar. It is %s instead."), *SlowFindInputName(Expression, Input).ToString(), InputValue->Type->GetSpelling().GetData());
 	}
@@ -286,8 +287,8 @@ void FEmitter::CheckInputIsScalar(const FExpressionInput* Input, FValue* InputVa
 
 void FEmitter::CheckInputIsScalar(const FExpressionInput* Input, FValue* InputValue, EScalarKind Kind)
 {
-	FArithmeticTypePtr ArithmeticType = InputValue->Type->AsArithmetic();
-	if (!ArithmeticType || !ArithmeticType->IsScalar() || ArithmeticType->ScalarKind != Kind)
+	FPrimitiveTypePtr PrimitiveType = InputValue->Type->AsPrimitive();
+	if (!PrimitiveType || !PrimitiveType->IsScalar() || PrimitiveType->ScalarKind != Kind)
 	{
 		Errorf(TEXT("Input '%s' expected to be a %s scalar. It is %s instead."), *SlowFindInputName(Expression, Input).ToString(), ScalarKindToString(Kind), InputValue->Type->GetSpelling().GetData());
 	}
@@ -349,7 +350,7 @@ FValue* FEmitter::EmitConstantBool1(bool InX)
 
 FValue* FEmitter::EmitConstantFloat1(TFloat InX)
 {
-	FConstant Scalar = MakePrototype<FConstant>(FArithmeticType::GetScalar(SK_Float));
+	FConstant Scalar = MakePrototype<FConstant>(FPrimitiveType::GetScalar(SK_Float));
 	Scalar.Float = InX;
 	return EmitPrototype(this, Scalar);
 }
@@ -380,7 +381,7 @@ FValue* FEmitter::EmitConstantFloat4(const FVector4f& InValue)
 
 FValue* FEmitter::EmitConstantInt1(TInteger InX)
 {
-	FConstant Scalar = MakePrototype<FConstant>(FArithmeticType::GetScalar(SK_Int));
+	FConstant Scalar = MakePrototype<FConstant>(FPrimitiveType::GetScalar(SK_Int));
 	Scalar.Integer = InX;
 	return EmitPrototype(this, Scalar);
 }
@@ -414,7 +415,7 @@ FValue* FEmitter::EmitVector2(FValue* InX, FValue* InY)
 	check(InX->Type->AsScalar());
 	check(InX->Type == InY->Type);
 
-	TDimensional<2> Vector = MakePrototype<TDimensional<2>>(FArithmeticType::GetVector(InX->Type->AsArithmetic()->ScalarKind, 2));
+	TDimensional<2> Vector = MakePrototype<TDimensional<2>>(FPrimitiveType::GetVector(InX->Type->AsPrimitive()->ScalarKind, 2));
 	TArrayView<FValue*> Components = Vector.GetComponents();
 	Components[0] = InX;
 	Components[1] = InY;
@@ -428,7 +429,7 @@ FValue* FEmitter::EmitVector3(FValue* InX, FValue* InY, FValue* InZ)
 	check(InX->Type == InY->Type);
 	check(InY->Type == InZ->Type);
 
-	TDimensional<3> Vector = MakePrototype<TDimensional<3>>(FArithmeticType::GetVector(InX->Type->AsArithmetic()->ScalarKind, 3));
+	TDimensional<3> Vector = MakePrototype<TDimensional<3>>(FPrimitiveType::GetVector(InX->Type->AsPrimitive()->ScalarKind, 3));
 	TArrayView<FValue*> Components = Vector.GetComponents();
 	Components[0] = InX;
 	Components[1] = InY;
@@ -444,7 +445,7 @@ FValue* FEmitter::EmitVector4(FValue* InX, FValue* InY, FValue* InZ, FValue* InW
 	check(InY->Type == InZ->Type);
 	check(InZ->Type == InW->Type);
 
-	TDimensional<4> Vector = MakePrototype<TDimensional<4>>(FArithmeticType::GetVector(InX->Type->AsArithmetic()->ScalarKind, 4));
+	TDimensional<4> Vector = MakePrototype<TDimensional<4>>(FPrimitiveType::GetVector(InX->Type->AsPrimitive()->ScalarKind, 4));
 	TArrayView<FValue*> Components = Vector.GetComponents();
 	Components[0] = InX;
 	Components[1] = InY;
@@ -463,8 +464,8 @@ FValue* FEmitter::GetExternalInput(EExternalInput Id)
 
 FValue* FEmitter::EmitSubscript(FValue* Value, int Index)
 {
-	FArithmeticTypePtr ArithmeticType = Value->Type->AsArithmetic();
-	if (!ArithmeticType)
+	FPrimitiveTypePtr PrimitiveType = Value->Type->AsPrimitive();
+	if (!PrimitiveType)
 	{
 		Errorf(TEXT("Value of type `%s` cannot be subscripted."), Value->Type->GetSpelling().GetData());
 		return nullptr;
@@ -490,7 +491,7 @@ FValue* FEmitter::EmitSubscript(FValue* Value, int Index)
 	}
 
 	// We can't resolve it at compile time: emit subscript value.
-	FSubscript Prototype = MakePrototype<FSubscript>(ArithmeticType->ToScalar());
+	FSubscript Prototype = MakePrototype<FSubscript>(PrimitiveType->ToScalar());
 	Prototype.Arg = Value;
 	Prototype.Index = Index;
 
@@ -502,9 +503,9 @@ FValue* FEmitter::EmitSwizzle(FValue* Value, FSwizzleMask Mask)
 	// At least one component must have been specified.
 	check(Mask.NumComponents > 0);
 
-	// We can only swizzle on non-matrix arithmetic types.
-	FArithmeticTypePtr ArithmeticType = Value->Type->AsVector();
-	if (!ArithmeticType || ArithmeticType->IsMatrix())
+	// We can only swizzle on non-matrix primitive types.
+	FPrimitiveTypePtr PrimitiveType = Value->Type->AsVector();
+	if (!PrimitiveType || PrimitiveType->IsMatrix())
 	{
 		Errorf(TEXT("Cannot swizzle a `%s` value."), Value->Type->GetSpelling().GetData());
 		return nullptr;
@@ -513,9 +514,9 @@ FValue* FEmitter::EmitSwizzle(FValue* Value, FSwizzleMask Mask)
 	// Make sure each component in the mask fits the number of components in Value.
 	for (EVectorComponent Component : Mask)
 	{
-		if ((int)Component >= ArithmeticType->NumRows)
+		if ((int)Component >= PrimitiveType->NumRows)
 		{
-			Errorf(TEXT("Value of type `%s` has no component `%s`."), ArithmeticType->Spelling.GetData(), VectorComponentToString(Component));
+			Errorf(TEXT("Value of type `%s` has no component `%s`."), PrimitiveType->Spelling.GetData(), VectorComponentToString(Component));
 			return nullptr;
 		}
 	}
@@ -528,7 +529,7 @@ FValue* FEmitter::EmitSwizzle(FValue* Value, FSwizzleMask Mask)
 
 	// If the requested number of components is the same as Value and the order in which the components
 	// are specified in the mask is sequential (e.g. x, y, z) then this is a no op, simply return Value as is.
-	if (Mask.NumComponents == ArithmeticType->GetNumComponents())
+	if (Mask.NumComponents == PrimitiveType->GetNumComponents())
 	{
 		bool InOrder = true;
 		for (int i = 0; i < Mask.NumComponents; ++i)
@@ -547,7 +548,7 @@ FValue* FEmitter::EmitSwizzle(FValue* Value, FSwizzleMask Mask)
 	}
 	
 	// Make the result vector type.
-	FArithmeticTypePtr ResultType = FArithmeticType::GetVector(ArithmeticType->ScalarKind, Mask.NumComponents);
+	FPrimitiveTypePtr ResultType = FPrimitiveType::GetVector(PrimitiveType->ScalarKind, Mask.NumComponents);
 	FDimensional* Result = NewDimensionalValue(this, ResultType);
 
 	for (int i = 0; i < Mask.NumComponents; ++i)
@@ -610,23 +611,95 @@ static bool FoldComparisonOperatorScalar(EBinaryOperator Operator, T Lhs, T Rhs)
 	}
 }
 
-static FValue* FoldBinaryOperatorScalar(FEmitter* Emitter, EBinaryOperator Operator, const FConstant* Lhs, const FConstant* Rhs)
+static FValue* TryFoldBinaryOperatorScalar(FEmitter* Emitter, EBinaryOperator Operator, FValue* Lhs, FValue* Rhs)
 {
-	FArithmeticTypePtr ArithType = Lhs->Type->AsArithmetic();
+	FPrimitiveTypePtr PrimitiveType = Lhs->Type->AsPrimitive();
+
+	// Check whether we can fold the operation based on identies (e.g. "x + 0 = x ∀ x ∈ R").
+	// These simplifications don't requre that both lhs and rhs are constant.
+	switch (Operator)
+	{
+		case BO_Add:
+			if (Lhs->IsConstantZero())
+			{
+				return Rhs;
+			}
+			else if (Rhs->IsConstantZero())
+			{
+				return Lhs;
+			}
+			break;
+
+		case BO_Subtract:
+			if (Rhs->IsConstantZero())
+			{
+				return Lhs;
+			}
+			else if (Lhs->IsConstantZero())
+			{
+				// return Emitter->EmitUnaryOperator(UO_Minus, Rhs);
+				UE_MIR_TODO();
+			}
+			break;
+
+		case BO_Multiply:
+			if (Lhs->IsConstantZero() || Rhs->IsConstantOne())
+			{
+				return Lhs;
+			}
+			else if (Lhs->IsConstantOne() || Rhs->IsConstantZero())
+			{
+				return Rhs;
+			}
+			break;
+
+		case BO_Divide:
+			if (Lhs->IsConstantZero() || Rhs->IsConstantOne())
+			{
+				return Lhs;
+			}
+			break;
+
+		case BO_GreaterThan:
+		case BO_LowerThan:
+		case BO_NotEquals:
+			if (Lhs->Equals(Rhs))
+			{
+				return Emitter->EmitConstantFalse();
+			}
+			break;
+
+		case BO_GreaterThanOrEquals:
+		case BO_LowerThanOrEquals:
+		case BO_Equals:
+			if (Lhs->Equals(Rhs))
+			{
+				return Emitter->EmitConstantTrue();
+			}
+			break;
+	}
+
+	// Verify that both lhs and rhs are constants, otherwise we cannot fold the operation.
+	const FConstant* LhsConstant = Lhs->As<FConstant>();
+	const FConstant* RhsConstant = Rhs->As<FConstant>();
+	if (!LhsConstant || !RhsConstant)
+	{
+		return nullptr;
+	}
 
 	if (IsArithmeticOperator(Operator))
 	{
-		switch (ArithType->ScalarKind)
+		switch (PrimitiveType->ScalarKind)
 		{
 			case SK_Int:
 			{
-				int Result = FoldScalarArithmeticOperator<TInteger>(Operator, Lhs->Integer, Rhs->Integer);
+				int Result = FoldScalarArithmeticOperator<TInteger>(Operator, LhsConstant->Integer, RhsConstant->Integer);
 				return Emitter->EmitConstantInt1(Result);
 			}
 
 			case SK_Float:
 			{
-				TFloat Result = FoldScalarArithmeticOperator<TFloat>(Operator, Lhs->Float, Rhs->Float);
+				TFloat Result = FoldScalarArithmeticOperator<TFloat>(Operator, LhsConstant->Float, RhsConstant->Float);
 				return Emitter->EmitConstantFloat1(Result);
 			}
 
@@ -637,14 +710,14 @@ static FValue* FoldBinaryOperatorScalar(FEmitter* Emitter, EBinaryOperator Opera
 	else if (IsComparisonOperator(Operator))
 	{
 		bool Result;
-		switch (ArithType->ScalarKind)
+		switch (PrimitiveType->ScalarKind)
 		{
 			case SK_Int:
-				Result = FoldComparisonOperatorScalar<TInteger>(Operator, Lhs->Integer, Rhs->Integer);
+				Result = FoldComparisonOperatorScalar<TInteger>(Operator, LhsConstant->Integer, RhsConstant->Integer);
 				break;
 
 			case SK_Float:
-				Result = FoldComparisonOperatorScalar<TFloat>(Operator, Lhs->Float, Rhs->Float);
+				Result = FoldComparisonOperatorScalar<TFloat>(Operator, LhsConstant->Float, RhsConstant->Float);
 				break;
 
 			default:
@@ -663,51 +736,96 @@ FValue* FEmitter::EmitBinaryOperator(EBinaryOperator Operator, FValue* Lhs, FVal
 	// Argument value types must always match.
 	check(Lhs->Type == Rhs->Type);
 
-	// Get the operands arithmetic type.
-	FArithmeticTypePtr ArithType = Lhs->Type->AsArithmetic();
-	check(ArithType);
+	// Get the operands primitive type.
+	FPrimitiveTypePtr PrimitiveType = Lhs->Type->AsPrimitive();
+	check(PrimitiveType);
 
-	// Try converting operands to scalar constants. If they're both so, fold the binary
-	// operator right away and return the result.
-	const FConstant* ScalarLhs = Lhs->As<FConstant>();
-	const FConstant* ScalarRhs = Rhs->As<FConstant>();
+	// Try converting operands to scalar constants
+	FConstant* ScalarLhs = Lhs->As<FConstant>();
+	FConstant* ScalarRhs = Rhs->As<FConstant>();
 
+	//  If they're both so, fold the binary operator right away and return the result.
 	if (ScalarLhs && ScalarRhs)
 	{
-		if (FValue* Value = FoldBinaryOperatorScalar(this, Operator, ScalarLhs, ScalarRhs))
+		if (FValue* Value = TryFoldBinaryOperatorScalar(this, Operator, ScalarLhs, ScalarRhs))
 		{
 			return Value;
 		}
 	}
 
-	// Determine the result type. If the operator is arithmetic, the result type will be the same
+	// Determine the result type. If the operator is primitive, the result type will be the same
 	// as the operands type. Otherwise it will have the same number of components but bool.
-	FArithmeticTypePtr ResultType = IsArithmeticOperator(Operator)
-		? ArithType
-		: FArithmeticType::Get(SK_Bool, ArithType->NumRows, ArithType->NumColumns);
+	FPrimitiveTypePtr ResultType = IsArithmeticOperator(Operator)
+		? PrimitiveType
+		: FPrimitiveType::Get(SK_Bool, PrimitiveType->NumRows, PrimitiveType->NumColumns);
+
+	// Check that at least one component of the resulting dimensional value would folded.
+	// If all components of resulting dimensional value are not folded, then instead of emitting
+	// an individual operator instruction for each component, simply emit a single binary operator
+	// instruction applied between lhs and rhs as a whole. (v1 + v2 rather than float2(v1.x + v2.x, v1.y + v2.y)
+	bool bSomeResultComponentWasFolded = false;
+	bool bResultIsIdenticalToLhs = true;
+	bool bResultIsIdenticalToRhs = true;
+	FValue* TempResultComponents[FDimensional::MaxNumComponents];
+	for (int i = 0; i < ResultType->GetNumComponents(); ++i)
+	{
+		// Extract the arguments individual components
+		FValue* LhsComponent = EmitSubscript(Lhs, i);
+		FValue* RhsComponent = EmitSubscript(Rhs, i);
+
+		// Try folding the operation, it may return null
+		FValue* ResultComponent = TryFoldBinaryOperatorScalar(this, Operator, LhsComponent, RhsComponent);
+
+		// Update the flags
+		bSomeResultComponentWasFolded |= (bool)ResultComponent;
+		bResultIsIdenticalToLhs &= ResultComponent && ResultComponent->Equals(LhsComponent);
+		bResultIsIdenticalToRhs &= ResultComponent && ResultComponent->Equals(RhsComponent);
 
-	// Now check that both operands are dimensional. If so, emit a new dimensional value with each component
-	// being the result of the binary operation applied to the corresponding components of the operands.
-	const FDimensional* DimensionalLhs = Lhs->As<FDimensional>();
-	const FDimensional* DimensionalRhs = Rhs->As<FDimensional>();
+		// Cache the results
+		TempResultComponents[i] = ResultComponent;
+	}
 
-	if (DimensionalLhs && DimensionalRhs)
+	// If result is identical to either lhs or rhs, simply return it
+	if (bResultIsIdenticalToLhs)
 	{
-		FDimensional* Result = NewDimensionalValue(this, ResultType);
+		return Lhs;
+	}
+	else if (bResultIsIdenticalToRhs)
+	{
+		return Rhs;
+	}
 
-		auto ResultComponents = Result->GetComponents();
-		auto LhsComponents = DimensionalLhs->GetComponents();
-		auto RhsComponents = DimensionalRhs->GetComponents();
+	// If some component was folded (it is either constant or the operation was a NOP), it is worth
+	// build the binary operation as a separate operation for each component, that is like
+	//    float2(lhs.x + rhs.x, lhs.y + rhs.y)
+	// rather than
+	//    lhs + rhs
+	// so that we retain as much compile-time information as possible.
+	if (bSomeResultComponentWasFolded)
+	{
+		// Make the new dimensional value
+		FDimensional* Result = NewDimensionalValue(this, ResultType);
+		TArrayView<FValue*> ResultComponents = Result->GetComponents();
 
+		// Create the operator instruction for each component pair
 		for (int i = 0; i < ResultType->GetNumComponents(); ++i)
 		{
-			ResultComponents[i] = EmitBinaryOperator(Operator, LhsComponents[i], RhsComponents[i]);
+			// Reuse cached result if possible
+			ResultComponents[i] = TempResultComponents[i];
+
+			// Otherwise emit the binary operation between the two components (this will create a new instruction)
+			if (!ResultComponents[i])
+			{
+				FValue* LhsComponent = EmitSubscript(Lhs, i);
+				FValue* RhsComponent = EmitSubscript(Rhs, i);
+				ResultComponents[i] = EmitBinaryOperator(Operator, LhsComponent, RhsComponent);
+			}
 		}
 
 		return Result;
 	}
 
-	// One (or both) of the operands is runtime only. Emit the runtime binary operation instruction.
+	// Finally, simply emit the binary operator instruction between lhs and rhs (i.e. "lhs + rhs")
 	FBinaryOperator Proto = MakePrototype<FBinaryOperator>(ResultType);
 	Proto.Operator = Operator;
 	Proto.LhsArg= Lhs;
@@ -719,7 +837,7 @@ FValue* FEmitter::EmitBinaryOperator(EBinaryOperator Operator, FValue* Lhs, FVal
 FValue* FEmitter::EmitBranch(FValue* Condition, FValue* True, FValue* False)
 {
 	// Condition must be of type bool
-	check(Condition->Type->IsBool1());
+	check(Condition->Type->IsBoolScalar());
 
 	// If the condition is a scalar constant, then simply evaluate the result now.
 	if (const FConstant* ConstCondition = Condition->As<FConstant>())
@@ -777,24 +895,24 @@ static FValue* CastConstant(FEmitter* Emitter, FConstant* Constant, EScalarKind
 	UE_MIR_UNREACHABLE();
 }
 
-static FValue* ConstructArithmeticValue(FEmitter* Emitter, FArithmeticTypePtr TargetArithmeticType, FValue* Initializer)
+static FValue* ConstructPrimitiveValue(FEmitter* Emitter, FPrimitiveTypePtr TargetPrimitiveType, FValue* Initializer)
 {
-	FArithmeticTypePtr InitializerArithmeticType = Initializer->Type->AsArithmetic();
-	if (!InitializerArithmeticType)
+	FPrimitiveTypePtr InitializerPrimitiveType = Initializer->Type->AsPrimitive();
+	if (!InitializerPrimitiveType)
 	{
-		Emitter->Errorf(TEXT("Cannot construct a '%s' from non arithmetic type '%s'."), TargetArithmeticType->Spelling.GetData(), Initializer->Type->GetSpelling().GetData());
+		Emitter->Errorf(TEXT("Cannot construct a '%s' from non primitive type '%s'."), TargetPrimitiveType->Spelling.GetData(), Initializer->Type->GetSpelling().GetData());
 		return nullptr;
 	}
 
 	// Construct a scalar from another scalar.
-	if (TargetArithmeticType->IsScalar())
+	if (TargetPrimitiveType->IsScalar())
 	{
 		//
 		Initializer = Emitter->EmitSubscript(Initializer, 0);
-		InitializerArithmeticType = Initializer->Type->AsArithmetic();
+		InitializerPrimitiveType = Initializer->Type->AsPrimitive();
 		
 		//
-		if (InitializerArithmeticType == TargetArithmeticType)
+		if (InitializerPrimitiveType == TargetPrimitiveType)
 		{
 			return Initializer;
 		}
@@ -802,7 +920,7 @@ static FValue* ConstructArithmeticValue(FEmitter* Emitter, FArithmeticTypePtr Ta
 		// Construct the scalar from a constant.
 		if (FConstant* ConstantInitializer = Initializer->As<FConstant>())
 		{
-			return CastConstant(Emitter, ConstantInitializer, InitializerArithmeticType->ScalarKind, TargetArithmeticType->ScalarKind);
+			return CastConstant(Emitter, ConstantInitializer, InitializerPrimitiveType->ScalarKind, TargetPrimitiveType->ScalarKind);
 		}
 		else
 		{
@@ -812,17 +930,17 @@ static FValue* ConstructArithmeticValue(FEmitter* Emitter, FArithmeticTypePtr Ta
 	}
 
 	// Construct a vector or matrix from a scalar. E.g. float4(3.14f)
-	if (!TargetArithmeticType->IsScalar() && InitializerArithmeticType->IsScalar())
+	if (!TargetPrimitiveType->IsScalar() && InitializerPrimitiveType->IsScalar())
 	{
 		// Create the result dimensional value.
-		FDimensional* Result = NewDimensionalValue(Emitter, TargetArithmeticType);
+		FDimensional* Result = NewDimensionalValue(Emitter, TargetPrimitiveType);
 
 		// Create a dimensional and initialize each of its components to the conversion
 		// of initializer value to the single component type.
-		FValue* Component = Emitter->EmitConstruct(TargetArithmeticType->ToScalar(), Initializer);
+		FValue* Component = Emitter->EmitConstruct(TargetPrimitiveType->ToScalar(), Initializer);
 
 		// Initialize all result components to the same scalar.
-		for (int i = 0; i < TargetArithmeticType->GetNumComponents(); ++i)
+		for (int i = 0; i < TargetPrimitiveType->GetNumComponents(); ++i)
 		{
 			Result->GetComponents()[i] = Component;
 		}
@@ -833,16 +951,16 @@ static FValue* ConstructArithmeticValue(FEmitter* Emitter, FArithmeticTypePtr Ta
 	// Construct a vector from another vector. If constructed vector is larger, initialize
 	// remaining components to zero. If it's smaller, truncate initializer vector and only use
 	// the necessary components.
-	if (TargetArithmeticType->IsVector() && InitializerArithmeticType->IsVector())
+	if (TargetPrimitiveType->IsVector() && InitializerPrimitiveType->IsVector())
 	{
-		int TargetNumComponents = TargetArithmeticType->GetNumComponents();
-		int InitializerNumComponents = InitializerArithmeticType->GetNumComponents();
+		int TargetNumComponents = TargetPrimitiveType->GetNumComponents();
+		int InitializerNumComponents = InitializerPrimitiveType->GetNumComponents();
 
 		// Create the result dimensional value.
-		FDimensional* Result = NewDimensionalValue(Emitter, TargetArithmeticType);
+		FDimensional* Result = NewDimensionalValue(Emitter, TargetPrimitiveType);
 
 		// Determine the result component type (scalar).
-		FArithmeticTypePtr ResultComponentType = TargetArithmeticType->ToScalar();
+		FPrimitiveTypePtr ResultComponentType = TargetPrimitiveType->ToScalar();
 
 		// For iterating over the components of the result dimensional value.
 		int Index = 0;
@@ -863,23 +981,23 @@ static FValue* ConstructArithmeticValue(FEmitter* Emitter, FArithmeticTypePtr Ta
 		return EmitNew(Emitter, Result);
 	}
 	
-	// The two arithmetic types are identical matrices that differ only by their scalar type.
-	if (TargetArithmeticType->NumRows 	 == InitializerArithmeticType->NumRows &&
-		TargetArithmeticType->NumColumns == InitializerArithmeticType->NumColumns)
+	// The two primitive types are identical matrices that differ only by their scalar type.
+	if (TargetPrimitiveType->NumRows 	 == InitializerPrimitiveType->NumRows &&
+		TargetPrimitiveType->NumColumns == InitializerPrimitiveType->NumColumns)
 	{
-		check(TargetArithmeticType->IsMatrix());
+		check(TargetPrimitiveType->IsMatrix());
 
 		// 
 		if (FDimensional* DimensionalInitializer = Initializer->As<FDimensional>())
 		{
 			// Create the result dimensional value.
-			FDimensional* Result = NewDimensionalValue(Emitter, TargetArithmeticType);
+			FDimensional* Result = NewDimensionalValue(Emitter, TargetPrimitiveType);
 			
 			// Determine the result component type (scalar).
-			FArithmeticTypePtr ResultComponentType = TargetArithmeticType->ToScalar();
+			FPrimitiveTypePtr ResultComponentType = TargetPrimitiveType->ToScalar();
 
 			// Convert components from the initializer vector.
-			for (int Index = 0, Num = TargetArithmeticType->GetNumComponents(); Index < Num; ++Index)
+			for (int Index = 0, Num = TargetPrimitiveType->GetNumComponents(); Index < Num; ++Index)
 			{
 				Result->GetComponents()[Index] = Emitter->EmitConstruct(ResultComponentType, DimensionalInitializer->GetComponents()[Index]);
 			}
@@ -889,13 +1007,13 @@ static FValue* ConstructArithmeticValue(FEmitter* Emitter, FArithmeticTypePtr Ta
 		else
 		{
 			// Initializer is an unknown value, construct target value casting initializer.
-			FCast Prototype = MakePrototype<FCast>(TargetArithmeticType);
+			FCast Prototype = MakePrototype<FCast>(TargetPrimitiveType);
 			Prototype.Arg = Initializer;
 			return EmitPrototype(Emitter, Prototype);
 		}
 	}
 
-	// Initializer value cannot be used to construct this arithmetic type.
+	// Initializer value cannot be used to construct this primitive type.
 	return nullptr;
 }
 
@@ -909,9 +1027,9 @@ FValue* FEmitter::EmitConstruct(FTypePtr Type, FValue* Initializer)
 	}
 	
 	FValue* Result{};
-	if (FArithmeticTypePtr ArithmeticType = Type->AsArithmetic())
+	if (FPrimitiveTypePtr PrimitiveType = Type->AsPrimitive())
 	{
-		Result = ConstructArithmeticValue(this, ArithmeticType, Initializer);
+		Result = ConstructPrimitiveValue(this, PrimitiveType, Initializer);
 	}
 
 	// No other legal conversions applicable. Report error if we haven't converted the value.
@@ -925,7 +1043,7 @@ FValue* FEmitter::EmitConstruct(FTypePtr Type, FValue* Initializer)
 
 FValue* FEmitter::EmitTextureSample(UTexture* Texture, FValue* TexCoord, ESamplerSourceMode SamplerSourceMode, ETextureMipValueMode MipValueMode, EMaterialSamplerType SamplerType)
 {
-	MIR::FTextureSample Prototype = MakePrototype<MIR::FTextureSample>(FArithmeticType::GetFloat4());
+	MIR::FTextureSample Prototype = MakePrototype<MIR::FTextureSample>(FPrimitiveType::GetFloat4());
 	Prototype.Texture = Texture;
 	Prototype.TexCoordArg = TexCoord;
 	Prototype.MipValueArg = nullptr;
@@ -937,7 +1055,7 @@ FValue* FEmitter::EmitTextureSample(UTexture* Texture, FValue* TexCoord, ESample
 	return EmitPrototype(this, Prototype);
 }
 
-FArithmeticTypePtr FEmitter::GetCommonArithmeticType(FArithmeticTypePtr A, FArithmeticTypePtr B)
+FPrimitiveTypePtr FEmitter::GetCommonPrimitiveType(FPrimitiveTypePtr A, FPrimitiveTypePtr B)
 {
 	// Trivial case: types are equal
 	if (A == B)
@@ -945,18 +1063,18 @@ FArithmeticTypePtr FEmitter::GetCommonArithmeticType(FArithmeticTypePtr A, FArit
 		return A;
 	}
 	
-	//	
-	if (A->IsMatrix() != B->IsMatrix())
+	// If both A and B are matrices, their dimensions must match (equality check above didn't trigger).
+	if (A->IsMatrix() || B->IsMatrix())
 	{
-		Errorf(TEXT("No common arithmetic type between `%s` and `%s`."), A->Spelling.GetData(), B->Spelling.GetData());
+		Errorf(TEXT("No common primitive type between `%s` and `%s`."), A->Spelling.GetData(), B->Spelling.GetData());
 		return nullptr;
 	}
 
+	// Neither A nor B are matrices, but single scalar or vector. Return the largest.
+	check(A->NumColumns == 1 && B->NumColumns == 1);
 	EScalarKind ScalarKind = FMath::Max(A->ScalarKind, B->ScalarKind);
 	int NumRows = FMath::Max(A->NumRows, B->NumRows);
-	int NumColumns = FMath::Max(A->NumColumns, B->NumColumns);
-
-	return FArithmeticType::Get(ScalarKind, NumRows, NumColumns);
+	return FPrimitiveType::Get(ScalarKind, NumRows, 1);
 }
 
 void FEmitter::Error(FString Message)
@@ -971,7 +1089,7 @@ void FEmitter::Error(FString Message)
 void FEmitter::Initialize()
 {
 	// Create and reference the true/false constants.
-	FConstant Temp = MakePrototype<FConstant>(FArithmeticType::GetBool1());
+	FConstant Temp = MakePrototype<FConstant>(FPrimitiveType::GetBool1());
 
 	Temp.Boolean = true;
 	ConstantTrue = EmitPrototype(this, Temp);
diff --git a/Engine/Source/Runtime/Engine/Private/Materials/MaterialIRInternal.cpp b/Engine/Source/Runtime/Engine/Private/Materials/MaterialIRInternal.cpp
new file mode 100644
--- /dev/null
+++ b/Engine/Source/Runtime/Engine/Private/Materials/MaterialIRInternal.cpp
@@ -0,0 +1,96 @@
+// Copyright Epic Games, Inc. All Rights Reserved.
+
+#include "Materials/MaterialIRInternal.h"
+#include "Materials/MaterialAttributeDefinitionMap.h"
+#include "Materials/MaterialIREmitter.h"
+#include "Materials/Material.h"
+#include "MaterialShared.h"
+
+#if WITH_EDITOR
+
+namespace UE::MIR::Internal {
+
+bool IsMaterialPropertyShared(EMaterialProperty InProperty)
+{
+	switch (InProperty)
+	{
+	case MP_Normal:
+	case MP_Tangent:
+	case MP_EmissiveColor:
+	case MP_Opacity:
+	case MP_OpacityMask:
+	case MP_BaseColor:
+	case MP_Metallic:
+	case MP_Specular:
+	case MP_Roughness:
+	case MP_Anisotropy:
+	case MP_AmbientOcclusion:
+	case MP_Refraction:
+	case MP_PixelDepthOffset:
+	case MP_SubsurfaceColor:
+	case MP_ShadingModel:
+	case MP_SurfaceThickness:
+	case MP_FrontMaterial:
+	case MP_Displacement:
+		return true;
+	default:
+		return false;
+	}
+}
+
+bool NextMaterialAttributeInput(UMaterial* BaseMaterial, int32& PropertyIndex, FMaterialInputDescription& Input)
+{
+	for (; PropertyIndex < MP_MAX; ++PropertyIndex)
+	{
+		EMaterialProperty Property = (EMaterialProperty)PropertyIndex;
+		if (MIR::Internal::IsMaterialPropertyShared(Property)
+			&& Property != MP_SubsurfaceColor
+			&& Property != MP_FrontMaterial
+			&& BaseMaterial->GetExpressionInputDescription(Property, Input))
+		{
+			return true;
+		}
+	}
+
+	return false;
+}
+
+MIR::FValue* CreateMaterialAttributeDefaultValue(FEmitter& Emitter, const UMaterial* Material, EMaterialProperty Property)
+{
+	EMaterialValueType Type = FMaterialAttributeDefinitionMap::GetValueType(Property);
+	FVector4f DefaultValue = FMaterialAttributeDefinitionMap::GetDefaultValue(Property);
+
+	switch (Type)
+	{
+		case MCT_ShadingModel: return Emitter.EmitConstantInt1(Material->GetShadingModels().GetFirstShadingModel());
+
+		case MCT_Float1: return Emitter.EmitConstantFloat1(  DefaultValue.X );
+		case MCT_Float2: return Emitter.EmitConstantFloat2({ DefaultValue.X, DefaultValue.Y });
+		case MCT_Float3: return Emitter.EmitConstantFloat3({ DefaultValue.X, DefaultValue.Y, DefaultValue.Z });
+		case MCT_Float4: return Emitter.EmitConstantFloat4(DefaultValue);
+
+		case MCT_UInt1: return Emitter.EmitConstantInt1(  (int32)DefaultValue.X );
+		case MCT_UInt2: return Emitter.EmitConstantInt2({ (int32)DefaultValue.X, (int32)DefaultValue.Y });
+		case MCT_UInt3: return Emitter.EmitConstantInt3({ (int32)DefaultValue.X, (int32)DefaultValue.Y, (int32)DefaultValue.Z });
+		case MCT_UInt4: return Emitter.EmitConstantInt4({ (int32)DefaultValue.X, (int32)DefaultValue.Y, (int32)DefaultValue.Z, (int32)DefaultValue.W });
+
+		default: UE_MIR_UNREACHABLE();
+	}
+}
+
+EMaterialTextureParameterType TextureMaterialValueTypeToParameterType(EMaterialValueType Type)
+{
+	switch (Type)
+	{
+		case MCT_Texture2D: return EMaterialTextureParameterType::Standard2D;
+		case MCT_Texture2DArray: return EMaterialTextureParameterType::Array2D;
+		case MCT_TextureCube: return EMaterialTextureParameterType::Cube;
+		case MCT_TextureCubeArray: return EMaterialTextureParameterType::ArrayCube;
+		case MCT_VolumeTexture: return EMaterialTextureParameterType::Volume;
+		default: UE_MIR_UNREACHABLE();
+	}
+}
+
+} // namespace UE::MIR::Internal
+
+#endif
diff --git a/Engine/Source/Runtime/Engine/Private/Materials/MaterialIRInternal.h b/Engine/Source/Runtime/Engine/Private/Materials/MaterialIRInternal.h
new file mode 100644
--- /dev/null
+++ b/Engine/Source/Runtime/Engine/Private/Materials/MaterialIRInternal.h
@@ -0,0 +1,35 @@
+// Copyright Epic Games, Inc. All Rights Reserved.
+#pragma once
+
+#include "Materials/MaterialIRCommon.h"
+#include "MaterialShared.h"
+
+#if WITH_EDITOR
+
+namespace UE::MIR::Internal {
+
+//
+bool IsMaterialPropertyShared(EMaterialProperty InProperty);
+
+//
+bool NextMaterialAttributeInput(UMaterial* BaseMaterial, int32& PropertyIndex, FMaterialInputDescription& Input);
+
+//
+FValue* CreateMaterialAttributeDefaultValue(FEmitter& Emitter, const UMaterial* Material, EMaterialProperty Property);
+ 
+//
+EMaterialTextureParameterType TextureMaterialValueTypeToParameterType(EMaterialValueType Type);
+
+//
+FValue* GetInputValue(FMaterialIRModuleBuilderImpl* Builder, const FExpressionInput* Input);
+
+//
+void SetInputValue(FMaterialIRModuleBuilderImpl* Builder, const FExpressionInput* Input, FValue* Value);
+
+//
+void SetOutputValue(FMaterialIRModuleBuilderImpl* Builder, const FExpressionOutput* Output, FValue* Value);
+
+} // namespace UE::MIR::Internal
+
+#endif // #if WITH_EDITOR
+
diff --git a/Engine/Source/Runtime/Engine/Private/Materials/MaterialIRModuleBuilder.cpp b/Engine/Source/Runtime/Engine/Private/Materials/MaterialIRModuleBuilder.cpp
--- a/Engine/Source/Runtime/Engine/Private/Materials/MaterialIRModuleBuilder.cpp
+++ b/Engine/Source/Runtime/Engine/Private/Materials/MaterialIRModuleBuilder.cpp
@@ -9,8 +9,11 @@
 #include "Materials/MaterialIR.h"
 #include "Materials/MaterialIRTypes.h"
 #include "Materials/MaterialIREmitter.h"
-#include "MaterialIRUtility.h"
+#include "Materials/MaterialIRInternal.h"
 #include "Materials/MaterialAttributeDefinitionMap.h"
+#include "Materials/MaterialExpressionFunctionInput.h"
+#include "Materials/MaterialExpressionFunctionOutput.h"
+#include "Materials/MaterialExpressionMaterialFunctionCall.h"
 #include "Materials/Material.h"
 #include "MaterialExpressionIO.h"
 #include "MaterialShared.h"
@@ -18,148 +21,298 @@
 #include "Materials/MaterialInsights.h"
 #include "Async/ParallelFor.h"
 #include "Engine/Texture.h"
-
 namespace MIR = UE::MIR;
 
-struct FMaterialIRModuleBuilder::FPrivate
+struct FAnalysisContext
 {
-	FMaterialIRModuleBuilder& Builder;
-	FMaterialIRModuleBuildParams Params;
-	FMaterialIRModule& Module;
-	MIR::FEmitter& Emitter;
-	TArray<UMaterialExpression*> ExpressionAnalysisStack;
-	TArray<MIR::FInstruction*> InstructionStack;
-
-	void Build_Initialize()
+	UMaterialExpressionMaterialFunctionCall* Call{};
+	TSet<UMaterialExpression*> BuiltExpressions{};
+	TArray<UMaterialExpression*> ExpressionStack{};
+	TMap<const FExpressionInput*, UE::MIR::FValue*> InputValues;
+	TMap<const FExpressionOutput*, UE::MIR::FValue*> OutputValues;
+
+	MIR::FValue* GetInputValue(const FExpressionInput* Input)
+	{
+		MIR::FValue** Value = InputValues.Find(Input);
+		return Value ? *Value : nullptr;
+	}
+
+	void SetInputValue(const FExpressionInput* Input, MIR::FValue* Value)
+	{
+		InputValues.Add(Input, Value);
+	}
+
+	MIR::FValue* GetOutputValue(const FExpressionOutput* Output)
+	{
+		MIR::FValue** Value = OutputValues.Find(Output);
+		return Value ? *Value : nullptr;
+	}
+	
+	void SetOutputValue(const FExpressionOutput* Output, MIR::FValue* Value)
 	{
-		Module.Empty();
-		Module.ShaderPlatform = Params.ShaderPlatform;
+		OutputValues.Add(Output, Value);
+	}
+};
 
-		Emitter.Initialize();
+struct FMaterialIRModuleBuilderImpl
+{
+	FMaterialIRModuleBuilder* Builder;
+	FMaterialIRModule* Module;
+	MIR::FEmitter* Emitter;
+	TArray<FAnalysisContext> AnalysisContextStack;
+	TArray<MIR::FValue*> AuxValueList;
+	TArray<MIR::FInstruction*> InstructionStack{};
+
+	void Step_Initialize()
+	{
+		Module->Empty();
+		Module->ShaderPlatform = Builder->ShaderPlatform;
+		
+		Emitter->Initialize();
+		AnalysisContextStack.Emplace();
 	}
 
-	void Build_GenerateOutputInstructions()
+	void Step_GenerateOutputInstructions()
 	{
 		// Prepare the array of FSetMaterialOutputInstr outputs from the material attributes inputs.
 		FMaterialInputDescription Input;
-		for (int32 Index = 0; UE::Utility::NextMaterialAttributeInput(Params.Material, Index, Input); ++Index)
+		for (int32 Index = 0; UE::MIR::Internal::NextMaterialAttributeInput(Builder->Material, Index, Input); ++Index)
 		{
 			EMaterialProperty Property = (EMaterialProperty)Index;
 
-			MIR::FSetMaterialOutput* Output = Emitter.EmitSetMaterialOutput(Property, nullptr);
+			MIR::FSetMaterialOutput* Output = Emitter->EmitSetMaterialOutput(Property, nullptr);
 
 			if (Input.bUseConstant)
 			{
-				Output->Arg = Emitter.EmitConstantFromShaderValue(Input.ConstantValue);
+				Output->Arg = Emitter->EmitConstantFromShaderValue(Input.ConstantValue);
 			}
 			else if (!Input.Input->IsConnected())
 			{
-				Output->Arg = UE::Utility::CreateMaterialAttributeDefaultValue(Emitter, Params.Material, Property);
+				Output->Arg = UE::MIR::Internal::CreateMaterialAttributeDefaultValue(*Emitter, Builder->Material, Property);
 			}
 			else
 			{
-				ExpressionAnalysisStack.Add(Input.Input->Expression);
+				AnalysisContextStack.Last().ExpressionStack.Add(Input.Input->Expression);
 			}
 		}
 	}
 
-	void Build_BuildIRGraph()
+	void Step_BuildMaterialExpressionsToIRGraph()
 	{
-		TSet<UMaterialExpression*> BuiltExpressions;
-
-		while (!ExpressionAnalysisStack.IsEmpty())
+		while (true)
 		{
-			Emitter.Expression = ExpressionAnalysisStack.Last();
+			FAnalysisContext& Context = AnalysisContextStack.Last();
 
-			// If expression is clean, nothing to be done.
-			if (BuiltExpressions.Contains(Emitter.Expression))
+			if (!Context.ExpressionStack.IsEmpty())
 			{
-				ExpressionAnalysisStack.Pop(EAllowShrinking::No);
-				continue;
+				// Some expression is on the expression stack of this context. Analyze it. This will
+				// have the effect of either building the expression or pushing its other expression
+				// dependencies onto the stack.
+				BuildTopMaterialExpression();
 			}
-
-			// Push to the expression stack all dependencies that still need to be analyzed.
-			for (FExpressionInputIterator It{ Emitter.Expression}; It; ++It)
+			else if (Context.Call)
 			{
-				// Ignore disconnected inputs and connected expressions already built.
-				if (!It->IsConnected() || BuiltExpressions.Contains(It->Expression))
-				{
-					continue;
-				}
-
-				ExpressionAnalysisStack.Push(It->Expression);
+				// There are no more expressions to analyze on the stack, this analysis context is complete.
+				// Context.Call isn't null so this context is for a function call, which has now been fully analyzed.
+				// Pop the callee context from the stack and resume analyzing the parent context (the caller).
+				PopFunctionCall();
+			}
+			else
+			{
+				// No other expressions on the stack to evaluate, nor this is a function
+				// call context but the root context. Nothing left to do so simply quit.
+				break;
 			}
+		}
+	}
+
+	void BuildTopMaterialExpression()
+	{
+		FAnalysisContext& CurrContext = AnalysisContextStack.Last();
+		Emitter->Expression = CurrContext.ExpressionStack.Last();
 
-			// If on top of the stack there's a different expression, we have a dependency to analyze first.
-			if (ExpressionAnalysisStack.Last() != Emitter.Expression) {
+		// If expression is clean, nothing to be done.
+		if (CurrContext.BuiltExpressions.Contains(Emitter->Expression))
+		{
+			CurrContext.ExpressionStack.Pop(EAllowShrinking::No);
+			return;
+		}
+
+		// Push to the expression stack all dependencies that still need to be analyzed.
+		for (FExpressionInputIterator It{ Emitter->Expression }; It; ++It)
+		{
+			// Ignore disconnected inputs and connected expressions already built.
+			if (!It->IsConnected() || CurrContext.BuiltExpressions.Contains(It->Expression))
+			{
 				continue;
 			}
 
-			// 
-			ExpressionAnalysisStack.Pop();
-			BuiltExpressions.Add(Emitter.Expression);
+			CurrContext.ExpressionStack.Push(It->Expression);
+		}
+
+		// If on top of the stack there's a different expression, we have a dependency to analyze first.
+		if (CurrContext.ExpressionStack.Last() != Emitter->Expression) {
+			return;
+		}
+
+		// Take the top expression out of the stack as ready for analysis. Also mark it as built.
+		CurrContext.ExpressionStack.Pop();
+		CurrContext.BuiltExpressions.Add(Emitter->Expression);
 
-			//
-			for (FExpressionInputIterator It{ Emitter.Expression}; It; ++It)
+		// Flow the value into this expression's inputs from their connected outputs.
+		for (FExpressionInputIterator It{ Emitter->Expression}; It; ++It)
+		{
+			FExpressionOutput* ConnectedOutput = It->GetConnectedOutput();
+			if (!ConnectedOutput)
 			{
-				if (FExpressionOutput* ConnectedOutput = It->GetConnectedOutput())
-				{
-					// Fetch the value flowing through connected output.
-					if (MIR::FValue** ValuePtr = Builder.OutputValues.Find(ConnectedOutput))
-					{
-						// Set the value flowing into this input.
-						Builder.InputValues.Add(It.Input, *ValuePtr);
-					}
-					else
-					{
-						Builder.InputValues.Remove(It.Input);
-					}
-				}
+				break;
 			}
 
+			// Fetch the value flowing through connected output.
+			if (MIR::FValue** ValuePtr = CurrContext.OutputValues.Find(ConnectedOutput))
+			{
+				// ...and flow it into this input.
+				CurrContext.InputValues.Add(It.Input, *ValuePtr);
+			}
+		}
+
+		if (auto Call = Cast<UMaterialExpressionMaterialFunctionCall>(Emitter->Expression))
+		{
+			// Function calls are handled internally as they manipulate the analysis context stack.
+			PushFunctionCall(Call);
+		}
+		else if (!Cast<UMaterialExpressionFunctionOutput>(Emitter->Expression))
+		{
+			// Ignore UMaterialExpressionFunctionOutputs as they're handled by PopFunctionCall().
+
 			// Invoke the expression build function. This will perform semantic analysis, error reporting and
 			// emit IR values for its outputs (which will flow into connected expressions inputs).
-			Emitter.Expression->Build(Emitter);
+			Emitter->Expression->Build(*Emitter);
 
 			// Populate the insight information about this expression pins.
-			if (Params.TargetInsight)
+			AddExpressionConnectionInsights(Emitter->Expression);
+		}
+	}
+
+	void PushFunctionCall(UMaterialExpressionMaterialFunctionCall* Call)
+	{
+		AuxValueList.Empty(Call->FunctionInputs.Num()); /* todo: is there a better (and idiomatic) way to allocate temp dynamic array on the stack? */
+
+		// Make sure each function input is connected and has a value. If so, cache the values flowing into this
+		// funcion call inside the auxiliary value array.
+		for (int i = 0; i < Call->FunctionInputs.Num(); ++i)
+		{
+			FFunctionExpressionInput& FunctionInput = Call->FunctionInputs[i];
+			MIR::FValue* Value = Emitter->Get(Call->GetInput(i));
+			if (Value)
 			{
-				AddExpressionConnectionInsights(Emitter.Expression);
+				MIR::FTypePtr Type = MIR::FType::FromMaterialValueType((EMaterialValueType)FunctionInput.ExpressionInput->GetInputType(0));
+				AuxValueList.Push(Emitter->EmitConstruct(Type, Value));
 			}
 		}
+
+		// If some error occurred (e.g. some function input wasn't linked in) early out.
+		if (Emitter->IsInvalid())
+		{
+			return;
+		}
+
+		// Push a new analysis context on the stack dedicated to this function call.
+		AnalysisContextStack.Emplace();
+		FAnalysisContext& ParentContext = AnalysisContextStack[AnalysisContextStack.Num() - 2];
+		FAnalysisContext& NewContext = AnalysisContextStack[AnalysisContextStack.Num() - 1];
+
+		// Set the function call. When the expressions stack in this new context is empty, this
+		// will be used to wire all values flowing inside the function outputs to the function call outputs.
+		NewContext.Call = Call;
+
+		// Forward values flowing into call inputs to called function inputs
+		for (int i = 0; i < Call->FunctionInputs.Num(); ++i)
+		{
+			FFunctionExpressionInput& FunctionInput = Call->FunctionInputs[i];
+			MIR::FValue* Value = AuxValueList[i];
+
+			// Bind the value flowing into the function call input to the function input
+			// expression (inside the function) in the new context.
+			NewContext.SetOutputValue(FunctionInput.ExpressionInput->GetOutput(0), Value);
+
+			// Mark the function input as built.
+			NewContext.BuiltExpressions.Add(FunctionInput.ExpressionInput.Get());
+		}
+
+		// Finally push the function outputs to the expression evaluation stack in the new context.
+		for (FFunctionExpressionOutput& FunctionOutput : Call->FunctionOutputs)
+		{
+			NewContext.ExpressionStack.Push(FunctionOutput.ExpressionOutput.Get());
+		}
+	}
+
+	void PopFunctionCall()
+	{
+		// Pull the values flowing into the function outputs out of the current
+		// context and flow them into the Call outputs in the parent context so that
+		// analysis can continue from the call expression.
+		FAnalysisContext& ParentContext = AnalysisContextStack[AnalysisContextStack.Num() - 2];
+		FAnalysisContext& CurrContext = AnalysisContextStack[AnalysisContextStack.Num() - 1];
+		UMaterialExpressionMaterialFunctionCall* Call = CurrContext.Call;
+
+		for (int i = 0; i < Call->FunctionOutputs.Num(); ++i)
+		{
+			FFunctionExpressionOutput& FunctionOutput = Call->FunctionOutputs[i];
+
+			// Get the value flowing into the function output inside the function in the current context.
+			MIR::FValue* Value = Emitter->Get(FunctionOutput.ExpressionOutput->GetInput(0));
+
+			// Get the function output type.
+			MIR::FTypePtr OutputType = MIR::FType::FromMaterialValueType((EMaterialValueType)FunctionOutput.ExpressionOutput->GetOutputType(0));
+
+			// Cast the value to the expected output type. This may fail (value will be null). 
+			Value = Emitter->EmitConstruct(OutputType, Value);
+
+			// And flow it to the relative function *call* output in the parent context.
+			ParentContext.SetOutputValue(Call->GetOutput(i), Value);
+		}
+
+		// Finally pop this context (the function call) to return to the caller.
+		AnalysisContextStack.Pop();
+
+		// Populate the insight information about this expression pins.
+		AddExpressionConnectionInsights(Call);
 	}
 
-	void Build_LinkMaterialOutputsToIncomingValues()
+	void Step_FlowValuesIntoMaterialOutputs()
 	{
-		for (MIR::FSetMaterialOutput* Output : Module.Outputs)
+		FAnalysisContext& Context = AnalysisContextStack.Last();
+
+		for (MIR::FSetMaterialOutput* Output : Module->Outputs)
 		{
 			FMaterialInputDescription Input;
-			ensure(Params.Material->GetExpressionInputDescription(Output->Property, Input));
+			ensure(Builder->Material->GetExpressionInputDescription(Output->Property, Input));
 
 			if (!Output->Arg)
 			{
-				MIR::FValue** ValuePtr = Builder.OutputValues.Find(Input.Input->GetConnectedOutput());
-				check(ValuePtr);
+				MIR::FValue** ValuePtr = Context.OutputValues.Find(Input.Input->GetConnectedOutput());
+				check(ValuePtr && *ValuePtr);
 
-				Builder.InputValues.Add(Input.Input, *ValuePtr);
+				MIR::Internal::SetInputValue(this, Input.Input, *ValuePtr);
 
 				MIR::FTypePtr OutputArgType = MIR::FType::FromShaderType(Input.Type);
-				Output->Arg = Emitter.EmitConstruct(OutputArgType, *ValuePtr);
+				Output->Arg = Emitter->EmitConstruct(OutputArgType, *ValuePtr);
 			}
 
-			if (Params.TargetInsight)
+			if (Builder->TargetInsight)
 			{
 				check(Output->Arg);
-				PushConnectionInsight(Params.Material, (int)Output->Property, Input.Input->Expression, Input.Input->OutputIndex, Output->Arg->Type);
+				PushConnectionInsight(Builder->Material, (int)Output->Property, Input.Input->Expression, Input.Input->OutputIndex, Output->Arg->Type);
 			}
 		}
 	}
 
-	void Build_AnalyzeIRGraph()
+	void Step_AnalyzeIRGraph()
 	{
 		InstructionStack.Reserve(64);
 
-		for (MIR::FSetMaterialOutput* Output : Module.Outputs)
+		for (MIR::FSetMaterialOutput* Output : Module->Outputs)
 		{
 			InstructionStack.Push(Output);
 		}
@@ -172,7 +325,7 @@ struct FMaterialIRModuleBuilder::FPrivate
 				if (UseValue && !(UseValue->Flags & MIR::VF_ValueAnalyzed))
 				{
 					UseValue->SetFlags(MIR::VF_ValueAnalyzed);
-					Build_AnalyzeValue(UseValue);
+					AnalyzeValue(UseValue);
 				}
 
 				MIR::FInstruction* Use = UseValue->AsInstruction();
@@ -192,36 +345,36 @@ struct FMaterialIRModuleBuilder::FPrivate
 		}
 	}
 
-	void Build_AnalyzeValue(MIR::FValue* Value)
+	void AnalyzeValue(MIR::FValue* Value)
 	{
 		if (auto ExternalInput = Value->As<MIR::FExternalInput>())
 		{
-			Module.Statistics.ExternalInputUsedMask[SF_Vertex][(int)ExternalInput->Id] = true;
-			Module.Statistics.ExternalInputUsedMask[SF_Pixel][(int)ExternalInput->Id] = true;
+			Module->Statistics.ExternalInputUsedMask[SF_Vertex][(int)ExternalInput->Id] = true;
+			Module->Statistics.ExternalInputUsedMask[SF_Pixel][(int)ExternalInput->Id] = true;
 		}
 		else if (auto TextureSample = Value->As<MIR::FTextureSample>())
 		{
-			EMaterialTextureParameterType ParamType = UE::Utility::TextureMaterialValueTypeToParameterType(TextureSample->Texture->GetMaterialType());
+			EMaterialTextureParameterType ParamType = MIR::Internal::TextureMaterialValueTypeToParameterType(TextureSample->Texture->GetMaterialType());
 
 			FMaterialTextureParameterInfo ParamInfo{};
 			ParamInfo.ParameterInfo = { "", EMaterialParameterAssociation::GlobalParameter, INDEX_NONE };
 			ParamInfo.SamplerSource = SSM_FromTextureAsset; // TODO - Is this needed?
 
-			ParamInfo.TextureIndex = Params.Material->GetReferencedTextures().Find(TextureSample->Texture);
+			ParamInfo.TextureIndex = Builder->Material->GetReferencedTextures().Find(TextureSample->Texture);
 			check(ParamInfo.TextureIndex != INDEX_NONE);
 
-			TextureSample->TextureParameterIndex = Module.CompilationOutput.UniformExpressionSet.FindOrAddTextureParameter(ParamType, ParamInfo);
+			TextureSample->TextureParameterIndex = Module->CompilationOutput.UniformExpressionSet.FindOrAddTextureParameter(ParamType, ParamInfo);
 		}
 	}
 
-	void Build_PopulateBlocks()
+	void Step_PopulateBlocks()
 	{
 		// This function walks the instruction graph and puts each instruction into the inner most possible block.
 		InstructionStack.Empty(InstructionStack.Max());
 
-		for (MIR::FSetMaterialOutput* Output : Module.Outputs)
+		for (MIR::FSetMaterialOutput* Output : Module->Outputs)
 		{
-			Output->Block = Module.RootBlock;
+			Output->Block = Module->RootBlock;
 			InstructionStack.Add(Output);
 		}
 
@@ -280,39 +433,44 @@ struct FMaterialIRModuleBuilder::FPrivate
 		}
 	}
 
-	void Build_Finalize()
+	void Step_Finalize()
 	{
 		/* Produce the module statistics */
 		for (int TexCoordIndex = 0; TexCoordIndex < MIR::TexCoordMaxNum; ++TexCoordIndex)
 		{
 			MIR::EExternalInput TexCoordInput = MIR::TexCoordIndexToExternalInput(TexCoordIndex);
-			if (Module.Statistics.ExternalInputUsedMask[SF_Vertex][(int)TexCoordInput])
+			if (Module->Statistics.ExternalInputUsedMask[SF_Vertex][(int)TexCoordInput])
 			{
-				Module.Statistics.NumVertexTexCoords = TexCoordIndex + 1;
+				Module->Statistics.NumVertexTexCoords = TexCoordIndex + 1;
 			}
-			if (Module.Statistics.ExternalInputUsedMask[SF_Pixel][(int)TexCoordInput])
+			if (Module->Statistics.ExternalInputUsedMask[SF_Pixel][(int)TexCoordInput])
 			{
-				Module.Statistics.NumPixelTexCoords = TexCoordIndex + 1;
+				Module->Statistics.NumPixelTexCoords = TexCoordIndex + 1;
 			}
 		}
 
 		/* Configure the compilation output */
-		FMaterialCompilationOutput& CompilationOutput = Module.CompilationOutput;
-		CompilationOutput.NumUsedUVScalars = Module.Statistics.NumPixelTexCoords * 2;
+		FMaterialCompilationOutput& CompilationOutput = Module->CompilationOutput;
+		CompilationOutput.NumUsedUVScalars = Module->Statistics.NumPixelTexCoords * 2;
 	}
 
 	void AddExpressionConnectionInsights(UMaterialExpression* Expression)
 	{
+		if (!Builder->TargetInsight)
+		{
+			return;
+		}
+
 		// Update expression inputs insight.
-		for (FExpressionInputIterator It{ Expression}; It; ++It)
+		for (FExpressionInputIterator It{ Expression }; It; ++It)
 		{
 			if (!It->IsConnected())
 			{
 				continue;
 			}
 
-			MIR::FValue** Value = Builder.InputValues.Find(It.Input);
-			PushConnectionInsight(Expression, It.Index, It->Expression, It->OutputIndex, Value ? (*Value)->Type : nullptr);
+			MIR::FValue* Value = MIR::Internal::GetInputValue(this, It.Input);
+			PushConnectionInsight(Expression, It.Index, It->Expression, It->OutputIndex, Value ? Value->Type : nullptr);
 		}
 	}
 	
@@ -325,32 +483,53 @@ struct FMaterialIRModuleBuilder::FPrivate
 		Insight.OutputIndex = OutputIndex,
 		Insight.ValueType = Type ? Type->ToValueType() : UE::Shader::EValueType::Any,
 		
-		Params.TargetInsight->ConnectionInsights.Push(Insight);
+		Builder->TargetInsight->ConnectionInsights.Push(Insight);
 	}
 };
 
-bool FMaterialIRModuleBuilder::Build(const FMaterialIRModuleBuildParams& Params, FMaterialIRModule* TargetModule)
+bool FMaterialIRModuleBuilder::Build(FMaterialIRModule* TargetModule)
 {
-	MIR::FEmitter Emitter{ this, Params.Material, TargetModule };
-	FPrivate Private{ *this, Params, *TargetModule, Emitter };
+	FMaterialIRModuleBuilderImpl Impl{ this, TargetModule };
+	
+	MIR::FEmitter Emitter{ &Impl, Material, TargetModule };
+	Impl.Emitter = &Emitter;
 
-	Private.Build_Initialize();
-	Private.Build_GenerateOutputInstructions();
-	Private.Build_BuildIRGraph();
+	Impl.Step_Initialize();
+	Impl.Step_GenerateOutputInstructions();
+	Impl.Step_BuildMaterialExpressionsToIRGraph();
 
-	if (Private.Emitter.IsInvalid())
+	if (Impl.Emitter->IsInvalid())
 	{
 		return false;
 	}
 
-	Private.Build_LinkMaterialOutputsToIncomingValues();
-	Private.Build_AnalyzeIRGraph();
-	Private.Build_PopulateBlocks();
-	Private.Build_Finalize();
+	Impl.Step_FlowValuesIntoMaterialOutputs();
+	Impl.Step_AnalyzeIRGraph();
+	Impl.Step_PopulateBlocks();
+	Impl.Step_Finalize();
 
 	UE::MIR::DebugDumpIRUseGraph(*TargetModule);
 
 	return true;
 }
 
+namespace UE::MIR::Internal {
+
+FValue* GetInputValue(FMaterialIRModuleBuilderImpl* Builder, const FExpressionInput* Input)
+{
+	return Builder->AnalysisContextStack.Last().GetInputValue(Input);
+}
+
+void SetInputValue(FMaterialIRModuleBuilderImpl* Builder, const FExpressionInput* Input, FValue* Value)
+{
+	Builder->AnalysisContextStack.Last().SetInputValue(Input, Value);
+}
+
+void SetOutputValue(FMaterialIRModuleBuilderImpl* Builder, const FExpressionOutput* Output, FValue* Value)
+{
+	Builder->AnalysisContextStack.Last().SetOutputValue(Output, Value);
+}
+
+} // namespace UE::MIR::Internal
+
 #endif // #if WITH_EDITOR
diff --git a/Engine/Source/Runtime/Engine/Private/Materials/MaterialIRToHLSLTranslator.cpp b/Engine/Source/Runtime/Engine/Private/Materials/MaterialIRToHLSLTranslator.cpp
--- a/Engine/Source/Runtime/Engine/Private/Materials/MaterialIRToHLSLTranslator.cpp
+++ b/Engine/Source/Runtime/Engine/Private/Materials/MaterialIRToHLSLTranslator.cpp
@@ -4,7 +4,7 @@
 #include "Materials/MaterialIRModule.h"
 #include "Materials/MaterialIRTypes.h"
 #include "Materials/MaterialIR.h"
-#include "MaterialIRUtility.h"
+#include "MaterialIRInternal.h"
 
 #include "ShaderCore.h"
 #include "MaterialShared.h"
@@ -226,7 +226,7 @@ struct FTranslator : FMaterialIRToHLSLTranslation
 		for (int32 PropertyIndex = 0; PropertyIndex < MP_MAX; ++PropertyIndex)
 		{
 			EMaterialProperty Property = (EMaterialProperty)PropertyIndex;
-			if (!UE::Utility::IsMaterialPropertyShared(Property))
+			if (!MIR::Internal::IsMaterialPropertyShared(Property))
 			{
 				continue;
 			}
@@ -268,14 +268,19 @@ struct FTranslator : FMaterialIRToHLSLTranslation
 
 			LowerInstruction(Instr);
 
-            if (Printer.Buffer.EndsWith(TEXT("}")))
+			if (Printer.Buffer.EndsWith(TEXT("}")))
 			{
-                Printer << NewLine;
-            }
-            else
+				Printer << NewLine;
+			}
+			else
 			{
-                Printer << EndOfStatement;
-            }
+				Printer << EndOfStatement;
+			}
+
+			if (Instr->Kind == MIR::VK_SetMaterialOutput)
+			{
+				Printer << NewLine;
+			}
         }
 
         NumLocals = OldNumLocals;
@@ -324,16 +329,16 @@ struct FTranslator : FMaterialIRToHLSLTranslation
 			default:
 				UE_MIR_UNREACHABLE();
 		}
-
+	
 		return NoOp;
 	}
 
 	void LowerConstant(const MIR::FConstant* Constant)
 	{
-		MIR::FArithmeticTypePtr ArithType = Constant->Type->AsArithmetic();
-		check(ArithType && ArithType->IsScalar());
+		MIR::FPrimitiveTypePtr checkNoEntry = Constant->Type->AsPrimitive();
+		check(checkNoEntry && checkNoEntry->IsScalar());
 
-		switch (ArithType->ScalarKind)
+		switch (checkNoEntry->ScalarKind)
 		{
 			case MIR::SK_Bool:  Printer.Buffer.Append(Constant->Boolean ? TEXT("true") : TEXT("false")); break;
 			case MIR::SK_Int:   Printer.Buffer.Appendf(TEXT("%") PRId64, Constant->Integer); break;
@@ -373,7 +378,7 @@ struct FTranslator : FMaterialIRToHLSLTranslation
 	
 	void LowerDimensional(const MIR::FDimensional* Dimensional) 
 	{
-		MIR::FArithmeticTypePtr ArithmeticType = Dimensional->Type->AsArithmetic();
+		MIR::FPrimitiveTypePtr ArithmeticType = Dimensional->Type->AsPrimitive();
 		check(ArithmeticType && ArithmeticType->IsVector());
 
 		Printer << ScalarKindToString(ArithmeticType->ScalarKind) << ArithmeticType->NumRows << BeginArgs;
@@ -440,7 +445,7 @@ struct FTranslator : FMaterialIRToHLSLTranslation
 	{
 		LowerValue(Subscript->Arg);
 
-		if (MIR::FArithmeticTypePtr ArgArithmeticType = Subscript->Arg->Type->AsVector())
+		if (MIR::FPrimitiveTypePtr ArgArithmeticType = Subscript->Arg->Type->AsVector())
 		{
 			const TCHAR* ComponentsStr[] = { TEXT(".x"), TEXT(".y"), TEXT(".z"), TEXT(".w") };
 			check(Subscript->Index <= ArgArithmeticType->GetNumComponents());
@@ -643,7 +648,7 @@ struct FTranslator : FMaterialIRToHLSLTranslation
 	
 	ENoOp LowerType(const MIR::FType* Type)
 	{
-		if (auto ArithmeticType = Type->AsArithmetic())
+		if (auto ArithmeticType = Type->AsPrimitive())
 		{
 			switch (ArithmeticType->ScalarKind)
 			{
diff --git a/Engine/Source/Runtime/Engine/Private/Materials/MaterialIRTypes.cpp b/Engine/Source/Runtime/Engine/Private/Materials/MaterialIRTypes.cpp
--- a/Engine/Source/Runtime/Engine/Private/Materials/MaterialIRTypes.cpp
+++ b/Engine/Source/Runtime/Engine/Private/Materials/MaterialIRTypes.cpp
@@ -12,7 +12,7 @@ const TCHAR* TypeKindToString(ETypeKind Kind)
 	switch (Kind)
 	{
 		case TK_Void: return TEXT("void");
-		case TK_Arithmetic: return TEXT("arithmetic");
+		case TK_Primitive: return TEXT("primitive");
 		default: UE_MIR_UNREACHABLE();
 	}
 }
@@ -31,19 +31,49 @@ FTypePtr FType::FromShaderType(const UE::Shader::FType& InShaderType)
 		case UE::Shader::EValueType::Float2:
 		case UE::Shader::EValueType::Float3:
 		case UE::Shader::EValueType::Float4:
-			return FArithmeticType::GetVector(SK_Float, (int)InShaderType.ValueType - (int)UE::Shader::EValueType::Float1 + 1);
+			return FPrimitiveType::GetVector(SK_Float, (int)InShaderType.ValueType - (int)UE::Shader::EValueType::Float1 + 1);
 
 		case UE::Shader::EValueType::Int1:
 		case UE::Shader::EValueType::Int2:
 		case UE::Shader::EValueType::Int3:
 		case UE::Shader::EValueType::Int4:
-			return FArithmeticType::GetVector(SK_Int, (int)InShaderType.ValueType - (int)UE::Shader::EValueType::Int1 + 1);
+			return FPrimitiveType::GetVector(SK_Int, (int)InShaderType.ValueType - (int)UE::Shader::EValueType::Int1 + 1);
 
 		case UE::Shader::EValueType::Bool1:
 		case UE::Shader::EValueType::Bool2:
 		case UE::Shader::EValueType::Bool3:
 		case UE::Shader::EValueType::Bool4:
-			return FArithmeticType::GetVector(SK_Bool, (int)InShaderType.ValueType - (int)UE::Shader::EValueType::Int1 + 1);
+			return FPrimitiveType::GetVector(SK_Bool, (int)InShaderType.ValueType - (int)UE::Shader::EValueType::Int1 + 1);
+
+		default:
+			UE_MIR_UNREACHABLE();
+	}
+}
+
+FTypePtr FType::FromMaterialValueType(EMaterialValueType Type)
+{
+	switch (Type)
+	{
+		case MCT_VoidStatement:
+			return FType::GetVoid();
+
+		case MCT_Float1:
+		case MCT_Float2:
+		case MCT_Float3:
+		case MCT_Float4:
+			return FPrimitiveType::GetVector(SK_Float, (int)Type - (int)MCT_Float1 + 1);
+
+		case MCT_Float:
+			return FPrimitiveType::GetVector(SK_Float, 4);
+
+		case MCT_UInt1:
+		case MCT_UInt2:
+		case MCT_UInt3:
+		case MCT_UInt4:
+			return FPrimitiveType::GetVector(SK_Int, (int)Type - (int)MCT_UInt1 + 1);
+
+		case MCT_Bool:
+			return FPrimitiveType::GetVector(SK_Bool, (int)Type - (int)UE::Shader::EValueType::Int1 + 1);
 
 		default:
 			UE_MIR_UNREACHABLE();
@@ -58,9 +88,9 @@ FTypePtr FType::GetVoid()
 
 FStringView FType::GetSpelling() const
 {
-	if (auto Arithmetic = AsArithmetic())
+	if (auto PrimitiveType = AsPrimitive())
 	{
-		return Arithmetic->Spelling;
+		return PrimitiveType->Spelling;
 	}
 	
 	UE_MIR_UNREACHABLE();
@@ -70,13 +100,13 @@ UE::Shader::EValueType FType::ToValueType() const
 {
 	using namespace UE::Shader;
 
-	if (FArithmeticTypePtr ArithmeticType = AsArithmetic())
+	if (FPrimitiveTypePtr PrimitiveType = AsPrimitive())
 	{
-		if (ArithmeticType->IsMatrix())
+		if (PrimitiveType->IsMatrix())
 		{
-			if (ArithmeticType->NumRows == 4 && ArithmeticType->NumColumns == 4)
+			if (PrimitiveType->NumRows == 4 && PrimitiveType->NumColumns == 4)
 			{
-				if (ArithmeticType->ScalarKind == SK_Float)
+				if (PrimitiveType->ScalarKind == SK_Float)
 				{
 					return EValueType::Float4x4;
 				}
@@ -89,13 +119,13 @@ UE::Shader::EValueType FType::ToValueType() const
 			return EValueType::Any;
 		}
 
-		check(ArithmeticType->NumColumns == 1 && ArithmeticType->NumRows <= 4);
+		check(PrimitiveType->NumColumns == 1 && PrimitiveType->NumRows <= 4);
 
-		switch (ArithmeticType->ScalarKind)
+		switch (PrimitiveType->ScalarKind)
 		{
-			case SK_Bool: 	return (EValueType)((int)EValueType::Bool1 + ArithmeticType->NumRows - 1); break;
-			case SK_Int: 	return (EValueType)((int)EValueType::Int1 + ArithmeticType->NumRows - 1); break;
-			case SK_Float: 	return (EValueType)((int)EValueType::Float1 + ArithmeticType->NumRows - 1); break;
+			case SK_Bool: 	return (EValueType)((int)EValueType::Bool1 + PrimitiveType->NumRows - 1); break;
+			case SK_Int: 	return (EValueType)((int)EValueType::Int1 + PrimitiveType->NumRows - 1); break;
+			case SK_Float: 	return (EValueType)((int)EValueType::Float1 + PrimitiveType->NumRows - 1); break;
 			default: UE_MIR_UNREACHABLE();
 		}
 	}
@@ -103,31 +133,31 @@ UE::Shader::EValueType FType::ToValueType() const
 	UE_MIR_UNREACHABLE();
 }
 
-bool FType::IsBool1() const
+bool FType::IsBoolScalar() const
 {
-	return this == FArithmeticType::GetBool1();
+	return this == FPrimitiveType::GetBool1();
 }
 
-FArithmeticTypePtr FType::AsArithmetic() const
+FPrimitiveTypePtr FType::AsPrimitive() const
 {
-	return Kind == TK_Arithmetic ? static_cast<FArithmeticTypePtr>(this) : nullptr; 
+	return Kind == TK_Primitive ? static_cast<FPrimitiveTypePtr>(this) : nullptr; 
 }
 
-FArithmeticTypePtr FType::AsScalar() const
+FPrimitiveTypePtr FType::AsScalar() const
 {
-	FArithmeticTypePtr Type = AsArithmetic();
+	FPrimitiveTypePtr Type = AsPrimitive();
 	return Type->IsScalar() ? Type : nullptr;
 }
 
-FArithmeticTypePtr FType::AsVector() const
+FPrimitiveTypePtr FType::AsVector() const
 {
-	FArithmeticTypePtr Type = AsArithmetic();
+	FPrimitiveTypePtr Type = AsPrimitive();
 	return Type->IsVector() ? Type : nullptr;
 }
 
-FArithmeticTypePtr FType::AsMatrix() const
+FPrimitiveTypePtr FType::AsMatrix() const
 {
-	FArithmeticTypePtr Type = AsArithmetic();
+	FPrimitiveTypePtr Type = AsPrimitive();
 	return Type->IsMatrix() ? Type : nullptr;
 }
 
@@ -142,109 +172,109 @@ const TCHAR* ScalarKindToString(EScalarKind Kind)
 	}
 }
 
-FArithmeticTypePtr FArithmeticType::GetBool1()
+FPrimitiveTypePtr FPrimitiveType::GetBool1()
 {
 	return GetScalar(SK_Bool);
 }
 
-FArithmeticTypePtr FArithmeticType::GetInt1()
+FPrimitiveTypePtr FPrimitiveType::GetInt1()
 {
 	return GetScalar(SK_Int);
 }
 
-FArithmeticTypePtr FArithmeticType::GetFloat1()
+FPrimitiveTypePtr FPrimitiveType::GetFloat1()
 {
 	return GetScalar(SK_Float);
 }
 
-FArithmeticTypePtr FArithmeticType::GetFloat2()
+FPrimitiveTypePtr FPrimitiveType::GetFloat2()
 {
 	return GetVector(SK_Float, 2);
 }
 
-FArithmeticTypePtr FArithmeticType::GetFloat3()
+FPrimitiveTypePtr FPrimitiveType::GetFloat3()
 {
 	return GetVector(SK_Float, 3);
 }
 
-FArithmeticTypePtr FArithmeticType::GetFloat4()
+FPrimitiveTypePtr FPrimitiveType::GetFloat4()
 {
 	return GetVector(SK_Float, 4);
 }
 
-const FArithmeticType* FArithmeticType::GetScalar(EScalarKind InScalarKind)
+const FPrimitiveType* FPrimitiveType::GetScalar(EScalarKind InScalarKind)
 {
 	return Get(InScalarKind, 1, 1);
 }
 
-const FArithmeticType* FArithmeticType::GetVector(EScalarKind InScalarKind, int NumComponents)
+const FPrimitiveType* FPrimitiveType::GetVector(EScalarKind InScalarKind, int NumComponents)
 {
 	check(NumComponents >= 1 && NumComponents <= 4);
 	return Get(InScalarKind, NumComponents, 1);
 }
 
-const FArithmeticType* FArithmeticType::GetMatrix(EScalarKind InScalarKind, int NumRows, int NumColumns)
+const FPrimitiveType* FPrimitiveType::GetMatrix(EScalarKind InScalarKind, int NumRows, int NumColumns)
 {
 	check(NumColumns > 1 && NumColumns <= 4);
 	check(NumRows > 1 && NumRows <= 4);
 	return Get(InScalarKind, NumRows, NumColumns);
 }
 
-FArithmeticTypePtr FArithmeticType::Get(EScalarKind InScalarKind, int NumRows, int NumColumns)
+FPrimitiveTypePtr FPrimitiveType::Get(EScalarKind InScalarKind, int NumRows, int NumColumns)
 {
 	check(InScalarKind >= 0 && InScalarKind <= SK_Float);
 	
 	static const FStringView Invalid = TEXT("invalid");
 
-	static const FArithmeticType Types[] {
-		{ { TK_Arithmetic }, { TEXT("bool") }, 		SK_Bool, 1, 1 },
-		{ { TK_Arithmetic }, Invalid, 				SK_Bool, 1, 2 }, 
-		{ { TK_Arithmetic }, Invalid, 				SK_Bool, 1, 3 },
-		{ { TK_Arithmetic }, Invalid, 				SK_Bool, 1, 4 },
-		{ { TK_Arithmetic }, { TEXT("bool2") },   	SK_Bool, 2, 1 },
-		{ { TK_Arithmetic }, { TEXT("bool2x2") }, 	SK_Bool, 2, 2 },
-		{ { TK_Arithmetic }, { TEXT("bool2x3") }, 	SK_Bool, 2, 3 },
-		{ { TK_Arithmetic }, { TEXT("bool2x4") }, 	SK_Bool, 2, 4 },
-		{ { TK_Arithmetic }, { TEXT("bool3") },   	SK_Bool, 3, 1 },
-		{ { TK_Arithmetic }, { TEXT("bool3x2") }, 	SK_Bool, 3, 2 },
-		{ { TK_Arithmetic }, { TEXT("bool3x3") }, 	SK_Bool, 3, 3 },
-		{ { TK_Arithmetic }, { TEXT("bool3x4") }, 	SK_Bool, 3, 4 },
-		{ { TK_Arithmetic }, { TEXT("bool4") },   	SK_Bool, 4, 1 },
-		{ { TK_Arithmetic }, { TEXT("bool4x2") }, 	SK_Bool, 4, 2 },
-		{ { TK_Arithmetic }, { TEXT("bool4x3") }, 	SK_Bool, 4, 3 },
-		{ { TK_Arithmetic }, { TEXT("bool4x4") }, 	SK_Bool, 4, 4 },
-		{ { TK_Arithmetic }, { TEXT("int") }, 		SK_Int, 1, 1 },
-		{ { TK_Arithmetic }, Invalid, 				SK_Int, 1, 2 },
-		{ { TK_Arithmetic }, Invalid, 				SK_Int, 1, 3 },
-		{ { TK_Arithmetic }, Invalid, 				SK_Int, 1, 4 },
-		{ { TK_Arithmetic }, { TEXT("int2") },   	SK_Int, 2, 1 },
-		{ { TK_Arithmetic }, { TEXT("int2x2") }, 	SK_Int, 2, 2 },
-		{ { TK_Arithmetic }, { TEXT("int2x3") }, 	SK_Int, 2, 3 },
-		{ { TK_Arithmetic }, { TEXT("int2x4") }, 	SK_Int, 2, 4 },
-		{ { TK_Arithmetic }, { TEXT("int3") },   	SK_Int, 3, 1 },
-		{ { TK_Arithmetic }, { TEXT("int3x2") }, 	SK_Int, 3, 2 },
-		{ { TK_Arithmetic }, { TEXT("int3x3") }, 	SK_Int, 3, 3 },
-		{ { TK_Arithmetic }, { TEXT("int3x4") }, 	SK_Int, 3, 4 },
-		{ { TK_Arithmetic }, { TEXT("int4") },   	SK_Int, 4, 1 },
-		{ { TK_Arithmetic }, { TEXT("int4x2") }, 	SK_Int, 4, 2 },
-		{ { TK_Arithmetic }, { TEXT("int4x3") }, 	SK_Int, 4, 3 },
-		{ { TK_Arithmetic }, { TEXT("int4x4") }, 	SK_Int, 4, 4 },
-		{ { TK_Arithmetic }, { TEXT("float") }, 	SK_Float, 1, 1 },
-		{ { TK_Arithmetic }, Invalid, 				SK_Float, 1, 2 },
-		{ { TK_Arithmetic }, Invalid, 				SK_Float, 1, 3 },
-		{ { TK_Arithmetic }, Invalid, 				SK_Float, 1, 4 },
-		{ { TK_Arithmetic }, { TEXT("float2") },   	SK_Float, 2, 1 },
-		{ { TK_Arithmetic }, { TEXT("float2x2") }, 	SK_Float, 2, 2 },
-		{ { TK_Arithmetic }, { TEXT("float2x3") }, 	SK_Float, 2, 3 },
-		{ { TK_Arithmetic }, { TEXT("float2x4") }, 	SK_Float, 2, 4 },
-		{ { TK_Arithmetic }, { TEXT("float3") },   	SK_Float, 3, 1 },
-		{ { TK_Arithmetic }, { TEXT("float3x2") }, 	SK_Float, 3, 2 },
-		{ { TK_Arithmetic }, { TEXT("float3x3") }, 	SK_Float, 3, 3 },
-		{ { TK_Arithmetic }, { TEXT("float3x4") }, 	SK_Float, 3, 4 },
-		{ { TK_Arithmetic }, { TEXT("float4") },   	SK_Float, 4, 1 },
-		{ { TK_Arithmetic }, { TEXT("float4x2") }, 	SK_Float, 4, 2 },
-		{ { TK_Arithmetic }, { TEXT("float4x3") }, 	SK_Float, 4, 3 },
-		{ { TK_Arithmetic }, { TEXT("float4x4") }, 	SK_Float, 4, 4 },
+	static const FPrimitiveType Types[] {
+		{ { TK_Primitive }, { TEXT("bool") }, 		SK_Bool, 1, 1 },
+		{ { TK_Primitive }, Invalid, 				SK_Bool, 1, 2 }, 
+		{ { TK_Primitive }, Invalid, 				SK_Bool, 1, 3 },
+		{ { TK_Primitive }, Invalid, 				SK_Bool, 1, 4 },
+		{ { TK_Primitive }, { TEXT("bool2") },   	SK_Bool, 2, 1 },
+		{ { TK_Primitive }, { TEXT("bool2x2") }, 	SK_Bool, 2, 2 },
+		{ { TK_Primitive }, { TEXT("bool2x3") }, 	SK_Bool, 2, 3 },
+		{ { TK_Primitive }, { TEXT("bool2x4") }, 	SK_Bool, 2, 4 },
+		{ { TK_Primitive }, { TEXT("bool3") },   	SK_Bool, 3, 1 },
+		{ { TK_Primitive }, { TEXT("bool3x2") }, 	SK_Bool, 3, 2 },
+		{ { TK_Primitive }, { TEXT("bool3x3") }, 	SK_Bool, 3, 3 },
+		{ { TK_Primitive }, { TEXT("bool3x4") }, 	SK_Bool, 3, 4 },
+		{ { TK_Primitive }, { TEXT("bool4") },   	SK_Bool, 4, 1 },
+		{ { TK_Primitive }, { TEXT("bool4x2") }, 	SK_Bool, 4, 2 },
+		{ { TK_Primitive }, { TEXT("bool4x3") }, 	SK_Bool, 4, 3 },
+		{ { TK_Primitive }, { TEXT("bool4x4") }, 	SK_Bool, 4, 4 },
+		{ { TK_Primitive }, { TEXT("int") }, 		SK_Int, 1, 1 },
+		{ { TK_Primitive }, Invalid, 				SK_Int, 1, 2 },
+		{ { TK_Primitive }, Invalid, 				SK_Int, 1, 3 },
+		{ { TK_Primitive }, Invalid, 				SK_Int, 1, 4 },
+		{ { TK_Primitive }, { TEXT("int2") },   	SK_Int, 2, 1 },
+		{ { TK_Primitive }, { TEXT("int2x2") }, 	SK_Int, 2, 2 },
+		{ { TK_Primitive }, { TEXT("int2x3") }, 	SK_Int, 2, 3 },
+		{ { TK_Primitive }, { TEXT("int2x4") }, 	SK_Int, 2, 4 },
+		{ { TK_Primitive }, { TEXT("int3") },   	SK_Int, 3, 1 },
+		{ { TK_Primitive }, { TEXT("int3x2") }, 	SK_Int, 3, 2 },
+		{ { TK_Primitive }, { TEXT("int3x3") }, 	SK_Int, 3, 3 },
+		{ { TK_Primitive }, { TEXT("int3x4") }, 	SK_Int, 3, 4 },
+		{ { TK_Primitive }, { TEXT("int4") },   	SK_Int, 4, 1 },
+		{ { TK_Primitive }, { TEXT("int4x2") }, 	SK_Int, 4, 2 },
+		{ { TK_Primitive }, { TEXT("int4x3") }, 	SK_Int, 4, 3 },
+		{ { TK_Primitive }, { TEXT("int4x4") }, 	SK_Int, 4, 4 },
+		{ { TK_Primitive }, { TEXT("float") }, 	SK_Float, 1, 1 },
+		{ { TK_Primitive }, Invalid, 				SK_Float, 1, 2 },
+		{ { TK_Primitive }, Invalid, 				SK_Float, 1, 3 },
+		{ { TK_Primitive }, Invalid, 				SK_Float, 1, 4 },
+		{ { TK_Primitive }, { TEXT("float2") },   	SK_Float, 2, 1 },
+		{ { TK_Primitive }, { TEXT("float2x2") }, 	SK_Float, 2, 2 },
+		{ { TK_Primitive }, { TEXT("float2x3") }, 	SK_Float, 2, 3 },
+		{ { TK_Primitive }, { TEXT("float2x4") }, 	SK_Float, 2, 4 },
+		{ { TK_Primitive }, { TEXT("float3") },   	SK_Float, 3, 1 },
+		{ { TK_Primitive }, { TEXT("float3x2") }, 	SK_Float, 3, 2 },
+		{ { TK_Primitive }, { TEXT("float3x3") }, 	SK_Float, 3, 3 },
+		{ { TK_Primitive }, { TEXT("float3x4") }, 	SK_Float, 3, 4 },
+		{ { TK_Primitive }, { TEXT("float4") },   	SK_Float, 4, 1 },
+		{ { TK_Primitive }, { TEXT("float4x2") }, 	SK_Float, 4, 2 },
+		{ { TK_Primitive }, { TEXT("float4x3") }, 	SK_Float, 4, 3 },
+		{ { TK_Primitive }, { TEXT("float4x4") }, 	SK_Float, 4, 4 },
 	};
 
 	int Index = InScalarKind * 4 * 4 + (NumRows - 1) * 4 + (NumColumns - 1);
@@ -252,9 +282,9 @@ FArithmeticTypePtr FArithmeticType::Get(EScalarKind InScalarKind, int NumRows, i
 	return &Types[Index];
 }
 
-FArithmeticTypePtr FArithmeticType::ToScalar() const
+FPrimitiveTypePtr FPrimitiveType::ToScalar() const
 {
-	return FArithmeticType::GetScalar(ScalarKind);
+	return FPrimitiveType::GetScalar(ScalarKind);
 }
 
 FTypePtr FTextureType::Get()
diff --git a/Engine/Source/Runtime/Engine/Private/Materials/MaterialIRUtility.cpp b/Engine/Source/Runtime/Engine/Private/Materials/MaterialIRUtility.cpp
deleted file mode 100644
diff --git a/Engine/Source/Runtime/Engine/Private/Materials/MaterialIRUtility.h b/Engine/Source/Runtime/Engine/Private/Materials/MaterialIRUtility.h
deleted file mode 100644
diff --git a/Engine/Source/Runtime/Engine/Private/Materials/MaterialShared.cpp b/Engine/Source/Runtime/Engine/Private/Materials/MaterialShared.cpp
--- a/Engine/Source/Runtime/Engine/Private/Materials/MaterialShared.cpp
+++ b/Engine/Source/Runtime/Engine/Private/Materials/MaterialShared.cpp
@@ -3361,7 +3361,8 @@ bool FMaterial::Translate_New(const FMaterialShaderMapId& InShaderMapId,
 
 	GetMaterialInterface()->MaterialInsight.Reset(new FMaterialInsights);
 
-	FMaterialIRModuleBuildParams TargetParams{
+	// Build the material
+	FMaterialIRModuleBuilder Builder = {
 		.Material = GetMaterialInterface()->GetMaterial(),
 		.ShaderPlatform = InShaderPlatform,
 		.TargetPlatform = InTargetPlatform,
@@ -3369,9 +3370,7 @@ bool FMaterial::Translate_New(const FMaterialShaderMapId& InShaderMapId,
 		.TargetInsight = GetMaterialInterface()->MaterialInsight.Get(),
 	};
 
-	// Build the material
-	FMaterialIRModuleBuilder Builder;
-	if (!Builder.Build(TargetParams, &Module))
+	if (!Builder.Build(&Module))
 	{
 		for (const FMaterialIRModule::FError& Error : Module.GetErrors())
 		{
diff --git a/Engine/Source/Runtime/Engine/Public/Materials/MaterialIR.h b/Engine/Source/Runtime/Engine/Public/Materials/MaterialIR.h
--- a/Engine/Source/Runtime/Engine/Public/Materials/MaterialIR.h
+++ b/Engine/Source/Runtime/Engine/Public/Materials/MaterialIR.h
@@ -57,7 +57,11 @@ struct FValue
 	TArrayView<const FValue*> GetUses() const;
 	TArrayView<FValue*> GetUses();
 	UTexture* GetTexture();
-	
+	bool IsConstantTrue() const;
+	bool IsConstantFalse() const;
+	bool IsConstantZero() const;
+	bool IsConstantOne() const;
+
 	template <typename T>
 	T* As() { return this && IsA(T::TypeKind) ? static_cast<T*>(this) : nullptr; }
 
@@ -83,6 +87,10 @@ struct FConstant : TValue<VK_Constant>
 		TFloat 		Float;
 	};
 
+	bool IsBool() const;
+	bool IsInteger() const;
+	bool IsFloat() const;
+
 	template <typename T>
 	T Get() const
 	{
@@ -188,6 +196,8 @@ struct TInstruction : FInstruction
 
 struct FDimensional : TInstruction<VK_Dimensional>
 {
+	static constexpr int MaxNumComponents = 16;
+
 	// Returns the constant array of component values. 
 	TArrayView<FValue* const> GetComponents() const;
 
@@ -196,7 +206,6 @@ struct FDimensional : TInstruction<VK_Dimensional>
 
 	// Returns whether all components are constant.
 	bool AreComponentsConstant() const;
-
 };
 
 template <int TDimension>
diff --git a/Engine/Source/Runtime/Engine/Public/Materials/MaterialIRCommon.h b/Engine/Source/Runtime/Engine/Public/Materials/MaterialIRCommon.h
--- a/Engine/Source/Runtime/Engine/Public/Materials/MaterialIRCommon.h
+++ b/Engine/Source/Runtime/Engine/Public/Materials/MaterialIRCommon.h
@@ -14,7 +14,7 @@ enum EMaterialProperty : int;
 class FMaterial;
 class FMaterialIRModule;
 class FMaterialIRModule;
-class FMaterialIRModuleBuilder;
+struct FMaterialIRModuleBuilder;
 class ITargetPlatform;
 class UMaterial;
 class UMaterialExpression;
@@ -25,6 +25,8 @@ struct FStaticParameterSet;
 struct FMaterialInputDescription;
 struct FMaterialInsights;
 class UTexture;
+class UMaterialExpressionMaterialFunctionCall;
+struct FMaterialIRModuleBuilderImpl;
 
 namespace UE::Shader
 {
@@ -38,10 +40,10 @@ namespace UE::MIR
 
 /* Types*/
 struct FType;
-struct FArithmeticType;
+struct FPrimitiveType;
 
 using FTypePtr = const FType*;
-using FArithmeticTypePtr = const FArithmeticType*;
+using FPrimitiveTypePtr = const FPrimitiveType*;
 
 /* IR */
 struct FValue;
diff --git a/Engine/Source/Runtime/Engine/Public/Materials/MaterialIRDebug.cpp b/Engine/Source/Runtime/Engine/Public/Materials/MaterialIRDebug.cpp
deleted file mode 100644
diff --git a/Engine/Source/Runtime/Engine/Public/Materials/MaterialIREmitter.h b/Engine/Source/Runtime/Engine/Public/Materials/MaterialIREmitter.h
--- a/Engine/Source/Runtime/Engine/Public/Materials/MaterialIREmitter.h
+++ b/Engine/Source/Runtime/Engine/Public/Materials/MaterialIREmitter.h
@@ -40,7 +40,7 @@ struct FSwizzleMask
 class FEmitter
 {
 public:
-	FEmitter(FMaterialIRModuleBuilder* InBuilder, UMaterial* InMaterial, FMaterialIRModule* InModule);
+	FEmitter(FMaterialIRModuleBuilderImpl* InBuilder, UMaterial* InMaterial, FMaterialIRModule* InModule);
 	
 	/* IO */
 
@@ -74,7 +74,7 @@ public:
 	FValue* GetScalar(const FExpressionInput* Input);
 	
 	//
-	FValue* GetArithmetic(const FExpressionInput* Input);
+	FValue* GetPrimitive(const FExpressionInput* Input);
 
 	// Gets the value flowing into `Input` and returns it after checking that its
 	// type matches `Kind`.
@@ -82,6 +82,7 @@ public:
 
 	/* Error Checking */
 
+	//
 	bool CheckValueValid(const FValue* Value);
 
 	//
@@ -131,7 +132,7 @@ public:
 
 	/* Types */
 
-	FArithmeticTypePtr GetCommonArithmeticType(FArithmeticTypePtr A, FArithmeticTypePtr B);
+	FPrimitiveTypePtr GetCommonPrimitiveType(FPrimitiveTypePtr A, FPrimitiveTypePtr B);
 
 	/* Error reporting */
 
@@ -146,7 +147,7 @@ public:
 	void Error(FString Message);
 
 	struct FPrivate;
-	friend FMaterialIRModuleBuilder;
+	friend FMaterialIRModuleBuilderImpl;
 
 private:
 	void Initialize();
@@ -154,7 +155,7 @@ private:
 	UMaterial* Material{};
 	FMaterialIRModule* Module{};
 	UMaterialExpression* Expression{};
-	FMaterialIRModuleBuilder* Builder{};
+	FMaterialIRModuleBuilderImpl* Builder{};
 	bool bHasExprBuildError = false;
 	FValue* ConstantTrue;
 	FValue* ConstantFalse;
diff --git a/Engine/Source/Runtime/Engine/Public/Materials/MaterialIRModule.h b/Engine/Source/Runtime/Engine/Public/Materials/MaterialIRModule.h
--- a/Engine/Source/Runtime/Engine/Public/Materials/MaterialIRModule.h
+++ b/Engine/Source/Runtime/Engine/Public/Materials/MaterialIRModule.h
@@ -44,7 +44,7 @@ private:
 	FStatistics Statistics;
 
 	friend UE::MIR::FEmitter;
-	friend FMaterialIRModuleBuilder;
+	friend FMaterialIRModuleBuilderImpl;
 };
 
 #endif // #if WITH_EDITOR
diff --git a/Engine/Source/Runtime/Engine/Public/Materials/MaterialIRModuleBuilder.h b/Engine/Source/Runtime/Engine/Public/Materials/MaterialIRModuleBuilder.h
--- a/Engine/Source/Runtime/Engine/Public/Materials/MaterialIRModuleBuilder.h
+++ b/Engine/Source/Runtime/Engine/Public/Materials/MaterialIRModuleBuilder.h
@@ -5,27 +5,16 @@
 
 #if WITH_EDITOR
 
-struct FMaterialIRModuleBuildParams
+//
+struct FMaterialIRModuleBuilder
 {
 	UMaterial* Material;
 	EShaderPlatform ShaderPlatform;
 	const ITargetPlatform* TargetPlatform;
 	const FStaticParameterSet& StaticParameters;
 	FMaterialInsights* TargetInsight{};
-};
-
-//
-class FMaterialIRModuleBuilder
-{
-public:
-	bool Build(const FMaterialIRModuleBuildParams& Params, FMaterialIRModule* TargetModule);
-
-private:
-	TMap<const FExpressionInput*, UE::MIR::FValue*> InputValues;
-	TMap<const FExpressionOutput*, UE::MIR::FValue*> OutputValues;
 
-	friend class UE::MIR::FEmitter;
-	struct FPrivate;
+	bool Build(FMaterialIRModule* TargetModule);
 };
 
 #endif
diff --git a/Engine/Source/Runtime/Engine/Public/Materials/MaterialIRTypes.h b/Engine/Source/Runtime/Engine/Public/Materials/MaterialIRTypes.h
--- a/Engine/Source/Runtime/Engine/Public/Materials/MaterialIRTypes.h
+++ b/Engine/Source/Runtime/Engine/Public/Materials/MaterialIRTypes.h
@@ -3,6 +3,7 @@
 
 #include "Materials/MaterialIRCommon.h"
 #include "Shader/ShaderTypes.h"
+#include "MaterialValueType.h"
 
 #if WITH_EDITOR
 
@@ -12,7 +13,7 @@ namespace UE::MIR
 enum ETypeKind
 {
 	TK_Void,
-	TK_Arithmetic,
+	TK_Primitive,
 	TK_Texture,
 };
 
@@ -26,23 +27,26 @@ struct FType
 	// Returns the type matching specified UE::Shader::FType.
 	static FTypePtr FromShaderType(const UE::Shader::FType& InShaderType);
 	
+	// Returns the type matching specified EMaterialValueType.
+	static FTypePtr FromMaterialValueType(EMaterialValueType Type);
+
 	// Returns the `void` type.
 	static FTypePtr GetVoid();
 
 	// Returns whether this type is a `bool` scalar.
-	bool IsBool1() const;
+	bool IsBoolScalar() const;
 
 	// Returns this type upcast this type to ArithmeticType if it is one. Otherwise it returns nullptr.
-	FArithmeticTypePtr AsArithmetic() const;
+	FPrimitiveTypePtr AsPrimitive() const;
 
 	// Returns this type upcast this type to ArithmeticType if it's a scalar. Otherwise it returns nullptr.
-	FArithmeticTypePtr AsScalar() const;
+	FPrimitiveTypePtr AsScalar() const;
 
 	// Returns this type upcast this type to ArithmeticType if it's a vector. Otherwise it returns nullptr.
-	FArithmeticTypePtr AsVector() const;
+	FPrimitiveTypePtr AsVector() const;
 
 	// Returns this type upcast this type to ArithmeticType if it's a matrix. Otherwise it returns nullptr.
-	FArithmeticTypePtr AsMatrix() const;
+	FPrimitiveTypePtr AsMatrix() const;
 
 	// Returns the this type name spelling (e.g. float4x4).
 	FStringView GetSpelling() const;
@@ -51,6 +55,8 @@ struct FType
 	UE::Shader::EValueType ToValueType() const;
 };
 
+// Primitive types of a single scalar.
+// Note: These are listed in precision order. Converting one to the other is then simply performed taking the max EScalarKind.
 enum EScalarKind
 {
 	SK_Bool, SK_Int, SK_Float,
@@ -58,30 +64,30 @@ enum EScalarKind
 
 const TCHAR* ScalarKindToString(EScalarKind Kind);
 
-struct FArithmeticType : FType
+struct FPrimitiveType : FType
 {
 	FStringView Spelling;
 	EScalarKind ScalarKind;
 	int NumRows;
 	int NumColumns;
 
-	static FArithmeticTypePtr GetBool1();
-	static FArithmeticTypePtr GetInt1();
-	static FArithmeticTypePtr GetFloat1();
-	static FArithmeticTypePtr GetFloat2();
-	static FArithmeticTypePtr GetFloat3();
-	static FArithmeticTypePtr GetFloat4();
+	static FPrimitiveTypePtr GetBool1();
+	static FPrimitiveTypePtr GetInt1();
+	static FPrimitiveTypePtr GetFloat1();
+	static FPrimitiveTypePtr GetFloat2();
+	static FPrimitiveTypePtr GetFloat3();
+	static FPrimitiveTypePtr GetFloat4();
 
-	static FArithmeticTypePtr GetScalar(EScalarKind InScalarKind);
-	static FArithmeticTypePtr GetVector(EScalarKind InScalarKind, int NumRows);
-	static FArithmeticTypePtr GetMatrix(EScalarKind InScalarKind, int NumColumns, int NumRows);
-	static FArithmeticTypePtr Get(EScalarKind InScalarKind, int NumRows, int NumColumns);
+	static FPrimitiveTypePtr GetScalar(EScalarKind InScalarKind);
+	static FPrimitiveTypePtr GetVector(EScalarKind InScalarKind, int NumRows);
+	static FPrimitiveTypePtr GetMatrix(EScalarKind InScalarKind, int NumColumns, int NumRows);
+	static FPrimitiveTypePtr Get(EScalarKind InScalarKind, int NumRows, int NumColumns);
 
 	int  GetNumComponents() const { return NumRows * NumColumns; }
 	bool IsScalar() const { return GetNumComponents() == 1; }
 	bool IsVector() const { return NumRows > 1 && NumColumns == 1; }
 	bool IsMatrix() const { return NumRows > 1 && NumColumns > 1; }
-	FArithmeticTypePtr ToScalar() const;
+	FPrimitiveTypePtr ToScalar() const;
 };
 
 struct FTextureType : FType
