Added TextureStreaming/CancelledMipsInKB csv stat

#rnx

diff --git a/Engine/Source/Runtime/Engine/Private/Streaming/Texture2DStreamIn_IO.cpp b/Engine/Source/Runtime/Engine/Private/Streaming/Texture2DStreamIn_IO.cpp
--- a/Engine/Source/Runtime/Engine/Private/Streaming/Texture2DStreamIn_IO.cpp
+++ b/Engine/Source/Runtime/Engine/Private/Streaming/Texture2DStreamIn_IO.cpp
@@ -12,6 +12,9 @@ Texture2DStreamIn.cpp: Stream in helper for 2D textures using texture streaming
 #include "Rendering/Texture2DResource.h"
 #include "Streaming/Texture2DStreamIn.h"
 #include "Streaming/Texture2DUpdate.h"
+#include "ProfilingDebugging/CsvProfiler.h"
+
+CSV_DECLARE_CATEGORY_EXTERN(TextureStreaming);
 
 #if PLATFORM_ANDROID
 #include "EngineLogs.h"
@@ -48,6 +51,7 @@ static void ValidateMipBulkDataSize(const UTexture2D& Texture, int32 MipSizeX, i
 
 void FTexture2DStreamIn_IO::SetIORequests(const FContext& Context)
 {
+	int32 InFlightBulkSize = 0;
 	const int32 BatchCount = CurrentFirstLODIdx - PendingFirstLODIdx;
 	FBulkDataBatchRequest::FBatchBuilder Batch = FBulkDataBatchRequest::NewBatch(BatchCount);
 	for (int32 MipIndex = PendingFirstLODIdx; MipIndex < CurrentFirstLODIdx && !IsCancelled(); ++MipIndex)
@@ -89,6 +93,8 @@ void FTexture2DStreamIn_IO::SetIORequests(const FContext& Context)
 
 			FIoBuffer Dst(FIoBuffer::Wrap, MipData[MipIndex].Data, BulkDataSize);
 			Batch.Read(MipMap.BulkData, 0, BulkDataSize, Priority | AIOP_FLAG_DONTCACHE | AIOP_FLAG_HW_TARGET_MEMORY, Dst);
+
+			InFlightBulkSize += BulkDataSize;
 		}
 		else // Bulk data size can only be 0 when not available, in which case, we need to recache the file state.
 		{
@@ -105,7 +111,7 @@ void FTexture2DStreamIn_IO::SetIORequests(const FContext& Context)
 
 	TaskSynchronization.Increment();
 
-	Batch.Issue([this](FBulkDataBatchRequest::EStatus Status)
+	Batch.Issue([this, InFlightBulkSize](FBulkDataBatchRequest::EStatus Status)
 	{
 		TaskSynchronization.Decrement();
 
@@ -122,6 +128,10 @@ void FTexture2DStreamIn_IO::SetIORequests(const FContext& Context)
 			FPlatformProcess::Sleep(FRenderAssetStreamingSettings::ExtraIOLatency * .001f); // Slow down the streaming.
 		}
 #endif
+		if (bIsCancelled)
+		{
+			CSV_CUSTOM_STAT(TextureStreaming, CancelledMipsInKB, float(InFlightBulkSize) / 1024.0f, ECsvCustomStatOp::Accumulate);
+		}
 		// The tick here is intended to schedule the success or cancel callback.
 		// Using TT_None ensure gets which could create a dead lock.
 		Tick(FTexture2DUpdate::TT_None);
