#UE Add more info to an assert

diff --git a/Engine/Source/Runtime/Engine/Private/WorldPartition/WorldPartitionActorDescInstance.cpp b/Engine/Source/Runtime/Engine/Private/WorldPartition/WorldPartitionActorDescInstance.cpp
--- a/Engine/Source/Runtime/Engine/Private/WorldPartition/WorldPartitionActorDescInstance.cpp
+++ b/Engine/Source/Runtime/Engine/Private/WorldPartition/WorldPartitionActorDescInstance.cpp
@@ -126,7 +126,7 @@ bool FWorldPartitionActorDescInstance::StartAsyncLoad()
 
 		AsyncLoadID = LoadPackageAsync(PackagePath, PackageName, FLoadPackageAsyncDelegate::CreateLambda([this, ActorPackage](const FName& PackageName, UPackage* Package, EAsyncLoadingResult::Type Result)
 		{
-			check(AsyncLoadID != INDEX_NONE);
+			checkf(AsyncLoadID != INDEX_NONE, TEXT("FWorldPartitionActorDescInstance failed to load package. ActorPackage:%s PackageName:%s"), *ActorPackage.ToString(), *PackageName.ToString());
 			AsyncLoadID = INDEX_NONE;
 
 			if ((Result != EAsyncLoadingResult::Succeeded) || !Package)
