#UE Disable secondary assert as well for tonight

diff --git a/Engine/Source/Runtime/Engine/Private/WorldPartition/WorldPartitionActorDescInstance.cpp b/Engine/Source/Runtime/Engine/Private/WorldPartition/WorldPartitionActorDescInstance.cpp
--- a/Engine/Source/Runtime/Engine/Private/WorldPartition/WorldPartitionActorDescInstance.cpp
+++ b/Engine/Source/Runtime/Engine/Private/WorldPartition/WorldPartitionActorDescInstance.cpp
@@ -145,7 +145,7 @@ bool FWorldPartitionActorDescInstance::StartAsyncLoad()
 				return;
 			}
 
-			check(ActorPtr->GetPackage() == Package);
+			//check(ActorPtr->GetPackage() == Package);
 		})
 		, PKG_None, INDEX_NONE, 0, InstancingContext);
 	}
