Fix a validation failure when verifying uploaded instance info in GPUScene.

diff --git a/Engine/Source/Runtime/Engine/Private/InstancedStaticMesh/ISMInstanceDataManager.cpp b/Engine/Source/Runtime/Engine/Private/InstancedStaticMesh/ISMInstanceDataManager.cpp
--- a/Engine/Source/Runtime/Engine/Private/InstancedStaticMesh/ISMInstanceDataManager.cpp
+++ b/Engine/Source/Runtime/Engine/Private/InstancedStaticMesh/ISMInstanceDataManager.cpp
@@ -660,14 +660,14 @@ bool FPrimitiveInstanceDataManager::FlushChanges(FInstanceUpdateComponentDesc &&
 	{
 		ChangeDesc.bInstancesChanged =  HasAnyInstanceChanges();
 
-		ChangeDesc.bPrimitiveTransformChanged = !ComponentData.PrimitiveLocalToWorld.Equals(PrimitiveLocalToWorld);
+		ChangeDesc.bPrimitiveTransformChanged = !ComponentData.PrimitiveLocalToWorld.Equals(PrimitiveLocalToWorld, UE_SMALL_NUMBER);
 		ChangeDesc.bMaterialUsageFlagsChanged = Flags != ComponentData.Flags;
 		ChangeDesc.bMaxDisplacementChanged = AbsMaxDisplacement != NewAbsMaxDisplacement;
 		ChangeDesc.bStaticMeshBoundsChanged = !StaticMeshBounds.Equals(ComponentData.StaticMeshBounds);
 	}
 	
 	// Yet another special case to handle externally managed data from landscape grass
-	if (Mode == EMode::ExternalLegacyData && (!ComponentData.PrimitiveLocalToWorld.Equals(PrimitiveLocalToWorld)))
+	if (Mode == EMode::ExternalLegacyData && (!ComponentData.PrimitiveLocalToWorld.Equals(PrimitiveLocalToWorld, UE_SMALL_NUMBER)))
 	{
 		ChangeDesc.bPrimitiveTransformChanged = true;
 	}
