Store SegmentIndex as int32 instead of uint8 (all code setting the segment index works with int32 and there are no checks on downcasts which can cause overflow with 256+ segments)

diff --git a/Engine/Plugins/FX/Niagara/Source/Niagara/Private/NiagaraRendererMeshes.cpp b/Engine/Plugins/FX/Niagara/Source/Niagara/Private/NiagaraRendererMeshes.cpp
--- a/Engine/Plugins/FX/Niagara/Source/Niagara/Private/NiagaraRendererMeshes.cpp
+++ b/Engine/Plugins/FX/Niagara/Source/Niagara/Private/NiagaraRendererMeshes.cpp
@@ -1761,7 +1761,7 @@ void FNiagaraRendererMeshes::GetDynamicRayTracingInstances(FRayTracingMaterialGa
 				bIsInstancedStereo,
 				false // bNeedsPrevTransform
 			);
-			MeshBatch.SegmentIndex = uint8(SectionIndex);
+			MeshBatch.SegmentIndex = SectionIndex;
 			MeshBatch.LODIndex = uint8(LODModel.LODIndex);
 
 #if !(UE_BUILD_SHIPPING || UE_BUILD_TEST)
diff --git a/Engine/Plugins/Runtime/GeometryCache/Source/GeometryCache/Private/GeometryCacheSceneProxy.cpp b/Engine/Plugins/Runtime/GeometryCache/Source/GeometryCache/Private/GeometryCacheSceneProxy.cpp
--- a/Engine/Plugins/Runtime/GeometryCache/Source/GeometryCache/Private/GeometryCacheSceneProxy.cpp
+++ b/Engine/Plugins/Runtime/GeometryCache/Source/GeometryCache/Private/GeometryCacheSceneProxy.cpp
@@ -556,7 +556,7 @@ void FGeometryCacheSceneProxy::GetDynamicRayTracingInstances(FRayTracingMaterial
 			const int32 MaterialIndex = TrackProxy->Materials.IsValidIndex(BatchInfo.MaterialIndex) ? BatchInfo.MaterialIndex : SegmentIndex;
 			MeshBatch.MaterialRenderProxy = TrackProxy->Materials[MaterialIndex]->GetRenderProxy();
 			MeshBatch.CastRayTracedShadow = IsShadowCast(Context.ReferenceView);
-			MeshBatch.SegmentIndex = static_cast<uint8>(SegmentIndex);
+			MeshBatch.SegmentIndex = SegmentIndex;
 			MeshBatch.ReverseCulling = false; // RayTracing does not want the transform orientation baked in
 			RayTracingInstance.Materials.Add(MeshBatch);
 		}
diff --git a/Engine/Source/Runtime/Engine/Public/MeshBatch.h b/Engine/Source/Runtime/Engine/Public/MeshBatch.h
--- a/Engine/Source/Runtime/Engine/Public/MeshBatch.h
+++ b/Engine/Source/Runtime/Engine/Public/MeshBatch.h
@@ -349,12 +349,14 @@ struct FMeshBatch
 	/** This is the threshold that will be used to know if we should use this mesh batch or use one with no tessellation enabled */
 	float TessellationDisablingShadowMapMeshSize;
 
+	/** Segment index of the original mesh */
+	int32 SegmentIndex;
+
 	/* Mesh Id in a primitive. Used for stable sorting of draws belonging to the same primitive. **/
 	uint16 MeshIdInPrimitive;
 
 	/** LOD index of the mesh, used for fading LOD transitions. */
 	int8 LODIndex;
-	uint8 SegmentIndex;
 
 	uint32 ReverseCulling : 1;
 	uint32 bDisableBackfaceCulling : 1;
@@ -474,9 +476,9 @@ struct FMeshBatch
 	,	MaterialRenderProxy(nullptr)
 	,	LCI(nullptr)
 	,	TessellationDisablingShadowMapMeshSize(0.0f)
+	,	SegmentIndex(INDEX_NONE)
 	,	MeshIdInPrimitive(0)
 	,	LODIndex(INDEX_NONE)
-	,	SegmentIndex(0xFF)
 	,	ReverseCulling(false)
 	,	bDisableBackfaceCulling(false)
 	,	CastShadow(true)
diff --git a/Engine/Source/Runtime/Renderer/Public/MeshPassProcessor.h b/Engine/Source/Runtime/Renderer/Public/MeshPassProcessor.h
--- a/Engine/Source/Runtime/Renderer/Public/MeshPassProcessor.h
+++ b/Engine/Source/Runtime/Renderer/Public/MeshPassProcessor.h
@@ -752,7 +752,7 @@ struct FMeshDrawCommandDebugData
 	uint32 PSOCollectorIndex;
 #endif
 	int8 LODIndex;
-	uint8 SegmentIndex;
+	int32 SegmentIndex;
 #endif
 };
 
