Fix incorrect eye / hair material conversion with Substrate.

When using SubstrateMaterialAttribute or SubstrateShadingModel, with a shading model driven by expression, the resulting packing was incorrect for hair and eye material. This is because the substrate operator was converted as a BSDF operator (vs. an eye / hair operator), and later the material was flagged as 'single' instead of 'complex'. The substrate export code was thus missing the complex header packing causing eye/hair data to be incorrect.

This CL add hair & eye BSDF feature, and refactor the feature tracking through a bitfield for simplicity.

diff --git a/Engine/Shaders/Public/ShaderVersion.ush b/Engine/Shaders/Public/ShaderVersion.ush
--- a/Engine/Shaders/Public/ShaderVersion.ush
+++ b/Engine/Shaders/Public/ShaderVersion.ush
@@ -3,4 +3,4 @@
 // Each time the console command is executed it generates a new GUID. As this file is included
 // in Platform.ush (which should be included in any shader) it allows to invalidate the shader DDC.
 // 
-#pragma message("UESHADERMETADATA_VERSION 42D3EF03-C067-42EF-9081-1471131640AC")
\ No newline at end of file
+#pragma message("UESHADERMETADATA_VERSION E153BDBF-C281-40FF-BCD1-3A139BC52C71")
\ No newline at end of file
diff --git a/Engine/Source/Editor/GraphEditor/Private/MaterialNodes/SGraphSubstrateMaterial.cpp b/Engine/Source/Editor/GraphEditor/Private/MaterialNodes/SGraphSubstrateMaterial.cpp
--- a/Engine/Source/Editor/GraphEditor/Private/MaterialNodes/SGraphSubstrateMaterial.cpp
+++ b/Engine/Source/Editor/GraphEditor/Private/MaterialNodes/SGraphSubstrateMaterial.cpp
@@ -167,14 +167,16 @@ static const TSharedRef<SWidget> InternalProcessOperator(
 			FString BSDFDesc = OutputType == ESubstrateWidgetOutputType::Node ? 
 											 FString(TEXT("BSDF")) : 
 											 FString::Printf(TEXT("BSDF (%s%s%s%s%s%s%s%s)")
-											, Op.bBSDFHasEdgeColor ? TEXT("F90 ") : TEXT("")
-											, Op.bBSDFHasSSS ? TEXT("SSS ") : TEXT("")
-											, Op.bBSDFHasMFPPluggedIn ? TEXT("MFP ") : TEXT("")
-											, Op.bBSDFHasAnisotropy ? TEXT("Ani ") : TEXT("")
-											, Op.bBSDFHasSecondRoughnessOrSimpleClearCoat ? TEXT("2Ro ") : TEXT("")
-											, Op.bBSDFHasFuzz ? TEXT("Fuz ") : TEXT("")
-											, Op.bBSDFHasGlint ? TEXT("Gli ") : TEXT("")
-											, Op.bBSDFHasSpecularProfile ? TEXT("Spc ") : TEXT("")
+											, Op.Has(ESubstrateBsdfFeature_EdgeColor) ? TEXT("F90 ") : TEXT("")
+											, Op.Has(ESubstrateBsdfFeature_SSS) ? TEXT("SSS ") : TEXT("")
+											, Op.Has(ESubstrateBsdfFeature_MFPPluggedIn) ? TEXT("MFP ") : TEXT("")
+											, Op.Has(ESubstrateBsdfFeature_Anisotropy) ? TEXT("Ani ") : TEXT("")
+											, Op.Has(ESubstrateBsdfFeature_SecondRoughnessOrSimpleClearCoat) ? TEXT("2Ro ") : TEXT("")
+											, Op.Has(ESubstrateBsdfFeature_Fuzz) ? TEXT("Fuz ") : TEXT("")
+											, Op.Has(ESubstrateBsdfFeature_Glint) ? TEXT("Gli ") : TEXT("")
+											, Op.Has(ESubstrateBsdfFeature_SpecularProfile) ? TEXT("Spc ") : TEXT("")
+											, Op.Has(ESubstrateBsdfFeature_Eye) ? TEXT("Eye ") : TEXT("")
+											, Op.Has(ESubstrateBsdfFeature_Hair) ? TEXT("Hai ") : TEXT("")
 											);
 
 			static FString ToolTip;
@@ -188,6 +190,8 @@ static const TSharedRef<SWidget> InternalProcessOperator(
 				ToolTip += TEXT("Ani means the BSDF anisotropic specular lighting is used.\n");
 				ToolTip += TEXT("Gli means the BSDF features glints.\n");
 				ToolTip += TEXT("Spc means the BSDF features specular profile.\n");
+				ToolTip += TEXT("Eye means the BSDF features eye is used.\n");
+				ToolTip += TEXT("Hai means the BSDF features hair is used.\n");
 			}
 
 			const EStyleColor Color = OverrideColor != EStyleColor::MAX ? OverrideColor : (bIsCurrent ? EStyleColor::AccentGreen : EStyleColor::AccentGray);
diff --git a/Engine/Source/Runtime/Core/Private/UObject/FortniteMainBranchObjectVersion.cpp b/Engine/Source/Runtime/Core/Private/UObject/FortniteMainBranchObjectVersion.cpp
--- a/Engine/Source/Runtime/Core/Private/UObject/FortniteMainBranchObjectVersion.cpp
+++ b/Engine/Source/Runtime/Core/Private/UObject/FortniteMainBranchObjectVersion.cpp
@@ -14,6 +14,6 @@ TMap<FGuid, FGuid> FFortniteMainBranchObjectVersion::GetSystemGuids()
 	SystemGuids.Add(DevGuids.Niagara_LatestScriptCompileVersion, FGuid("E486231DA967C64B95D4E73131056FA0"));
 	SystemGuids.Add(DevGuids.SkeletalMeshDerivedDataVersion, FGuid("B9AEAA3EE7AC4FCFB3860E8F3E06DD75"));
 	SystemGuids.Add(DevGuids.STATICMESH_DERIVEDDATA_VER, FGuid("1C2EBE3310464FD7995ADE1D6AB92BDC"));
-	SystemGuids.Add(DevGuids.MaterialTranslationDDCVersion, FGuid("C42A7403BBDD420EBA946325D6F79B1B"));
+	SystemGuids.Add(DevGuids.MaterialTranslationDDCVersion, FGuid("3A4DC1DF808D4096BEEA6BA2E3ACC38D"));
 	return SystemGuids;
 }
diff --git a/Engine/Source/Runtime/Engine/Private/Materials/HLSLMaterialTranslator.cpp b/Engine/Source/Runtime/Engine/Private/Materials/HLSLMaterialTranslator.cpp
--- a/Engine/Source/Runtime/Engine/Private/Materials/HLSLMaterialTranslator.cpp
+++ b/Engine/Source/Runtime/Engine/Private/Materials/HLSLMaterialTranslator.cpp
@@ -12675,14 +12675,7 @@ bool FHLSLMaterialTranslator::FSubstrateCompilationContext::SubstrateGenerateDer
 			{
 				// Disable all optional features for now to fit.
 				// SUBSTRATE_TODO we will need to refine that to account for platforms supporting SSS for instance.
-				It.bBSDFHasSSS = false;
-				It.bBSDFHasMFPPluggedIn = false;
-				It.bBSDFHasEdgeColor = false;
-				It.bBSDFHasFuzz = false;
-				It.bBSDFHasSecondRoughnessOrSimpleClearCoat = false;
-				It.bBSDFHasAnisotropy = false;
-				It.bBSDFHasGlint = false;
-				It.bBSDFHasSpecularProfile = false;
+				It.BSDFFeatures = 0;
 			}
 			SubstrateSimplificationStatus.bSlabSimplificationStepHasBeenRun = true;
 		}
@@ -13018,9 +13011,23 @@ bool FHLSLMaterialTranslator::FSubstrateCompilationContext::SubstrateGenerateDer
 					{
 					case SUBSTRATE_BSDF_TYPE_SLAB:
 					{
-						SubstrateMaterialComplexity.bIsSimple = SubstrateMaterialComplexity.bIsSimple && !bMayHaveColoredWeight && !It.bBSDFHasAnisotropy && !It.bBSDFHasEdgeColor && !It.bBSDFHasFuzz && !It.bBSDFHasSecondRoughnessOrSimpleClearCoat && !It.bBSDFHasMFPPluggedIn && !It.bBSDFHasSSS && !It.bBSDFHasGlint && !It.bBSDFHasSpecularProfile;
-						SubstrateMaterialComplexity.bIsSingle = SubstrateMaterialComplexity.bIsSingle && !bMayHaveColoredWeight && !It.bBSDFHasAnisotropy && !It.bBSDFHasGlint && !It.bBSDFHasSpecularProfile;
-						SubstrateMaterialComplexity.bIsComplexSpecial |= It.bBSDFHasGlint || It.bBSDFHasSpecularProfile;
+						const uint32 SingleMask = 
+							  ESubstrateBsdfFeature_EdgeColor
+							| ESubstrateBsdfFeature_Fuzz
+							| ESubstrateBsdfFeature_SSS
+							| ESubstrateBsdfFeature_SecondRoughnessOrSimpleClearCoat
+							| ESubstrateBsdfFeature_MFPPluggedIn;
+						const uint32 CompexMask = 
+							  ESubstrateBsdfFeature_Anisotropy
+							| ESubstrateBsdfFeature_Eye
+							| ESubstrateBsdfFeature_Hair;
+						const uint32 CompexSpecialMask = 
+							  ESubstrateBsdfFeature_Glint
+							| ESubstrateBsdfFeature_SpecularProfile;
+
+						SubstrateMaterialComplexity.bIsSimple = SubstrateMaterialComplexity.bIsSimple && !bMayHaveColoredWeight && !It.Has(CompexMask) && !It.Has(CompexSpecialMask) && !It.Has(SingleMask);
+						SubstrateMaterialComplexity.bIsSingle = SubstrateMaterialComplexity.bIsSingle && !bMayHaveColoredWeight && !It.Has(CompexMask) && !It.Has(CompexSpecialMask);
+						SubstrateMaterialComplexity.bIsComplexSpecial |= It.Has(CompexSpecialMask);
 						break;
 					}
 					case SUBSTRATE_BSDF_TYPE_HAIR:
@@ -13146,13 +13153,13 @@ bool FHLSLMaterialTranslator::FSubstrateCompilationContext::SubstrateGenerateDer
 					case SUBSTRATE_BSDF_TYPE_SLAB:
 					{
 						// Compute values closer to the reality for HasSSS and IsSimpleVolume, now that we know that we know the topology of the material.
-						const bool bIsSimpleVolume = !It.bIsBottom && It.bBSDFHasMFPPluggedIn;
-						const bool bHasSSS = It.bIsBottom && It.bBSDFHasSSS;
+						const bool bIsSimpleVolume = !It.bIsBottom && It.Has(ESubstrateBsdfFeature_MFPPluggedIn);
+						const bool bHasSSS = It.bIsBottom && It.Has(ESubstrateBsdfFeature_SSS);
 
 						SubstrateMaterialRequestedSizeByte += UintByteSize;
 						SubstrateMaterialRequestedSizeByte += UintByteSize;
 
-						if (It.bBSDFHasEdgeColor || It.bBSDFHasSecondRoughnessOrSimpleClearCoat)
+						if (It.Has(ESubstrateBsdfFeature_EdgeColor | ESubstrateBsdfFeature_SecondRoughnessOrSimpleClearCoat))
 						{
 							SubstrateMaterialRequestedSizeByte += UintByteSize;
 						}
@@ -13160,21 +13167,31 @@ bool FHLSLMaterialTranslator::FSubstrateCompilationContext::SubstrateGenerateDer
 						{
 							SubstrateMaterialRequestedSizeByte += UintByteSize;
 						}
-						if (It.bBSDFHasFuzz)
+						if (It.Has(ESubstrateBsdfFeature_Fuzz))
 						{
 							SubstrateMaterialRequestedSizeByte += UintByteSize;
 						}
-						if (It.bBSDFHasGlint)
+						if (It.Has(ESubstrateBsdfFeature_Glint))
 						{
 							SubstrateMaterialRequestedSizeByte += UintByteSize;
 							SubstrateMaterialRequestedSizeByte += UintByteSize;
 							SubstrateMaterialRequestedSizeByte += UintByteSize;
 							SubstrateMaterialRequestedSizeByte += UintByteSize;
 						}
-						if (It.bBSDFHasSpecularProfile)
+						if (It.Has(ESubstrateBsdfFeature_SpecularProfile))
 						{
 							SubstrateMaterialRequestedSizeByte += UintByteSize;
 						}
+						if (It.Has(ESubstrateBsdfFeature_Eye))
+						{
+							SubstrateMaterialRequestedSizeByte += UintByteSize;
+							SubstrateMaterialRequestedSizeByte += UintByteSize;
+						}
+						if (It.Has(ESubstrateBsdfFeature_Hair))
+						{
+							SubstrateMaterialRequestedSizeByte += UintByteSize;
+							SubstrateMaterialRequestedSizeByte += UintByteSize;
+						}
 						break;
 					}
 					case SUBSTRATE_BSDF_TYPE_HAIR:
@@ -15663,9 +15680,19 @@ void FHLSLMaterialTranslator::PrepareEnvironmentDefines()
 						{
 							if (!It.IsDiscarded() && It.MaxDistanceFromLeaves == DistanceToLeaves)
 							{
-								SubstrateMaterialContextDescription += FString::Printf(TEXT("\tIdx=%d Op=%s ParentIdx=%d LeftIndex=%d RightIndex=%d BSDFIdx=%d LayerDepth=%d IsTop=%d IsBot=%d BSDFType=%s SSS=%d MFP=%d F90=%d Rough2=%d Fuzz=%d Aniso=%d Glint=%d SpecularProfile=%d\n"),
+								SubstrateMaterialContextDescription += FString::Printf(TEXT("\tIdx=%d Op=%s ParentIdx=%d LeftIndex=%d RightIndex=%d BSDFIdx=%d LayerDepth=%d IsTop=%d IsBot=%d BSDFType=%s SSS=%d MFP=%d F90=%d Rough2=%d Fuzz=%d Aniso=%d Glint=%d SpecularProfile=%d Eye=%d Hair=%d\n"),
 									It.Index, GetSubstrateOperatorStr(It.OperatorType), It.ParentIndex, It.LeftIndex, It.RightIndex, It.BSDFIndex, It.LayerDepth, It.bIsTop, It.bIsBottom,
-									*GetSubstrateBSDFName(It.BSDFType), It.bBSDFHasSSS, It.bBSDFHasMFPPluggedIn, It.bBSDFHasEdgeColor, It.bBSDFHasSecondRoughnessOrSimpleClearCoat, It.bBSDFHasFuzz, It.bBSDFHasAnisotropy, It.bBSDFHasGlint, It.bBSDFHasSpecularProfile);
+									*GetSubstrateBSDFName(It.BSDFType), 
+									It.Has(ESubstrateBsdfFeature_SSS), 
+									It.Has(ESubstrateBsdfFeature_MFPPluggedIn), 
+									It.Has(ESubstrateBsdfFeature_EdgeColor), 
+									It.Has(ESubstrateBsdfFeature_SecondRoughnessOrSimpleClearCoat), 
+									It.Has(ESubstrateBsdfFeature_Fuzz), 
+									It.Has(ESubstrateBsdfFeature_Anisotropy), 
+									It.Has(ESubstrateBsdfFeature_Glint), 
+									It.Has(ESubstrateBsdfFeature_SpecularProfile), 
+									It.Has(ESubstrateBsdfFeature_Eye), 
+									It.Has(ESubstrateBsdfFeature_Hair));
 							}
 						}
 					}
diff --git a/Engine/Source/Runtime/Engine/Private/Materials/MaterialExpressionSubstrate.cpp b/Engine/Source/Runtime/Engine/Private/Materials/MaterialExpressionSubstrate.cpp
--- a/Engine/Source/Runtime/Engine/Private/Materials/MaterialExpressionSubstrate.cpp
+++ b/Engine/Source/Runtime/Engine/Private/Materials/MaterialExpressionSubstrate.cpp
@@ -551,14 +551,11 @@ FSubstrateOperator* UMaterialExpressionSubstrateShadingModels::SubstrateGenerate
 	// Note Thickness has no meaning/usage in the context of SubstrateLegacyConversionNode
 	int32 ThicknessIndex = Compiler->SubstrateThicknessStackGetThicknessIndex();
 
-	auto AddDefaultWorstCase = [&](bool bSSS, bool bFuzz)
+	auto AddDefaultWorstCase = [&](uint32 In)
 		{
 			FSubstrateOperator& SlabOperator = Compiler->SubstrateCompilationRegisterOperator(SUBSTRATE_OPERATOR_BSDF_LEGACY, Compiler->SubstrateTreeStackGetPathUniqueId(), this, Parent, Compiler->SubstrateTreeStackGetParentPathUniqueId());
 			SlabOperator.BSDFType = SUBSTRATE_BSDF_TYPE_SLAB;
-			SlabOperator.bBSDFHasSSS = bSSS;
-			SlabOperator.bBSDFHasMFPPluggedIn = bSSS;
-			SlabOperator.bBSDFHasFuzz = bFuzz;
-			SlabOperator.bBSDFHasAnisotropy = Anisotropy.IsConnected();
+			SlabOperator.BSDFFeatures = In | (Anisotropy.IsConnected() ? ESubstrateBsdfFeature_Anisotropy : ESubstrateBsdfFeature_None);
 			SlabOperator.bBSDFWritesEmissive = EmissiveColor.IsConnected();
 			SlabOperator.ThicknessIndex = ThicknessIndex;
 
@@ -580,7 +577,13 @@ FSubstrateOperator* UMaterialExpressionSubstrateShadingModels::SubstrateGenerate
 			Operator.ThicknessIndex = ThicknessIndex;
 			return &Operator;
 		}
-		return AddDefaultWorstCase(true, true);
+
+		// Be sure to track eye/hair feature, even though they are not part of Slab BSDF. This is important later for issuing 
+		// the correct material complexity, as hair/eye requires 'complex' complexity (not 'simple') for correct packing
+		uint32 BsdfFeatures = ESubstrateBsdfFeature_SSS | ESubstrateBsdfFeature_Fuzz;
+		if (ShadingModels.HasShadingModel(MSM_Eye))  { BsdfFeatures |= ESubstrateBsdfFeature_Eye; }
+		if (ShadingModels.HasShadingModel(MSM_Hair)) { BsdfFeatures |= ESubstrateBsdfFeature_Hair; }
+		return AddDefaultWorstCase(BsdfFeatures);
 	}
 	// else
 	{
@@ -596,39 +599,38 @@ FSubstrateOperator* UMaterialExpressionSubstrateShadingModels::SubstrateGenerate
 		}
 		else if (ShadingModels.HasShadingModel(MSM_DefaultLit))
 		{
-			return AddDefaultWorstCase(false, false);
+			return AddDefaultWorstCase(ESubstrateBsdfFeature_None);
 		}
 		else if (ShadingModels.HasShadingModel(MSM_ThinTranslucent))
 		{
-			return AddDefaultWorstCase(false, false);
+			return AddDefaultWorstCase(ESubstrateBsdfFeature_None);
 		}
 		else if (ShadingModels.HasShadingModel(MSM_SubsurfaceProfile))
 		{
-			return AddDefaultWorstCase(true, false);
+			return AddDefaultWorstCase(ESubstrateBsdfFeature_SSS);
 		}
 		else if (ShadingModels.HasShadingModel(MSM_Subsurface))
 		{
-			return AddDefaultWorstCase(true, false);
+			return AddDefaultWorstCase(ESubstrateBsdfFeature_SSS);
 		}
 		else if (ShadingModels.HasShadingModel(MSM_TwoSidedFoliage))
 		{
-			return AddDefaultWorstCase(true, false);
+			return AddDefaultWorstCase(ESubstrateBsdfFeature_SSS);
 		}
 		else if (ShadingModels.HasShadingModel(MSM_PreintegratedSkin))
 		{
-			return AddDefaultWorstCase(true, false);
+			return AddDefaultWorstCase(ESubstrateBsdfFeature_SSS);
 		}
 		else if (ShadingModels.HasShadingModel(MSM_Cloth))
 		{
-			return AddDefaultWorstCase(false, true);
+			return AddDefaultWorstCase(ESubstrateBsdfFeature_Fuzz);
 		}
 		else if (ShadingModels.HasShadingModel(MSM_ClearCoat))
 		{
 			FSubstrateOperator& Operator = Compiler->SubstrateCompilationRegisterOperator(SUBSTRATE_OPERATOR_BSDF_LEGACY, Compiler->SubstrateTreeStackGetPathUniqueId(), this, Parent, Compiler->SubstrateTreeStackGetParentPathUniqueId());
 			Operator.BSDFType = SUBSTRATE_BSDF_TYPE_SLAB;
 			Operator.ThicknessIndex = ThicknessIndex;
-			Operator.bBSDFHasSecondRoughnessOrSimpleClearCoat = true;
-			Operator.bBSDFHasAnisotropy = Anisotropy.IsConnected();
+			Operator.BSDFFeatures = ESubstrateBsdfFeature_SecondRoughnessOrSimpleClearCoat | (Anisotropy.IsConnected() ? ESubstrateBsdfFeature_Anisotropy : ESubstrateBsdfFeature_None);
 			Operator.bBSDFWritesEmissive = EmissiveColor.IsConnected();
 			return &Operator;
 		}
@@ -636,6 +638,7 @@ FSubstrateOperator* UMaterialExpressionSubstrateShadingModels::SubstrateGenerate
 		{
 			FSubstrateOperator& Operator = Compiler->SubstrateCompilationRegisterOperator(SUBSTRATE_OPERATOR_BSDF_LEGACY, Compiler->SubstrateTreeStackGetPathUniqueId(), this, Parent, Compiler->SubstrateTreeStackGetParentPathUniqueId());
 			Operator.BSDFType = SUBSTRATE_BSDF_TYPE_HAIR;
+			Operator.BSDFFeatures = ESubstrateBsdfFeature_Hair;
 			Operator.ThicknessIndex = ThicknessIndex;
 			Operator.bBSDFWritesEmissive = EmissiveColor.IsConnected();
 			return &Operator;
@@ -644,6 +647,7 @@ FSubstrateOperator* UMaterialExpressionSubstrateShadingModels::SubstrateGenerate
 		{
 			FSubstrateOperator& Operator = Compiler->SubstrateCompilationRegisterOperator(SUBSTRATE_OPERATOR_BSDF_LEGACY, Compiler->SubstrateTreeStackGetPathUniqueId(), this, Parent, Compiler->SubstrateTreeStackGetParentPathUniqueId());
 			Operator.BSDFType = SUBSTRATE_BSDF_TYPE_EYE;
+			Operator.BSDFFeatures = ESubstrateBsdfFeature_Eye;
 			Operator.ThicknessIndex = ThicknessIndex;
 			Operator.bBSDFWritesEmissive = EmissiveColor.IsConnected();
 			return &Operator;
@@ -716,16 +720,16 @@ int32 UMaterialExpressionSubstrateSlabBSDF::Compile(class FMaterialCompiler* Com
 
 	// We also cannot ignore the tangent when using the default Tangent because GetTangentBasis
 	// used in SubstrateGetBSDFSharedBasis cannot be relied on for smooth tangent used for lighting on any mesh.
-	const bool bHasAnisotropy = SubstrateOperator.bBSDFHasAnisotropy > 0;
+	const bool bHasAnisotropy = SubstrateOperator.Has(ESubstrateBsdfFeature_Anisotropy);
 
 	int32 SSSProfileCodeChunk = INDEX_NONE;
-	if (SubstrateOperator.bBSDFHasSSS > 0)
+	if (SubstrateOperator.Has(ESubstrateBsdfFeature_SSS))
 	{
 		SSSProfileCodeChunk = Compiler->ForceCast(Compiler->ScalarParameter(GetSubsurfaceProfileParameterName(), 1.0f), MCT_Float1);
 	}
 
 	int32 SpecularProfileCodeChunk = INDEX_NONE;
-	if (SubstrateOperator.bBSDFHasSpecularProfile > 0)
+	if (SubstrateOperator.Has(ESubstrateBsdfFeature_SpecularProfile))
 	{
 		const FName SpecularProfileParameterName = CreateSpecularProfileParameterName(SpecularProfile);
 		SpecularProfileCodeChunk = Compiler->ForceCast(Compiler->ScalarParameter(SpecularProfileParameterName, 1.0f), MCT_Float1);
@@ -769,31 +773,31 @@ int32 UMaterialExpressionSubstrateSlabBSDF::Compile(class FMaterialCompiler* Com
 	int32 GlintUVCodeChunk = CompileWithDefaultFloat2(Compiler, GlintUV, 0.0f, 0.0f);
 
 	// Disable some features if requested by the simplification process
-	if (SubstrateOperator.bBSDFHasMFPPluggedIn == 0)
+	if (!SubstrateOperator.Has(ESubstrateBsdfFeature_MFPPluggedIn))
 	{
 		SSSMFPCodeChunk = Compiler->Constant3(0.0f, 0.0f, 0.0f);
 	}
-	if (SubstrateOperator.bBSDFHasEdgeColor == 0)
+	if (!SubstrateOperator.Has(ESubstrateBsdfFeature_EdgeColor))
 	{
 		F90CodeChunk = Compiler->Constant3(1.0f, 1.0f, 1.0f);
 	}
-	if (SubstrateOperator.bBSDFHasFuzz == 0)
+	if (!SubstrateOperator.Has(ESubstrateBsdfFeature_Fuzz))
 	{
 		FuzzAmountCodeChunk = Compiler->Constant(0.0f);
 	}
-	if (SubstrateOperator.bBSDFHasSecondRoughnessOrSimpleClearCoat == 0)
+	if (!SubstrateOperator.Has(ESubstrateBsdfFeature_SecondRoughnessOrSimpleClearCoat))
 	{
 		SecondRoughnessWeightCodeChunk = Compiler->Constant(0.0f);
 	}
-	if (SubstrateOperator.bBSDFHasAnisotropy == 0)
+	if (!SubstrateOperator.Has(ESubstrateBsdfFeature_Anisotropy))
 	{
 		AnisotropyCodeChunk = Compiler->Constant(0.0f);
 	}
-	if (SubstrateOperator.bBSDFHasGlint == 0)
+	if (!SubstrateOperator.Has(ESubstrateBsdfFeature_Glint))
 	{
 		GlintValueCodeChunk = Compiler->Constant(0.0f);
 	}
-	if (SubstrateOperator.bBSDFHasSpecularProfile == 0)
+	if (!SubstrateOperator.Has(ESubstrateBsdfFeature_SpecularProfile))
 	{
 		SpecularProfileCodeChunk = INDEX_NONE;
 	}
@@ -1150,14 +1154,16 @@ FSubstrateOperator* UMaterialExpressionSubstrateSlabBSDF::SubstrateGenerateMater
 {
 	FSubstrateOperator& SubstrateOperator = Compiler->SubstrateCompilationRegisterOperator(SUBSTRATE_OPERATOR_BSDF, Compiler->SubstrateTreeStackGetPathUniqueId(), this, Parent, Compiler->SubstrateTreeStackGetParentPathUniqueId());
 	SubstrateOperator.BSDFType = SUBSTRATE_BSDF_TYPE_SLAB;
-	SubstrateOperator.bBSDFHasEdgeColor = HasEdgeColor();
-	SubstrateOperator.bBSDFHasFuzz = HasFuzz();
-	SubstrateOperator.bBSDFHasSecondRoughnessOrSimpleClearCoat = HasSecondRoughness();
-	SubstrateOperator.bBSDFHasSSS = HasSSS();
-	SubstrateOperator.bBSDFHasMFPPluggedIn = HasMFPPluggedIn();
-	SubstrateOperator.bBSDFHasAnisotropy = HasAnisotropy();
-	SubstrateOperator.bBSDFHasGlint = HasGlint();
-	SubstrateOperator.bBSDFHasSpecularProfile = HasSpecularProfile();
+
+	if (HasEdgeColor())			{ SubstrateOperator.BSDFFeatures |= ESubstrateBsdfFeature_EdgeColor; }
+	if (HasFuzz())				{ SubstrateOperator.BSDFFeatures |= ESubstrateBsdfFeature_Fuzz; }
+	if (HasSecondRoughness())	{ SubstrateOperator.BSDFFeatures |= ESubstrateBsdfFeature_SecondRoughnessOrSimpleClearCoat; }
+	if (HasSSS())				{ SubstrateOperator.BSDFFeatures |= ESubstrateBsdfFeature_SSS; }
+	if (HasMFPPluggedIn())		{ SubstrateOperator.BSDFFeatures |= ESubstrateBsdfFeature_MFPPluggedIn; }
+	if (HasAnisotropy())		{ SubstrateOperator.BSDFFeatures |= ESubstrateBsdfFeature_Anisotropy; }
+	if (HasGlint())				{ SubstrateOperator.BSDFFeatures |= ESubstrateBsdfFeature_Glint; }
+	if (HasSpecularProfile())	{ SubstrateOperator.BSDFFeatures |= ESubstrateBsdfFeature_SpecularProfile; }
+
 	SubstrateOperator.ThicknessIndex = Compiler->SubstrateThicknessStackGetThicknessIndex();
 	SubstrateOperator.bBSDFWritesEmissive = EmissiveColor.IsConnected();
 	return &SubstrateOperator;
@@ -1384,14 +1390,7 @@ FSubstrateOperator* UMaterialExpressionSubstrateSimpleClearCoatBSDF::SubstrateGe
 {
 	FSubstrateOperator& SubstrateOperator = Compiler->SubstrateCompilationRegisterOperator(SUBSTRATE_OPERATOR_BSDF, Compiler->SubstrateTreeStackGetPathUniqueId(), this, Parent, Compiler->SubstrateTreeStackGetParentPathUniqueId());
 	SubstrateOperator.BSDFType = SUBSTRATE_BSDF_TYPE_SLAB;
-	SubstrateOperator.bBSDFHasEdgeColor = false;
-	SubstrateOperator.bBSDFHasFuzz = false;
-	SubstrateOperator.bBSDFHasSecondRoughnessOrSimpleClearCoat = true;	// This node explicitly requires simple clear coat
-	SubstrateOperator.bBSDFHasSSS = false;
-	SubstrateOperator.bBSDFHasMFPPluggedIn = false;
-	SubstrateOperator.bBSDFHasAnisotropy = false;
-	SubstrateOperator.bBSDFHasGlint = false;
-	SubstrateOperator.bBSDFHasSpecularProfile = false;
+	SubstrateOperator.BSDFFeatures = ESubstrateBsdfFeature_SecondRoughnessOrSimpleClearCoat; // This node explicitly requires simple clear coat
 	SubstrateOperator.ThicknessIndex = Compiler->SubstrateThicknessStackGetThicknessIndex();
 	SubstrateOperator.bBSDFWritesEmissive = EmissiveColor.IsConnected();
 	return &SubstrateOperator;
@@ -2029,6 +2028,7 @@ FSubstrateOperator* UMaterialExpressionSubstrateHairBSDF::SubstrateGenerateMater
 {
 	FSubstrateOperator& SubstrateOperator = Compiler->SubstrateCompilationRegisterOperator(SUBSTRATE_OPERATOR_BSDF, Compiler->SubstrateTreeStackGetPathUniqueId(), this, Parent, Compiler->SubstrateTreeStackGetParentPathUniqueId());
 	SubstrateOperator.BSDFType = SUBSTRATE_BSDF_TYPE_HAIR;
+	SubstrateOperator.BSDFFeatures = ESubstrateBsdfFeature_Hair;
 	SubstrateOperator.ThicknessIndex = Compiler->SubstrateThicknessStackGetThicknessIndex();
 	SubstrateOperator.bBSDFWritesEmissive = EmissiveColor.IsConnected();
 	return &SubstrateOperator;
@@ -2145,6 +2145,7 @@ FSubstrateOperator* UMaterialExpressionSubstrateEyeBSDF::SubstrateGenerateMateri
 {
 	FSubstrateOperator& SubstrateOperator = Compiler->SubstrateCompilationRegisterOperator(SUBSTRATE_OPERATOR_BSDF, Compiler->SubstrateTreeStackGetPathUniqueId(), this, Parent, Compiler->SubstrateTreeStackGetParentPathUniqueId());
 	SubstrateOperator.BSDFType = SUBSTRATE_BSDF_TYPE_EYE;
+	SubstrateOperator.BSDFFeatures = ESubstrateBsdfFeature_Eye;
 	SubstrateOperator.ThicknessIndex = Compiler->SubstrateThicknessStackGetThicknessIndex();
 	SubstrateOperator.bBSDFWritesEmissive = EmissiveColor.IsConnected();
 	return &SubstrateOperator;
@@ -3794,15 +3795,11 @@ FSubstrateOperator* UMaterialExpressionSubstrateConvertMaterialAttributes::Subst
 
 	const bool bHasAnisotropy = IsMaterialAttributeInputConnected(Cached, MP_Anisotropy);
 
-	auto AddDefaultWorstCase = [&](bool bSSS, bool bFuzz)
+	auto AddDefaultWorstCase = [&](uint32 In)
 		{
-
 			FSubstrateOperator& SlabOperator = Compiler->SubstrateCompilationRegisterOperator(SUBSTRATE_OPERATOR_BSDF_LEGACY, Compiler->SubstrateTreeStackGetPathUniqueId(), this, Parent, Compiler->SubstrateTreeStackGetParentPathUniqueId());
 			SlabOperator.BSDFType = SUBSTRATE_BSDF_TYPE_SLAB;
-			SlabOperator.bBSDFHasSSS = bSSS;
-			SlabOperator.bBSDFHasMFPPluggedIn = bSSS;
-			SlabOperator.bBSDFHasFuzz = bFuzz;
-			SlabOperator.bBSDFHasAnisotropy = bHasAnisotropy;
+			SlabOperator.BSDFFeatures  = In | (bHasAnisotropy ? ESubstrateBsdfFeature_Anisotropy : ESubstrateBsdfFeature_None);
 			SlabOperator.ThicknessIndex = ThicknessIndex;
 
 			return &SlabOperator;
@@ -3815,7 +3812,10 @@ FSubstrateOperator* UMaterialExpressionSubstrateConvertMaterialAttributes::Subst
 	const bool bHasShadingModelFromExpression = IsMaterialAttributeInputConnected(Cached, MP_ShadingModel) || ShadingModelOverride == MSM_FromMaterialExpression; // We keep HasShadingModelFromExpression in case all shading models cannot be safely recovered from material functions.
 	if ((ShadingModels.CountShadingModels() > 1) || bHasShadingModelFromExpression)
 	{
-		return AddDefaultWorstCase(true, true);
+		uint32 Features = ESubstrateBsdfFeature_SSS | ESubstrateBsdfFeature_Fuzz;
+		if (ShadingModels.HasShadingModel(MSM_Eye))  { Features |= ESubstrateBsdfFeature_Eye; }
+		if (ShadingModels.HasShadingModel(MSM_Hair)) { Features |= ESubstrateBsdfFeature_Hair; }
+		return AddDefaultWorstCase(Features);
 	}
 	else
 	{
@@ -3830,45 +3830,45 @@ FSubstrateOperator* UMaterialExpressionSubstrateConvertMaterialAttributes::Subst
 		}
 		else if (ShadingModels.HasShadingModel(MSM_DefaultLit))
 		{
-			return AddDefaultWorstCase(false, false);
+			return AddDefaultWorstCase(ESubstrateBsdfFeature_None);
 		}
 		else if (ShadingModels.HasShadingModel(MSM_ThinTranslucent))
 		{
-			return AddDefaultWorstCase(false, false);
+			return AddDefaultWorstCase(ESubstrateBsdfFeature_None);
 		}
 		else if (ShadingModels.HasShadingModel(MSM_SubsurfaceProfile))
 		{
-			return AddDefaultWorstCase(true, false);
+			return AddDefaultWorstCase(ESubstrateBsdfFeature_SSS);
 		}
 		else if (ShadingModels.HasShadingModel(MSM_Subsurface))
 		{
-			return AddDefaultWorstCase(true, false);
+			return AddDefaultWorstCase(ESubstrateBsdfFeature_SSS);
 		}
 		else if (ShadingModels.HasShadingModel(MSM_TwoSidedFoliage))
 		{
-			return AddDefaultWorstCase(true, false);
+			return AddDefaultWorstCase(ESubstrateBsdfFeature_SSS);
 		}
 		else if (ShadingModels.HasShadingModel(MSM_PreintegratedSkin))
 		{
-			return AddDefaultWorstCase(true, false);
+			return AddDefaultWorstCase(ESubstrateBsdfFeature_SSS);
 		}
 		else if (ShadingModels.HasShadingModel(MSM_Cloth))
 		{
-			return AddDefaultWorstCase(false, true);
+			return AddDefaultWorstCase(ESubstrateBsdfFeature_Fuzz);
 		}
 		else if (ShadingModels.HasShadingModel(MSM_ClearCoat))
 		{
 			FSubstrateOperator& Operator = Compiler->SubstrateCompilationRegisterOperator(SUBSTRATE_OPERATOR_BSDF_LEGACY, Compiler->SubstrateTreeStackGetPathUniqueId(), this, Parent, Compiler->SubstrateTreeStackGetParentPathUniqueId());
 			Operator.BSDFType = SUBSTRATE_BSDF_TYPE_SLAB;
 			Operator.ThicknessIndex = ThicknessIndex;
-			Operator.bBSDFHasSecondRoughnessOrSimpleClearCoat = true;
-			Operator.bBSDFHasAnisotropy = bHasAnisotropy;
+			Operator.BSDFFeatures = ESubstrateBsdfFeature_SecondRoughnessOrSimpleClearCoat | (bHasAnisotropy ? ESubstrateBsdfFeature_Anisotropy : ESubstrateBsdfFeature_None);
 			return &Operator;
 		}
 		else if (ShadingModels.HasShadingModel(MSM_Hair))
 		{
 			FSubstrateOperator& Operator = Compiler->SubstrateCompilationRegisterOperator(SUBSTRATE_OPERATOR_BSDF_LEGACY, Compiler->SubstrateTreeStackGetPathUniqueId(), this, Parent, Compiler->SubstrateTreeStackGetParentPathUniqueId());
 			Operator.BSDFType = SUBSTRATE_BSDF_TYPE_HAIR;
+			Operator.BSDFFeatures = ESubstrateBsdfFeature_Hair;
 			Operator.ThicknessIndex = ThicknessIndex;
 			return &Operator;
 		}
@@ -3876,6 +3876,7 @@ FSubstrateOperator* UMaterialExpressionSubstrateConvertMaterialAttributes::Subst
 		{
 			FSubstrateOperator& Operator = Compiler->SubstrateCompilationRegisterOperator(SUBSTRATE_OPERATOR_BSDF_LEGACY, Compiler->SubstrateTreeStackGetPathUniqueId(), this, Parent, Compiler->SubstrateTreeStackGetParentPathUniqueId());
 			Operator.BSDFType = SUBSTRATE_BSDF_TYPE_EYE;
+			Operator.BSDFFeatures = ESubstrateBsdfFeature_Eye;
 			Operator.ThicknessIndex = ThicknessIndex;
 			return &Operator;
 		}
diff --git a/Engine/Source/Runtime/Engine/Public/Rendering/SubstrateMaterialShared.cpp b/Engine/Source/Runtime/Engine/Public/Rendering/SubstrateMaterialShared.cpp
--- a/Engine/Source/Runtime/Engine/Public/Rendering/SubstrateMaterialShared.cpp
+++ b/Engine/Source/Runtime/Engine/Public/Rendering/SubstrateMaterialShared.cpp
@@ -47,14 +47,8 @@ FSubstrateOperator::FSubstrateOperator()
 	BSDFIndex = INDEX_NONE;
 	BSDFType = 0;
 	BSDFRegisteredSharedLocalBasis = FSubstrateRegisteredSharedLocalBasis();
-	bBSDFHasSSS = false;
-	bBSDFHasMFPPluggedIn = false;
-	bBSDFHasEdgeColor = false;
-	bBSDFHasFuzz = false;
-	bBSDFHasSecondRoughnessOrSimpleClearCoat = false;
-	bBSDFHasAnisotropy = false;
-	bBSDFHasGlint = false;
-	bBSDFHasSpecularProfile = false;
+	BSDFFeatures = 0;
+	SubUsage = SUBSTRATE_OPERATOR_SUBUSAGE_NONE;
 
 	bBSDFWritesEmissive = false;
 	bBSDFWritesAmbientOcclusion = false;
@@ -72,28 +66,14 @@ FSubstrateOperator::FSubstrateOperator()
 void FSubstrateOperator::CombineFlagsForParameterBlending(FSubstrateOperator& A, FSubstrateOperator& B)
 {
 #if WITH_EDITOR
-	bBSDFHasSSS = A.bBSDFHasSSS || B.bBSDFHasSSS;
-	bBSDFHasMFPPluggedIn = A.bBSDFHasMFPPluggedIn || B.bBSDFHasMFPPluggedIn;
-	bBSDFHasEdgeColor = A.bBSDFHasEdgeColor || B.bBSDFHasEdgeColor;
-	bBSDFHasFuzz = A.bBSDFHasFuzz || B.bBSDFHasFuzz;
-	bBSDFHasSecondRoughnessOrSimpleClearCoat = A.bBSDFHasSecondRoughnessOrSimpleClearCoat || B.bBSDFHasSecondRoughnessOrSimpleClearCoat;
-	bBSDFHasAnisotropy = A.bBSDFHasAnisotropy || B.bBSDFHasAnisotropy;
-	bBSDFHasGlint = A.bBSDFHasGlint || B.bBSDFHasGlint;
-	bBSDFHasSpecularProfile = A.bBSDFHasSpecularProfile || B.bBSDFHasSpecularProfile;
+	BSDFFeatures = A.BSDFFeatures | B.BSDFFeatures;
 #endif
 }
 
 void FSubstrateOperator::CopyFlagsForParameterBlending(FSubstrateOperator& A)
 {
 #if WITH_EDITOR
-	bBSDFHasSSS = A.bBSDFHasSSS;
-	bBSDFHasMFPPluggedIn = A.bBSDFHasMFPPluggedIn;
-	bBSDFHasEdgeColor = A.bBSDFHasEdgeColor;
-	bBSDFHasFuzz = A.bBSDFHasFuzz;
-	bBSDFHasSecondRoughnessOrSimpleClearCoat = A.bBSDFHasSecondRoughnessOrSimpleClearCoat;
-	bBSDFHasAnisotropy = A.bBSDFHasAnisotropy;
-	bBSDFHasGlint = A.bBSDFHasGlint;
-	bBSDFHasSpecularProfile = A.bBSDFHasSpecularProfile;
+	BSDFFeatures = A.BSDFFeatures;
 #endif
 }
 
@@ -106,7 +86,6 @@ bool FSubstrateOperator::IsDiscarded() const
 #endif
 }
 
-
 IMPLEMENT_TYPE_LAYOUT(FSubstrateRegisteredSharedLocalBasis);
 
 FSubstrateRegisteredSharedLocalBasis::FSubstrateRegisteredSharedLocalBasis()
diff --git a/Engine/Source/Runtime/Engine/Public/Rendering/SubstrateMaterialShared.h b/Engine/Source/Runtime/Engine/Public/Rendering/SubstrateMaterialShared.h
--- a/Engine/Source/Runtime/Engine/Public/Rendering/SubstrateMaterialShared.h
+++ b/Engine/Source/Runtime/Engine/Public/Rendering/SubstrateMaterialShared.h
@@ -24,6 +24,21 @@ public:
 	LAYOUT_FIELD_EDITORONLY(uint8, GraphSharedLocalBasisIndex);
 };
 
+enum ESubstrateBsdfFeature
+{
+	ESubstrateBsdfFeature_None                             = 0u,
+	ESubstrateBsdfFeature_SSS                              = 1u<<0u,
+	ESubstrateBsdfFeature_MFPPluggedIn                     = 1u<<1u,
+	ESubstrateBsdfFeature_EdgeColor                        = 1u<<2u,
+	ESubstrateBsdfFeature_Fuzz                             = 1u<<3u,
+	ESubstrateBsdfFeature_SecondRoughnessOrSimpleClearCoat = 1u<<4u,
+	ESubstrateBsdfFeature_Anisotropy                       = 1u<<5u,
+	ESubstrateBsdfFeature_Glint                            = 1u<<6u,
+	ESubstrateBsdfFeature_SpecularProfile                  = 1u<<7u,
+	ESubstrateBsdfFeature_Eye                              = 1u<<8u,
+	ESubstrateBsdfFeature_Hair                             = 1u<<9u
+};
+
 struct FSubstrateOperator
 {
 	DECLARE_TYPE_LAYOUT(FSubstrateOperator, NonVirtual);
@@ -48,14 +63,7 @@ public:
 	LAYOUT_FIELD_EDITORONLY(int32, BSDFIndex);		// Index in the array of BSDF if a BSDF operator
 	LAYOUT_FIELD_EDITORONLY(int32, BSDFType);
 	LAYOUT_FIELD_EDITORONLY(FSubstrateRegisteredSharedLocalBasis, BSDFRegisteredSharedLocalBasis);
-	LAYOUT_FIELD_EDITORONLY(uint8, bBSDFHasSSS);
-	LAYOUT_FIELD_EDITORONLY(uint8, bBSDFHasMFPPluggedIn);
-	LAYOUT_FIELD_EDITORONLY(uint8, bBSDFHasEdgeColor);
-	LAYOUT_FIELD_EDITORONLY(uint8, bBSDFHasFuzz);
-	LAYOUT_FIELD_EDITORONLY(uint8, bBSDFHasSecondRoughnessOrSimpleClearCoat);
-	LAYOUT_FIELD_EDITORONLY(uint8, bBSDFHasAnisotropy);
-	LAYOUT_FIELD_EDITORONLY(uint8, bBSDFHasGlint);
-	LAYOUT_FIELD_EDITORONLY(uint8, bBSDFHasSpecularProfile);
+	LAYOUT_FIELD_EDITORONLY(uint32,BSDFFeatures);
 	LAYOUT_FIELD_EDITORONLY(uint8, SubUsage);		// Sometimes, Unlit or Weight operators are used to transport data for other meaning (e.g. Light Function or ConvertToDecal)
 
 	LAYOUT_FIELD_EDITORONLY(uint8, bBSDFWritesEmissive);
@@ -76,6 +84,24 @@ public:
 	void CopyFlagsForParameterBlending(FSubstrateOperator& A);
 
 	bool IsDiscarded() const;
+	
+	bool Has(ESubstrateBsdfFeature In) const 
+	{ 
+	#if WITH_EDITOR
+		return (BSDFFeatures & In) != 0; 
+	#else
+		return false;
+	#endif
+	}
+
+	bool Has(uint32 In) const 
+	{ 
+	#if WITH_EDITOR
+		return (BSDFFeatures & In) != 0; 
+	#else
+		return false;
+	#endif
+	}
 };
 
 #define SUBSTRATE_COMPILATION_OUTPUT_MAX_OPERATOR 24
