Open around atomic that always increments up, if we abort we can just grab a new index next time

diff --git a/Engine/Source/Runtime/Engine/Private/Curves/KeyHandle.cpp b/Engine/Source/Runtime/Engine/Private/Curves/KeyHandle.cpp
--- a/Engine/Source/Runtime/Engine/Private/Curves/KeyHandle.cpp
+++ b/Engine/Source/Runtime/Engine/Private/Curves/KeyHandle.cpp
@@ -1,11 +1,12 @@
 // Copyright Epic Games, Inc. All Rights Reserved.
 
 #include "Curves/KeyHandle.h"
+#include "AutoRTFM/AutoRTFM.h"
 
 #include UE_INLINE_GENERATED_CPP_BY_NAME(KeyHandle)
 
-
-FKeyHandle::FKeyHandle()
+// Always open around this CTOR as we always increment this atomic. If we abort we simply will just grab a new index when trying again
+UE_AUTORTFM_ALWAYS_OPEN FKeyHandle::FKeyHandle()
 {
 	static std::atomic<uint32> LastKeyHandleIndex = 1;
 	Index = ++LastKeyHandleIndex;
