Fix missing custom tangent input when converting eye material using material attributes.

diff --git a/Engine/Source/Runtime/Engine/Private/Materials/Material.cpp b/Engine/Source/Runtime/Engine/Private/Materials/Material.cpp
--- a/Engine/Source/Runtime/Engine/Private/Materials/Material.cpp
+++ b/Engine/Source/Runtime/Engine/Private/Materials/Material.cpp
@@ -3524,6 +3524,27 @@ void UMaterial::ConvertMaterialToSubstrateMaterial()
 				}
 			}
 
+			// If the graph contains a 'tangent output' node and uses 'eye' shading model, we add the custom tangent to the material attribute list
+			const bool bRequireCustomTangent = ShadingModels.HasShadingModel(MSM_Eye) && TangentOutput;
+			if (bRequireCustomTangent)
+			{
+				const FGuid MaterialAttributeGuid = FMaterialAttributeDefinitionMap::GetID(MP_MaterialAttributes);
+				const FGuid CustomEyeTangentGuid = FMaterialAttributeDefinitionMap::GetCustomAttributeID(TEXT("CustomEyeTangent"));
+
+				UMaterialExpressionSetMaterialAttributes* SetAttributesNode = NewObject<UMaterialExpressionSetMaterialAttributes>(this);
+				SetAttributesNode->Material = this;
+				SetAttributesNode->Inputs[0].Connect(0, EditorOnly->MaterialAttributes.Expression);
+				SetAttributesNode->AttributeSetTypes.Add(CustomEyeTangentGuid);
+				SetAttributesNode->Inputs.Add(*TangentOutput->GetInput(0));
+				EditorOnly->ExpressionCollection.AddExpression(SetAttributesNode); // Add the SetAttributesNode to ensure the MaterialAttribute input are correctly tracked & cached
+
+				// Connect to conversion node
+				ConvertAttributeNode->MaterialAttributes.Connect(0, SetAttributesNode);
+
+				// Connect to root node
+				EditorOnly->MaterialAttributes.Connect(0, SetAttributesNode);
+			}
+
 			// Connect converted Substrate data to root node
 			EditorOnly->FrontMaterial.Connect(0, ConvertAttributeNode);
 
