ReferenceSkeleton: support for root re-parenting and deletion when conditions are met

#rnx

diff --git a/Engine/Source/Runtime/Engine/Private/Animation/ReferenceSkeleton.cpp b/Engine/Source/Runtime/Engine/Private/Animation/ReferenceSkeleton.cpp
--- a/Engine/Source/Runtime/Engine/Private/Animation/ReferenceSkeleton.cpp
+++ b/Engine/Source/Runtime/Engine/Private/Animation/ReferenceSkeleton.cpp
@@ -4,6 +4,7 @@
 #include "Animation/Skeleton.h"
 #include "Animation/SkeletonRemappingRegistry.h"
 #include "EngineLogs.h"
+#include "Algo/Count.h"
 #include "Engine/SkeletalMesh.h"
 
 FReferenceSkeletonModifier::FReferenceSkeletonModifier(USkeleton* InSkeleton)
@@ -104,16 +105,22 @@ void FReferenceSkeleton::Remove(const FName InBoneName, const bool bRemoveChildr
 		UE_LOG(LogAnimation, Error, TEXT("Remove: '%s' not found."), *InBoneName.ToString());
 		return;
 	}
-	
-	if (RawBoneIndex == 0)
-	{
-		UE_LOG(LogAnimation, Error, TEXT("Remove: cannot remove root bone."));
-		return;
-	}
 
 	const FMeshBoneInfo& BoneInfo = RawRefBoneInfo[RawBoneIndex];
 	const int32 RawParentIndex = BoneInfo.ParentIndex;
-	
+
+	// is this a root?
+	if (RawParentIndex == INDEX_NONE)
+	{
+		// is this the only root?
+		const int32 NumRoots = Algo::CountIf(RawRefBoneInfo, [](const FMeshBoneInfo& BoneInfo) { return BoneInfo.ParentIndex == INDEX_NONE; });
+		if (NumRoots == 1)
+		{
+			UE_LOG(LogAnimation, Error, TEXT("Remove: cannot remove root bone as it's the only one in the hierarchy."));
+			return;
+		}
+	}
+
 	// Make sure our arrays are in sync.
 	checkSlow((RawRefBoneInfo.Num() == RawRefBonePose.Num()) && (RawRefBoneInfo.Num() == RawNameToIndexMap.Num()));
 
@@ -322,6 +329,18 @@ namespace FReferenceSkeletonLocals
 				}
 			};
 
+			// look for the first root has it must be inserted first
+			// this can happen if the initial root has been re-parented to a new bone
+			const int32 FirstRoot = Elements.IndexOfByPredicate([](const FElement& Element)
+			{
+				return Element.Parent == nullptr;
+			});
+
+			if (ensure(FirstRoot != INDEX_NONE))
+			{
+				ElementToIndex(Elements[FirstRoot], ElementToIndex);
+			}
+			
 			// parse the full hierarchy
 			for (const FElement& Element: Elements)
 			{
@@ -342,16 +361,16 @@ int32 FReferenceSkeleton::SetParent(const FName InBoneName, const FName InParent
 	}
 
 	const int32 BoneIndex = FindRawBoneIndex(InBoneName);
-	if (BoneIndex < 1)
-	{ // we do not allow to change the root's parent
-		UE_LOG(LogAnimation, Error, TEXT("SetParent: cannot re-parent root bone."));
+	if (BoneIndex == INDEX_NONE)
+	{
+		UE_LOG(LogAnimation, Error, TEXT("SetParent: bone '%s' not found."), *InBoneName.ToString());
 		return INDEX_NONE;
 	}
 	
 	const int32 NewParentIndex = FindRawBoneIndex(InParentName);
 	if (NewParentIndex == INDEX_NONE && InParentName != NAME_None)
 	{
-		UE_LOG(LogAnimation, Error, TEXT("SetParent: '%s' not found."), *InParentName.ToString());
+		UE_LOG(LogAnimation, Error, TEXT("SetParent: parent bone '%s' not found."), *InParentName.ToString());
 		return INDEX_NONE;
 	}
 
