PR #9967: Expose UPlayerInput::DebugExecBindings to editor for more convenient ?

diff --git a/Engine/Source/Runtime/Engine/Classes/GameFramework/PlayerInput.h b/Engine/Source/Runtime/Engine/Classes/GameFramework/PlayerInput.h
--- a/Engine/Source/Runtime/Engine/Classes/GameFramework/PlayerInput.h
+++ b/Engine/Source/Runtime/Engine/Classes/GameFramework/PlayerInput.h
@@ -26,6 +26,7 @@ struct FInputActionBinding;
 struct FInputAxisBinding;
 struct FInputKeyBinding;
 
+
 /** Struct containing mappings for legacy method of binding keys to exec commands. */
 USTRUCT()
 struct FKeyBind
@@ -33,43 +34,43 @@ struct FKeyBind
 	GENERATED_USTRUCT_BODY()
 
 	/** The key to be bound to the command */
-	UPROPERTY(config)
+	UPROPERTY(config, EditAnywhere, Category = "Input")
 	FKey Key;
 
 	/** The command to execute when the key is pressed/released */
-	UPROPERTY(config)
+	UPROPERTY(config, EditAnywhere, Category = "Input")
 	FString Command;
 
 	/** Whether the control key needs to be held when the key event occurs */
-	UPROPERTY(config)
+	UPROPERTY(config, EditAnywhere, Category = "Input")
 	uint8 Control:1;
 
 	/** Whether the shift key needs to be held when the key event occurs */
-	UPROPERTY(config)
+	UPROPERTY(config, EditAnywhere, Category = "Input")
 	uint8 Shift:1;
 
 	/** Whether the alt key needs to be held when the key event occurs */
-	UPROPERTY(config)
+	UPROPERTY(config, EditAnywhere, Category = "Input")
 	uint8 Alt:1;
 
 	/** Whether the command key needs to be held when the key event occurs */
-	UPROPERTY(config)
+	UPROPERTY(config, EditAnywhere, Category = "Input")
 	uint8 Cmd:1;
 
 	/** Whether the control key must not be held when the key event occurs */
-	UPROPERTY(config)
+	UPROPERTY(config, EditAnywhere, Category = "Input")
 	uint8 bIgnoreCtrl:1;
 
 	/** Whether the shift key must not be held when the key event occurs */
-	UPROPERTY(config)
+	UPROPERTY(config, EditAnywhere, Category = "Input")
 	uint8 bIgnoreShift:1;
 
 	/** Whether the alt key must not be held when the key event occurs */
-	UPROPERTY(config)
+	UPROPERTY(config, EditAnywhere, Category = "Input")
 	uint8 bIgnoreAlt:1;
 
 	/** Whether the command key must not be held when the key event occurs */
-	UPROPERTY(config)
+	UPROPERTY(config, EditAnywhere, Category = "Input")
 	uint8 bIgnoreCmd:1;
 
 	UPROPERTY(transient)
@@ -435,7 +436,7 @@ private:
 
 public:
 	/** Generic bindings of keys to Exec()-compatible strings for development purposes only */
-	UPROPERTY(config)
+	UPROPERTY(config, EditAnywhere, Category = "Input")
 	TArray<struct FKeyBind> DebugExecBindings;
 
 	/** This player's version of the Axis Properties */
@@ -448,7 +449,7 @@ public:
 	TArray<struct FInputAxisKeyMapping> AxisMappings;
 
 	/** List of Axis Mappings that have been inverted */
-	UPROPERTY(config)
+	UPROPERTY(config, EditAnywhere, Category = "Input")
 	TArray<FName> InvertedAxis;
 
 	/** Gets the axis properties for a given AxisKey.  Returns if true if AxisKey was found in the AxisConfig array. */
