#UE Explicitly call async load completion callback when flushing in FWorldPartitionActorDescInstance. This avoids a complication where the callback for the async load may not happen if the flush happened while in a callstack processing another flush

diff --git a/Engine/Source/Runtime/Engine/Private/WorldPartition/WorldPartitionActorDescInstance.cpp b/Engine/Source/Runtime/Engine/Private/WorldPartition/WorldPartitionActorDescInstance.cpp
--- a/Engine/Source/Runtime/Engine/Private/WorldPartition/WorldPartitionActorDescInstance.cpp
+++ b/Engine/Source/Runtime/Engine/Private/WorldPartition/WorldPartitionActorDescInstance.cpp
@@ -107,7 +107,6 @@ FBox FWorldPartitionActorDescInstance::GetRuntimeBounds() const
 
 bool FWorldPartitionActorDescInstance::StartAsyncLoad()
 {
-	static FText FailedToLoad(LOCTEXT("FailedToLoadReason", "Failed to load"));
 	UnloadedReason = nullptr;
 
 	if (ActorPtr.IsExplicitlyNull() || ActorPtr.IsStale())
@@ -126,26 +125,11 @@ bool FWorldPartitionActorDescInstance::StartAsyncLoad()
 
 		AsyncLoadID = LoadPackageAsync(PackagePath, PackageName, FLoadPackageAsyncDelegate::CreateLambda([this, ActorPackage](const FName& PackageName, UPackage* Package, EAsyncLoadingResult::Type Result)
 		{
-			//checkf(AsyncLoadID != INDEX_NONE, TEXT("FWorldPartitionActorDescInstance failed to load package. ActorPackage:%s PackageName:%s"), *ActorPackage.ToString(), *PackageName.ToString());
-			AsyncLoadID = INDEX_NONE;
-
-			if ((Result != EAsyncLoadingResult::Succeeded) || !Package)
-			{
-				UE_LOG(LogWorldPartition, Warning, TEXT("Can't load actor guid `%s` ('%s') from package '%s'"), *GetGuid().ToString(), *GetActorNameString(), *ActorPackage.ToString());
-				UnloadedReason = &FailedToLoad;
-				return;
-			}
-
-			ActorPtr = FindObject<AActor>(nullptr, * GetActorSoftPath().ToString());
-
-			if (!ActorPtr.IsValid())
+			if (AsyncLoadID != INDEX_NONE)
 			{
-				UE_LOG(LogWorldPartition, Warning, TEXT("Can't find actor guid `%s` ('%s') in package '%s'"), *GetGuid().ToString(), *GetActorNameString(), *ActorPackage.ToString());
-				UnloadedReason = &FailedToLoad;
-				return;
+				AsyncLoadFinished(ActorPackage, PackageName, Package, Result == EAsyncLoadingResult::Succeeded);
+				AsyncLoadID = INDEX_NONE;
 			}
-
-			//check(ActorPtr->GetPackage() == Package);
 		})
 		, PKG_None, INDEX_NONE, 0, InstancingContext);
 	}
@@ -157,9 +141,43 @@ void FWorldPartitionActorDescInstance::FlushAsyncLoad() const
 {
 	if (AsyncLoadID != INDEX_NONE)
 	{
-		FlushAsyncLoading(AsyncLoadID);
+		// Instead relying on AsyncLoading to call the callback during a flush, we'll do it ourselves here explicitly
+		// This is because during a callstack where we are already async loading, we may not get completion callbacks called before returning from this flush.
+		// Setting AsyncLoadID to INDEX_NONE before flushing so that we do not end up calling AsyncLoadFinished twice in cases where the completion callback is called during the flush
+		int32 IDToFlush = AsyncLoadID;
 		AsyncLoadID = INDEX_NONE;
+		FlushAsyncLoading(IDToFlush);
+
+		// Now call AsyncLoadFinished to assign ActorPtr
+		const FName ActorPackage = GetActorPackage();
+		const FLinkerInstancingContext* InstancingContext = GetContainerInstance()->GetInstancingContext();
+		const FName PackageName = InstancingContext ? InstancingContext->RemapPackage(ActorPackage) : ActorPackage;
+		UPackage* Package = FindPackage(nullptr, *PackageName.ToString());
+		AsyncLoadFinished(ActorPackage, PackageName, Package, Package != nullptr);
+	}
+}
+
+void FWorldPartitionActorDescInstance::AsyncLoadFinished(const FName& ActorPackage, const FName& PackageName, UPackage* Package, bool bSuccessful) const
+{
+	static FText FailedToLoad(LOCTEXT("FailedToLoadReason", "Failed to load"));
+
+	if (!bSuccessful || !Package)
+	{
+		UE_LOG(LogWorldPartition, Warning, TEXT("Can't load actor guid `%s` ('%s') from package '%s'"), *GetGuid().ToString(), *GetActorNameString(), *ActorPackage.ToString());
+		UnloadedReason = &FailedToLoad;
+		return;
 	}
+
+	ActorPtr = FindObject<AActor>(nullptr, *GetActorSoftPath().ToString());
+
+	if (!ActorPtr.IsValid())
+	{
+		UE_LOG(LogWorldPartition, Warning, TEXT("Can't find actor guid `%s` ('%s') in package '%s'"), *GetGuid().ToString(), *GetActorNameString(), *ActorPackage.ToString());
+		UnloadedReason = &FailedToLoad;
+		return;
+	}
+
+	check(ActorPtr->GetPackage() == Package);
 }
 
 void FWorldPartitionActorDescInstance::MarkUnload()
diff --git a/Engine/Source/Runtime/Engine/Public/WorldPartition/WorldPartitionActorDescInstance.h b/Engine/Source/Runtime/Engine/Public/WorldPartition/WorldPartitionActorDescInstance.h
--- a/Engine/Source/Runtime/Engine/Public/WorldPartition/WorldPartitionActorDescInstance.h
+++ b/Engine/Source/Runtime/Engine/Public/WorldPartition/WorldPartitionActorDescInstance.h
@@ -127,6 +127,7 @@ protected:
 
 	ENGINE_API bool StartAsyncLoad();
 	ENGINE_API void FlushAsyncLoad() const;
+	void AsyncLoadFinished(const FName& ActorPackage, const FName& PackageName, UPackage* Package, bool bSuccessful) const;
 	ENGINE_API void MarkUnload();
 
 	UWorldPartition* GetLoadedChildWorldPartition() const { return ActorDesc->GetLoadedChildWorldPartition(this); }
