Alpha: Prevent cook errors for propagate alpha previously set to legacy EAlphaChannelMode::AllowThroughTonemapper (2) by force-updating the config.

Previous fix attempt at CL 35569163 was backed out due to cvar/uproperty mismatch warning ("ConsoleVariable-Backed Enum Property of type... "). The only other solution would be to revert the uproperty & cvar back to int while adding engine-level support for cvar clamping at load time.

#rnx

diff --git a/Engine/Source/Runtime/Engine/Classes/Engine/RendererSettings.h b/Engine/Source/Runtime/Engine/Classes/Engine/RendererSettings.h
--- a/Engine/Source/Runtime/Engine/Classes/Engine/RendererSettings.h
+++ b/Engine/Source/Runtime/Engine/Classes/Engine/RendererSettings.h
@@ -1452,6 +1452,7 @@ private:
 	// Notification about missing required shader models.
 	TWeakPtr<class SNotificationItem> ShaderModelNotificationPtr;
 
+	void PreInitPropertiesFixup();
 	void CheckForMissingShaderModels();
 	void FixAntiAliasingOnShadingPathChange(FPropertyChangedEvent& PropertyChangedEvent);
 #endif // WITH_EDITOR
diff --git a/Engine/Source/Runtime/Engine/Private/RendererSettings.cpp b/Engine/Source/Runtime/Engine/Private/RendererSettings.cpp
--- a/Engine/Source/Runtime/Engine/Private/RendererSettings.cpp
+++ b/Engine/Source/Runtime/Engine/Private/RendererSettings.cpp
@@ -102,6 +102,13 @@ URendererSettings::URendererSettings(const FObjectInitializer& ObjectInitializer
 	WhiteChromaticityCoordinate = FVector2D::ZeroVector;
 	bUseLegacyLuminanceFactors = false;
 	bEnableVirtualTextureOpacityMask = false;
+
+#if WITH_EDITOR
+	if (IsTemplate())
+	{
+		PreInitPropertiesFixup();
+	}
+#endif // #if WITH_EDITOR
 }
 
 void URendererSettings::PostInitProperties()
@@ -164,6 +171,22 @@ void UpdateDependentPropertyInConfigFile(URendererSettings* RendererSettings, FN
 	}
 }
 
+void URendererSettings::PreInitPropertiesFixup()
+{
+	int32 PropagateAlphaIntValue = 0;
+	if (GConfig->GetInt(TEXT("/Script/Engine.RendererSettings"), TEXT("r.PostProcessing.PropagateAlpha"), PropagateAlphaIntValue, GEngineIni))
+	{
+		constexpr int32 LegacyAllowThroughTonemapper = 2;
+		if (PropagateAlphaIntValue == LegacyAllowThroughTonemapper)
+		{
+			UE_LOG(LogEngine, Warning, TEXT("r.PostProcessing.PropagateAlpha config value under /Script/Engine.RendererSettings was automatically converted from 2 to True, please version the change."));
+
+			bEnableAlphaChannelInPostProcessing = true;
+			UpdateDependentPropertyInConfigFile(this, GET_MEMBER_NAME_CHECKED(URendererSettings, bEnableAlphaChannelInPostProcessing));
+		}
+	}
+}
+
 void URendererSettings::FixAntiAliasingOnShadingPathChange(FPropertyChangedEvent& PropertyChangedEvent)
 {
 	if (PropertyChangedEvent.GetPropertyName() == GET_MEMBER_NAME_CHECKED(URendererSettings, MobileShadingPath))
