[Material Translator] Some minor fixes and added comments to the new material translator code.

#rnx

diff --git a/Engine/Source/Runtime/Engine/Private/Materials/MaterialExpressionsIR.cpp b/Engine/Source/Runtime/Engine/Private/Materials/MaterialExpressionsIR.cpp
--- a/Engine/Source/Runtime/Engine/Private/Materials/MaterialExpressionsIR.cpp
+++ b/Engine/Source/Runtime/Engine/Private/Materials/MaterialExpressionsIR.cpp
@@ -532,9 +532,15 @@ void UMaterialExpressionTextureSample::Build(UE::MIR::FEmitter& Emitter)
 
 void UMaterialExpressionTextureCoordinate::Build(UE::MIR::FEmitter& Emitter)
 {
-	MIR::FValue* OutputValue = Emitter.GetExternalInput(MIR::TexCoordIndexToExternalInput(CoordinateIndex));
+	if (UnMirrorU || UnMirrorV)
+	{
+		Emitter.Error("unmirroring unsupported");
+	}
 
-	// todo: add tiling support
+	MIR::FValue* OutputValue = Emitter.GetExternalInput(MIR::TexCoordIndexToExternalInput(CoordinateIndex));
+	
+	// Multiply the UV input by the UV tiling constants
+	OutputValue = Emitter.EmitBinaryOperator(MIR::BO_Multiply, OutputValue, Emitter.EmitConstantFloat2({ UTiling, VTiling }));
 	
 	Emitter.Put(GetOutput(0), OutputValue);
 }
diff --git a/Engine/Source/Runtime/Engine/Private/Materials/MaterialIR.cpp b/Engine/Source/Runtime/Engine/Private/Materials/MaterialIR.cpp
--- a/Engine/Source/Runtime/Engine/Private/Materials/MaterialIR.cpp
+++ b/Engine/Source/Runtime/Engine/Private/Materials/MaterialIR.cpp
@@ -27,6 +27,24 @@ const TCHAR* ValueKindToString(EValueKind Kind)
 	};
 }
 
+uint32 FValue::GetSizeInBytes() const
+{
+	switch (Kind)
+	{
+		case VK_Constant: return sizeof(FConstant);
+		case VK_ExternalInput: return sizeof(FExternalInput);
+		case VK_MaterialParameter: return sizeof(FMaterialParameter);
+		case VK_SetMaterialOutput: return sizeof(FSetMaterialOutput);
+		case VK_BinaryOperator: return sizeof(FBinaryOperator);
+		case VK_Dimensional: return sizeof(FDimensional) + sizeof(FValue*) * static_cast<const FDimensional*>(this)->GetComponents().Num();
+		case VK_Branch: return sizeof(FBranch);
+		case VK_Subscript: return sizeof(FSubscript);
+		case VK_Cast: return sizeof(FCast);
+		case VK_TextureSample: return sizeof(FTextureSample);
+		default: UE_MIR_UNREACHABLE();
+	}
+}
+
 FInstruction* FValue::AsInstruction()
 {
 	return this && (Kind > VK_InstructionBegin && Kind < VK_InstructionEnd) ? static_cast<FInstruction*>(this) : nullptr;
@@ -64,24 +82,6 @@ bool FValue::Equals(const FValue* Other) const
 	return FMemory::Memcmp(this, Other, SizeInBytes) == 0;
 }
 
-uint32 FValue::GetSizeInBytes() const
-{
-	switch (Kind)
-	{
-		case VK_Constant: return sizeof(FConstant);
-		case VK_ExternalInput: return sizeof(FExternalInput);
-		case VK_MaterialParameter: return sizeof(FMaterialParameter);
-		case VK_SetMaterialOutput: return sizeof(FSetMaterialOutput);
-		case VK_BinaryOperator: return sizeof(FBinaryOperator);
-		case VK_Dimensional: return sizeof(FDimensional) + sizeof(FValue*) * static_cast<const FDimensional*>(this)->GetComponents().Num();
-		case VK_Branch: return sizeof(FBranch);
-		case VK_Subscript: return sizeof(FSubscript);
-		case VK_Cast: return sizeof(FCast);
-		case VK_TextureSample: return sizeof(FTextureSample);
-		default: UE_MIR_UNREACHABLE();
-	}
-}
-
 TArrayView<const FValue*> FValue::GetUses() const
 {
 	TArrayView<FValue*> Uses = const_cast<FValue*>(this)->GetUses();
@@ -138,54 +138,80 @@ TArrayView<FValue*> FValue::GetUses()
 	}
 }
 
-UTexture* FValue::GetTexture()
+bool FValue::IsScalar() const
 {
-	if (auto Parameter = As<FMaterialParameter>())
-	{
-		UObject* TextureObject = Parameter->Metadata.Value.AsTextureObject();
-		if (TextureObject)
-		{
-			return static_cast<UTexture*>(TextureObject);
-		}
-	}
+	return Type->AsScalar() != nullptr;
+}
 
-	return nullptr;
+bool FValue::IsVector() const
+{
+	return Type->AsVector() != nullptr;
 }
 
-bool FValue::IsConstantTrue() const
+bool FValue::IsTrue() const
 {
 	const MIR::FConstant* Constant = As<MIR::FConstant>();
 	return Constant && Constant->IsBool() && Constant->Boolean == true;
 }
 
-bool FValue::IsConstantFalse() const
+bool FValue::IsFalse() const
 {
 	const MIR::FConstant* Constant = As<MIR::FConstant>();
 	return Constant && Constant->IsBool() && Constant->Boolean == false;
 }
 
-bool FValue::IsConstantZero() const
+bool FValue::IsExactlyZero() const
 {
-	const MIR::FConstant* Constant = As<MIR::FConstant>();
-	if (!Constant)
+	if (const MIR::FConstant* Constant = As<MIR::FConstant>())
 	{
-		return false;
+		return (Constant->IsInteger() && Constant->Integer == 0)
+			|| (Constant->IsFloat() && Constant->Float == 0.0f);
 	}
+	return false;
+}
 
-	return (Constant->IsInteger() && Constant->Integer == 0)
-		|| (Constant->IsFloat() && Constant->Float == 0.0f);
+bool FValue::IsNearlyZero() const
+{
+	if (const MIR::FConstant* Constant = As<MIR::FConstant>())
+	{
+		return (Constant->IsInteger() && Constant->Integer == 0)
+			|| (Constant->IsFloat() && FMath::IsNearlyZero(Constant->Float));
+	}
+	return false;
 }
 
-bool FValue::IsConstantOne() const
+bool FValue::IsExactlyOne() const
 {
-	const MIR::FConstant* Constant = As<MIR::FConstant>();
-	if (!Constant)
+	if (const MIR::FConstant* Constant = As<MIR::FConstant>())
 	{
-		return false;
+		return (Constant->IsInteger() && Constant->Integer == 1)
+			|| (Constant->IsFloat() && Constant->Float == 1.0f);
+	}
+	return false;
+}
+
+bool FValue::IsNearlyOne() const
+{
+	if (const MIR::FConstant* Constant = As<MIR::FConstant>())
+	{
+		return (Constant->IsInteger() && Constant->Integer == 1)
+			|| (Constant->IsFloat() && FMath::IsNearlyEqual(Constant->Float, 1.0f));
+	}
+	return false;
+}
+
+UTexture* FValue::GetTexture()
+{
+	if (auto Parameter = As<FMaterialParameter>())
+	{
+		UObject* TextureObject = Parameter->Metadata.Value.AsTextureObject();
+		if (TextureObject)
+		{
+			return static_cast<UTexture*>(TextureObject);
+		}
 	}
 
-	return (Constant->IsInteger() && Constant->Integer == 1)
-		|| (Constant->IsFloat() && Constant->Float == 1.0f);
+	return nullptr;
 }
 
 bool FConstant::IsBool() const
diff --git a/Engine/Source/Runtime/Engine/Private/Materials/MaterialIREmitter.cpp b/Engine/Source/Runtime/Engine/Private/Materials/MaterialIREmitter.cpp
--- a/Engine/Source/Runtime/Engine/Private/Materials/MaterialIREmitter.cpp
+++ b/Engine/Source/Runtime/Engine/Private/Materials/MaterialIREmitter.cpp
@@ -620,22 +620,22 @@ static FValue* TryFoldBinaryOperatorScalar(FEmitter* Emitter, EBinaryOperator Op
 	switch (Operator)
 	{
 		case BO_Add:
-			if (Lhs->IsConstantZero())
+			if (Lhs->IsNearlyZero())
 			{
 				return Rhs;
 			}
-			else if (Rhs->IsConstantZero())
+			else if (Rhs->IsNearlyZero())
 			{
 				return Lhs;
 			}
 			break;
 
 		case BO_Subtract:
-			if (Rhs->IsConstantZero())
+			if (Rhs->IsNearlyZero())
 			{
 				return Lhs;
 			}
-			else if (Lhs->IsConstantZero())
+			else if (Lhs->IsNearlyZero())
 			{
 				// return Emitter->EmitUnaryOperator(UO_Minus, Rhs);
 				UE_MIR_TODO();
@@ -643,18 +643,18 @@ static FValue* TryFoldBinaryOperatorScalar(FEmitter* Emitter, EBinaryOperator Op
 			break;
 
 		case BO_Multiply:
-			if (Lhs->IsConstantZero() || Rhs->IsConstantOne())
+			if (Lhs->IsNearlyZero() || Rhs->IsNearlyOne())
 			{
 				return Lhs;
 			}
-			else if (Lhs->IsConstantOne() || Rhs->IsConstantZero())
+			else if (Lhs->IsNearlyOne() || Rhs->IsNearlyZero())
 			{
 				return Rhs;
 			}
 			break;
 
 		case BO_Divide:
-			if (Lhs->IsConstantZero() || Rhs->IsConstantOne())
+			if (Lhs->IsNearlyZero() || Rhs->IsNearlyOne())
 			{
 				return Lhs;
 			}
@@ -734,26 +734,21 @@ static FValue* TryFoldBinaryOperatorScalar(FEmitter* Emitter, EBinaryOperator Op
 FValue* FEmitter::EmitBinaryOperator(EBinaryOperator Operator, FValue* Lhs, FValue* Rhs)
 {
 	// Argument value types must always match.
-	check(Lhs->Type == Rhs->Type);
+	if (Lhs->Type != Rhs->Type)
+	{
+		Errorf(TEXT("operands of binary operator `%s` do not have the same type."), BinaryOperatorToString(Operator));
+		return nullptr;
+	}
 
 	// Get the operands primitive type.
 	FPrimitiveTypePtr PrimitiveType = Lhs->Type->AsPrimitive();
-	check(PrimitiveType);
-
-	// Try converting operands to scalar constants
-	FConstant* ScalarLhs = Lhs->As<FConstant>();
-	FConstant* ScalarRhs = Rhs->As<FConstant>();
-
-	//  If they're both so, fold the binary operator right away and return the result.
-	if (ScalarLhs && ScalarRhs)
+	if (!PrimitiveType)
 	{
-		if (FValue* Value = TryFoldBinaryOperatorScalar(this, Operator, ScalarLhs, ScalarRhs))
-		{
-			return Value;
-		}
+		Errorf(TEXT("invalid binary operator on non primitive types."));
+		return nullptr;
 	}
 
-	// Determine the result type. If the operator is primitive, the result type will be the same
+	// Determine the result type. If the operator is arithmetic, the result type will be the same
 	// as the operands type. Otherwise it will have the same number of components but bool.
 	FPrimitiveTypePtr ResultType = IsArithmeticOperator(Operator)
 		? PrimitiveType
@@ -766,7 +761,11 @@ FValue* FEmitter::EmitBinaryOperator(EBinaryOperator Operator, FValue* Lhs, FVal
 	bool bSomeResultComponentWasFolded = false;
 	bool bResultIsIdenticalToLhs = true;
 	bool bResultIsIdenticalToRhs = true;
-	FValue* TempResultComponents[FDimensional::MaxNumComponents];
+
+	// Allocate the temporary array to store the folded component results
+	FMemMark Mark(FMemStack::Get());
+	TArrayView<FValue*> TempResultComponents = MakeTemporaryArray<FValue*>(Mark, ResultType->GetNumComponents());
+	
 	for (int i = 0; i < ResultType->GetNumComponents(); ++i)
 	{
 		// Extract the arguments individual components
@@ -803,6 +802,13 @@ FValue* FEmitter::EmitBinaryOperator(EBinaryOperator Operator, FValue* Lhs, FVal
 	// so that we retain as much compile-time information as possible.
 	if (bSomeResultComponentWasFolded)
 	{
+		// If result type is scalar, simply return the single folded result (instead of creating a dimensional value)
+		if (ResultType->IsScalar())
+		{
+			check(TempResultComponents[0]);
+			return TempResultComponents[0];
+		}
+
 		// Make the new dimensional value
 		FDimensional* Result = NewDimensionalValue(this, ResultType);
 		TArrayView<FValue*> ResultComponents = Result->GetComponents();
diff --git a/Engine/Source/Runtime/Engine/Private/Materials/MaterialIRInternal.h b/Engine/Source/Runtime/Engine/Private/Materials/MaterialIRInternal.h
--- a/Engine/Source/Runtime/Engine/Private/Materials/MaterialIRInternal.h
+++ b/Engine/Source/Runtime/Engine/Private/Materials/MaterialIRInternal.h
@@ -6,6 +6,13 @@
 
 #if WITH_EDITOR
 
+template <typename T>
+TArrayView<T> MakeTemporaryArray(FMemMark&, int Count)
+{
+	auto Ptr = (T*)FMemStack::Get().Alloc(sizeof(T) * Count, alignof(T));
+	return { Ptr, Count };
+}
+
 namespace UE::MIR::Internal {
 
 //
diff --git a/Engine/Source/Runtime/Engine/Private/Materials/MaterialIRModuleBuilder.cpp b/Engine/Source/Runtime/Engine/Private/Materials/MaterialIRModuleBuilder.cpp
--- a/Engine/Source/Runtime/Engine/Private/Materials/MaterialIRModuleBuilder.cpp
+++ b/Engine/Source/Runtime/Engine/Private/Materials/MaterialIRModuleBuilder.cpp
@@ -60,7 +60,6 @@ struct FMaterialIRModuleBuilderImpl
 	FMaterialIRModule* Module;
 	MIR::FEmitter* Emitter;
 	TArray<FAnalysisContext> AnalysisContextStack;
-	TArray<MIR::FValue*> AuxValueList;
 	TArray<MIR::FInstruction*> InstructionStack{};
 
 	void Step_Initialize()
@@ -196,7 +195,8 @@ struct FMaterialIRModuleBuilderImpl
 
 	void PushFunctionCall(UMaterialExpressionMaterialFunctionCall* Call)
 	{
-		AuxValueList.Empty(Call->FunctionInputs.Num()); /* todo: is there a better (and idiomatic) way to allocate temp dynamic array on the stack? */
+		FMemMark Mark(FMemStack::Get());
+		TArrayView<MIR::FValue*> CallInputValues = MakeTemporaryArray<MIR::FValue*>(Mark, Call->FunctionInputs.Num());
 
 		// Make sure each function input is connected and has a value. If so, cache the values flowing into this
 		// funcion call inside the auxiliary value array.
@@ -207,7 +207,7 @@ struct FMaterialIRModuleBuilderImpl
 			if (Value)
 			{
 				MIR::FTypePtr Type = MIR::FType::FromMaterialValueType((EMaterialValueType)FunctionInput.ExpressionInput->GetInputType(0));
-				AuxValueList.Push(Emitter->EmitConstruct(Type, Value));
+				CallInputValues[i] = Emitter->EmitConstruct(Type, Value);
 			}
 		}
 
@@ -230,11 +230,10 @@ struct FMaterialIRModuleBuilderImpl
 		for (int i = 0; i < Call->FunctionInputs.Num(); ++i)
 		{
 			FFunctionExpressionInput& FunctionInput = Call->FunctionInputs[i];
-			MIR::FValue* Value = AuxValueList[i];
 
 			// Bind the value flowing into the function call input to the function input
 			// expression (inside the function) in the new context.
-			NewContext.SetOutputValue(FunctionInput.ExpressionInput->GetOutput(0), Value);
+			NewContext.SetOutputValue(FunctionInput.ExpressionInput->GetOutput(0), CallInputValues[i]);
 
 			// Mark the function input as built.
 			NewContext.BuiltExpressions.Add(FunctionInput.ExpressionInput.Get());
diff --git a/Engine/Source/Runtime/Engine/Public/Materials/MaterialIR.h b/Engine/Source/Runtime/Engine/Public/Materials/MaterialIR.h
--- a/Engine/Source/Runtime/Engine/Public/Materials/MaterialIR.h
+++ b/Engine/Source/Runtime/Engine/Public/Materials/MaterialIR.h
@@ -42,29 +42,69 @@ enum EValueFlags
 	VF_InstructionAnalyzed = 2,
 };
 
+//
 struct FValue
 {
 	EValueKind Kind : 8{};
 	EValueFlags Flags : 8{};
 	FTypePtr  Type{};
 
+	// Returns the size in bytes of this value instance.
+	uint32 GetSizeInBytes() const;
+	
+	// Sets the specified value flags.
 	void SetFlags(EValueFlags InFlags) { Flags = (EValueFlags)(Flags | InFlags); }
+
+	// Returns whether this value is of specified kind.
 	bool IsA(EValueKind InKind) const { return Kind == InKind; }
+
+	// Tries to cast this value to an instruction and returns it (nullptr otherwise).
 	FInstruction* AsInstruction();
+
+	// Tries to cast this value to an instruction and returns it (nullptr otherwise).
 	const FInstruction* AsInstruction() const;
-	bool Equals(const FValue* Other) const;
-	uint32 GetSizeInBytes() const;
+
+	// Returns the array of this value's uses. An use is another value referenced by this one (e.g. the operands of a binary expression).
 	TArrayView<const FValue*> GetUses() const;
+
+	// Returns the array of this value's uses. An use is another value referenced by this one (e.g. the operands of a binary expression).
 	TArrayView<FValue*> GetUses();
+
+	// Returns whether this value exactly equals Other.
+	bool Equals(const FValue* Other) const;
+
+	// Returns whether this value is a scalar (its type is Primitive with exactly 1 component).
+	bool IsScalar() const;
+
+	// Returns whether this value is a vector (tis type is Primitive with 1-4 rows and exactly 1 column).
+	bool IsVector() const;
+
+	// Returns whether this value is a constant boolean with value true.
+	bool IsTrue() const;
+
+	// Returns whether this value is a constant boolean with value false.
+	bool IsFalse() const;
+
+	// Returns whether this value is arithmetic and exactly zero.
+	bool IsExactlyZero() const;
+
+	// Returns whether this value is arithmetic and approximately zero.
+	bool IsNearlyZero() const;
+
+	// Returns whether this value is arithmetic and exactly one.
+	bool IsExactlyOne() const;
+
+	// Returns whether this value is arithmetic and approximately one.
+	bool IsNearlyOne() const;
+	
+	// Returns this value's textures if it has one (nullptr otherwise).
 	UTexture* GetTexture();
-	bool IsConstantTrue() const;
-	bool IsConstantFalse() const;
-	bool IsConstantZero() const;
-	bool IsConstantOne() const;
 
+	// Tries to cast this value to specified type T and returns the casted pointer, if possible (nullptr otherwise).
 	template <typename T>
 	T* As() { return this && IsA(T::TypeKind) ? static_cast<T*>(this) : nullptr; }
 
+	// Tries to cast this value to specified type T and returns the casted pointer, if possible (nullptr otherwise).
 	template <typename T>
 	const T* As() const { return this && IsA(T::TypeKind) ? static_cast<const T*>(this) : nullptr; }
 };
