Make sure packet simulation is applied at most once to a packet.

diff --git a/Engine/Source/Runtime/Engine/Private/NetConnection.cpp b/Engine/Source/Runtime/Engine/Private/NetConnection.cpp
--- a/Engine/Source/Runtime/Engine/Private/NetConnection.cpp
+++ b/Engine/Source/Runtime/Engine/Private/NetConnection.cpp
@@ -2139,7 +2139,13 @@ void UNetConnection::FlushPacketOrderCache(bool bFlushWholeCache/*=false*/)
 			{
 				UE_LOG(LogNet, VeryVerbose, TEXT("'Out of Order' Packet Cache, replaying packet with cache index: %i (bFlushWholeCache: %i)"), PacketOrderCacheStartIdx, (int32)bFlushWholeCache);
 
-				ReceivedPacket(*CurCachePacket.Get());
+#if DO_ENABLE_NET_TEST
+				// Packets in the packet order cache have already had packet simulation applied. Make sure it's not applied again.
+				TGuardValue<bool> ReinjectingGuard(bIsReinjectingDelayedPackets, true);
+#endif
+
+				constexpr bool bIsReinjectedPacket = true;
+				ReceivedPacket(*CurCachePacket.Get(), bIsReinjectedPacket);
 
 				CurCachePacket.Reset();
 
