Added nullptr guard to AudioComponent->Sound

#rnx

diff --git a/Engine/Source/Runtime/Engine/Private/AudioVirtualLoop.cpp b/Engine/Source/Runtime/Engine/Private/AudioVirtualLoop.cpp
--- a/Engine/Source/Runtime/Engine/Private/AudioVirtualLoop.cpp
+++ b/Engine/Source/Runtime/Engine/Private/AudioVirtualLoop.cpp
@@ -321,14 +321,16 @@ bool FAudioVirtualLoop::Update(float DeltaTime, bool bForceUpdate)
 		{
 			if (UAudioComponent* AudioComponent = UAudioComponent::GetAudioComponentFromID(AudioComponentID))
 			{
-				// Use instance parameters which will not contain transient parameters
-				check(AudioComponent->Sound);
-				TArray<FAudioParameter> InstanceParams = AudioComponent->GetInstanceParameters();
-
-				// create proxies from the UObject references in InstanceParams
-				AudioComponent->Sound->InitParameters(InstanceParams);
-				
-				ParameterTransmitter->OnRealizeVirtualizedActiveSound(MoveTemp(InstanceParams));
+				if (ensure(AudioComponent->Sound != nullptr))
+				{
+					// Use instance parameters which will not contain transient parameters
+					TArray<FAudioParameter> InstanceParams = AudioComponent->GetInstanceParameters();
+
+					// create proxies from the UObject references in InstanceParams
+					AudioComponent->Sound->InitParameters(InstanceParams);
+
+					ParameterTransmitter->OnRealizeVirtualizedActiveSound(MoveTemp(InstanceParams));
+				}
 			}
 		}
 	}
