Avoid crash from SpawnDecalAttached edge case usages

diff --git a/Engine/Source/Runtime/Engine/Private/GameplayStatics.cpp b/Engine/Source/Runtime/Engine/Private/GameplayStatics.cpp
--- a/Engine/Source/Runtime/Engine/Private/GameplayStatics.cpp
+++ b/Engine/Source/Runtime/Engine/Private/GameplayStatics.cpp
@@ -2208,7 +2208,7 @@ int32 UGameplayStatics::GetMaxAudioChannelCount(const UObject* WorldContextObjec
 
 UDecalComponent* CreateDecalComponent(class UMaterialInterface* DecalMaterial, FVector DecalSize, UWorld* World, AActor* Actor, float LifeSpan)
 {
-	if (World && World->GetNetMode() == NM_DedicatedServer)
+	if (World == nullptr || World->GetNetMode() == NM_DedicatedServer)
 	{
 		return nullptr;
 	}
