Add support for getting the mesh paint texture asset's RHI texture on the render thread.

diff --git a/Engine/Source/Runtime/Engine/Private/SceneView.cpp b/Engine/Source/Runtime/Engine/Private/SceneView.cpp
--- a/Engine/Source/Runtime/Engine/Private/SceneView.cpp
+++ b/Engine/Source/Runtime/Engine/Private/SceneView.cpp
@@ -27,6 +27,7 @@
 #include "Camera/CameraComponent.h"
 #include "Camera/CameraTypes.h"
 #include "UObject/Interface.h"
+#include "TextureResource.h"
 
 DEFINE_LOG_CATEGORY(LogBufferVisualization);
 DEFINE_LOG_CATEGORY(LogNaniteVisualization);
@@ -424,15 +425,37 @@ EVertexColorViewMode::Type GetMeshPaintVisualizeChannels()
 	return GMeshPaintVisualizeChannels;
 }
 static TWeakObjectPtr<UTexture> GMeshPaintVisualizeTexture = nullptr;
+static FRHITexture* GMeshPaintVisualizeTextureRHI_GameThread = nullptr;
+static FRHITexture* GMeshPaintVisualizeTextureRHI_RenderThread = nullptr;
 void SetMeshPaintVisualizeTexture(TWeakObjectPtr<UTexture> Texture)
 {
 	GMeshPaintVisualizeTexture = Texture;
+		
+	UTexture* TexturePtr = Texture.Get();
+	FTextureResource* TextureResource = TexturePtr != nullptr ? TexturePtr->GetResource() : nullptr;
+	FRHITexture* TextureRHI = TextureResource != nullptr ? TextureResource->GetTexture2DRHI() : nullptr;
+	if (TextureRHI != GMeshPaintVisualizeTextureRHI_GameThread)
+	{
+		GMeshPaintVisualizeTextureRHI_GameThread = TextureRHI;
+		ENQUEUE_RENDER_COMMAND(SetMeshPaintVisualizeTexture)([TextureRHI](FRHICommandListImmediate& RHICmdList)
+		{
+			GMeshPaintVisualizeTextureRHI_RenderThread = TextureRHI;
+		});
+	}
+}
+FRHITexture* GetMeshPaintVisualizeTexture_RenderThread()
+{
+	return GMeshPaintVisualizeTextureRHI_RenderThread;
 }
 static int32 GMeshPaintVisualizeUVChannel = 0;
 void SetMeshPaintVisualizeTextureCoordinateIndex(int32 Index)
 {
 	GMeshPaintVisualizeUVChannel = Index;
 }
+int32 GetMeshPaintVisualizeTextureCoordinateIndex()
+{
+	return GMeshPaintVisualizeUVChannel;
+}
 
 FMaterialRenderProxy* GetMeshPaintVisualizeMaterialRenderProxy(bool bIsSelected, bool bIsHovered)
 {
diff --git a/Engine/Source/Runtime/Engine/Public/SceneManagement.h b/Engine/Source/Runtime/Engine/Public/SceneManagement.h
--- a/Engine/Source/Runtime/Engine/Public/SceneManagement.h
+++ b/Engine/Source/Runtime/Engine/Public/SceneManagement.h
@@ -2800,7 +2800,9 @@ extern ENGINE_API EMeshPaintVisualizeShowMode::Type GetMeshPaintVisualizeShowMod
 extern ENGINE_API void SetMeshPaintVisualizeChannels(EVertexColorViewMode::Type Channels);
 extern ENGINE_API EVertexColorViewMode::Type GetMeshPaintVisualizeChannels();
 extern ENGINE_API void SetMeshPaintVisualizeTexture(TWeakObjectPtr<UTexture> Texture);
+extern ENGINE_API FRHITexture* GetMeshPaintVisualizeTexture_RenderThread();
 extern ENGINE_API void SetMeshPaintVisualizeTextureCoordinateIndex(int32 Index);
+extern ENGINE_API int32 GetMeshPaintVisualizeTextureCoordinateIndex();
 /** Get the mesh paint visualization material proxy based on the current global settings. */
 extern ENGINE_API FMaterialRenderProxy* GetMeshPaintVisualizeMaterialRenderProxy(bool bIsSelected, bool bIsHovered);
 
