Changed csvstat CancelledMipsInKB to CancelledMipsInMB

#rnx

diff --git a/Engine/Source/Runtime/Engine/Private/Streaming/Texture2DStreamIn_IO.cpp b/Engine/Source/Runtime/Engine/Private/Streaming/Texture2DStreamIn_IO.cpp
--- a/Engine/Source/Runtime/Engine/Private/Streaming/Texture2DStreamIn_IO.cpp
+++ b/Engine/Source/Runtime/Engine/Private/Streaming/Texture2DStreamIn_IO.cpp
@@ -130,7 +130,7 @@ void FTexture2DStreamIn_IO::SetIORequests(const FContext& Context)
 #endif
 		if (bIsCancelled)
 		{
-			CSV_CUSTOM_STAT(TextureStreaming, CancelledMipsInKB, float(InFlightBulkSize) / 1024.0f, ECsvCustomStatOp::Accumulate);
+			CSV_CUSTOM_STAT(TextureStreaming, CancelledMipsInMB, float(InFlightBulkSize) / (1024.0f * 1024.0f), ECsvCustomStatOp::Accumulate);
 		}
 		// The tick here is intended to schedule the success or cancel callback.
 		// Using TT_None ensure gets which could create a dead lock.
