SkinnedMeshComponent: Better fix to Hidden Materials working with static view relevance.

This undoes CL 35506030 and instead adds a call to MarkRenderStateDirty() when changing which materials are hidden. This handles both hiding and unhiding (previous fix only fixed hiding).

diff --git a/Engine/Source/Runtime/Engine/Private/Components/SkinnedMeshComponent.cpp b/Engine/Source/Runtime/Engine/Private/Components/SkinnedMeshComponent.cpp
--- a/Engine/Source/Runtime/Engine/Private/Components/SkinnedMeshComponent.cpp
+++ b/Engine/Source/Runtime/Engine/Private/Components/SkinnedMeshComponent.cpp
@@ -3521,6 +3521,8 @@ void USkinnedMeshComponent::ShowMaterialSection(int32 MaterialID, int32 SectionI
 				InMeshObject->SetHiddenMaterials(LODIndex, HiddenMaterials);
 			});
 		}
+
+		MarkRenderStateDirty();
 	}
 }
 
@@ -3550,6 +3552,7 @@ void USkinnedMeshComponent::ShowAllMaterialSections(int32 LODIndex)
 						InMeshObject->SetHiddenMaterials(LODIndex, HiddenMaterials);
 					});
 			}
+			MarkRenderStateDirty();
 		}
 	}
 }
diff --git a/Engine/Source/Runtime/Engine/Private/SkeletalMesh.cpp b/Engine/Source/Runtime/Engine/Private/SkeletalMesh.cpp
--- a/Engine/Source/Runtime/Engine/Private/SkeletalMesh.cpp
+++ b/Engine/Source/Runtime/Engine/Private/SkeletalMesh.cpp
@@ -7418,7 +7418,6 @@ FPrimitiveViewRelevance FSkeletalMeshSceneProxy::GetViewRelevance(const FSceneVi
 			|| EngineShowFlags.VisualizeGPUSkinCache
 			|| (IsSelected() && (EngineShowFlags.VertexColors || EngineShowFlags.Bounds))
 			|| bForceWireframe
-			|| MeshObject->IsAnyMaterialHidden(MeshObject->GetLOD())
 #if WITH_EDITORONLY_DATA
 			|| MeshObject->SelectedEditorMaterial != -1
 			|| MeshObject->SelectedEditorSection != -1
diff --git a/Engine/Source/Runtime/Engine/Private/SkeletalRender.cpp b/Engine/Source/Runtime/Engine/Private/SkeletalRender.cpp
--- a/Engine/Source/Runtime/Engine/Private/SkeletalRender.cpp
+++ b/Engine/Source/Runtime/Engine/Private/SkeletalRender.cpp
@@ -247,25 +247,6 @@ bool FSkeletalMeshObject::IsMaterialHidden(int32 InLODIndex,int32 MaterialIdx) c
 	check(LODInfo.IsValidIndex(InLODIndex));
 	return LODInfo[InLODIndex].HiddenMaterials.IsValidIndex(MaterialIdx) && LODInfo[InLODIndex].HiddenMaterials[MaterialIdx];
 }
-
-/**
- * Determine if any material section entry for an LOD is hidden or not
- *
- * @param InLODIndex - LOD entry to get hidden material flags for
- */
-bool FSkeletalMeshObject::IsAnyMaterialHidden(int32 InLODIndex) const
-{
-	check(LODInfo.IsValidIndex(InLODIndex));
-	for (const bool bIsHidden : LODInfo[InLODIndex].HiddenMaterials)
-	{
-		if (bIsHidden)
-		{
-			return true;
-		}
-	}
-	return false;
-}
-
 /**
  * Initialize the array of LODInfo based on the settings of the current skel mesh component
  */
diff --git a/Engine/Source/Runtime/Engine/Public/SkeletalRenderPublic.h b/Engine/Source/Runtime/Engine/Public/SkeletalRenderPublic.h
--- a/Engine/Source/Runtime/Engine/Public/SkeletalRenderPublic.h
+++ b/Engine/Source/Runtime/Engine/Public/SkeletalRenderPublic.h
@@ -245,13 +245,6 @@ public:
 	 * @param MaterialIdx - index of the material section to check
 	 */
 	bool IsMaterialHidden(int32 InLODIndex,int32 MaterialIdx) const;
-
-	/**
-	 * Determine if any material section entry for an LOD is hidden or not
-	 *
-	 * @param InLODIndex - LOD entry to get hidden material flags for
-	 */
-	bool IsAnyMaterialHidden(int32 InLODIndex) const;
 	
 	/**
 	 * Initialize the array of LODInfo based on the settings of the current skel mesh component
