Add analytic derivatives support for GetPixelPosition / CustomPrimitiveData / dependent texture reads when the derivatives are 0

#rnx

diff --git a/Engine/Source/Runtime/Engine/Private/Materials/HLSLMaterialTranslator.cpp b/Engine/Source/Runtime/Engine/Private/Materials/HLSLMaterialTranslator.cpp
--- a/Engine/Source/Runtime/Engine/Private/Materials/HLSLMaterialTranslator.cpp
+++ b/Engine/Source/Runtime/Engine/Private/Materials/HLSLMaterialTranslator.cpp
@@ -5952,7 +5952,17 @@ int32 FHLSLMaterialTranslator::GetPixelPosition()
 	{
 		return Errorf(TEXT("GetPixelPosition() node is only available in vertex or pixel shader input."));
 	}
-	return AddCodeChunk(MCT_Float2, TEXT("GetPixelPosition(Parameters)"));
+
+	FString FiniteCode = TEXT("GetPixelPosition(Parameters)");
+	if (IsAnalyticDerivEnabled())
+	{
+		FString AnalyticCode = DerivativeAutogen.ConstructDeriv(FiniteCode, TEXT("float2(1.0f, 0.0f)"), TEXT("float2(0.0f, 1.0f)"), EDerivativeType::Float2);
+		return AddCodeChunkInnerDeriv(*FiniteCode, *AnalyticCode, MCT_Float2, false, EDerivativeStatus::Valid);
+	}
+	else
+	{
+		return AddCodeChunk(MCT_Float2, *FiniteCode);
+	}
 }
 
 int32 FHLSLMaterialTranslator::ParticleMacroUV()
@@ -7572,7 +7582,10 @@ int32 FHLSLMaterialTranslator::TextureSample(
 			}
 
 			SampleCodeAnalytic = ApplySamplerType(SampleCodeAnalytic, SamplerType);
-			SamplingCodeIndex = AddCodeChunkInnerDeriv(*SampleCodeFinite, *SampleCodeAnalytic, MCT_Float4, false, EDerivativeStatus::NotValid);
+
+			EDerivativeStatus TextureSampleDerivativeStatus = (UvDerivativeStatus == EDerivativeStatus::Zero) ? EDerivativeStatus::Zero : EDerivativeStatus::NotValid;
+
+			SamplingCodeIndex = AddCodeChunkInnerDeriv(*SampleCodeFinite, *SampleCodeAnalytic, MCT_Float4, false, TextureSampleDerivativeStatus);
 		}
 		else
 		{
@@ -7633,7 +7646,9 @@ int32 FHLSLMaterialTranslator::TextureSample(
 				SampleCodeAnalytic = ApplySamplerType(SampleCodeAnalytic, SamplerType);
 			}
 
-			SamplingCodeIndex = AddCodeChunkInnerDeriv(*SampleCodeFinite, *SampleCodeAnalytic, MCT_Float4, false /*?*/, EDerivativeStatus::NotValid);
+			EDerivativeStatus TextureSampleDerivativeStatus = (UvDerivativeStatus == EDerivativeStatus::Zero) ? EDerivativeStatus::Zero : EDerivativeStatus::NotValid;
+
+			SamplingCodeIndex = AddCodeChunkInnerDeriv(*SampleCodeFinite, *SampleCodeAnalytic, MCT_Float4, false /*?*/, TextureSampleDerivativeStatus);
 		}
 		else
 		{
@@ -12010,7 +12025,7 @@ int32 FHLSLMaterialTranslator::CustomPrimitiveData(int32 OutputIndex, EMaterialV
 		HlslCode.AppendChar(')');
 	}
 
-	return AddCodeChunk(Type, TEXT("%s"), *HlslCode);
+	return AddCodeChunkZeroDeriv(Type, TEXT("%s"), *HlslCode);
 }
 
 int32 FHLSLMaterialTranslator::ShadingModel(EMaterialShadingModel InSelectedShadingModel)
