Force DefautLit shading model, on platfroms that has per-pixel shadiing models disabled.
Trying to choose a shading model from compiled ones do not always work correctly

diff --git a/Engine/Shaders/Private/MobileBasePassCommon.ush b/Engine/Shaders/Private/MobileBasePassCommon.ush
--- a/Engine/Shaders/Private/MobileBasePassCommon.ush
+++ b/Engine/Shaders/Private/MobileBasePassCommon.ush
@@ -4,6 +4,9 @@
 	MobileBasePassCommon.ush: Base pass definitions used by both vertex and pixel shader
 =============================================================================*/
 
+// Just used to trigger MobileBase pass shader recompiles, guid should be regenerated on merge conflicts
+#pragma message("UESHADERMETADATA_VERSION C1A9D426-8014-4723-8C69-E32EA8808D15")
+
 #include "BasePassCommon.ush"
 #define FogStruct MobileBasePass.Fog
 #define PlanarReflectionStruct MobileBasePass.PlanarReflection
diff --git a/Engine/Source/Runtime/Engine/Private/Materials/Material.cpp b/Engine/Source/Runtime/Engine/Private/Materials/Material.cpp
--- a/Engine/Source/Runtime/Engine/Private/Materials/Material.cpp
+++ b/Engine/Source/Runtime/Engine/Private/Materials/Material.cpp
@@ -261,8 +261,8 @@ int32 FMaterialResource::CompilePropertyAndSetMaterialProperty(EMaterialProperty
 			}
 			else
 			{
-				FMaterialShadingModelField ShadingModels = Compiler->GetMaterialShadingModels();
-				Ret = Compiler->ShadingModel(ShadingModels.GetFirstShadingModel());
+				// if platform does not support per-pixel shading models, always fallback to DefaultLit
+				Ret = Compiler->ShadingModel(MSM_DefaultLit);
 			}
 			break;
 		case MP_MaterialAttributes:
@@ -4131,18 +4131,23 @@ void UMaterial::PostLoad()
 		bUseFullPrecision_DEPRECATED = false;
 	}
 
-	if (!GIsEditor)
+#if !WITH_EDITOR
+
+	if (GetFullName().Contains(TEXT("clearcoat_bn")))
 	{
-		// Filter out ShadingModels field to a current platform settings
-		FilterOutPlatformShadingModels(GMaxRHIShaderPlatform, ShadingModels);
-		// Override material shader model if it was filtered out
-		if (!ShadingModels.HasShadingModel(ShadingModel) &&
-			(ShadingModel != MSM_FromMaterialExpression || !AllowPerPixelShadingModels(GMaxRHIShaderPlatform)))
-		{
-			ShadingModel = ShadingModels.GetFirstShadingModel();
-		}
+		UE_LOG(LogMaterial, Display, TEXT("DEBUG break"));
 	}
 
+	// Filter out ShadingModels field to a current platform settings
+	FilterOutPlatformShadingModels(GMaxRHIShaderPlatform, ShadingModels);
+	// Override material shader model if it was filtered out
+	if (!ShadingModels.HasShadingModel(ShadingModel) &&
+		(ShadingModel != MSM_FromMaterialExpression || !AllowPerPixelShadingModels(GMaxRHIShaderPlatform)))
+	{
+		ShadingModel = ShadingModels.GetFirstShadingModel();
+	}
+#endif
+
 #if WITH_EDITOR
 	// Create exec flow expressions, if needed
 	CreateExecutionFlowExpressions();
diff --git a/Engine/Source/Runtime/Engine/Private/Materials/MaterialInstance.cpp b/Engine/Source/Runtime/Engine/Private/Materials/MaterialInstance.cpp
--- a/Engine/Source/Runtime/Engine/Private/Materials/MaterialInstance.cpp
+++ b/Engine/Source/Runtime/Engine/Private/Materials/MaterialInstance.cpp
@@ -2405,11 +2405,10 @@ void UMaterialInstance::UpdateOverridableBaseProperties()
 		BasePropertyOverrides.BlendMode = BlendMode;
 	}
 
-	if (!GIsEditor)
-	{
-		// Filter out ShadingModels field to a current platform settings
-		FilterOutPlatformShadingModels(GMaxRHIShaderPlatform, ShadingModels);
-	}
+#if !WITH_EDITOR
+	// Filter out ShadingModels field to a current platform settings
+	FilterOutPlatformShadingModels(GMaxRHIShaderPlatform, ShadingModels);
+#endif
 
 	if (BasePropertyOverrides.bOverride_TwoSided)
 	{
@@ -3374,11 +3373,10 @@ void UMaterialInstance::PostLoad()
 		}
 	}
 
-	if (!GIsEditor)
-	{
-		// Filter out ShadingModels field to a current platform settings
-		FilterOutPlatformShadingModels(GMaxRHIShaderPlatform, ShadingModels);
-	}
+#if !WITH_EDITOR
+	// Filter out ShadingModels field to a current platform settings
+	FilterOutPlatformShadingModels(GMaxRHIShaderPlatform, ShadingModels);
+#endif
 
 #if WITH_EDITOR
 	UpdateCachedData();
diff --git a/Engine/Source/Runtime/Engine/Private/Materials/MaterialInterface.cpp b/Engine/Source/Runtime/Engine/Private/Materials/MaterialInterface.cpp
--- a/Engine/Source/Runtime/Engine/Private/Materials/MaterialInterface.cpp
+++ b/Engine/Source/Runtime/Engine/Private/Materials/MaterialInterface.cpp
@@ -1891,7 +1891,7 @@ void UMaterialInterface::FilterOutPlatformShadingModels(const EShaderPlatform In
 	const FStaticShaderPlatform Platform(InPlatform);
 	if (ShadingModels.CountShadingModels() > 1 && !AllowPerPixelShadingModels(Platform))
 	{
-		ShadingModels = FMaterialShadingModelField(ShadingModels.GetFirstShadingModel());
+		ShadingModels = FMaterialShadingModelField(MSM_DefaultLit);
 	}
 
 	uint32 ShadingModelsMask = GetPlatformShadingModelsMask(Platform);
@@ -1901,7 +1901,7 @@ void UMaterialInterface::FilterOutPlatformShadingModels(const EShaderPlatform In
 		ShadingModels.SetShadingModelField(FilteredShadingModels);
 		if (!ShadingModels.IsValid())
 		{
-			ShadingModels.AddShadingModel(MSM_DefaultLit);
+			ShadingModels = FMaterialShadingModelField(MSM_DefaultLit);
 		}
 	}
 }
