[NetPhysics] Improved handling of client server desync.
- Inputs sent from client and server are sent from latest processed frame, with redundant inputs back in time
- If inputs arrive late to the server and the buffer is empty, inject the input for the next tick
- If client receives target states from the server before the client has simulated the corresponding tick, compare that state after the client has simulated the tick
- Deactivate logic that checks if all NetworkPhysicsComponents have inputs and states when performing a rewind, this often stopped resimulations from triggering due to not all sim-proxies having all data aligned for one single frame.
- Allow larger time accumulation towards fixed physics tick when having physics prediction enabled
- Rebalance time dilation and physics tick sync corrections
- Base time dilation settings on time rather than physics ticks

diff --git a/Engine/Source/Runtime/Engine/Classes/GameFramework/PlayerController.h b/Engine/Source/Runtime/Engine/Classes/GameFramework/PlayerController.h
--- a/Engine/Source/Runtime/Engine/Classes/GameFramework/PlayerController.h
+++ b/Engine/Source/Runtime/Engine/Classes/GameFramework/PlayerController.h
@@ -2428,6 +2428,9 @@ private:
 	/** The latest physics step we've received from the client. */
 	int32 ServerLatestAsyncPhysicsStepReceived = INDEX_NONE;
 
+	/** Time that tick offset has been desynced for since it was last in sync, used to trigger a correction */
+	float NetworkPhysicsTickOffsetDesyncAccumulatedTime = 0.0f;
+
 public:
 
 	/** Update the tick offset in between the local client and the server */
diff --git a/Engine/Source/Runtime/Engine/Private/PhysicsEngine/NetworkPhysicsComponent.cpp b/Engine/Source/Runtime/Engine/Private/PhysicsEngine/NetworkPhysicsComponent.cpp
--- a/Engine/Source/Runtime/Engine/Private/PhysicsEngine/NetworkPhysicsComponent.cpp
+++ b/Engine/Source/Runtime/Engine/Private/PhysicsEngine/NetworkPhysicsComponent.cpp
@@ -214,6 +214,9 @@ int32 FNetworkPhysicsCallback::TriggerRewindIfNeeded_Internal(int32 LatestStepCo
 
 	if (RewindData)
 	{
+		const int32 TargetStateComparisonFrame = RewindData->CompareTargetsToLastFrame();
+		ResimFrame = (ResimFrame == INDEX_NONE) ? TargetStateComparisonFrame : (TargetStateComparisonFrame == INDEX_NONE) ? ResimFrame : FMath::Min(TargetStateComparisonFrame, ResimFrame);
+
 		const int32 ReplicationFrame = RewindData->GetResimFrame();
 		ResimFrame = (ResimFrame == INDEX_NONE) ? ReplicationFrame : (ReplicationFrame == INDEX_NONE) ? ResimFrame : FMath::Min(ReplicationFrame, ResimFrame);
 
@@ -221,7 +224,7 @@ int32 FNetworkPhysicsCallback::TriggerRewindIfNeeded_Internal(int32 LatestStepCo
 		{
 			const int32 ValidFrame = RewindData->FindValidResimFrame(ResimFrame);
 #if DEBUG_NETWORK_PHYSICS || DEBUG_REWIND_DATA
-			UE_LOG(LogChaos, Log, TEXT("CLIENT | PT | TriggerRewindIfNeeded_Internal | Requested Resim Frame = %d | Valid Resim Frame = %d"), ResimFrame, ValidFrame);
+			UE_LOG(LogChaos, Log, TEXT("CLIENT | PT | TriggerRewindIfNeeded_Internal | Requested Resim Frame = %d (%d / %d) | Valid Resim Frame = %d"), ResimFrame, TargetStateComparisonFrame, ReplicationFrame, ValidFrame);
 #endif
 			ResimFrame = ValidFrame;
 		}
@@ -682,6 +685,7 @@ void UNetworkPhysicsComponent::OnRep_SetReplicatedStates()
 			AsyncInput->StateData = StateHelper->CreateUniqueRewindHistory(StateRedundancy);
 		}
 
+		AsyncInput->StateData->ResetFast();
 		ReplicatedStates.History->CopyAllData(*AsyncInput->StateData.Get(), /*bIncludeUnimportant*/ true, /*bIncludeImportant*/ true);
 	}
 
@@ -718,6 +722,7 @@ void UNetworkPhysicsComponent::OnRep_SetReplicatedInputs()
 			AsyncInput->InputData = InputHelper->CreateUniqueRewindHistory(InputRedundancy);
 		}
 
+		AsyncInput->InputData->ResetFast();
 		ReplicatedInputs.History->CopyAllData(*AsyncInput->InputData.Get(), /*bIncludeUnimportant*/ true, /*bIncludeImportant*/ true);
 	}
 }
@@ -748,6 +753,7 @@ void UNetworkPhysicsComponent::ServerReceiveInputData_Implementation(const FNetw
 			ClientInputs.History->ValidateDataInHistory(ActorComponent);
 		}
 
+		AsyncInput->InputData->ResetFast();
 		ClientInputs.History->CopyAllData(*AsyncInput->InputData.Get(), /*bIncludeUnimportant*/ true, /*bIncludeImportant*/ true);
 	}
 }
@@ -1080,7 +1086,7 @@ const FNetworkPhysicsSettingsNetworkPhysicsComponent& FAsyncNetworkPhysicsCompon
 	return SettingsComponent ? SettingsComponent->Settings.NetworkPhysicsComponentSettings : SettingsNetworkPhysicsComponent_Default;
 };
 
-void FAsyncNetworkPhysicsComponent::ConsumeAsyncInput()
+void FAsyncNetworkPhysicsComponent::ConsumeAsyncInput(const int32 PhysicsStep)
 {
 	if (const FAsyncNetworkPhysicsComponentInput* AsyncInput = GetConsumerInput_Internal())
 	{
@@ -1129,8 +1135,9 @@ void FAsyncNetworkPhysicsComponent::ConsumeAsyncInput()
 				// Setup rewind data if not already done, and get history size
 				const int32 NumFrames = SetupRewindData();
 
-				// Create input data and history
+				// Create input history and local data properties
 				InputData = (*AsyncInput->InputHelper)->CreateUniqueData();
+				LatestInputReceiveData = (*AsyncInput->InputHelper)->CreateUniqueData();
 				InputHistory = MakeShareable((*AsyncInput->InputHelper)->CreateUniqueRewindHistory(NumFrames).Release());
 				RegisterDataHistoryInRewindData();
 			}
@@ -1139,7 +1146,7 @@ void FAsyncNetworkPhysicsComponent::ConsumeAsyncInput()
 				// Setup rewind data if not already done, and get history size
 				const int32 NumFrames = SetupRewindData();
 
-				// Create state data and history
+				// Create state history and local property
 				StateData = (*AsyncInput->StateHelper)->CreateUniqueData();
 				StateHistory = MakeShareable((*AsyncInput->StateHelper)->CreateUniqueRewindHistory(NumFrames).Release());
 				RegisterDataHistoryInRewindData();
@@ -1157,23 +1164,61 @@ void FAsyncNetworkPhysicsComponent::ConsumeAsyncInput()
 				{
 					TriggerResimulation(ResimFrame);
 				}
+
 #if DEBUG_NETWORK_PHYSICS
 				{
 					FString NetRoleString = IsServer() ? FString("SERVER") : (IsLocallyControlled() ? FString("AUTONO") : FString("PROXY "));
-					ReceiveData->DebugData(FString::Printf(TEXT("%s | PT | RECEIVE DATA | bImportant: %d | Name: %s"), *NetRoleString, bImportant, *GetActorName()));
+					ReceiveData->DebugData(FString::Printf(TEXT("%s | PT | RECEIVE DATA | LatestFrame: %d | bImportant: %d | Name: %s"), *NetRoleString, ReceiveData->GetLatestFrame(), bImportant, *GetActorName()));
 				}
 #endif
+
+				// Reset the received data after having consumed it
+				ReceiveData->ResetFast();
 			};
 
 			// Receive Inputs
 			if (AsyncInput->InputData && AsyncInput->InputData->HasDataInHistory())
 			{
+				/* Extract latest received input from client on the server, to be used if input buffer runs empty
+				* TODO, improve flow to not require this before ReceiveHelper */
+				FNetworkPhysicsData* PhysicsData = nullptr;
+				if (IsServer() && LatestInputReceiveData)
+				{
+					PhysicsData = LatestInputReceiveData.Get();
+					if (AsyncInput->InputData->ExtractData(AsyncInput->InputData->GetLatestFrame(), false, PhysicsData, true) == false)
+					{
+						// Extraction failed
+						ensureMsgf(false, TEXT("Failed to extract latest input data from received inputs"));
+						PhysicsData = nullptr;
+
+#if DEBUG_NETWORK_PHYSICS
+						AsyncInput->InputData->DebugData(FString::Printf(TEXT("SERVER | PT | Failed to extract | LatestFrame: %d | Name: %s"), AsyncInput->InputData->GetLatestFrame(), *GetActorName()));
+#endif
+					}
+				}
+
 				// Validate data in the received inputs on the server
 				if (!ComponentSettings.GetValidateDataOnGameThread() && IsServer() && ActorComponent.IsValid() && ActorComponent.Get()->IsBeingDestroyed() == false)
 				{
 					AsyncInput->InputData->ValidateDataInHistory(ActorComponent.Get());
 				}
 				ReceiveHelper(InputHistory.Get(), AsyncInput->InputData.Get(), /*bImportant*/false, ComponentSettings.GetCompareInputToTriggerRewind());
+
+				/* If the server-side input history doesn't have any entries ahead of the current physics tick, the input buffer is empty, inject the latest received input as the input for the current tick.
+				* This happens during a desync where the client is far behind the server */
+				if (IsServer() && InputHistory->GetLatestFrame() < PhysicsStep && PhysicsData)
+				{
+#if DEBUG_NETWORK_PHYSICS
+					UE_LOG(LogChaos, Log, TEXT("SERVER | PT | Input Buffer Empty, Injecting Received Input at frame %d || LocalFrame = %d || ServerFrame = %d || InputFrame = %d || Data: %s || Actor: %s")
+						, PhysicsStep, PhysicsData->LocalFrame, PhysicsData->ServerFrame, PhysicsData->InputFrame, *PhysicsData->DebugData(), *GetActorName());
+#endif
+					
+					// Record data in InputHistory
+					PhysicsData->LocalFrame = PhysicsStep;
+					PhysicsData->ServerFrame = PhysicsStep;
+					InputHistory->RecordData(PhysicsStep, PhysicsData);
+				}
+
 			}
 
 			// Receive States
@@ -1229,7 +1274,7 @@ FAsyncNetworkPhysicsComponentOutput& FAsyncNetworkPhysicsComponent::GetAsyncOutp
 
 void FAsyncNetworkPhysicsComponent::OnPreProcessInputs_Internal(const int32 PhysicsStep)
 {
-	ConsumeAsyncInput();
+	ConsumeAsyncInput(PhysicsStep);
 
 	const FNetworkPhysicsSettingsNetworkPhysicsComponent& ComponentSettings = GetComponentSettings();
 	const bool bIsServer = IsServer();
@@ -1400,8 +1445,8 @@ void FAsyncNetworkPhysicsComponent::OnPostProcessInputs_Internal(const int32 Phy
 
 	if (ActorComponent.IsValid() && ActorComponent.Get()->IsBeingDestroyed() == false)
 	{
+		// Cache current input if we are locally controlled
 		const bool bShouldCacheInputHistory = IsLocallyControlled() && !bIsSolverResim;
-		// For the inputs client local ones are ground truth otherwise use the replicated ones coming from the server
 		if (bShouldCacheInputHistory && (InputData != nullptr))
 		{
 			// Prepare to gather input data
@@ -1422,6 +1467,7 @@ void FAsyncNetworkPhysicsComponent::OnPostProcessInputs_Internal(const int32 Phy
 #endif
 		}
 
+		// Cache current state if this is the server of we are comparing predicted states on autonomous proxy
 		const bool bShouldCacheStateHistory = bIsServer || (ComponentSettings.GetCompareStateToTriggerRewind() && bShouldCacheInputHistory);
 		if (StateHistory && StateData && bShouldCacheStateHistory)
 		{
@@ -1458,12 +1504,12 @@ void FAsyncNetworkPhysicsComponent::OnPostProcessInputs_Internal(const int32 Phy
 
 	// Marshal inputs and states from PT to GT for networking
 	FAsyncNetworkPhysicsComponentOutput& AsyncOutput = GetAsyncOutput_Internal();
-	SendInputData_Internal(AsyncOutput);
-	SendStateData_Internal(AsyncOutput);
+	SendInputData_Internal(AsyncOutput, PhysicsStep);
+	SendStateData_Internal(AsyncOutput, PhysicsStep);
 	FinalizeOutputData_Internal();
 }
 
-void FAsyncNetworkPhysicsComponent::SendInputData_Internal(FAsyncNetworkPhysicsComponentOutput& AsyncOutput)
+void FAsyncNetworkPhysicsComponent::SendInputData_Internal(FAsyncNetworkPhysicsComponentOutput& AsyncOutput, const int32 PhysicsStep)
 {
 	const bool bIsServer = IsServer();
 
@@ -1473,7 +1519,7 @@ void FAsyncNetworkPhysicsComponent::SendInputData_Internal(FAsyncNetworkPhysicsC
 		const FNetworkPhysicsSettingsNetworkPhysicsComponent& ComponentSettings = GetComponentSettings();
 
 		// Send latest N frames from history
-		const int32 ToFrame = FMath::Max(0, InputHistory->GetLatestFrame());
+		const int32 ToFrame = FMath::Max(0, PhysicsStep);
 
 		// -- Default / Unreliable Flow --
 		if (ComponentSettings.GetEnableUnreliableFlow())
@@ -1532,14 +1578,14 @@ void FAsyncNetworkPhysicsComponent::SendInputData_Internal(FAsyncNetworkPhysicsC
 	}
 }
 
-void FAsyncNetworkPhysicsComponent::SendStateData_Internal(FAsyncNetworkPhysicsComponentOutput& AsyncOutput)
+void FAsyncNetworkPhysicsComponent::SendStateData_Internal(FAsyncNetworkPhysicsComponentOutput& AsyncOutput, const int32 PhysicsStep)
 {
 	if (IsServer() && StateHistory && AsyncOutput.StateData)
 	{
 		const FNetworkPhysicsSettingsNetworkPhysicsComponent& ComponentSettings = GetComponentSettings();
 
 		// Send latest N frames from history
-		const int32 ToFrame = FMath::Max(0, StateHistory->GetLatestFrame());
+		const int32 ToFrame = FMath::Max(0, PhysicsStep);
 
 		// -- Default / Unreliable Flow --
 		if (ComponentSettings.GetEnableUnreliableFlow())
diff --git a/Engine/Source/Runtime/Engine/Private/PhysicsEngine/PhysLevel.cpp b/Engine/Source/Runtime/Engine/Private/PhysicsEngine/PhysLevel.cpp
--- a/Engine/Source/Runtime/Engine/Private/PhysicsEngine/PhysLevel.cpp
+++ b/Engine/Source/Runtime/Engine/Private/PhysicsEngine/PhysLevel.cpp
@@ -146,8 +146,24 @@ void UWorld::SetupPhysicsTickFunctions(float DeltaSeconds)
 	FVector DefaultGravity( 0.f, 0.f, GetGravityZ() );
 
 	static const auto CVar_MaxPhysicsDeltaTime = IConsoleManager::Get().FindTConsoleVariableDataFloat(TEXT("p.MaxPhysicsDeltaTime"));
-	PhysScene->SetUpForFrame(&DefaultGravity, DeltaSeconds, UPhysicsSettings::Get()->MinPhysicsDeltaTime, UPhysicsSettings::Get()->MaxPhysicsDeltaTime,
-		UPhysicsSettings::Get()->MaxSubstepDeltaTime, UPhysicsSettings::Get()->MaxSubsteps, UPhysicsSettings::Get()->bSubstepping);
+
+	UPhysicsSettings* Settings = UPhysicsSettings::Get();
+
+	float MinPhysicsDeltaTime = Settings->MinPhysicsDeltaTime;
+	float MaxPhysicsDeltaTime = Settings->MaxPhysicsDeltaTime;
+	float MaxSubstepDeltaTime = Settings->MaxSubstepDeltaTime;
+
+	/* When using physics prediction, allow max delta time at least equal to max supported latency for physics prediction
+	* NOTE: These values clamp how much game thread delta time that can accumulate towards ticking the fixed physics steps, 
+	* if we clamp this accumulation too much we hinder time dilation from correcting desyncs and we also make the client (and server) more prone to desyncing physics due to dropping physics steps. */
+	if (Settings->PhysicsPrediction.bEnablePhysicsPrediction && Settings->bTickPhysicsAsync)
+	{
+		MinPhysicsDeltaTime = 0.0f;
+		MaxPhysicsDeltaTime = MaxPhysicsDeltaTime <= UE_SMALL_NUMBER ? 0.0f : FMath::Max((Settings->PhysicsPrediction.MaxSupportedLatencyPrediction / 1000.0f), MaxPhysicsDeltaTime);
+		MaxSubstepDeltaTime = MaxSubstepDeltaTime <= UE_SMALL_NUMBER ? 0.0f : FMath::Max((Settings->PhysicsPrediction.MaxSupportedLatencyPrediction / 1000.0f) / Settings->MaxSubsteps, MaxSubstepDeltaTime);
+	}
+
+	PhysScene->SetUpForFrame(&DefaultGravity, DeltaSeconds, MinPhysicsDeltaTime, MaxPhysicsDeltaTime, MaxSubstepDeltaTime, Settings->MaxSubsteps, Settings->bSubstepping);
 }
 
 void UWorld::StartPhysicsSim()
diff --git a/Engine/Source/Runtime/Engine/Private/PlayerController.cpp b/Engine/Source/Runtime/Engine/Private/PlayerController.cpp
--- a/Engine/Source/Runtime/Engine/Private/PlayerController.cpp
+++ b/Engine/Source/Runtime/Engine/Private/PlayerController.cpp
@@ -116,32 +116,47 @@ namespace UE::Gameplay::CVars
 
 namespace NetworkPhysicsCvars
 {
-	/* DEPRECATED 5.4 */
-	int32 NumRedundantCmds = 3;
-	FAutoConsoleVariableRef CVarNumRedundantCmds(TEXT("np2.NumRedundantCmds"), NumRedundantCmds, TEXT("(DEPRECATED 5.4, only part of the legacy physics frame offset logic) Number of redundant user cmds to send per frame"));
+		/* DEPRECATED 5.4 */
+		int32 NumRedundantCmds = 3;
+		FAutoConsoleVariableRef CVarNumRedundantCmds(TEXT("np2.NumRedundantCmds"), NumRedundantCmds, TEXT("(DEPRECATED 5.4, only part of the legacy physics frame offset logic) Number of redundant user cmds to send per frame"));
 
 #if (UE_BUILD_SHIPPING || UE_BUILD_TEST)
-	int32 EnableDebugRPC = 0;
+		int32 EnableDebugRPC = 0;
 #else
-	int32 EnableDebugRPC = 1;
+		int32 EnableDebugRPC = 1;
 #endif
-	/* DEPRECATED 5.4 */
-	FAutoConsoleVariableRef CVarEnableDebugRPC(TEXT("np2.EnableDebugRPC"), EnableDebugRPC, TEXT("(DEPRECATED 5.4, only part of the legacy physics frame offset logic) Sends extra debug information to clients about server side input buffering"));
-	
-	/* DEPRECATED 5.4 */
-	int32 NetworkPhysicsPredictionFrameOffset = 4;
-	FAutoConsoleVariableRef CVarNetworkPhysicsPredictionFrameOffset(TEXT("np2.NetworkPhysicsPredictionFrameOffset"), NetworkPhysicsPredictionFrameOffset, TEXT("(DEPRECATED 5.4, use np2.PredictionAsyncFrameBuffer instead) Additional frame offset to be added to the local to server offset used by network prediction"));
-	
-	int32 PredictionAsyncFrameBuffer = 3;
-	FAutoConsoleVariableRef CVarPredictionAsyncFrameBuffer(TEXT("np2.PredictionAsyncFrameBuffer"), PredictionAsyncFrameBuffer, TEXT("Additional frame offset to be added to the local to server offset used by network prediction"));
+		/* DEPRECATED 5.4 */
+		FAutoConsoleVariableRef CVarEnableDebugRPC(TEXT("np2.EnableDebugRPC"), EnableDebugRPC, TEXT("(DEPRECATED 5.4, only part of the legacy physics frame offset logic) Sends extra debug information to clients about server side input buffering"));
 
-	int32 TickOffsetUpdateInterval = 10;
-	FAutoConsoleVariableRef CVarTickOffsetUpdateInterval(TEXT("np2.TickOffsetUpdateInterval"), TickOffsetUpdateInterval, TEXT("How many physics ticks to wait between each tick offset update. Lowest viable value = 1, which means update each tick. Deactivate physics offset updates by setting to 0 or negative value."));
-	
-	int32 TickOffsetCorrectionLimit = 10;
-	FAutoConsoleVariableRef CVarTickOffsetCorrectionLimit(TEXT("np2.TickOffsetCorrectionLimit"), TickOffsetCorrectionLimit, TEXT("If the client gets out of sync with physics ticks more than this limit, cut the losses and reset the offset."));
-	
-	float TimeDilationAmount = 0.01f;
+		/* DEPRECATED 5.4 */
+		int32 NetworkPhysicsPredictionFrameOffset = 4;
+		FAutoConsoleVariableRef CVarNetworkPhysicsPredictionFrameOffset(TEXT("np2.NetworkPhysicsPredictionFrameOffset"), NetworkPhysicsPredictionFrameOffset, TEXT("(DEPRECATED 5.4, use np2.PredictionAsyncFrameBuffer instead) Additional frame offset to be added to the local to server offset used by network prediction"));
+
+		/* DEPRECATED 5.5 */
+		int32 PredictionAsyncFrameBuffer = 3;
+		FAutoConsoleVariableRef CVarPredictionAsyncFrameBuffer(TEXT("np2.PredictionAsyncFrameBuffer"), PredictionAsyncFrameBuffer, TEXT("(DEPRECATED 5.5, Use np2.TickOffsetBufferTime instead) Additional frame offset to be added to the local to server offset used by network prediction"));
+
+		/* DEPRECATED 5.5 */
+		int32 TickOffsetUpdateInterval = 10;
+		FAutoConsoleVariableRef CVarTickOffsetUpdateInterval(TEXT("np2.TickOffsetUpdateInterval"), TickOffsetUpdateInterval, TEXT("(DEPRECATED 5.5, Use np2.TickOffsetUpdateIntervalTime instead) How many physics ticks to wait between each tick offset update. Lowest viable value = 1, which means update each tick. Deactivate physics offset updates by setting to 0 or negative value."));
+
+		/* DEPRECATED 5.5 */
+		int32 TickOffsetCorrectionLimit = 10;
+		FAutoConsoleVariableRef CVarTickOffsetCorrectionLimit(TEXT("np2.TickOffsetCorrectionLimit"), TickOffsetCorrectionLimit, TEXT("(DEPRECATED 5.5, Use np2.TickOffsetCorrectionSizeTimeLimit instead) If the client gets out of sync with physics ticks more than this limit, cut the losses and reset the offset."));
+
+	int32 TickOffsetUpdateIntervalTime = 100;
+	FAutoConsoleVariableRef CVarTickOffsetUpdateIntervalTime(TEXT("np2.TickOffsetUpdateIntervalTime"), TickOffsetUpdateIntervalTime, TEXT("Value in milliseconds, default 100. How long time between syncing the tick offset between client and server. Deactivate syncing by setting value 0."));
+
+	int32 TickOffsetBufferTime = 60;
+	FAutoConsoleVariableRef CVarTickOffsetBufferTime(TEXT("np2.TickOffsetBufferTime"), TickOffsetBufferTime, TEXT("Value in milliseconds, default 60. Additional offset to be added to the local to server offset used by network prediction, this results in a buffer server-side for incoming data that uses the client/server physics offset."));
+
+	int32 TickOffsetCorrectionSizeTimeLimit = 200;
+	FAutoConsoleVariableRef CVarTickOffsetCorrectionSizeTimeLimit(TEXT("np2.TickOffsetCorrectionSizeTimeLimit"), TickOffsetCorrectionSizeTimeLimit, TEXT("Value in milliseconds, Default 200. Note: Keep this equal to or larger than np2.TickOffsetBufferTime. If the client gets out of sync with physics ticks and the desync is larger than this value, reset the offset."));
+
+	int32 TickOffsetCorrectionTimeLimit = 1000;
+	FAutoConsoleVariableRef CVarTickOffsetCorrectionTimeLimit(TEXT("np2.TickOffsetCorrectionTimeLimit"), TickOffsetCorrectionTimeLimit, TEXT("Value in milliseconds, Default 1000. If the client gets out of sync with physics ticks and can't get in sync again for this amount of time, reset the offset."));
+
+	float TimeDilationAmount = 0.02f;
 	FAutoConsoleVariableRef CVarTimeDilationAmount(TEXT("np2.TimeDilationAmount"), TimeDilationAmount, TEXT("Server-side CVar, Disable TimeDilation by setting to 0 | Default: 0.01 | Value is in percent where 0.01 = 1% dilation. Example: 1.0/0.01 = 100, meaning that over the time it usually takes to tick 100 physics steps we will tick 99 or 101 depending on if we dilate up or down."));
 
 	bool TimeDilationEscalation = true;
@@ -153,10 +168,10 @@ namespace NetworkPhysicsCvars
 	float TimeDilationEscalationDecayMax = 0.5f;
 	FAutoConsoleVariableRef CVarTimeDilationEscalationDecayMax(TEXT("np2.TimeDilationEscalationDecayMax"), TimeDilationEscalationDecayMax, TEXT("Value is a multiplier, Default: 0.5. The max decay value for escalated time dilation. Lower value means higher decay."));
 
-	float TimeDilationMax = 1.1f;
+	float TimeDilationMax = 1.25f;
 	FAutoConsoleVariableRef CVarTimeDilationMax(TEXT("np2.TimeDilationMax"), TimeDilationMax, TEXT("Max value of the time dilation multiplier."));
 
-	float TimeDilationMin = 0.9f;
+	float TimeDilationMin = 0.75f;
 	FAutoConsoleVariableRef CVarTimeDilationMin(TEXT("np2.TimeDilationMin"), TimeDilationMin, TEXT("Min value of the time dilation multiplier"));
 }
 
@@ -6224,7 +6239,9 @@ FAsyncPhysicsTimestamp APlayerController::GetPhysicsTimestamp(float DeltaSeconds
 void APlayerController::UpdateServerAsyncPhysicsTickOffset()
 {
 	FAsyncPhysicsTimestamp Timestamp = GetPhysicsTimestamp();
-	if (NetworkPhysicsCvars::TickOffsetUpdateInterval <= 0 || ClientLatestAsyncPhysicsStepSent + NetworkPhysicsCvars::TickOffsetUpdateInterval > Timestamp.LocalFrame)
+	const int32 TickOffsetUpdateIntervalSteps = (NetworkPhysicsCvars::TickOffsetUpdateIntervalTime / 1000.f) / UPhysicsSettings::Get()->AsyncFixedTimeStepSize;
+
+	if (TickOffsetUpdateIntervalSteps <= 0 || ClientLatestAsyncPhysicsStepSent + TickOffsetUpdateIntervalSteps > Timestamp.LocalFrame)
 	{
 		//Only send a new timestamp if enough physics ticks have passed, based on CVar.
 		//If GT is running faster than physics sim the physics timestep will not have changed, so no need to send another update to server
@@ -6260,18 +6277,27 @@ void APlayerController::ServerSendLatestAsyncPhysicsTimestamp_Implementation(FAs
 
 	// Get current server timestamp and add the frame buffer to the ServerFrame
 	FAsyncPhysicsTimestamp ActualTimestamp = GetPhysicsTimestamp();
-	ActualTimestamp.ServerFrame += NetworkPhysicsCvars::PredictionAsyncFrameBuffer;
+	const int32 BufferTickSize = FMath::CeilToInt((NetworkPhysicsCvars::TickOffsetBufferTime / 1000.f) / UPhysicsSettings::Get()->AsyncFixedTimeStepSize);
+	ActualTimestamp.ServerFrame += BufferTickSize;
 
 	// Mark offset as assigned when we get a valid predicted server frame.
 	const int32 PredictedServerFrame = Timestamp.ServerFrame;
 	bNetworkPhysicsTickOffsetAssigned |= PredictedServerFrame != INDEX_NONE;
 
-	// Send update to client if offset is not assigned or over correction limit
+	const float CurrentTime = GetWorld()->GetTimeSeconds();
+	const int32 TimestampDiff = FMath::Abs(PredictedServerFrame - ActualTimestamp.ServerFrame);
+	const float TimestampTimeDiff = TimestampDiff * UPhysicsSettings::Get()->AsyncFixedTimeStepSize;
+	NetworkPhysicsTickOffsetDesyncAccumulatedTime = (TimestampDiff == 0) ? CurrentTime : NetworkPhysicsTickOffsetDesyncAccumulatedTime;
+
+	// Send update to client if offset is not assigned or over correction limits
 	// Note that we are sending the current ServerFrame along with the frame buffer added, to the client.
-	if (!bNetworkPhysicsTickOffsetAssigned || FMath::Abs(PredictedServerFrame - ActualTimestamp.ServerFrame) > NetworkPhysicsCvars::TickOffsetCorrectionLimit)
+	if (!bNetworkPhysicsTickOffsetAssigned
+		|| TimestampTimeDiff > (NetworkPhysicsCvars::TickOffsetCorrectionSizeTimeLimit / 1000.f)
+		|| CurrentTime - NetworkPhysicsTickOffsetDesyncAccumulatedTime > (NetworkPhysicsCvars::TickOffsetCorrectionTimeLimit / 1000.0f))
 	{
 		Timestamp.ServerFrame = ActualTimestamp.ServerFrame;
 		NetworkPhysicsTickOffset = Timestamp.ServerFrame - Timestamp.LocalFrame;
+		NetworkPhysicsTickOffsetDesyncAccumulatedTime = CurrentTime;
 		ClientSetupNetworkPhysicsTimestamp(Timestamp); /* Reliable RPC */
 	}
 
@@ -6355,7 +6381,7 @@ void APlayerController::ClientAckTimeDilation_Implementation(float TimeDilation,
 	}
 	ClientLatestTimeDilationServerStep = ServerStep;
 
-	if(UWorld* World = GetWorld())
+	if (UWorld* World = GetWorld())
 	{
 		World->GetPhysicsScene()->SetNetworkDeltaTimeScale(TimeDilation);
 	}
diff --git a/Engine/Source/Runtime/Engine/Public/Physics/NetworkPhysicsComponent.h b/Engine/Source/Runtime/Engine/Public/Physics/NetworkPhysicsComponent.h
--- a/Engine/Source/Runtime/Engine/Public/Physics/NetworkPhysicsComponent.h
+++ b/Engine/Source/Runtime/Engine/Public/Physics/NetworkPhysicsComponent.h
@@ -1135,7 +1135,7 @@ private:
 	void OnPostProcessInputs_Internal(const int32 PhysicsStep);
 
 	// Consume data from async input
-	void ConsumeAsyncInput();
+	void ConsumeAsyncInput(const int32 PhysicsStep);
 
 	/** Get the rigid solver */
 	Chaos::FPBDRigidsSolver* GetRigidSolver();
@@ -1153,10 +1153,10 @@ private:
 	const float GetCurrentInputDecay(const FNetworkPhysicsData* PhysicsData);
 
 	/** Populate input data in AsyncOutput to send over the network */
-	void SendInputData_Internal(FAsyncNetworkPhysicsComponentOutput& AsyncOutput);
+	void SendInputData_Internal(FAsyncNetworkPhysicsComponentOutput& AsyncOutput, const int32 PhysicsStep);
 
 	/** Populate state data in AsyncOutput to send over the network */
-	void SendStateData_Internal(FAsyncNetworkPhysicsComponentOutput& AsyncOutput);
+	void SendStateData_Internal(FAsyncNetworkPhysicsComponentOutput& AsyncOutput, const int32 PhysicsStep);
 
 private:
 	friend UNetworkPhysicsComponent;
@@ -1190,6 +1190,9 @@ private:
 
 	// Local temporary inputs data used by pre/post process inputs functions
 	TUniquePtr<FNetworkPhysicsData> InputData;
+	
+	// Local temporary inputs data used by ConsumeAsyncInput
+	TUniquePtr<FNetworkPhysicsData> LatestInputReceiveData;
 
 	// Local temporary states data used by pre/post process inputs functions
 	TUniquePtr<FNetworkPhysicsData> StateData;
diff --git a/Engine/Source/Runtime/Experimental/Chaos/Private/PBDRigidsSolver.cpp b/Engine/Source/Runtime/Experimental/Chaos/Private/PBDRigidsSolver.cpp
--- a/Engine/Source/Runtime/Experimental/Chaos/Private/PBDRigidsSolver.cpp
+++ b/Engine/Source/Runtime/Experimental/Chaos/Private/PBDRigidsSolver.cpp
@@ -1930,7 +1930,7 @@ namespace Chaos
 			}
 #endif
 
-			if ((ResimStep < LastStep) && NumResimSteps <= MarshallingManager.GetNumHistory_Internal())
+			if ((ResimStep <= LastStep) && NumResimSteps <= MarshallingManager.GetNumHistory_Internal())
 			{
 				FResimDebugInfo DebugInfo;
 				QUICK_SCOPE_CYCLE_COUNTER(ChaosRewindAndResim);
diff --git a/Engine/Source/Runtime/Experimental/Chaos/Private/RewindData.cpp b/Engine/Source/Runtime/Experimental/Chaos/Private/RewindData.cpp
--- a/Engine/Source/Runtime/Experimental/Chaos/Private/RewindData.cpp
+++ b/Engine/Source/Runtime/Experimental/Chaos/Private/RewindData.cpp
@@ -257,6 +257,49 @@ void FRewindData::ApplyTargets(const int32 Frame, const bool bResetSimulation)
 	}
 }
 
+const int32 FRewindData::CompareTargetsToLastFrame()
+{
+	int32 RewindFrame = INDEX_NONE;
+	const FFrameAndPhase FrameAndPhase{ CurrentFrame() - 1, FFrameAndPhase::PostPushData };
+
+	if (LatestTargetFrame < FrameAndPhase.Frame)
+	{
+		// Early out if we only have targets earlier than the previous simulated frame
+		// NOTE: This is the normal flow, we should only run this logic when the client is desynced behind the server and we receive targets from the server ahead of time.
+		return RewindFrame;
+	}
+
+	// Iterate over targets that exist for current frame
+	for (FDirtyParticleInfo& DirtyParticleInfo : DirtyParticles)
+	{
+		// TODO: Only iterate targets that are not predicted/interpolated
+		FGeometryParticleHandle* PTParticle = DirtyParticleInfo.GetObjectPtr();
+		if (PTParticle)
+		{
+			FGeometryParticleStateBase& History = DirtyParticleInfo.GetHistory();
+			if (!History.TargetPositions.IsEmpty())
+			{
+				// Compare with particle for this frame and mark resim if needed from CurrentFrame()
+				if (const FParticlePositionRotation* TargetState = History.TargetPositions.Read(FrameAndPhase, PropertiesPool))
+				{
+					if (const FParticlePositionRotation* XRState = History.ParticlePositionRotation.Read(FrameAndPhase, PropertiesPool))
+					{
+						const FVec3 ErrorOffset = TargetState->GetX() - XRState->GetX();
+						const bool ShouldTriggerResim = ErrorOffset.Size() >= FPhysicsSolverBase::ResimulationErrorThreshold();
+
+						if (ShouldTriggerResim)
+						{
+							RewindFrame = CurrentFrame() - 1;
+						}
+					}
+				}
+			}
+		}
+	}
+
+	return RewindFrame;
+}
+
 CHAOS_API bool bResimAllowRewindToResimulatedFrames = false;
 FAutoConsoleVariableRef CVarResimAllowRewindToResimulatedFrames(TEXT("p.Resim.AllowRewindToResimulatedFrames"), bResimAllowRewindToResimulatedFrames, TEXT("Allow rewinding back to a frame that was previously part of a resimulation. If a resimulation is performed between frame 100-110, allow a new resim from 105-115 if needed, else next resim will be able to start from frame 111."));
 
@@ -999,6 +1042,9 @@ int32 FRewindData::FindValidResimFrame(const int32 RequestedFrame)
 			}
 		}
 		
+/*
+* todo: Only check input and state histories for actors that will actively resimulate
+
 		if (bHasTargetHistory)
 		{
 			for (TWeakPtr<FBaseRewindHistory>& InputHistory : InputHistories)
@@ -1026,6 +1072,7 @@ int32 FRewindData::FindValidResimFrame(const int32 RequestedFrame)
 				}
 			}
 		}
+*/
 
 		if (bHasTargetHistory)
 		{
@@ -1053,6 +1100,7 @@ void FRewindData::PushStateAtFrame(FGeometryParticleHandle& Handle, int32 Frame,
 	FDirtyParticleInfo& Info = FindOrAddDirtyObj(Handle);
 	FGeometryParticleStateBase& Latest = Info.GetHistory();
 	const FFrameAndPhase FrameAndPhase{ Frame, Phase };
+	LatestTargetFrame = bRecordingHistory ? FMath::Max(LatestTargetFrame, Frame) : LatestTargetFrame;
 
 	if (bRecordingHistory || Latest.TargetPositions.IsClean(FrameAndPhase))
 	{
diff --git a/Engine/Source/Runtime/Experimental/Chaos/Public/RewindData.h b/Engine/Source/Runtime/Experimental/Chaos/Public/RewindData.h
--- a/Engine/Source/Runtime/Experimental/Chaos/Public/RewindData.h
+++ b/Engine/Source/Runtime/Experimental/Chaos/Public/RewindData.h
@@ -1513,6 +1513,7 @@ public:
 	, DataIdxOffset(0)
 	, bNeedsSave(false)
 	, bResimOptimization(InResimOptimization)
+	, LatestTargetFrame(0)
 	{
 	}
 
@@ -1522,6 +1523,7 @@ public:
 		CurFrame = InCurrentFrame;
 		LatestFrame = InCurrentFrame;
 		bResimOptimization = InResimOptimization;
+		LatestTargetFrame = 0;
 		Managers = TCircularBuffer<FFrameManagerInfo>(NumFrames + 1);
 	}
 
@@ -1717,6 +1719,11 @@ public:
 	/** Get the latest frame resim has been blocked from rewinding past */
 	int32 GetBlockedResimFrame() { return BlockResimFrame; }
 
+	/** Check if we have received targets already for the last frame simulated,
+	* if so compare those with the result of the simulation and if they desync return the frame value to request a rewind for to correct the desync
+	* NOTE: This only happens when the client has desynced and is behind the server, so we receive server states for frames not yet simulated */
+	const int32 CHAOS_API CompareTargetsToLastFrame();
+
 private:
 	friend class FPBDRigidsSolver;
 
@@ -1872,9 +1879,8 @@ private:
 		return TObjState(State, Handle, PropertiesPool, { Frame, Phase });
 	}
 
+	/** Apply the cached history state for the given frame cached particles and joints */
 	bool RewindToFrame(int32 RewindFrame);
-	
-
 
 	/** Apply targets positions and velocities while resimulating */
 	void ApplyTargets(const int32 Frame, const bool bResetSimulation);
@@ -1904,6 +1910,7 @@ private:
 	bool bNeedsSave;	//Indicates that some data is pointing at head and requires saving before a rewind
 	bool bResimOptimization;
 	int32 ResimFrame = INDEX_NONE;
+	int32 LatestTargetFrame;
 
 	// Used to block rewinding past a physics change we currently don't handle
 	int32 BlockResimFrame = INDEX_NONE;
diff --git a/Engine/Source/Runtime/PhysicsCore/Private/ChaosScene.cpp b/Engine/Source/Runtime/PhysicsCore/Private/ChaosScene.cpp
--- a/Engine/Source/Runtime/PhysicsCore/Private/ChaosScene.cpp
+++ b/Engine/Source/Runtime/PhysicsCore/Private/ChaosScene.cpp
@@ -335,7 +335,7 @@ void FChaosScene::SetUpForFrame(const FVector* NewGrav,float InDeltaSeconds /*=
 
 	if(bSubstepping)
 	{
-		MDeltaTime = FMath::Min(InDeltaSeconds, InMaxSubsteps * InMaxSubstepDeltaTime);
+		MDeltaTime = InMaxSubstepDeltaTime > 0.f ? FMath::Min(InDeltaSeconds, InMaxSubsteps * InMaxSubstepDeltaTime) : InDeltaSeconds;
 	}
 	else
 	{
