Fixed Runtime Data Layer initial state on clients : Clients now only apply initial state for client-only runtime Data Layers.

diff --git a/Engine/Source/Runtime/Engine/Private/WorldPartition/DataLayer/WorldDataLayers.cpp b/Engine/Source/Runtime/Engine/Private/WorldPartition/DataLayer/WorldDataLayers.cpp
--- a/Engine/Source/Runtime/Engine/Private/WorldPartition/DataLayer/WorldDataLayers.cpp
+++ b/Engine/Source/Runtime/Engine/Private/WorldPartition/DataLayer/WorldDataLayers.cpp
@@ -108,6 +108,53 @@ void AWorldDataLayers::PostRegisterAllComponents()
 	MARK_PROPERTY_DIRTY_FROM_NAME(AWorldDataLayers, ReplicatedArray, this); \
 	ReplicatedArray = SourceArray;
 
+bool AWorldDataLayers::CanChangeDataLayerRuntimeState(const UDataLayerInstance* InDataLayerInstance, AWorldDataLayers::ESetDataLayerRuntimeStateError* OutReason) const
+{
+	if (!InDataLayerInstance->IsRuntime())
+	{
+		if (OutReason)
+		{
+			*OutReason = ESetDataLayerRuntimeStateError::NotRuntime;
+		}
+		return false;
+	}
+
+	const ENetMode NetMode = GetNetMode();
+
+	if (InDataLayerInstance->IsClientOnly())
+	{
+		if (NetMode != NM_Standalone && NetMode != NM_Client)
+		{
+			if (OutReason)
+			{
+				*OutReason = ESetDataLayerRuntimeStateError::ClientOnlyFromServer;
+			}
+			return false;
+		}
+	}
+	else if (InDataLayerInstance->IsServerOnly())
+	{
+		if (NetMode == NM_Client)
+		{
+			if (OutReason)
+			{
+				*OutReason = ESetDataLayerRuntimeStateError::ServerOnlyFromClient;
+			}
+			return false;
+		}
+	}
+	else if (NetMode == NM_Client)
+	{
+		if (OutReason)
+		{
+			*OutReason = ESetDataLayerRuntimeStateError::AuthoritativeFromClient;
+		}
+		return false;
+	}
+
+	return true;
+}
+
 void AWorldDataLayers::InitializeDataLayerRuntimeStates()
 {
 	if (IsExternalDataLayerWorldDataLayers())
@@ -124,23 +171,26 @@ void AWorldDataLayers::InitializeDataLayerRuntimeStates()
 
 		ForEachDataLayerInstance([this, &RuntimeDataLayerInstances](UDataLayerInstance* DataLayerInstance)
 		{
-			const bool bDataLayerClientOnly = DataLayerInstance->IsClientOnly();
-			const bool bDataLayerServerOnly = DataLayerInstance->IsServerOnly();
-			const bool bIsLocalDataLayer = bDataLayerClientOnly || bDataLayerServerOnly;
-			TSet<FName>& TargetLoadedDataLayerNames = bIsLocalDataLayer ? LocalLoadedDataLayerNames : LoadedDataLayerNames;
-			TSet<FName>& TargetActiveDataLayerNames = bIsLocalDataLayer ? LocalActiveDataLayerNames : ActiveDataLayerNames;
-
-			if (DataLayerInstance->IsRuntime())
+			if (CanChangeDataLayerRuntimeState(DataLayerInstance))
 			{
-				if (DataLayerInstance->GetInitialRuntimeState() == EDataLayerRuntimeState::Loaded)
-				{
-					TargetLoadedDataLayerNames.Add(DataLayerInstance->GetDataLayerFName());
-				}
-				else if (DataLayerInstance->GetInitialRuntimeState() == EDataLayerRuntimeState::Activated)
+				const bool bDataLayerClientOnly = DataLayerInstance->IsClientOnly();
+				const bool bDataLayerServerOnly = DataLayerInstance->IsServerOnly();
+				const bool bIsLocalDataLayer = bDataLayerClientOnly || bDataLayerServerOnly;
+				TSet<FName>& TargetLoadedDataLayerNames = bIsLocalDataLayer ? LocalLoadedDataLayerNames : LoadedDataLayerNames;
+				TSet<FName>& TargetActiveDataLayerNames = bIsLocalDataLayer ? LocalActiveDataLayerNames : ActiveDataLayerNames;
+
+				if (DataLayerInstance->IsRuntime())
 				{
-					TargetActiveDataLayerNames.Add(DataLayerInstance->GetDataLayerFName());
+					if (DataLayerInstance->GetInitialRuntimeState() == EDataLayerRuntimeState::Loaded)
+					{
+						TargetLoadedDataLayerNames.Add(DataLayerInstance->GetDataLayerFName());
+					}
+					else if (DataLayerInstance->GetInitialRuntimeState() == EDataLayerRuntimeState::Activated)
+					{
+						TargetActiveDataLayerNames.Add(DataLayerInstance->GetDataLayerFName());
+					}
+					RuntimeDataLayerInstances.Add(DataLayerInstance);
 				}
-				RuntimeDataLayerInstances.Add(DataLayerInstance);
 			}
 			return true;
 		});
@@ -181,53 +231,53 @@ void AWorldDataLayers::ResetDataLayerRuntimeStates()
 
 bool AWorldDataLayers::SetDataLayerRuntimeState(const UDataLayerInstance* InDataLayerInstance, EDataLayerRuntimeState InState, bool bInIsRecursive)
 {
-	if (!InDataLayerInstance || !InDataLayerInstance->IsRuntime())
+	if (!InDataLayerInstance)
 	{
 		return false;
 	}
 
-	const ENetMode NetMode = GetNetMode();
-	const bool bDataLayerClientOnly = InDataLayerInstance->IsClientOnly();
-	const bool bDataLayerServerOnly = InDataLayerInstance->IsServerOnly();
-	const bool bIsLocalDataLayer = bDataLayerClientOnly || bDataLayerServerOnly;
-	TSet<FName>& TargetLoadedDataLayerNames = bIsLocalDataLayer ? LocalLoadedDataLayerNames : LoadedDataLayerNames;
-	TSet<FName>& TargetActiveDataLayerNames = bIsLocalDataLayer ? LocalActiveDataLayerNames : ActiveDataLayerNames;
-
 	const EDataLayerRuntimeState CurrentState = GetDataLayerRuntimeStateByName(InDataLayerInstance->GetDataLayerFName());
 
-	if (bDataLayerClientOnly)
+	ESetDataLayerRuntimeStateError Reason;
+	if (!CanChangeDataLayerRuntimeState(InDataLayerInstance, &Reason))
 	{
-		if (NetMode != NM_Standalone && NetMode != NM_Client)
+		switch (Reason)
 		{
-			UE_LOG(LogWorldPartition, Verbose, TEXT("Client Only Data Layer state change '%s' was ignored: %s -> %s"),
+		case ESetDataLayerRuntimeStateError::NotRuntime:
+			UE_LOG(LogWorldPartition, Verbose, TEXT("Non-Runtime Data Layer '%s' state change was ignored"),
+				*InDataLayerInstance->GetDataLayerShortName());
+			break;
+		case ESetDataLayerRuntimeStateError::ClientOnlyFromServer:
+			UE_LOG(LogWorldPartition, Verbose, TEXT("Client Only Data Layer '%s' state change was ignored: %s -> %s"),
 				*InDataLayerInstance->GetDataLayerShortName(),
 				*StaticEnum<EDataLayerRuntimeState>()->GetDisplayNameTextByValue((int64)CurrentState).ToString(),
 				*StaticEnum<EDataLayerRuntimeState>()->GetDisplayNameTextByValue((int64)InState).ToString());
-			return false;
-		}
-	}
-	else if (bDataLayerServerOnly)
-	{
-		if (NetMode == NM_Client)
-		{
-			UE_LOG(LogWorldPartition, Verbose, TEXT("Server Only Data Layer state change '%s' was ignored: %s -> %s"),
+			break;
+		case ESetDataLayerRuntimeStateError::ServerOnlyFromClient:
+			UE_LOG(LogWorldPartition, Verbose, TEXT("Server Only Data Layer '%s' state change was ignored: %s -> %s"),
 				*InDataLayerInstance->GetDataLayerShortName(),
 				*StaticEnum<EDataLayerRuntimeState>()->GetDisplayNameTextByValue((int64)CurrentState).ToString(),
 				*StaticEnum<EDataLayerRuntimeState>()->GetDisplayNameTextByValue((int64)InState).ToString());
-			return false;
+			break;
+		case ESetDataLayerRuntimeStateError::AuthoritativeFromClient:
+				UE_LOG(LogWorldPartition, Verbose, TEXT("Data Layer '%s' state change was ignored on client: %s -> %s"),
+					*InDataLayerInstance->GetDataLayerShortName(),
+					*StaticEnum<EDataLayerRuntimeState>()->GetDisplayNameTextByValue((int64)CurrentState).ToString(),
+					*StaticEnum<EDataLayerRuntimeState>()->GetDisplayNameTextByValue((int64)InState).ToString());
+				break;
 		}
-	}
-	else if (NetMode == NM_Client)
-	{
-		UE_LOG(LogWorldPartition, Verbose, TEXT("Data Layer '%s' state change was ignored on client: %s -> %s"),
-			*InDataLayerInstance->GetDataLayerShortName(),
-			*StaticEnum<EDataLayerRuntimeState>()->GetDisplayNameTextByValue((int64)CurrentState).ToString(),
-			*StaticEnum<EDataLayerRuntimeState>()->GetDisplayNameTextByValue((int64)InState).ToString());
 		return false;
 	}
 
 	if (CurrentState != InState)
 	{
+		const ENetMode NetMode = GetNetMode();
+		const bool bDataLayerClientOnly = InDataLayerInstance->IsClientOnly();
+		const bool bDataLayerServerOnly = InDataLayerInstance->IsServerOnly();
+		const bool bIsLocalDataLayer = bDataLayerClientOnly || bDataLayerServerOnly;
+		TSet<FName>& TargetLoadedDataLayerNames = bIsLocalDataLayer ? LocalLoadedDataLayerNames : LoadedDataLayerNames;
+		TSet<FName>& TargetActiveDataLayerNames = bIsLocalDataLayer ? LocalActiveDataLayerNames : ActiveDataLayerNames;
+
 		TargetLoadedDataLayerNames.Remove(InDataLayerInstance->GetDataLayerFName());
 		TargetActiveDataLayerNames.Remove(InDataLayerInstance->GetDataLayerFName());
 
diff --git a/Engine/Source/Runtime/Engine/Public/WorldPartition/DataLayer/WorldDataLayers.h b/Engine/Source/Runtime/Engine/Public/WorldPartition/DataLayer/WorldDataLayers.h
--- a/Engine/Source/Runtime/Engine/Public/WorldPartition/DataLayer/WorldDataLayers.h
+++ b/Engine/Source/Runtime/Engine/Public/WorldPartition/DataLayer/WorldDataLayers.h
@@ -245,6 +245,14 @@ protected:
 	ENGINE_API void OnRep_EffectiveLoadedDataLayerNames();
 
 private:
+	enum class ESetDataLayerRuntimeStateError
+	{
+		NotRuntime,
+		ClientOnlyFromServer,
+		ServerOnlyFromClient,
+		AuthoritativeFromClient
+	};
+
 	ENGINE_API const FWorldDataLayersEffectiveStates& GetEffectiveStates() const;
 
 	// External Data Layers
@@ -261,6 +269,7 @@ private:
 	ENGINE_API void OnDataLayerManagerDeinitialized();
 	ENGINE_API void ResolveEffectiveRuntimeState(const UDataLayerInstance* InDataLayer, bool bInNotifyChange = true);
 	ENGINE_API void DumpDataLayerRecursively(const UDataLayerInstance* DataLayer, FString Prefix, FOutputDevice& OutputDevice) const;
+	bool CanChangeDataLayerRuntimeState(const UDataLayerInstance* InDataLayerInstance, ESetDataLayerRuntimeStateError* OutReason = nullptr) const;
 
 	//~ Begin IDataLayerInstanceProvider interface
 	ENGINE_API virtual TSet<TObjectPtr<UDataLayerInstance>>& GetDataLayerInstances() override;
