[World Partition HashSet] Support streaming sources sector angle.

diff --git a/Engine/Source/Runtime/Engine/Private/Tests/WorldPartition/WorldPartitionStaticSpatialIndexTests.cpp b/Engine/Source/Runtime/Engine/Private/Tests/WorldPartition/WorldPartitionStaticSpatialIndexTests.cpp
--- a/Engine/Source/Runtime/Engine/Private/Tests/WorldPartition/WorldPartitionStaticSpatialIndexTests.cpp
+++ b/Engine/Source/Runtime/Engine/Private/Tests/WorldPartition/WorldPartitionStaticSpatialIndexTests.cpp
@@ -38,7 +38,8 @@ namespace WorldPartitionTests
 		const double StartTime = FPlatformTime::Seconds();
 		for (int32 ListNumTests = 0; ListNumTests < Tests.Num(); ListNumTests++)
 		{
-			SpatialIndex.ForEachIntersectingElement(Tests[ListNumTests], [&Results](const int32& Value) { Results.Add(Value); });
+			FStaticSpatialIndex::FSphere Sphere(Tests[ListNumTests].Center, Tests[ListNumTests].W);
+			SpatialIndex.ForEachIntersectingElement(Sphere, [&Results](const int32& Value) { Results.Add(Value); });
 		}
 		const double RunTime = FPlatformTime::Seconds() - StartTime;
 
diff --git a/Engine/Source/Runtime/Engine/Private/WorldPartition/RuntimeHashSet/WorldPartitionRuntimeHashSet.cpp b/Engine/Source/Runtime/Engine/Private/WorldPartition/RuntimeHashSet/WorldPartitionRuntimeHashSet.cpp
--- a/Engine/Source/Runtime/Engine/Private/WorldPartition/RuntimeHashSet/WorldPartitionRuntimeHashSet.cpp
+++ b/Engine/Source/Runtime/Engine/Private/WorldPartition/RuntimeHashSet/WorldPartitionRuntimeHashSet.cpp
@@ -16,6 +16,8 @@
 #include "Misc/HashBuilder.h"
 #include "Misc/ArchiveMD5.h"
 
+#include "WorldPartition/RuntimeHashSet/StaticSpatialIndex.h"
+
 void FRuntimePartitionStreamingData::CreatePartitionsSpatialIndex() const
 {
 	if (!SpatialIndex)
@@ -505,17 +507,29 @@ void UWorldPartitionRuntimeHashSet::ForEachStreamingCellsQuery(const FWorldParti
 		{
 			QuerySource.ForEachShape(InLoadingRange, InGridName, false, [InSpatialIndex, QueryCache, &ShouldAddCell, &QuerySource, &Func](const FSphericalSector& Shape)
 			{
-				const FSphere ShapeSphere(Shape.GetCenter(), Shape.GetRadius());
-
-				InSpatialIndex->ForEachIntersectingElement(ShapeSphere, [QueryCache, &ShouldAddCell, &Shape, &QuerySource, &Func](UWorldPartitionRuntimeCell* RuntimeCell)
+				auto ForEachIntersectingElement = [InSpatialIndex, QueryCache, &ShouldAddCell, &Shape, &QuerySource, &Func]<typename ShapeType>(const ShapeType& InShape)
 				{
-					if (QueryCache)
+					InSpatialIndex->ForEachIntersectingElement(InShape, [QueryCache, &ShouldAddCell, &Shape, &QuerySource, &Func](UWorldPartitionRuntimeCell* RuntimeCell)
 					{
-						QueryCache->AddCellInfo(RuntimeCell, Shape);
-					}
+						if (QueryCache)
+						{
+							QueryCache->AddCellInfo(RuntimeCell, Shape);
+						}
 
-					return !ShouldAddCell(RuntimeCell, QuerySource) || Func(RuntimeCell);
-				});
+						return !ShouldAddCell(RuntimeCell, QuerySource) || Func(RuntimeCell);
+					});
+				};
+
+				if (Shape.IsSphere())
+				{
+					const FStaticSpatialIndex::FSphere Sphere(Shape.GetCenter(), Shape.GetRadius());
+					ForEachIntersectingElement(Sphere);
+				}
+				else
+				{
+					const FStaticSpatialIndex::FCone Cone(Shape.GetCenter(), Shape.GetAxis(), Shape.GetRadius(), Shape.GetAngle());
+					ForEachIntersectingElement(Cone);
+				}
 			});
 		}
 
@@ -610,8 +624,6 @@ void UWorldPartitionRuntimeHashSet::ForEachStreamingCellsSources(const TArray<FW
 				{
 					Source.ForEachShape(StreamingData->GetLoadingRange(), false, [this, &Source, StreamingData, &Context, &Func](const FSphericalSector& Shape)
 					{
-						const FSphere ShapeSphere(Shape.GetCenter(), Shape.GetRadius());
-
 						auto ForEachIntersectingElementFunc = [this, &Source, &Shape, &Context, &Func](UWorldPartitionRuntimeCell* Cell)
 						{
 #if WITH_EDITOR
@@ -630,21 +642,35 @@ void UWorldPartitionRuntimeHashSet::ForEachStreamingCellsSources(const TArray<FW
 							}
 						};
 
-						if (StreamingData->SpatialIndex.IsValid())
+						auto ForEachIntersectingElement = [&StreamingData, &Source, &ForEachIntersectingElementFunc]<typename ShapeType>(const ShapeType& Shape)
 						{
-							if (Source.bForce2D)
+							if (StreamingData->SpatialIndex.IsValid())
 							{
-								StreamingData->SpatialIndexForce2D->ForEachIntersectingElement(ShapeSphere, ForEachIntersectingElementFunc);
+								if (Source.bForce2D)
+								{
+									StreamingData->SpatialIndexForce2D->ForEachIntersectingElement(Shape, ForEachIntersectingElementFunc);
+								}
+								else
+								{
+									StreamingData->SpatialIndex->ForEachIntersectingElement(Shape, ForEachIntersectingElementFunc);
+								}
 							}
-							else
+
+							if (StreamingData->SpatialIndex2D.IsValid())
 							{
-								StreamingData->SpatialIndex->ForEachIntersectingElement(ShapeSphere, ForEachIntersectingElementFunc);
+								StreamingData->SpatialIndex2D->ForEachIntersectingElement(Shape, ForEachIntersectingElementFunc);
 							}
-						}
+						};
 
-						if (StreamingData->SpatialIndex2D.IsValid())
+						if (Shape.IsSphere())
+						{
+							const FStaticSpatialIndex::FSphere ShapeSphere(Shape.GetCenter(), Shape.GetRadius());
+							ForEachIntersectingElement(ShapeSphere);
+						}
+						else
 						{
-							StreamingData->SpatialIndex2D->ForEachIntersectingElement(ShapeSphere, ForEachIntersectingElementFunc);
+							const FStaticSpatialIndex::FCone ShapeCone(Shape.GetCenter(), Shape.GetAxis(), Shape.GetRadius(), Shape.GetAngle());
+							ForEachIntersectingElement(ShapeCone);
 						}
 					});
 				}
diff --git a/Engine/Source/Runtime/Engine/Private/WorldPartition/RuntimeHashSet/WorldPartitionRuntimeHashSetDebugDraw.cpp b/Engine/Source/Runtime/Engine/Private/WorldPartition/RuntimeHashSet/WorldPartitionRuntimeHashSetDebugDraw.cpp
--- a/Engine/Source/Runtime/Engine/Private/WorldPartition/RuntimeHashSet/WorldPartitionRuntimeHashSetDebugDraw.cpp
+++ b/Engine/Source/Runtime/Engine/Private/WorldPartition/RuntimeHashSet/WorldPartitionRuntimeHashSetDebugDraw.cpp
@@ -364,58 +364,70 @@ void UWorldPartitionRuntimeHashSet::Draw3D(const TArray<FWorldPartitionStreaming
 			{
 				Source.ForEachShape(StreamingData.GetLoadingRange(), StreamingData.Name, false, [this, &StreamingData, VisualizeMode, &DataLayerDebugColors, ContentBundleManager, OwningWorld, &WorldPartitionTransform](const FSphericalSector& Shape)
 				{
-					const FSphere ShapeSphere(Shape.GetCenter(), Shape.GetRadius());
-
-					StreamingData.SpatialIndex.Get()->ForEachIntersectingElement(ShapeSphere, [this, VisualizeMode, &DataLayerDebugColors, ContentBundleManager, OwningWorld, &WorldPartitionTransform](UWorldPartitionRuntimeCell* Cell)
+					auto ForEachIntersectingElement = [&]<typename ShapeType>(const ShapeType& InShape)
 					{
-						const FVector2D CellStreamingBoundsSize = FVector2D(Cell->GetStreamingBounds().GetSize());
-						const FVector2D CellStreamingBoundsMin = FVector2D(Cell->GetStreamingBounds().Min);
+						StreamingData.SpatialIndex.Get()->ForEachIntersectingElement(InShape, [this, VisualizeMode, &DataLayerDebugColors, ContentBundleManager, OwningWorld, &WorldPartitionTransform](UWorldPartitionRuntimeCell* Cell)
+						{
+							const FVector2D CellStreamingBoundsSize = FVector2D(Cell->GetStreamingBounds().GetSize());
+							const FVector2D CellStreamingBoundsMin = FVector2D(Cell->GetStreamingBounds().Min);
 
-						float CellOpacity = 0.0f;
-						TArray<FLinearColor> CellColors;
+							float CellOpacity = 0.0f;
+							TArray<FLinearColor> CellColors;
 
-						switch (GShowRuntimeHashSetDebugDisplayMode)
-						{
-						case 0:
-							CellColors.Add(Cell->GetDebugColor(VisualizeMode));
-							CellOpacity = 0.25f / FMath::Max<float>(GShowRuntimeHashSetDebugDisplayLevelCount, 1);
-							break;
-						case 1:
-							if (DataLayerDebugColors.Num() && Cell->GetDataLayers().Num())
+							switch (GShowRuntimeHashSetDebugDisplayMode)
 							{
-								for (const FName& DataLayer : Cell->GetDataLayers())
+							case 0:
+								CellColors.Add(Cell->GetDebugColor(VisualizeMode));
+								CellOpacity = 0.25f / FMath::Max<float>(GShowRuntimeHashSetDebugDisplayLevelCount, 1);
+								break;
+							case 1:
+								if (DataLayerDebugColors.Num() && Cell->GetDataLayers().Num())
 								{
-									CellColors.Add(DataLayerDebugColors[DataLayer]);
+									for (const FName& DataLayer : Cell->GetDataLayers())
+									{
+										CellColors.Add(DataLayerDebugColors[DataLayer]);
+									}
+									CellOpacity = 0.67f;
 								}
-								CellOpacity = 0.67f;
-							}
-							break;
-						case 2:
-							if (ContentBundleManager && Cell->GetContentBundleID().IsValid())
-							{
-								if (const FContentBundleBase* ContentBundle = ContentBundleManager->GetContentBundle(OwningWorld, Cell->GetContentBundleID()))
+								break;
+							case 2:
+								if (ContentBundleManager && Cell->GetContentBundleID().IsValid())
 								{
-									check(ContentBundle->GetDescriptor());
-									CellColors.Add(ContentBundle->GetDescriptor()->GetDebugColor());
-									CellOpacity = 0.67f;
+									if (const FContentBundleBase* ContentBundle = ContentBundleManager->GetContentBundle(OwningWorld, Cell->GetContentBundleID()))
+									{
+										check(ContentBundle->GetDescriptor());
+										CellColors.Add(ContentBundle->GetDescriptor()->GetDebugColor());
+										CellOpacity = 0.67f;
+									}
 								}
+								break;
 							}
-							break;
-						}
 
-						if (CellColors.IsEmpty())
-						{
-							CellColors.Add(FLinearColor::White);
-							CellOpacity = 0.1f;
-						}
+							if (CellColors.IsEmpty())
+							{
+								CellColors.Add(FLinearColor::White);
+								CellOpacity = 0.1f;
+							}
+
+							// Draw Cell using its debug color
+							const FBox Box(Cell->GetStreamingBounds());
+							const FVector BoxCenter(Box.GetCenter());
+							const FColor BoxColor(CellColors[0].CopyWithNewOpacity(CellOpacity).ToFColor(true));
+							const FVector CellPos = WorldPartitionTransform.TransformPosition(BoxCenter);
+							DrawDebugBox(OwningWorld, CellPos, Box.GetExtent(), WorldPartitionTransform.GetRotation(), BoxColor.WithAlpha(255), false, -1.f, 255, 20.f);
+						});
+					};
 
-						// Draw Cell using its debug color
-						const FBox Box(Cell->GetStreamingBounds());
-						const FVector BoxCenter(Box.GetCenter());
-						const FColor BoxColor(CellColors[0].CopyWithNewOpacity(CellOpacity).ToFColor(true));
-						const FVector CellPos = WorldPartitionTransform.TransformPosition(BoxCenter);
-						DrawDebugBox(OwningWorld, CellPos, Box.GetExtent(), WorldPartitionTransform.GetRotation(), BoxColor.WithAlpha(255), false, -1.f, 255, 20.f);
-					});
+					if (Shape.IsSphere())
+					{
+						const FStaticSpatialIndex::FSphere Sphere(Shape.GetCenter(), Shape.GetRadius());
+						ForEachIntersectingElement(Sphere);
+					}
+					else
+					{
+						const FStaticSpatialIndex::FCone Cone(Shape.GetCenter(), Shape.GetAxis(), Shape.GetRadius(), Shape.GetAngle());
+						ForEachIntersectingElement(Cone);
+					}
 				});
 			}
 		}
diff --git a/Engine/Source/Runtime/Engine/Public/WorldPartition/RuntimeHashSet/StaticSpatialIndex.h b/Engine/Source/Runtime/Engine/Public/WorldPartition/RuntimeHashSet/StaticSpatialIndex.h
--- a/Engine/Source/Runtime/Engine/Public/WorldPartition/RuntimeHashSet/StaticSpatialIndex.h
+++ b/Engine/Source/Runtime/Engine/Public/WorldPartition/RuntimeHashSet/StaticSpatialIndex.h
@@ -35,6 +35,128 @@ namespace FStaticSpatialIndex
 		const typename Profile::FVector ClosestPoint = Profile::FVector::Max(InBox.Min, Profile::FVector::Min(InCenter, InBox.Max));
 		return (ClosestPoint - InCenter).SizeSquared() <= InRadiusSquared;
 	}
+
+	struct FSphere
+	{
+		FSphere(const FVector& InCenter, const FVector::FReal InRadius)
+			: Center(InCenter)
+			, Radius(InRadius)
+		{}
+
+		FVector Center;
+		FVector::FReal Radius;
+	};
+
+	struct FCone : FSphere
+	{
+		FCone(const FVector& InCenter, const FVector& InAxis, const FVector::FReal InRadius, const FVector::FReal InAngle)
+			: FSphere(InCenter, InRadius)
+			, Axis(InAxis)
+			, Angle(InAngle)
+		{}
+
+		FVector Axis;
+		FVector::FReal Angle;
+	};
+
+	template <typename Profile>
+	inline bool FastConeAABBIntersection(const typename Profile::FVector& InCenter, const typename Profile::FReal InRadiusSquared, const typename Profile::FVector& InAxis, const typename Profile::FReal InAngle, const typename Profile::FReal InSinHalfAngle, const typename Profile::FReal InCosHalfAngle, const typename Profile::FBox& InBox)
+	{
+		if (FastSphereAABBIntersection<Profile>(InCenter, InRadiusSquared, InBox))
+		{
+			if constexpr (Profile::Is3D)
+			{
+				const typename Profile::FVector BoxExtent = InBox.GetExtent();
+				const typename Profile::FReal BoxExtentSizeSqr = BoxExtent.SizeSquared();
+				const typename Profile::FReal BoxExtentSize = FMath::Sqrt(BoxExtentSizeSqr);
+				const typename Profile::FVector U = InCenter - (BoxExtentSize / InSinHalfAngle) * InAxis;
+			
+				typename Profile::FVector D = InBox.GetCenter() - U;
+				typename Profile::FReal DSqr = D | D;
+				typename Profile::FReal E = InAxis | D;
+
+				if(E > 0.0f && E * E >= DSqr * FMath::Square(InCosHalfAngle))
+				{
+					D = InBox.GetCenter() - InCenter;
+					DSqr = D | D;
+					E = -(InAxis | D);
+					if(E > 0.0f && E * E >= DSqr * FMath::Square(InSinHalfAngle))
+					{
+						return DSqr <= BoxExtentSizeSqr;
+					}
+
+					return true;
+				}
+
+				return false;
+			}
+			else
+			{
+				const FVector2D CenterToMinXMinY((FVector2D(InBox.Min.X, InBox.Min.Y) - InCenter).GetSafeNormal());
+				const FVector2D CenterToMaxXMinY((FVector2D(InBox.Max.X, InBox.Min.Y) - InCenter).GetSafeNormal());
+				const FVector2D CenterToMaxXMaxY((FVector2D(InBox.Max.X, InBox.Max.Y) - InCenter).GetSafeNormal());
+				const FVector2D CenterToMinXMaxY((FVector2D(InBox.Min.X, InBox.Max.Y) - InCenter).GetSafeNormal());
+
+				if (InAngle <= 180.0f)
+				{
+					const Profile::FReal SinAngleMinXMinY = FVector2D::CrossProduct(InAxis, CenterToMinXMinY);
+					const Profile::FReal SinAngleMaxXMinY = FVector2D::CrossProduct(InAxis, CenterToMaxXMinY);
+					const Profile::FReal SinAngleMaxXMaxY = FVector2D::CrossProduct(InAxis, CenterToMaxXMaxY);
+					const Profile::FReal SinAngleMinXMaxY = FVector2D::CrossProduct(InAxis, CenterToMinXMaxY);
+
+					if (SinAngleMinXMinY < -InSinHalfAngle && SinAngleMaxXMinY < -InSinHalfAngle && SinAngleMaxXMaxY < -InSinHalfAngle && SinAngleMinXMaxY < -InSinHalfAngle)
+					{
+						return false;
+					}
+
+					if (SinAngleMinXMinY > InSinHalfAngle && SinAngleMaxXMinY > InSinHalfAngle && SinAngleMaxXMaxY > InSinHalfAngle && SinAngleMinXMaxY > InSinHalfAngle)
+					{
+						return false;
+					}
+
+					const Profile::FReal CosAngleMinXMinY = FVector2D::DotProduct(InAxis, CenterToMinXMinY);
+					const Profile::FReal CosAngleMaxXMinY = FVector2D::DotProduct(InAxis, CenterToMaxXMinY);
+					const Profile::FReal CosAngleMaxXMaxY = FVector2D::DotProduct(InAxis, CenterToMaxXMaxY);
+					const Profile::FReal CosAngleMinXMaxY = FVector2D::DotProduct(InAxis, CenterToMinXMaxY);
+
+					if (CosAngleMinXMinY < 0 && CosAngleMaxXMinY < 0 && CosAngleMaxXMaxY < 0 && CosAngleMinXMaxY < 0)
+					{
+						return false;
+					}
+
+					return true;
+				}
+				else
+				{
+					const Profile::FReal CosAngleMinXMinY = FVector2D::DotProduct(InAxis,CenterToMinXMinY);
+					const Profile::FReal CosAngleMaxXMinY = FVector2D::DotProduct(InAxis,CenterToMaxXMinY);
+					const Profile::FReal CosAngleMaxXMaxY = FVector2D::DotProduct(InAxis,CenterToMaxXMaxY);
+					const Profile::FReal CosAngleMinXMaxY = FVector2D::DotProduct(InAxis,CenterToMinXMaxY);
+
+					if (CosAngleMinXMinY >= 0 || CosAngleMaxXMinY >= 0 || CosAngleMaxXMaxY >= 0 || CosAngleMinXMaxY >= 0)
+					{
+						return true;
+					}
+
+					const Profile::FReal InvInSinHalfAngle = FMath::Sin((360.0 - InAngle) * 0.5 * UE_PI / 180.0);
+					const Profile::FReal InvSinAngleMinXMinY = FVector2D::CrossProduct(-InAxis, CenterToMinXMinY);
+					const Profile::FReal InvSinAngleMaxXMinY = FVector2D::CrossProduct(-InAxis, CenterToMaxXMinY);
+					const Profile::FReal InvSinAngleMaxXMaxY = FVector2D::CrossProduct(-InAxis, CenterToMaxXMaxY);
+					const Profile::FReal InvSinAngleMinXMaxY = FVector2D::CrossProduct(-InAxis, CenterToMinXMaxY);
+
+					if (InvSinAngleMinXMinY > -InvInSinHalfAngle && InvSinAngleMaxXMinY > -InvInSinHalfAngle && InvSinAngleMaxXMaxY > -InvInSinHalfAngle && InvSinAngleMinXMaxY > -InvInSinHalfAngle &&
+						InvSinAngleMinXMinY < InvInSinHalfAngle && InvSinAngleMaxXMinY < InvInSinHalfAngle && InvSinAngleMaxXMaxY < InvInSinHalfAngle && InvSinAngleMinXMaxY < InvInSinHalfAngle)
+					{
+						return false;
+					}
+
+					return true;
+				}
+			}
+		}
+
+		return false;
+	}
 }
 
 template <typename Profile>
@@ -92,7 +214,7 @@ public:
 	}
 
 	template <class Func>
-	void ForEachIntersectingElement(const FSphere& InSphere, Func InFunc) const
+	void ForEachIntersectingElement(const FStaticSpatialIndex::FSphere& InSphere, Func InFunc) const
 	{
 		TOverrideVoidReturnInvoker Invoker(true, InFunc);
 
@@ -102,6 +224,17 @@ public:
 		});
 	}
 
+	template <class Func>
+	void ForEachIntersectingElement(const FStaticSpatialIndex::FCone& InCone, Func InFunc) const
+	{
+		TOverrideVoidReturnInvoker Invoker(true, InFunc);
+
+		SpatialIndex.ForEachIntersectingElement(InCone, [this, &Invoker](uint32 ValueIndex)
+		{
+			return Invoker(Elements[ValueIndex].Value);
+		});
+	}
+
 	void AddReferencedObjects(FReferenceCollector& Collector)
 	{
 		if constexpr (TIsPointerOrObjectPtrToBaseOf<ValueType, UObject>::Value)
@@ -214,9 +347,9 @@ namespace FStaticSpatialIndex
 
 		}
 
-		bool ForEachIntersectingElement(const FSphere& InSphere, TFunctionRef<bool(uint32 InValueIndex)> InFunc) const
+		bool ForEachIntersectingElement(const FStaticSpatialIndex::FSphere& InSphere, TFunctionRef<bool(uint32 InValueIndex)> InFunc) const
 		{
-			const FSphere::FReal RadiusSquared = FMath::Square(InSphere.W);
+			const typename Profile::FReal RadiusSquared = FMath::Square(InSphere.Radius);
 
 			uint32 BoxStride;
 			const FBox* Box = this->DataInterface.GetBoxes(0, BoxStride);
@@ -233,6 +366,27 @@ namespace FStaticSpatialIndex
 			return true;
 		}
 
+		bool ForEachIntersectingElement(const FStaticSpatialIndex::FCone& InCone, TFunctionRef<bool(uint32 InValueIndex)> InFunc) const
+		{
+			const typename Profile::FReal RadiusSquared = FMath::Square(InCone.Radius);
+			const typename Profile::FReal HalfSinAngle = FMath::Sin(InCone.Angle * 0.5 * UE_PI / 180.0);
+			const typename Profile::FReal HalfCosAngle = FMath::Cos(InCone.Angle * 0.5 * UE_PI / 180.0);
+
+			uint32 BoxStride;
+			const FBox* Box = this->DataInterface.GetBoxes(0, BoxStride);
+			for (int32 ValueIndex = 0; ValueIndex < this->DataInterface.GetNumBox(); ValueIndex++, *(uint8**)&Box += BoxStride)
+			{
+				if (FastConeAABBIntersection<Profile>(FVector(InCone.Center), RadiusSquared, InCone.Axis, InCone.Angle, HalfSinAngle, HalfCosAngle, *Box))
+				{
+					if (!InFunc(ValueIndex))
+					{
+						return false;
+					}
+				}
+			}
+			return true;
+		}
+
 		uint32 GetAllocatedSize() const
 		{
 			return sizeof(*this);
@@ -337,12 +491,20 @@ namespace FStaticSpatialIndex
 			return ForEachIntersectingElementRecursive(&RootNode, InBox, InFunc);
 		}
 
-		bool ForEachIntersectingElement(const FSphere& InSphere, TFunctionRef<bool(uint32 InValueIndex)> InFunc) const
+		bool ForEachIntersectingElement(const FStaticSpatialIndex::FSphere& InSphere, TFunctionRef<bool(uint32 InValueIndex)> InFunc) const
 		{
-			const FSphere::FReal RadiusSquared = FMath::Square(InSphere.W);
+			const typename Profile::FReal RadiusSquared = FMath::Square(InSphere.Radius);
 			return ForEachIntersectingElementRecursive(&RootNode, FVector(InSphere.Center), RadiusSquared, InFunc);
 		}
 
+		bool ForEachIntersectingElement(const FStaticSpatialIndex::FCone& InCone, TFunctionRef<bool(uint32 InValueIndex)> InFunc) const
+		{
+			const typename Profile::FReal RadiusSquared = FMath::Square(InCone.Radius);
+			const typename Profile::FReal HalfSinAngle = FMath::Sin(InCone.Angle * 0.5 * UE_PI / 180.0);
+			const typename Profile::FReal HalfCosAngle = FMath::Cos(InCone.Angle * 0.5 * UE_PI / 180.0);
+			return ForEachIntersectingElementRecursive(&RootNode, FVector(InCone.Center), RadiusSquared, FVector(InCone.Axis), InCone.Angle, HalfSinAngle, HalfCosAngle, InFunc);
+		}
+
 		uint32 GetAllocatedSize() const
 		{
 			TFunction<uint32(const FNode*, uint32)> GetAllocatedSizeRecursive = [this, &GetAllocatedSizeRecursive](const FNode* Node, uint32 BaseSize)
@@ -484,7 +646,7 @@ namespace FStaticSpatialIndex
 			return true;
 		}
 
-		bool ForEachIntersectingElementRecursive(const FNode* InNode, const FVector& InSphereCenter, FSphere::FReal InRadiusSquared, TFunctionRef<bool(uint32 InValueIndex)> InFunc) const
+		bool ForEachIntersectingElementRecursive(const FNode* InNode, const FVector& InSphereCenter, typename Profile::FReal InRadiusSquared, TFunctionRef<bool(uint32 InValueIndex)> InFunc) const
 		{
 			if (InNode->Content.template IsType<typename FNode::FNodeType>())
 			{
@@ -520,6 +682,41 @@ namespace FStaticSpatialIndex
 			return true;
 		}
 
+		bool ForEachIntersectingElementRecursive(const FNode* InNode, const FVector& InConeCenter, typename Profile::FReal InRadiusSquared, const FVector& InAxis, float InAngle, float InSinHalfAngle, float InCosHalfAngle, TFunctionRef<bool(uint32 InValueIndex)> InFunc) const
+		{
+			if (InNode->Content.template IsType<typename FNode::FNodeType>())
+			{
+				for (auto& ChildNode : InNode->Content.template Get<typename FNode::FNodeType>())
+				{
+					if (FastSphereAABBIntersection<Profile>(FVector(InConeCenter), InRadiusSquared, ChildNode.GetBox()))
+					{
+						if (!ForEachIntersectingElementRecursive(&ChildNode, InConeCenter, InRadiusSquared, InAxis, InAngle, InSinHalfAngle, InCosHalfAngle, InFunc))
+						{
+							return false;
+						}
+					}
+				}
+			}
+			else
+			{
+				uint32 BoxStride;
+				const FBox* Box = this->DataInterface.GetBoxes(InNode->Content.template Get<typename FNode::FLeafType>().StartIndex, BoxStride);
+
+				for (uint32 ValueIndex : InNode->Content.template Get<typename FNode::FLeafType>())
+				{
+					if (FastConeAABBIntersection<Profile>(FVector(InConeCenter), InRadiusSquared, InAxis, InAngle, InSinHalfAngle, InCosHalfAngle, *Box))
+					{
+						if (!InFunc(ValueIndex))
+						{
+							return false;
+						}
+					}
+
+					*(uint8**)&Box += BoxStride;
+				}				
+			}
+			return true;
+		}
 		FNode RootNode;
 	};
 
