Added missing 'typename' prior to dependent type names.

diff --git a/Engine/Source/Runtime/Engine/Public/WorldPartition/RuntimeHashSet/StaticSpatialIndex.h b/Engine/Source/Runtime/Engine/Public/WorldPartition/RuntimeHashSet/StaticSpatialIndex.h
--- a/Engine/Source/Runtime/Engine/Public/WorldPartition/RuntimeHashSet/StaticSpatialIndex.h
+++ b/Engine/Source/Runtime/Engine/Public/WorldPartition/RuntimeHashSet/StaticSpatialIndex.h
@@ -99,10 +99,10 @@ namespace FStaticSpatialIndex
 
 				if (InAngle <= 180.0f)
 				{
-					const Profile::FReal SinAngleMinXMinY = FVector2D::CrossProduct(InAxis, CenterToMinXMinY);
-					const Profile::FReal SinAngleMaxXMinY = FVector2D::CrossProduct(InAxis, CenterToMaxXMinY);
-					const Profile::FReal SinAngleMaxXMaxY = FVector2D::CrossProduct(InAxis, CenterToMaxXMaxY);
-					const Profile::FReal SinAngleMinXMaxY = FVector2D::CrossProduct(InAxis, CenterToMinXMaxY);
+					const typename Profile::FReal SinAngleMinXMinY = FVector2D::CrossProduct(InAxis, CenterToMinXMinY);
+					const typename Profile::FReal SinAngleMaxXMinY = FVector2D::CrossProduct(InAxis, CenterToMaxXMinY);
+					const typename Profile::FReal SinAngleMaxXMaxY = FVector2D::CrossProduct(InAxis, CenterToMaxXMaxY);
+					const typename Profile::FReal SinAngleMinXMaxY = FVector2D::CrossProduct(InAxis, CenterToMinXMaxY);
 
 					if (SinAngleMinXMinY < -InSinHalfAngle && SinAngleMaxXMinY < -InSinHalfAngle && SinAngleMaxXMaxY < -InSinHalfAngle && SinAngleMinXMaxY < -InSinHalfAngle)
 					{
@@ -114,10 +114,10 @@ namespace FStaticSpatialIndex
 						return false;
 					}
 
-					const Profile::FReal CosAngleMinXMinY = FVector2D::DotProduct(InAxis, CenterToMinXMinY);
-					const Profile::FReal CosAngleMaxXMinY = FVector2D::DotProduct(InAxis, CenterToMaxXMinY);
-					const Profile::FReal CosAngleMaxXMaxY = FVector2D::DotProduct(InAxis, CenterToMaxXMaxY);
-					const Profile::FReal CosAngleMinXMaxY = FVector2D::DotProduct(InAxis, CenterToMinXMaxY);
+					const typename Profile::FReal CosAngleMinXMinY = FVector2D::DotProduct(InAxis, CenterToMinXMinY);
+					const typename Profile::FReal CosAngleMaxXMinY = FVector2D::DotProduct(InAxis, CenterToMaxXMinY);
+					const typename Profile::FReal CosAngleMaxXMaxY = FVector2D::DotProduct(InAxis, CenterToMaxXMaxY);
+					const typename Profile::FReal CosAngleMinXMaxY = FVector2D::DotProduct(InAxis, CenterToMinXMaxY);
 
 					if (CosAngleMinXMinY < 0 && CosAngleMaxXMinY < 0 && CosAngleMaxXMaxY < 0 && CosAngleMinXMaxY < 0)
 					{
@@ -128,21 +128,21 @@ namespace FStaticSpatialIndex
 				}
 				else
 				{
-					const Profile::FReal CosAngleMinXMinY = FVector2D::DotProduct(InAxis,CenterToMinXMinY);
-					const Profile::FReal CosAngleMaxXMinY = FVector2D::DotProduct(InAxis,CenterToMaxXMinY);
-					const Profile::FReal CosAngleMaxXMaxY = FVector2D::DotProduct(InAxis,CenterToMaxXMaxY);
-					const Profile::FReal CosAngleMinXMaxY = FVector2D::DotProduct(InAxis,CenterToMinXMaxY);
+					const typename Profile::FReal CosAngleMinXMinY = FVector2D::DotProduct(InAxis,CenterToMinXMinY);
+					const typename Profile::FReal CosAngleMaxXMinY = FVector2D::DotProduct(InAxis,CenterToMaxXMinY);
+					const typename Profile::FReal CosAngleMaxXMaxY = FVector2D::DotProduct(InAxis,CenterToMaxXMaxY);
+					const typename Profile::FReal CosAngleMinXMaxY = FVector2D::DotProduct(InAxis,CenterToMinXMaxY);
 
 					if (CosAngleMinXMinY >= 0 || CosAngleMaxXMinY >= 0 || CosAngleMaxXMaxY >= 0 || CosAngleMinXMaxY >= 0)
 					{
 						return true;
 					}
 
-					const Profile::FReal InvInSinHalfAngle = FMath::Sin((360.0 - InAngle) * 0.5 * UE_PI / 180.0);
-					const Profile::FReal InvSinAngleMinXMinY = FVector2D::CrossProduct(-InAxis, CenterToMinXMinY);
-					const Profile::FReal InvSinAngleMaxXMinY = FVector2D::CrossProduct(-InAxis, CenterToMaxXMinY);
-					const Profile::FReal InvSinAngleMaxXMaxY = FVector2D::CrossProduct(-InAxis, CenterToMaxXMaxY);
-					const Profile::FReal InvSinAngleMinXMaxY = FVector2D::CrossProduct(-InAxis, CenterToMinXMaxY);
+					const typename Profile::FReal InvInSinHalfAngle = FMath::Sin((360.0 - InAngle) * 0.5 * UE_PI / 180.0);
+					const typename Profile::FReal InvSinAngleMinXMinY = FVector2D::CrossProduct(-InAxis, CenterToMinXMinY);
+					const typename Profile::FReal InvSinAngleMaxXMinY = FVector2D::CrossProduct(-InAxis, CenterToMaxXMinY);
+					const typename Profile::FReal InvSinAngleMaxXMaxY = FVector2D::CrossProduct(-InAxis, CenterToMaxXMaxY);
+					const typename Profile::FReal InvSinAngleMinXMaxY = FVector2D::CrossProduct(-InAxis, CenterToMinXMaxY);
 
 					if (InvSinAngleMinXMinY > -InvInSinHalfAngle && InvSinAngleMaxXMinY > -InvInSinHalfAngle && InvSinAngleMaxXMaxY > -InvInSinHalfAngle && InvSinAngleMinXMaxY > -InvInSinHalfAngle &&
 						InvSinAngleMinXMinY < InvInSinHalfAngle && InvSinAngleMaxXMinY < InvInSinHalfAngle && InvSinAngleMaxXMaxY < InvInSinHalfAngle && InvSinAngleMinXMaxY < InvInSinHalfAngle)
