[NetPhysics] Tick-offset sync to handle when solver resets and ticks numbers start again from 0.

diff --git a/Engine/Source/Runtime/Engine/Classes/GameFramework/PlayerController.h b/Engine/Source/Runtime/Engine/Classes/GameFramework/PlayerController.h
--- a/Engine/Source/Runtime/Engine/Classes/GameFramework/PlayerController.h
+++ b/Engine/Source/Runtime/Engine/Classes/GameFramework/PlayerController.h
@@ -2419,14 +2419,8 @@ private:
 	int32 NetworkPhysicsTickOffset = INDEX_NONE;
 	bool bNetworkPhysicsTickOffsetAssigned = false;
 
-	/** The latest server step we've received a time dilation for. Needed for out of order updates */
-	int32 ClientLatestTimeDilationServerStep = INDEX_NONE;
-
-	/** The latest physics step we've sent to the server. Due to async we need to avoid duplicate sends */
-	int32 ClientLatestAsyncPhysicsStepSent = INDEX_NONE;
-
-	/** The latest physics step we've received from the client. */
-	int32 ServerLatestAsyncPhysicsStepReceived = INDEX_NONE;
+	/** Tick Offset sync interval timer */
+	float TickOffsetSyncCountdown = 0.f;
 
 	/** Time that tick offset has been desynced for since it was last in sync, used to trigger a correction */
 	float NetworkPhysicsTickOffsetDesyncAccumulatedTime = 0.0f;
diff --git a/Engine/Source/Runtime/Engine/Private/PlayerController.cpp b/Engine/Source/Runtime/Engine/Private/PlayerController.cpp
--- a/Engine/Source/Runtime/Engine/Private/PlayerController.cpp
+++ b/Engine/Source/Runtime/Engine/Private/PlayerController.cpp
@@ -5274,8 +5274,9 @@ void APlayerController::TickActor( float DeltaSeconds, ELevelTick TickType, FAct
 	// Clear old axis inputs since we are done with them. 
 	RotationInput = FRotator::ZeroRotator;
 
-	if(UPhysicsSettings::Get()->PhysicsPrediction.bEnablePhysicsPrediction && GetLocalRole() == ROLE_AutonomousProxy && bIsClient)
+	if (UPhysicsSettings::Get()->PhysicsPrediction.bEnablePhysicsPrediction && GetLocalRole() == ROLE_AutonomousProxy && bIsClient)
 	{
+		TickOffsetSyncCountdown += DeltaSeconds;
 		UpdateServerAsyncPhysicsTickOffset();
 	}
 
@@ -6238,18 +6239,13 @@ FAsyncPhysicsTimestamp APlayerController::GetPhysicsTimestamp(float DeltaSeconds
 
 void APlayerController::UpdateServerAsyncPhysicsTickOffset()
 {
-	FAsyncPhysicsTimestamp Timestamp = GetPhysicsTimestamp();
-	const int32 TickOffsetUpdateIntervalSteps = (NetworkPhysicsCvars::TickOffsetUpdateIntervalTime / 1000.f) / UPhysicsSettings::Get()->AsyncFixedTimeStepSize;
-
-	if (TickOffsetUpdateIntervalSteps <= 0 || ClientLatestAsyncPhysicsStepSent + TickOffsetUpdateIntervalSteps > Timestamp.LocalFrame)
+	if (NetworkPhysicsCvars::TickOffsetUpdateIntervalTime <= 0 || TickOffsetSyncCountdown < (NetworkPhysicsCvars::TickOffsetUpdateIntervalTime / 1000.0f))
 	{
-		//Only send a new timestamp if enough physics ticks have passed, based on CVar.
-		//If GT is running faster than physics sim the physics timestep will not have changed, so no need to send another update to server
-		//This ensures monotonic increase
 		return;
 	}
+	TickOffsetSyncCountdown = 0.0f;
 
-	ClientLatestAsyncPhysicsStepSent = Timestamp.LocalFrame;
+	FAsyncPhysicsTimestamp Timestamp = GetPhysicsTimestamp();
 	Timestamp.ServerFrame = bNetworkPhysicsTickOffsetAssigned ? Timestamp.ServerFrame : INDEX_NONE; // If offset is not yet assigned, set an invalid ServerFrame
 	ServerSendLatestAsyncPhysicsTimestamp(Timestamp);
 }
@@ -6258,15 +6254,6 @@ void APlayerController::ServerSendLatestAsyncPhysicsTimestamp_Implementation(FAs
 {
 	ensure(UPhysicsSettings::Get()->PhysicsPrediction.bEnablePhysicsPrediction);
 
-	//This tells the server how the client thinks the async physics tick will line up.
-	//If we have already received a more up to date timestamp from the client, early out
-	if (Timestamp.LocalFrame <= ServerLatestAsyncPhysicsStepReceived)
-	{
-		return;
-	}
-
-	ServerLatestAsyncPhysicsStepReceived = Timestamp.LocalFrame;
-
 	PRAGMA_DISABLE_DEPRECATION_WARNINGS
 	// Only cache the most up to date timestamp based on LocalFrame
 	if (Timestamp.LocalFrame > ServerPendingTimestamp_DEPRECATED.LocalFrame)
@@ -6375,12 +6362,6 @@ void APlayerController::ClientSetupNetworkPhysicsTimestamp_Implementation(FAsync
 
 void APlayerController::ClientAckTimeDilation_Implementation(float TimeDilation, int32 ServerStep)
 {
-	if (ServerStep <= ClientLatestTimeDilationServerStep)
-	{
-		return;
-	}
-	ClientLatestTimeDilationServerStep = ServerStep;
-
 	if (UWorld* World = GetWorld())
 	{
 		World->GetPhysicsScene()->SetNetworkDeltaTimeScale(TimeDilation);
