[NetPhysics] Only enable rewind data if async physics and physics prediction is used, else print a warning if it's trying to get enabled.

diff --git a/Engine/Source/Runtime/Engine/Classes/PhysicsEngine/PhysicsSettings.h b/Engine/Source/Runtime/Engine/Classes/PhysicsEngine/PhysicsSettings.h
--- a/Engine/Source/Runtime/Engine/Classes/PhysicsEngine/PhysicsSettings.h
+++ b/Engine/Source/Runtime/Engine/Classes/PhysicsEngine/PhysicsSettings.h
@@ -93,7 +93,8 @@ struct FPhysicsPredictionSettings
 
 	/** Enable networked physics prediction (experimental)
 	* Note: If an AActor::PhysicsReplicationMode is set to use Resimulation this will allow physics to cache history which is required by resimulation replication.
-	* Note: This can also affect how physics is solved even when not using resimulation. */
+	* Note: This can also affect how physics is solved even when not using resimulation. 
+	* Note: Physics Prediction only works when Physics -> Framerate -> Tick Physics Async is enabled */
 	UPROPERTY(EditAnywhere, Category = "Replication")
 	bool bEnablePhysicsPrediction;
 
diff --git a/Engine/Source/Runtime/Engine/Private/PhysicsEngine/PhysicsReplication.cpp b/Engine/Source/Runtime/Engine/Private/PhysicsEngine/PhysicsReplication.cpp
--- a/Engine/Source/Runtime/Engine/Private/PhysicsEngine/PhysicsReplication.cpp
+++ b/Engine/Source/Runtime/Engine/Private/PhysicsEngine/PhysicsReplication.cpp
@@ -387,7 +387,7 @@ void FPhysicsReplication::OnTick(float DeltaSeconds, TMap<TWeakObjectPtr<UPrimit
 
 	int32 LocalFrameOffset = 0; // LocalFrame = ServerFrame + LocalFrameOffset;
 
-	if (FPhysicsSolverBase::IsNetworkPhysicsPredictionEnabled())
+	if (UPhysicsSettings::Get()->PhysicsPrediction.bEnablePhysicsPrediction)
 	{
 		if (UWorld* World = GetOwningWorld())
 		{
@@ -921,7 +921,7 @@ void FPhysicsReplicationAsync::OnPreSimulate_Internal()
 			if (!bRewindDataExist && Input.RepMode == EPhysicsReplicationMode::Resimulation)
 			{
 				// We don't have rewind data but an actor is set to replicate using resimulation; we need to enable rewind capture.
-				if (ensure(Chaos::FPBDRigidsSolver::IsNetworkPhysicsPredictionEnabled()))
+				if (ensure(Chaos::FPBDRigidsSolver::IsNetworkPhysicsPredictionEnabled() && RigidsSolver->IsUsingFixedDt()))
 				{
 					const int32 NumFrames = FMath::Max<int32>(1, Chaos::FPBDRigidsSolver::GetPhysicsHistoryCount());
 					RigidsSolver->EnableRewindCapture(NumFrames, true);
diff --git a/Engine/Source/Runtime/Engine/Private/PlayerController.cpp b/Engine/Source/Runtime/Engine/Private/PlayerController.cpp
--- a/Engine/Source/Runtime/Engine/Private/PlayerController.cpp
+++ b/Engine/Source/Runtime/Engine/Private/PlayerController.cpp
@@ -5274,10 +5274,22 @@ void APlayerController::TickActor( float DeltaSeconds, ELevelTick TickType, FAct
 	// Clear old axis inputs since we are done with them. 
 	RotationInput = FRotator::ZeroRotator;
 
-	if (UPhysicsSettings::Get()->PhysicsPrediction.bEnablePhysicsPrediction && GetLocalRole() == ROLE_AutonomousProxy && bIsClient)
+	if (bIsClient && UPhysicsSettings::Get()->PhysicsPrediction.bEnablePhysicsPrediction && GetLocalRole() == ROLE_AutonomousProxy)
 	{
-		TickOffsetSyncCountdown += DeltaSeconds;
-		UpdateServerAsyncPhysicsTickOffset();
+		if (UWorld* World = GetWorld())
+		{
+			if (FPhysScene* PhysScene = World->GetPhysicsScene())
+			{
+				if (Chaos::FPhysicsSolver* Solver = PhysScene->GetSolver())
+				{
+					if (Solver->IsUsingFixedDt())
+					{
+						TickOffsetSyncCountdown += DeltaSeconds;
+						UpdateServerAsyncPhysicsTickOffset();
+					}
+				}
+			}
+		}
 	}
 
 #if !(UE_BUILD_SHIPPING || UE_BUILD_TEST)
diff --git a/Engine/Source/Runtime/Experimental/Chaos/Private/PBDRigidsSolver.cpp b/Engine/Source/Runtime/Experimental/Chaos/Private/PBDRigidsSolver.cpp
--- a/Engine/Source/Runtime/Experimental/Chaos/Private/PBDRigidsSolver.cpp
+++ b/Engine/Source/Runtime/Experimental/Chaos/Private/PBDRigidsSolver.cpp
@@ -1102,6 +1102,12 @@ namespace Chaos
 
 	void FPBDRigidsSolver::EnableRewindCapture(int32 NumFrames, bool InUseCollisionResimCache)
 	{
+		if (!IsUsingFixedDt())
+		{
+			UE_LOG(LogChaos, Warning, TEXT("PBDRigidsSolver::EnableRewindCapture Failed - Physics need to run with fixed tick. Project Settings -> Physics -> Framerate -> Tick Physics Async needs to be enabled."));
+			return;
+		}
+
 		//TODO: this function calls both internal and external - sort of assumed during initialization. Should decide what thread it's called on and mark it as either external or internal
 		if (MRewindData.IsValid())
 		{
