World Partition - HLOD: Improved hash computation to avoid unnecessary HLOD rebuilds in some cases
* Some of the code still relied on possibly non-persistent GetTypeHash()
* Fixed an issue with the transient Geometry Collection ISM Pool components getting a random value for their InstancingRandomSeed property.
* Improved verbose logging to ease debugging in the future

diff --git a/Engine/Plugins/Editor/WorldPartitionHLODUtilities/Source/Private/WorldPartition/HLOD/Builders/HLODBuilderMeshApproximate.cpp b/Engine/Plugins/Editor/WorldPartitionHLODUtilities/Source/Private/WorldPartition/HLOD/Builders/HLODBuilderMeshApproximate.cpp
--- a/Engine/Plugins/Editor/WorldPartitionHLODUtilities/Source/Private/WorldPartition/HLOD/Builders/HLODBuilderMeshApproximate.cpp
+++ b/Engine/Plugins/Editor/WorldPartitionHLODUtilities/Source/Private/WorldPartition/HLOD/Builders/HLODBuilderMeshApproximate.cpp
@@ -59,7 +59,7 @@ uint32 UHLODBuilderMeshApproximateSettings::GetCRC() const
 	{
 		uint32 MaterialCRC = UHLODProxy::GetCRC(HLODMaterial);
 		UE_LOG(LogHLODBuilder, VeryVerbose, TEXT(" - Material = %d"), MaterialCRC);
-		Hash = HashCombine(Hash, MaterialCRC);
+		Hash = HashCombineFast(Hash, MaterialCRC);
 	}
 
 	return Hash;
diff --git a/Engine/Plugins/Editor/WorldPartitionHLODUtilities/Source/Private/WorldPartition/HLOD/Builders/HLODBuilderMeshMerge.cpp b/Engine/Plugins/Editor/WorldPartitionHLODUtilities/Source/Private/WorldPartition/HLOD/Builders/HLODBuilderMeshMerge.cpp
--- a/Engine/Plugins/Editor/WorldPartitionHLODUtilities/Source/Private/WorldPartition/HLOD/Builders/HLODBuilderMeshMerge.cpp
+++ b/Engine/Plugins/Editor/WorldPartitionHLODUtilities/Source/Private/WorldPartition/HLOD/Builders/HLODBuilderMeshMerge.cpp
@@ -82,7 +82,7 @@ uint32 UHLODBuilderMeshMergeSettings::GetCRC() const
 	{
 		uint32 MaterialCRC = UHLODProxy::GetCRC(HLODMaterial);
 		UE_LOG(LogHLODBuilder, VeryVerbose, TEXT(" - Material = %d"), MaterialCRC);
-		Hash = HashCombine(Hash, MaterialCRC);
+		Hash = HashCombineFast(Hash, MaterialCRC);
 	}
 
 	return Hash;
diff --git a/Engine/Plugins/Editor/WorldPartitionHLODUtilities/Source/Private/WorldPartition/HLOD/Builders/HLODBuilderMeshSimplify.cpp b/Engine/Plugins/Editor/WorldPartitionHLODUtilities/Source/Private/WorldPartition/HLOD/Builders/HLODBuilderMeshSimplify.cpp
--- a/Engine/Plugins/Editor/WorldPartitionHLODUtilities/Source/Private/WorldPartition/HLOD/Builders/HLODBuilderMeshSimplify.cpp
+++ b/Engine/Plugins/Editor/WorldPartitionHLODUtilities/Source/Private/WorldPartition/HLOD/Builders/HLODBuilderMeshSimplify.cpp
@@ -64,7 +64,7 @@ uint32 UHLODBuilderMeshSimplifySettings::GetCRC() const
 	{
 		uint32 MaterialCRC = UHLODProxy::GetCRC(HLODMaterial);
 		UE_LOG(LogHLODBuilder, VeryVerbose, TEXT(" - Material = %d"), MaterialCRC);
-		Hash = HashCombine(Hash, MaterialCRC);
+		Hash = HashCombineFast(Hash, MaterialCRC);
 	}
 
 	return Hash;
diff --git a/Engine/Source/Runtime/Engine/Private/HLODProxy.cpp b/Engine/Source/Runtime/Engine/Private/HLODProxy.cpp
--- a/Engine/Source/Runtime/Engine/Private/HLODProxy.cpp
+++ b/Engine/Source/Runtime/Engine/Private/HLODProxy.cpp
@@ -31,6 +31,8 @@
 
 #include UE_INLINE_GENERATED_CPP_BY_NAME(HLODProxy)
 
+DEFINE_LOG_CATEGORY_STATIC(LogHLODProxy, Log, All);
+
 #if WITH_EDITOR
 
 void UHLODProxy::SetMap(const UWorld* InMap)
@@ -443,26 +445,41 @@ uint32 UHLODProxy::GetCRC(UStaticMeshComponent* InComponent, uint32 InCRC, const
 	ComponentLocation = TransformComponents.TransformPosition(ComponentLocation);
 	ComponentRotation = TransformComponents.TransformRotation(ComponentRotation.Quaternion()).Rotator();
 	AppendRoundedTransform(ComponentRotation, ComponentLocation, ComponentScale, Ar);
+	UE_LOG(LogHLODProxy, VeryVerbose, TEXT("         - Transform = %x"), Ar.GetCrc());
 
 	// Include other relevant properties
 	Ar << InComponent->ForcedLodModel;
+	UE_LOG(LogHLODProxy, VeryVerbose, TEXT("         - ForcedLodModel = %x"), Ar.GetCrc());
+
 	Ar << InComponent->HLODBatchingPolicy;
+	UE_LOG(LogHLODProxy, VeryVerbose, TEXT("         - HLODBatchingPolicy = %x"), Ar.GetCrc());
+
 	bool bCastShadow = InComponent->CastShadow;
 	Ar << bCastShadow;
+	UE_LOG(LogHLODProxy, VeryVerbose, TEXT("         - bCastShadow = %x"), Ar.GetCrc());
+
 	bool bCastStaticShadow = InComponent->bCastStaticShadow;
 	Ar << bCastStaticShadow;
+	UE_LOG(LogHLODProxy, VeryVerbose, TEXT("         - bCastStaticShadow = %x"), Ar.GetCrc());
+
 	bool bCastDynamicShadow = InComponent->bCastDynamicShadow;
 	Ar << bCastDynamicShadow;
+	UE_LOG(LogHLODProxy, VeryVerbose, TEXT("         - bCastDynamicShadow = %x"), Ar.GetCrc());
+
 	bool bCastFarShadow = InComponent->bCastFarShadow;
 	Ar << bCastFarShadow;
+	UE_LOG(LogHLODProxy, VeryVerbose, TEXT("         - bCastFarShadow = %x"), Ar.GetCrc());
+
 	int32 Width, Height;
 	InComponent->GetLightMapResolution(Width, Height);
 	Ar << Width;
 	Ar << Height;
+	UE_LOG(LogHLODProxy, VeryVerbose, TEXT("         - LightMapResolution = %x"), Ar.GetCrc());
 
 	if (!InComponent->GetCustomPrimitiveData().Data.IsEmpty())
 	{
 		Ar << InComponent->GetCustomPrimitiveData();
+		UE_LOG(LogHLODProxy, VeryVerbose, TEXT("         - CustomPrimitiveData = %x"), Ar.GetCrc());
 	}
 	
 	// incorporate vertex colors
@@ -471,6 +488,7 @@ uint32 UHLODProxy::GetCRC(UStaticMeshComponent* InComponent, uint32 InCRC, const
 		if(LODInfo.OverrideVertexColors)
 		{
 			Ar.Serialize((uint8*)LODInfo.OverrideVertexColors->GetVertexData(), LODInfo.OverrideVertexColors->GetNumVertices() * LODInfo.OverrideVertexColors->GetStride());
+			UE_LOG(LogHLODProxy, VeryVerbose, TEXT("         - OverrideVertexColors = %x"), Ar.GetCrc());
 		}
 	}
 
@@ -480,15 +498,22 @@ uint32 UHLODProxy::GetCRC(UStaticMeshComponent* InComponent, uint32 InCRC, const
 		for (const FInstancedStaticMeshInstanceData& InstanceData : ISMC->PerInstanceSMData)
 		{
 			AppendRoundedTransform(FTransform(InstanceData.Transform), Ar);
+			UE_LOG(LogHLODProxy, VeryVerbose, TEXT("         - Instance Transform = %x"), Ar.GetCrc());
 		}
 
 		Ar << ISMC->PerInstanceSMCustomData;
+		UE_LOG(LogHLODProxy, VeryVerbose, TEXT("         - PerInstanceSMCustomData = %x"), Ar.GetCrc());
+
 		Ar << ISMC->InstancingRandomSeed;
+		UE_LOG(LogHLODProxy, VeryVerbose, TEXT("         - InstancingRandomSeed = %x"), Ar.GetCrc());
 
 		for (FInstancedStaticMeshRandomSeed& ISMCRandomSeed : ISMC->AdditionalRandomSeeds)
 		{
 			Ar << ISMCRandomSeed.StartInstanceIndex;
+			UE_LOG(LogHLODProxy, VeryVerbose, TEXT("         - StartInstanceIndex = %x"), Ar.GetCrc());
+
 			Ar << ISMCRandomSeed.RandomSeed;
+			UE_LOG(LogHLODProxy, VeryVerbose, TEXT("         - RandomSeed = %x"), Ar.GetCrc());
 		}
 	}
 
@@ -648,7 +673,7 @@ FName UHLODProxy::GenerateKeyForActor(const ALODActor* LODActor, bool bMustUndoL
 		// Append all components CRCs
 		for (uint32 ComponentCRC : ComponentsCRCs)
 		{
-			CRC = HashCombine(CRC, ComponentCRC);
+			CRC = HashCombineFast(CRC, ComponentCRC);
 		}
 
 		Key += TEXT("_");
diff --git a/Engine/Source/Runtime/Engine/Private/WorldPartition/HLOD/HLODBuilder.cpp b/Engine/Source/Runtime/Engine/Private/WorldPartition/HLOD/HLODBuilder.cpp
--- a/Engine/Source/Runtime/Engine/Private/WorldPartition/HLOD/HLODBuilder.cpp
+++ b/Engine/Source/Runtime/Engine/Private/WorldPartition/HLOD/HLODBuilder.cpp
@@ -74,8 +74,9 @@ uint32 UHLODBuilder::ComputeHLODHash(const UActorComponent* InSourceComponent) c
 		// CRC static mesh
 		if (UStaticMesh* StaticMesh = StaticMeshComponent->GetStaticMesh())
 		{
-			ComponentCRC = UHLODProxy::GetCRC(StaticMesh, ComponentCRC);
-			UE_LOG(LogHLODBuilder, VeryVerbose, TEXT("     - Static Mesh (%s) = %x"), *StaticMesh->GetName(), ComponentCRC);
+			uint32 StaticMeshCRC = UHLODProxy::GetCRC(StaticMesh);
+			UE_LOG(LogHLODBuilder, VeryVerbose, TEXT("     - Static Mesh (%s) = %x"), *StaticMesh->GetName(), StaticMeshCRC);
+			ComponentCRC = HashCombineFast(ComponentCRC, StaticMeshCRC);
 		}
 
 		// CRC materials
@@ -85,29 +86,33 @@ uint32 UHLODBuilder::ComputeHLODHash(const UActorComponent* InSourceComponent) c
 			UMaterialInterface* MaterialInterface = StaticMeshComponent->GetMaterial(MaterialIndex);
 			if (MaterialInterface)
 			{
-				ComponentCRC = UHLODProxy::GetCRC(MaterialInterface, ComponentCRC);
-				UE_LOG(LogHLODBuilder, VeryVerbose, TEXT("     - Material (%s) = %x"), *MaterialInterface->GetName(), ComponentCRC);
+				uint32 MaterialInterfaceCRC = UHLODProxy::GetCRC(MaterialInterface);
+				UE_LOG(LogHLODBuilder, VeryVerbose, TEXT("     - Material (%s) = %x"), *MaterialInterface->GetName(), MaterialInterfaceCRC);
+				ComponentCRC = HashCombineFast(ComponentCRC, MaterialInterfaceCRC);
 
 				TArray<UTexture*> Textures;
 				MaterialInterface->GetUsedTextures(Textures, EMaterialQualityLevel::High, true, ERHIFeatureLevel::SM5, true);
 				for (UTexture* Texture : Textures)
 				{
-					ComponentCRC = UHLODProxy::GetCRC(Texture, ComponentCRC);
-					UE_LOG(LogHLODBuilder, VeryVerbose, TEXT("     - Texture (%s) = %x"), *Texture->GetName(), ComponentCRC);
+					uint32 TextureCRC = UHLODProxy::GetCRC(Texture);
+					UE_LOG(LogHLODBuilder, VeryVerbose, TEXT("     - Texture (%s) = %x"), *Texture->GetName(), TextureCRC);
+					ComponentCRC = HashCombineFast(ComponentCRC, TextureCRC);
 				}
 			}
 			UMaterialInterface* NaniteOverride = MaterialInterface ? MaterialInterface->GetNaniteOverride() : nullptr;
 			if (NaniteOverride)
 			{
-				ComponentCRC = UHLODProxy::GetCRC(NaniteOverride, ComponentCRC);
-				UE_LOG(LogHLODBuilder, VeryVerbose, TEXT("     - Material (%s) = %x"), *NaniteOverride->GetName(), ComponentCRC);
+				uint32 NaniteOverrideCRC = UHLODProxy::GetCRC(NaniteOverride);
+				UE_LOG(LogHLODBuilder, VeryVerbose, TEXT("     - Nanite Override Material (%s) = %x"), *NaniteOverride->GetName(), NaniteOverrideCRC);
+				ComponentCRC = HashCombineFast(ComponentCRC, NaniteOverrideCRC);
 
 				TArray<UTexture*> Textures;
 				NaniteOverride->GetUsedTextures(Textures, EMaterialQualityLevel::High, true, ERHIFeatureLevel::SM5, true);
 				for (UTexture* Texture : Textures)
 				{
-					ComponentCRC = UHLODProxy::GetCRC(Texture, ComponentCRC);
-					UE_LOG(LogHLODBuilder, VeryVerbose, TEXT("     - Texture (%s) = %x"), *Texture->GetName(), ComponentCRC);
+					uint32 TextureCRC = UHLODProxy::GetCRC(Texture);
+					UE_LOG(LogHLODBuilder, VeryVerbose, TEXT("     - Nanite Override Texture (%s) = %x"), *Texture->GetName(), TextureCRC);
+					ComponentCRC = HashCombineFast(ComponentCRC, TextureCRC);
 				}
 			}
 		}
diff --git a/Engine/Source/Runtime/Engine/Private/WorldPartition/HLOD/HLODSourceActors.cpp b/Engine/Source/Runtime/Engine/Private/WorldPartition/HLOD/HLODSourceActors.cpp
--- a/Engine/Source/Runtime/Engine/Private/WorldPartition/HLOD/HLODSourceActors.cpp
+++ b/Engine/Source/Runtime/Engine/Private/WorldPartition/HLOD/HLODSourceActors.cpp
@@ -16,15 +16,15 @@ static uint32 GetHLODHash(const UHLODLayer* InHLODLayer)
 
 	uint32 HLODHash;
 
-	HLODHash = GetTypeHash(HLODLayer.GetLayerType());
+	HLODHash = static_cast<uint8>(HLODLayer.GetLayerType());
 	UE_LOG(LogHLODHash, VeryVerbose, TEXT(" - LayerType = %d"), HLODHash);
-
-	HLODHash = HashCombine(HLODLayer.GetHLODBuilderSettings()->GetCRC(), HLODHash);
+	
+	HLODHash = HashCombineFast(HLODLayer.GetHLODBuilderSettings()->GetCRC(), HLODHash);
 	UE_LOG(LogHLODHash, VeryVerbose, TEXT(" - HLODBuilderSettings = %d"), HLODHash);
 
-	HLODHash = HashCombine(HLODLayer.GetCellSize(), HLODHash);
+	HLODHash = HashCombineFast(HLODLayer.GetCellSize(), HLODHash);
 	UE_LOG(LogHLODHash, VeryVerbose, TEXT(" - CellSize = %d"), HLODHash);
-
+	
 	return HLODHash;
 }
 
diff --git a/Engine/Source/Runtime/Engine/Private/WorldPartition/HLOD/HLODSourceActorsFromCell.cpp b/Engine/Source/Runtime/Engine/Private/WorldPartition/HLOD/HLODSourceActorsFromCell.cpp
--- a/Engine/Source/Runtime/Engine/Private/WorldPartition/HLOD/HLODSourceActorsFromCell.cpp
+++ b/Engine/Source/Runtime/Engine/Private/WorldPartition/HLOD/HLODSourceActorsFromCell.cpp
@@ -56,7 +56,7 @@ uint32 UWorldPartitionHLODSourceActorsFromCell::GetHLODHash() const
 	// Source Actors
 	uint32 SourceActorsHash = GetHLODHash(Actors);
 	UE_LOG(LogHLODHash, VeryVerbose, TEXT(" - Source Actors (%d actors) = %x"), Actors.Num(), SourceActorsHash);
-	HLODHash = HashCombine(HLODHash, SourceActorsHash);
+	HLODHash = HashCombineFast(HLODHash, SourceActorsHash);
 
 	return HLODHash;
 }
diff --git a/Engine/Source/Runtime/Engine/Private/WorldPartition/HLOD/HLODSourceActorsFromLevel.cpp b/Engine/Source/Runtime/Engine/Private/WorldPartition/HLOD/HLODSourceActorsFromLevel.cpp
--- a/Engine/Source/Runtime/Engine/Private/WorldPartition/HLOD/HLODSourceActorsFromLevel.cpp
+++ b/Engine/Source/Runtime/Engine/Private/WorldPartition/HLOD/HLODSourceActorsFromLevel.cpp
@@ -35,12 +35,8 @@ ULevelStreaming* UWorldPartitionHLODSourceActorsFromLevel::LoadSourceActors(bool
 
 uint32 UWorldPartitionHLODSourceActorsFromLevel::GetHLODHash() const
 {
-	uint32 HLODHash = Super::GetHLODHash();
-
 	// SourceLevel
-	HLODHash = HashCombine(HLODHash, GetTypeHash(SourceLevel));
-	
-	return HLODHash;
+	return FCrc::StrCrc32(*SourceLevel.GetUniqueID().GetLongPackageName(), Super::GetHLODHash());
 }
 
 
diff --git a/Engine/Source/Runtime/Experimental/GeometryCollectionEngine/Private/GeometryCollection/GeometryCollectionISMPoolComponent.cpp b/Engine/Source/Runtime/Experimental/GeometryCollectionEngine/Private/GeometryCollection/GeometryCollectionISMPoolComponent.cpp
--- a/Engine/Source/Runtime/Experimental/GeometryCollectionEngine/Private/GeometryCollection/GeometryCollectionISMPoolComponent.cpp
+++ b/Engine/Source/Runtime/Experimental/GeometryCollectionEngine/Private/GeometryCollection/GeometryCollectionISMPoolComponent.cpp
@@ -220,6 +220,10 @@ void FGeometryCollectionISM::InitISM(const FGeometryCollectionStaticMeshInstance
 	ISMComponent->bDisallowNanite = bDisallowNanite;
 	ISMComponent->SetMeshDrawCommandStatsCategory(MeshInstance.Desc.StatsCategory);
 	ISMComponent->ComponentTags = MeshInstance.Desc.Tags;
+
+	// Use a fixed seed to avoid getting a different seed at every run (see UInstancedStaticMeshComponent::OnRegister())
+	// A possible improvement would be to compute an hash from the owner Geometry Collection component and use that as the seed.
+	ISMComponent->InstancingRandomSeed = 1;	
 }
 
 FInstanceGroups::FInstanceGroupId FGeometryCollectionISM::AddInstanceGroup(int32 InstanceCount, TArrayView<const float> CustomDataFloats)
