[shaders] speculative fix for one potential cause of DDC corruption. The previous attempt to delete the worker output files on relaunch was not sufficient to handle the case where relaunch happens so quickly that a worker may still be running (and so generate an output file for the wrong jobs after the attempt at stale output deletion). Instead we now forcibly terminate all worker processes when abandoning workers, and delete any files which have already been written at that point.

diff --git a/Engine/Source/Runtime/Engine/Private/ShaderCompiler/ShaderCompilerThreadRunnable.cpp b/Engine/Source/Runtime/Engine/Private/ShaderCompiler/ShaderCompilerThreadRunnable.cpp
--- a/Engine/Source/Runtime/Engine/Private/ShaderCompiler/ShaderCompilerThreadRunnable.cpp
+++ b/Engine/Source/Runtime/Engine/Private/ShaderCompiler/ShaderCompilerThreadRunnable.cpp
@@ -589,9 +589,6 @@ bool FShaderCompileThreadRunnable::LaunchWorkersIfNeeded()
 		LastCheckForWorkersTime = CurrentTime;
 	}
 
-	bool bClearStaleOutputs = Manager->bClearStaleWorkerOutputs;
-	Manager->bClearStaleWorkerOutputs = false;
-
 	FScopeLock WorkerScopeLock(&WorkerInfosLock);
 	for (int32 WorkerIndex = 0; WorkerIndex < WorkerInfos.Num(); WorkerIndex++)
 	{
@@ -645,12 +642,6 @@ bool FShaderCompileThreadRunnable::LaunchWorkersIfNeeded()
 			{
 				constexpr bool bRelativePath = true;
 				const FString WorkingDirectory = GetWorkingDirectoryForWorker(WorkerIndex, bRelativePath);
-				
-				if (bClearStaleOutputs)
-				{
-					// Delete any potential stale output files; these can persist if we had previously abandoned running with workers due to an unexpected SCW termination.
-					DiscardWorkerOutputFile(WorkerIndex);
-				}
 
 				// Store the handle with this thread so that we will know not to launch it again
 				CurrentWorkerInfo.WorkerProcess = Manager->LaunchWorker(WorkingDirectory, Manager->ProcessId, WorkerIndex, GWorkerInputFilename, GWorkerOutputFilename);
@@ -1253,6 +1244,14 @@ int32 FShaderCompileThreadRunnable::CompilingLoop()
 
 			if (bAbandonWorkers)
 			{
+				// Immediately terminate all worker processes and delete any output files they may have generated;
+				// we will re-run all these jobs locally instead.
+				for (int32 WorkerIndex = 0; WorkerIndex < WorkerInfos.Num(); ++WorkerIndex)
+				{
+					WorkerInfos[WorkerIndex]->TerminateWorkerProcess();
+					DiscardWorkerOutputFile(WorkerIndex);
+				}
+
 				// Fall back to local compiles if the SCW crashed.
 				// This is nasty but needed to work around issues where message passing through files to SCW is unreliable on random PCs
 				Manager->bAllowCompilingThroughWorkers = false;
@@ -1297,7 +1296,6 @@ int32 FShaderCompileThreadRunnable::CompilingLoop()
 				if (Manager->NumSingleThreadedRunsBeforeRetry == 0)
 				{
 					UE_LOG(LogShaderCompilers, Display, TEXT("Retry shader compiling through workers."));
-					Manager->bClearStaleWorkerOutputs = true;
 					Manager->bAllowCompilingThroughWorkers = true;
 				}
 			}
diff --git a/Engine/Source/Runtime/Engine/Public/ShaderCompiler.h b/Engine/Source/Runtime/Engine/Public/ShaderCompiler.h
--- a/Engine/Source/Runtime/Engine/Public/ShaderCompiler.h
+++ b/Engine/Source/Runtime/Engine/Public/ShaderCompiler.h
@@ -789,8 +789,6 @@ private:
 	bool bPromptToRetryFailedShaderCompiles;
 	/** If enabled when we enter the prompt to retry we will break in the debugger if one is attached rather than prompting. */
 	bool bDebugBreakOnPromptToRetryShaderCompile = false;
-	/** If true, when relaunching workers will clear out any existing output files. This is only used when workers are abandoned due to an unexpected termination */
-	bool bClearStaleWorkerOutputs = false;
 	/** Whether to log out shader job completion times on the worker thread.  Useful for tracking down which global shader is taking a long time. */
 	bool bLogJobCompletionTimes;
 	/** Target execution time for ProcessAsyncResults.  Larger values speed up async shader map processing but cause more hitchiness while async compiling is happening. */
