Prevent modification of FSoundWaveProxy while it is being rendered.

diff --git a/Engine/Plugins/Runtime/AudioCapture/Source/AudioCaptureEditor/Private/AudioRecordingManager.cpp b/Engine/Plugins/Runtime/AudioCapture/Source/AudioCaptureEditor/Private/AudioRecordingManager.cpp
--- a/Engine/Plugins/Runtime/AudioCapture/Source/AudioCaptureEditor/Private/AudioRecordingManager.cpp
+++ b/Engine/Plugins/Runtime/AudioCapture/Source/AudioCaptureEditor/Private/AudioRecordingManager.cpp
@@ -411,9 +411,8 @@ TObjectPtr<USoundWave> FAudioRecordingManager::CreateSoundWaveAsset(const FRecor
 		NewSoundWave->SetSoundAssetCompressionType(ESoundAssetCompressionType::BinkAudio);
 		NewSoundWave->SetTimecodeInfo(GetTimecodeInfo(InSourceSettings));
 
-		// Initialize SoundWaveData so that it is synchronized with the owning object
-		// (note that for serialized cases, this happens in PostLoad(), PostImport(), etc.)
-		NewSoundWave->SoundWaveDataPtr->InitializeDataFromSoundWave(*NewSoundWave);
+		// Post import initializes USoundWave internal data structures 
+		NewSoundWave->PostImport();
 
 		if (bCreatedPackage)
 		{
diff --git a/Engine/Source/Runtime/Engine/Classes/Sound/SoundWave.h b/Engine/Source/Runtime/Engine/Classes/Sound/SoundWave.h
--- a/Engine/Source/Runtime/Engine/Classes/Sound/SoundWave.h
+++ b/Engine/Source/Runtime/Engine/Classes/Sound/SoundWave.h
@@ -703,7 +703,7 @@ public:
 
 #if WITH_EDITOR
 	/** The current revision of our compressed audio data. Used to tell when a chunk in the cache is stale. */
-	TSharedPtr<FThreadSafeCounter> CurrentChunkRevision{ MakeShared<FThreadSafeCounter>() };
+	std::atomic<int32> CurrentChunkRevision = 0;
 #endif
 
 private:
@@ -1098,7 +1098,12 @@ public:
 	// or zero if it is not a streaming source.
 	ENGINE_API uint32 GetNumChunks() const;
 
-	ENGINE_API uint32 GetSizeOfChunk(uint32 ChunkIndex);
+	ENGINE_API uint32 GetSizeOfChunk(uint32 ChunkIndex) const;
+
+	// Calculates audio streaming cache usage for this sound wave. 
+	// @param OutTotalBytesOfAudioData - The accumulated number of bytes of all audio data of all chunks for this sound wave.
+	// @param OutMaxChunkBytesOfAudioData - The maximum number of bytes of audio data for all audio chunks for this sound wave. 
+	ENGINE_API void GetChunkSizeStats(uint32& OutTotalBytesOfAudioData, uint32& OutMaxChunkBytesOfAudioData) const;
 
 	ENGINE_API virtual void BeginDestroy() override;
 #if WITH_EDITOR
@@ -1122,6 +1127,9 @@ private:
 	/**  Utility function used internally to wait or poll a task while maintaining thread-safety. */
 	ENGINE_API bool WaitAsyncTaskWithTimeout(float InTimeoutInSeconds);
 
+	/** Creates and initializes a new FSoundWaveData and FSoundWaveProxy. */
+	void CreateNewSoundWaveData();
+
 public:
 #endif // WITH_EDITOR
 
@@ -1516,9 +1524,10 @@ public:
 		return (ESoundWavePrecacheState)PrecacheState.GetValue();
 	}
 
-	TSharedPtr<FSoundWaveData, ESPMode::ThreadSafe> SoundWaveDataPtr{ MakeShared<FSoundWaveData>() };
 
 private:
+	TSharedPtr<FSoundWaveData, ESPMode::ThreadSafe> SoundWaveDataPtr{ MakeShared<FSoundWaveData>() };
+
 	friend class FSoundWaveProxy;
 	friend class USoundFactory;
 };
@@ -1529,12 +1538,6 @@ class FSoundWaveData
 public:
 	UE_NONCOPYABLE(FSoundWaveData);
 
-	struct MaxChunkSizeResults
-	{
-		uint32 MaxUnevictableSize = 0;
-		uint32 MaxSizeInCache = 0;
-	};
-	
 	FSoundWaveData()
 		: bIsLooping(0)
 		, bIsTemplate(0)
@@ -1569,8 +1572,6 @@ public:
 	const TArray<FSoundWaveCuePoint>& GetLoopRegions() const { return LoopRegions; }
 	void SetAllCuePoints(const TArray<FSoundWaveCuePoint>& InCuePoints);
 
-	ENGINE_API MaxChunkSizeResults GetMaxChunkSizeResults() const;
-
 	ENGINE_API uint32 GetNumChunks() const;
 	ENGINE_API uint32 GetSizeOfChunk(uint32 ChunkIndex) const;
 	int32 GetNumFrames() const { return NumFrames; }
@@ -1659,7 +1660,7 @@ private:
 #endif //WITH_EDITORONLY_DATA
 	
 #if WITH_EDITOR
-	std::atomic<int32> CurrentChunkRevision;
+	int32 CurrentChunkRevision;
 #endif // #if WITH_EDITOR
 
 	FName NameCached;
@@ -1730,8 +1731,6 @@ public:
 	ENGINE_API const TArray<FSoundWaveCuePoint>& GetCuePoints() const;
 	ENGINE_API const TArray<FSoundWaveCuePoint>& GetLoopRegions() const;
 
-	ENGINE_API FSoundWaveData::MaxChunkSizeResults GetMaxChunkSizeResults() const;
-
 	ENGINE_API bool IsLooping() const;
 	ENGINE_API bool IsTemplate() const;
 	ENGINE_API bool IsStreaming() const;
diff --git a/Engine/Source/Runtime/Engine/Private/AudioDevice.cpp b/Engine/Source/Runtime/Engine/Private/AudioDevice.cpp
--- a/Engine/Source/Runtime/Engine/Private/AudioDevice.cpp
+++ b/Engine/Source/Runtime/Engine/Private/AudioDevice.cpp
@@ -1897,11 +1897,13 @@ bool FAudioDevice::HandleAudioMemoryInfo(const TCHAR* Cmd, FOutputDevice& Ar)
 					break;
 			}
 
-			check(SoundWave->SoundWaveDataPtr);
-			FSoundWaveData::MaxChunkSizeResults MaxChunkSizes = SoundWave->SoundWaveDataPtr->GetMaxChunkSizeResults();
+			// Calculate streaming cache size usage. 
+			uint32 MaxSizeInCache = 0;
+			uint32 MaxUnevictableSize = 0;
+			SoundWave->GetChunkSizeStats(MaxSizeInCache, MaxUnevictableSize);
 
 			// Add the info to the SoundWaveObjects array
-			SoundWaveObjects.Add(FSoundWaveInfo(SoundWave, TrueResourceSize, SoundGroupName, SoundWave->Duration, LoadType, MaxChunkSizes.MaxUnevictableSize, MaxChunkSizes.MaxSizeInCache));
+			SoundWaveObjects.Add(FSoundWaveInfo(SoundWave, TrueResourceSize, SoundGroupName, SoundWave->Duration, LoadType, MaxUnevictableSize, MaxSizeInCache));
 
 			// Track total resource usage
 			TotalResourceSize += TrueResourceSize;
diff --git a/Engine/Source/Runtime/Engine/Private/SoundWave.cpp b/Engine/Source/Runtime/Engine/Private/SoundWave.cpp
--- a/Engine/Source/Runtime/Engine/Private/SoundWave.cpp
+++ b/Engine/Source/Runtime/Engine/Private/SoundWave.cpp
@@ -332,19 +332,6 @@ void FSoundWaveData::SetAllCuePoints(const TArray<FSoundWaveCuePoint>& InCuePoin
 	}
 }
 
-FSoundWaveData::MaxChunkSizeResults FSoundWaveData::GetMaxChunkSizeResults() const
-{
-	MaxChunkSizeResults Results;
-
-	for (auto& Chunk : RunningPlatformData.GetChunks())
-	{
-		Results.MaxUnevictableSize = FMath::Max<uint32>(Results.MaxUnevictableSize, Chunk.AudioDataSize);
-		Results.MaxSizeInCache += Chunk.AudioDataSize;
-	}
-
-	return Results;
-}
-
 uint32 FSoundWaveData::GetSizeOfChunk(uint32 ChunkIndex) const
 {
 	check(ChunkIndex < GetNumChunks());
@@ -359,6 +346,7 @@ uint32 FSoundWaveData::GetSizeOfChunk(uint32 ChunkIndex) const
 	return RunningPlatformData.GetChunks()[ChunkIndex].AudioDataSize;
 }
 
+
 void FSoundWaveData::ReleaseCompressedAudio()
 {
 	LLM_SCOPE(ELLMTag::AudioSoundWaves);
@@ -1887,9 +1875,17 @@ void USoundWave::InvalidateCompressedData(bool bFreeResources, bool bRebuildStre
 
 	CompressedDataGuid = FGuid::NewGuid();
 
+#if WITH_EDITOR
+	// In editor, do not modify SoundWaveDataPtr because it may be accessed by 
+	// another system through the FSoundWaveProxy. Instead, create a new SoundWaveDataPtr
+	// to ensure that SoundWaveDataPtr is not being read while it is being written to. 
+	const bool bIsRetained = SoundWaveDataPtr->FirstChunk.IsValid();
+	CreateNewSoundWaveData();
+#else
 	SoundWaveDataPtr->DiscardZerothChunkData();
 	SoundWaveDataPtr->CompressedFormatData.FlushData();
 	RemoveAudioResource();
+#endif // WITH_EDITOR
 
 	if (bFreeResources)
 	{
@@ -1908,14 +1904,13 @@ void USoundWave::InvalidateCompressedData(bool bFreeResources, bool bRebuildStre
 	{
 		CachePlatformData(true /* bAsyncCache */);
 
-		SoundWaveDataPtr->CurrentChunkRevision += 1;
+		CurrentChunkRevision += 1;
+		SoundWaveDataPtr->CurrentChunkRevision = CurrentChunkRevision;
 	}
 
-
 	// If this sound wave is retained, release and re-retain the new chunk.
-	if (SoundWaveDataPtr->FirstChunk.IsValid())
+	if (bIsRetained)
 	{
-		ReleaseCompressedAudio();
 		RetainCompressedAudio(true);
 	}
 #endif
@@ -2011,15 +2006,15 @@ void USoundWave::PostLoad()
 	}
 #endif
 
+	MaxDistance = ComputeMaxDistance();
+
 	// Don't need to do anything in post load if this is a source bus or procedural audio
 	if (this->IsA(USoundSourceBus::StaticClass()) || bProcedural)
 	{
 		return;
 	}
 
-
 	CacheInheritedLoadingBehavior();
-	MaxDistance = ComputeMaxDistance();
 	
 	if (FApp::CanEverRenderAudio())
 	{
@@ -2209,7 +2204,7 @@ uint32 USoundWave::GetNumChunks() const
 	}
 }
 
-uint32 USoundWave::GetSizeOfChunk(uint32 ChunkIndex)
+uint32 USoundWave::GetSizeOfChunk(uint32 ChunkIndex) const
 {
 	check(ChunkIndex < GetNumChunks());
 
@@ -2223,6 +2218,21 @@ uint32 USoundWave::GetSizeOfChunk(uint32 ChunkIndex)
 	}
 }
 
+void USoundWave::GetChunkSizeStats(uint32& OutTotalBytesOfAudioData, uint32& OutMaxChunkBytesOfAudioData) const
+{
+	OutTotalBytesOfAudioData = 0;
+	OutMaxChunkBytesOfAudioData = 0;
+
+	if (SoundWaveDataPtr.IsValid())
+	{
+		for (const FStreamedAudioChunk& Chunk : SoundWaveDataPtr->RunningPlatformData.GetChunks())
+		{
+			OutTotalBytesOfAudioData += Chunk.AudioDataSize;
+			OutMaxChunkBytesOfAudioData = FMath::Max((uint32)Chunk.AudioDataSize, OutMaxChunkBytesOfAudioData);
+		}
+	}
+}
+
 void USoundWave::BeginDestroy()
 {
 	Super::BeginDestroy();
@@ -2918,12 +2928,24 @@ void USoundWave::BakeEnvelopeAnalysis()
 void USoundWave::PostEditChangeProperty(FPropertyChangedEvent& PropertyChangedEvent)
 {
 	Super::PostEditChangeProperty(PropertyChangedEvent);
-	check(SoundWaveDataPtr);
 
-	if (PropertyChangedEvent.Property == nullptr)
+	MaxDistance = ComputeMaxDistance();
+
+	// Source buses and procedural sources do not have sample data and should 
+	// avoid going through creation of streaming compressed data.
+	const bool bHasSampleData = !(bIsSourceBus || bProcedural);
+	if (!bHasSampleData)
+	{
+		return;
+	}
+
+	if (PropertyChangedEvent.Property == nullptr) 
 	{
 		//an empty event property field might mean the update comes from an undo
 		//we can't discern what properties where reverted so we update the asset and bakes
+	
+		// Avoid modifying FSoundWaveProxy that is still in use by outside systems. 
+		CreateNewSoundWaveData();
 
 		UpdateAsset();
 		BakeFFTAnalysis();
@@ -2943,9 +2965,6 @@ void USoundWave::PostEditChangeProperty(FPropertyChangedEvent& PropertyChangedEv
 	static const FName CloudStreamingFName = GET_MEMBER_NAME_CHECKED(USoundWave, bEnableCloudStreaming);
 	static const FName CloudStreamingPlatformSettingsFName = GET_MEMBER_NAME_CHECKED(USoundWave, PlatformSettings);
 
-	// force proxy state to be up to date
-	SoundWaveDataPtr->InitializeDataFromSoundWave(*this);
-
 	if (FProperty* PropertyThatChanged = PropertyChangedEvent.Property)
 	{
 		const FName& Name = PropertyThatChanged->GetFName();
@@ -2959,6 +2978,9 @@ void USoundWave::PostEditChangeProperty(FPropertyChangedEvent& PropertyChangedEv
 		// Prevent constant re-compression of SoundWave while properties are being changed interactively
 		if (PropertyChangedEvent.ChangeType != EPropertyChangeType::Interactive)
 		{
+			// Avoid modifying FSoundWaveProxy that is still in use by outside systems. 
+			CreateNewSoundWaveData();
+
 			// Regenerate on save any compressed sound formats or if analysis needs to be re-done
 			if (Name == LoadingBehaviorFName || Name == InlinedAudioInSecondsFName)
 			{
@@ -3007,7 +3029,41 @@ void USoundWave::PostEditChangeProperty(FPropertyChangedEvent& PropertyChangedEv
 		}
 	}
 
-	MaxDistance = ComputeMaxDistance();
+}
+
+void USoundWave::CreateNewSoundWaveData()
+{
+	// Create a new FSoundWaveData to avoid altering an FSoundWaveData while it is
+	// being decoded by FSoundWaveProxy clients.
+	//
+	// USoundWave currently does not have a way to stop clients which utilize
+	// an FSoundWaveProxy or FSoundWaveData outside of this USoundWave. 
+	// 
+	// USoundWave::FreeResources(...) has the ability to stop clients accessing
+	// FSoundWaveData only if those clients expose the USoundWave as the WaveData
+	// on one of an FActiveSound's WaveInstances, or if the USoundWave is being 
+	// played by the "realtime" audio playback system which utilizies FSoundBuffers
+	// to decode and track resource usage. 
+	// 
+	// In the future, all clients should use the proxy system and USoundWave should 
+	// be extended to support notifying clients that it has been edited and replaced 
+	// an existing proxy with a new proxy. Clients can then choose whether to
+	// continue playing their original proxy, or retrieve the updated proxy. 
+	SoundWaveDataPtr = MakeShared<FSoundWaveData, ESPMode::ThreadSafe>();
+
+	const bool bHasSampleData = !(bIsSourceBus || bProcedural);
+	if (bHasSampleData)
+	{
+		// Source buses and procedural sources do not have sample data and should 
+		// avoid going through creation of streaming compressed data.
+		//
+		// procedural sources and source buses should be rebased to a yet-to-be
+		// created `USoundWaveBase`
+		SoundWaveDataPtr->InitializeDataFromSoundWave(*this);
+
+		// Associate the Proxy with the new SoundWaveDataPtr
+		Proxy = CreateSoundWaveProxy();
+	}
 }
 
 bool USoundWave::CanEditChange(const FProperty* InProperty) const
@@ -4769,12 +4825,6 @@ const TArray<FSoundWaveCuePoint>& FSoundWaveProxy::GetLoopRegions() const
 	return SoundWaveDataPtr->GetLoopRegions();
 }
 
-FSoundWaveData::MaxChunkSizeResults FSoundWaveProxy::GetMaxChunkSizeResults() const
-{
-	check(SoundWaveDataPtr);
-	return SoundWaveDataPtr->GetMaxChunkSizeResults();
-}
-
 bool FSoundWaveProxy::IsLooping() const
 {
 	check(SoundWaveDataPtr);
