Fix for crash when cooking invalid textures.

diff --git a/Engine/Source/Runtime/Engine/Private/Texture.cpp b/Engine/Source/Runtime/Engine/Private/Texture.cpp
--- a/Engine/Source/Runtime/Engine/Private/Texture.cpp
+++ b/Engine/Source/Runtime/Engine/Private/Texture.cpp
@@ -3927,61 +3927,63 @@ void UTexture::GetBuiltTextureSize(const ITargetPlatformSettings* TargetPlatform
 	// @todo Oodle : verify against TextureCompressorModule
 	// @todo Oodle : with cinematic mips or not? maybe add a bool arg
 	
-	int32 SizeX,SizeY,SizeZ;
+	int32 SizeX=0,SizeY=0,SizeZ=0;
 
 #if WITH_EDITORONLY_DATA
-	FIntPoint SourceSize = Source.GetLogicalSize();
-	SizeX = SourceSize.X;
-	SizeY = SourceSize.Y;
-	
-	SizeZ = Source.GetNumSlices();
-	if ( Source.IsLongLatCubemap() )
+	if (Source.IsValid())
 	{
-		SizeZ *= 6;
-	}
+		FIntPoint SourceSize = Source.GetLogicalSize();
+		SizeX = SourceSize.X;
+		SizeY = SourceSize.Y;
 	
-	// Volumes mip down Z, other types don't
-	ETextureClass TextureClass = GetTextureClass();
-	bool bIsVolume = ( TextureClass == ETextureClass::Volume );
+		SizeZ = Source.GetNumSlices();
+		if ( Source.IsLongLatCubemap() )
+		{
+			SizeZ *= 6;
+		}
+	
+		// Volumes mip down Z, other types don't
+		ETextureClass TextureClass = GetTextureClass();
+		bool bIsVolume = ( TextureClass == ETextureClass::Volume );
 	
-	UE::TextureBuildUtilities::GetPowerOfTwoTargetTextureSize(
-		SizeX, SizeY, SizeZ,
-		bIsVolume, PowerOfTwoMode, ResizeDuringBuildX, ResizeDuringBuildY,
-		SizeX, SizeY, SizeZ);
+		UE::TextureBuildUtilities::GetPowerOfTwoTargetTextureSize(
+			SizeX, SizeY, SizeZ,
+			bIsVolume, PowerOfTwoMode, ResizeDuringBuildX, ResizeDuringBuildY,
+			SizeX, SizeY, SizeZ);
 
-	if (Source.IsLongLatCubemap())
-	{
-		SizeX = SizeY = UE::TextureBuildUtilities::ComputeLongLatCubemapExtents(SizeX, MaxTextureSize);
-	}
+		if (Source.IsLongLatCubemap())
+		{
+			SizeX = SizeY = UE::TextureBuildUtilities::ComputeLongLatCubemapExtents(SizeX, MaxTextureSize);
+		}
 
-	//we need to really have the actual top mip size of output platformdata
-	//	(hence the LODBias check below)
-	// trying to reproduce here exactly what TextureCompressor + serialization will do = brittle
+		//we need to really have the actual top mip size of output platformdata
+		//	(hence the LODBias check below)
+		// trying to reproduce here exactly what TextureCompressor + serialization will do = brittle
 
-	if ( MaxTextureSize != 0 )
-	{
-		while( SizeX > MaxTextureSize || SizeY > MaxTextureSize )
+		if ( MaxTextureSize != 0 )
 		{
-			SizeX = FMath::Max(SizeX>>1,1);
-			SizeY = FMath::Max(SizeY>>1,1);
-			if ( bIsVolume )
+			while( SizeX > MaxTextureSize || SizeY > MaxTextureSize )
 			{
-				SizeZ = FMath::Max(SizeZ>>1,1);
+				SizeX = FMath::Max(SizeX>>1,1);
+				SizeY = FMath::Max(SizeY>>1,1);
+				if ( bIsVolume )
+				{
+					SizeZ = FMath::Max(SizeZ>>1,1);
+				}
 			}
 		}
-	}
 	
-	const bool bVirtualTextureStreaming = VirtualTextureStreaming && IsVirtualTexturingEnabled(TargetPlatformSettings);
+		const bool bVirtualTextureStreaming = VirtualTextureStreaming && IsVirtualTexturingEnabled(TargetPlatformSettings);
 
-	const UTextureLODSettings& LODSettings = TargetPlatformSettings->GetTextureLODSettings();
- 	const uint32 LODBiasNoCinematics = FMath::Max<int32>(LODSettings.CalculateLODBias(SizeX, SizeY, MaxTextureSize, LODGroup, LODBias, 0, MipGenSettings, bVirtualTextureStreaming), 0);
-	SizeX = FMath::Max<int32>(SizeX >> LODBiasNoCinematics, 1);
-	SizeY = FMath::Max<int32>(SizeY >> LODBiasNoCinematics, 1);
-	if ( bIsVolume )
-	{
-		SizeZ = FMath::Max<int32>(SizeZ >> LODBiasNoCinematics, 1);
+		const UTextureLODSettings& LODSettings = TargetPlatformSettings->GetTextureLODSettings();
+		const uint32 LODBiasNoCinematics = FMath::Max<int32>(LODSettings.CalculateLODBias(SizeX, SizeY, MaxTextureSize, LODGroup, LODBias, 0, MipGenSettings, bVirtualTextureStreaming), 0);
+		SizeX = FMath::Max<int32>(SizeX >> LODBiasNoCinematics, 1);
+		SizeY = FMath::Max<int32>(SizeY >> LODBiasNoCinematics, 1);
+		if ( bIsVolume )
+		{
+			SizeZ = FMath::Max<int32>(SizeZ >> LODBiasNoCinematics, 1);
+		}
 	}
-
 #else // WITH_EDITORONLY_DATA
 
 	// no Editor data
diff --git a/Engine/Source/Runtime/Engine/Private/TextureDerivedData.cpp b/Engine/Source/Runtime/Engine/Private/TextureDerivedData.cpp
--- a/Engine/Source/Runtime/Engine/Private/TextureDerivedData.cpp
+++ b/Engine/Source/Runtime/Engine/Private/TextureDerivedData.cpp
@@ -728,7 +728,13 @@ static void ModifyMaxTextureResolutionBuildSettingsForPlatformLimit(
 {
 	check( ! OutSettings.bVirtualStreamable );
 	check( OutSettings.TextureFormatName != NAME_None );
-			
+	
+	if (!Texture.Source.IsValid())
+	{
+		// Nothing to do - texture can't be built.
+		return;
+	}
+
 	// GetBuiltTextureSize is the size after LODBias
 	int32 BuiltSizeX=0,BuiltSizeY=0,BuiltSizeZ=0;
 	Texture.GetBuiltTextureSize(TargetPlatform,BuiltSizeX,BuiltSizeY,BuiltSizeZ);
