Simplified redundant and confusing calculations in FViewMatrices::Init

diff --git a/Engine/Source/Runtime/Engine/Private/SceneView.cpp b/Engine/Source/Runtime/Engine/Private/SceneView.cpp
--- a/Engine/Source/Runtime/Engine/Private/SceneView.cpp
+++ b/Engine/Source/Runtime/Engine/Private/SceneView.cpp
@@ -690,8 +690,8 @@ void FViewMatrices::Init(const FMinimalInitializer& Initializer)
 	TranslatedViewMatrix = LocalTranslatedViewMatrix;
 	InvTranslatedViewMatrix = LocalInvTranslatedViewMatrix;
 
-	OverriddenTranslatedViewMatrix = FTranslationMatrix(-GetPreViewTranslation()) * GetViewMatrix();
-	OverriddenInvTranslatedViewMatrix = GetInvViewMatrix() * FTranslationMatrix(GetPreViewTranslation());
+	OverriddenTranslatedViewMatrix = LocalTranslatedViewMatrix;
+	OverriddenInvTranslatedViewMatrix = LocalInvTranslatedViewMatrix;
 
 	TranslatedViewProjectionMatrix = LocalTranslatedViewMatrix * ProjectionMatrix;
 	InvTranslatedViewProjectionMatrix = InvProjectionMatrix * LocalInvTranslatedViewMatrix;
@@ -1201,8 +1201,8 @@ void FViewMatrices::UpdateViewMatrix(const FVector& ViewLocation, const FRotator
 
 	TranslatedViewMatrix = ViewRotationMatrix;
 	InvTranslatedViewMatrix = TranslatedViewMatrix.GetTransposed();
-	OverriddenTranslatedViewMatrix = FTranslationMatrix(-PreViewTranslation) * ViewMatrix;
-	OverriddenInvTranslatedViewMatrix = InvViewMatrix * FTranslationMatrix(PreViewTranslation);
+	OverriddenTranslatedViewMatrix = ViewRotationMatrix;
+	OverriddenInvTranslatedViewMatrix = OverriddenTranslatedViewMatrix.GetTransposed();;
 
 	// Compute a transform from view origin centered world-space to clip space.
 	TranslatedViewProjectionMatrix = GetTranslatedViewMatrix() * GetProjectionMatrix();
