[World Partition] Don't allow deleting the default partition object.

diff --git a/Engine/Source/Runtime/Engine/Private/WorldPartition/RuntimeHashSet/WorldPartitionRuntimeHashSet.cpp b/Engine/Source/Runtime/Engine/Private/WorldPartition/RuntimeHashSet/WorldPartitionRuntimeHashSet.cpp
--- a/Engine/Source/Runtime/Engine/Private/WorldPartition/RuntimeHashSet/WorldPartitionRuntimeHashSet.cpp
+++ b/Engine/Source/Runtime/Engine/Private/WorldPartition/RuntimeHashSet/WorldPartitionRuntimeHashSet.cpp
@@ -1025,6 +1025,14 @@ void UWorldPartitionRuntimeHashSet::PostEditChangeChainProperty(FPropertyChanged
 			}
 		}
 	}
+	else if (PropertyName == GET_MEMBER_NAME_CHECKED(UWorldPartitionRuntimeHashSet, RuntimePartitions))
+	{
+		if (RuntimePartitions.IsEmpty())
+		{
+			// Reapply the default values as we need at least a main partition
+			SetDefaultValues();
+		}
+	}
 }
 
 void UWorldPartitionRuntimeHashSet::PostDuplicate(EDuplicateMode::Type DuplicateMode)
