World Partition - Streaming generation will always use the attached parent actor spatially loaded flag.
Fixed streaming generation to test for invalid runtime bounds since bounds can now be invalid if non-spatially loaded actor has a spatially loaded parent actor.
This is a backout of 26994479 combined with a fix.

diff --git a/Engine/Source/Runtime/Engine/Private/WorldPartition/WorldPartitionActorDescView.cpp b/Engine/Source/Runtime/Engine/Private/WorldPartition/WorldPartitionActorDescView.cpp
--- a/Engine/Source/Runtime/Engine/Private/WorldPartition/WorldPartitionActorDescView.cpp
+++ b/Engine/Source/Runtime/Engine/Private/WorldPartition/WorldPartitionActorDescView.cpp
@@ -78,13 +78,12 @@ bool FWorldPartitionActorDescView::GetIsSpatiallyLoaded() const
 		return false;
 	}
 
-	bool bIsSpatiallyLoaded = ActorDesc->GetIsSpatiallyLoaded();
-	if (bIsSpatiallyLoaded && ParentView)
+	if (ParentView)
 	{
-		bIsSpatiallyLoaded = ParentView->GetIsSpatiallyLoaded();
+		return ParentView->GetIsSpatiallyLoaded();
 	}
-
-	return bIsSpatiallyLoaded;
+	
+	return ActorDesc->GetIsSpatiallyLoaded();
 }
 
 bool FWorldPartitionActorDescView::GetActorIsHLODRelevant() const
diff --git a/Engine/Source/Runtime/Engine/Private/WorldPartition/WorldPartitionStreamingGeneration.cpp b/Engine/Source/Runtime/Engine/Private/WorldPartition/WorldPartitionStreamingGeneration.cpp
--- a/Engine/Source/Runtime/Engine/Private/WorldPartition/WorldPartitionStreamingGeneration.cpp
+++ b/Engine/Source/Runtime/Engine/Private/WorldPartition/WorldPartitionStreamingGeneration.cpp
@@ -196,13 +196,12 @@ bool FStreamingGenerationActorDescView::GetIsSpatiallyLoaded() const
 		return false;
 	}
 
-	bool bIsSpatiallyLoaded = Super::GetIsSpatiallyLoaded();
-	if (bIsSpatiallyLoaded && ParentView)
+	if (ParentView)
 	{
-		bIsSpatiallyLoaded = ParentView->GetIsSpatiallyLoaded();
+		return ParentView->GetIsSpatiallyLoaded();
 	}
 
-	return bIsSpatiallyLoaded;
+	return Super::GetIsSpatiallyLoaded();
 }
 
 FSoftObjectPath FStreamingGenerationActorDescView::GetHLODLayer() const
@@ -952,9 +951,12 @@ class FWorldPartitionStreamingGenerator
 					if (ActorDescView.GetIsSpatiallyLoaded())
 					{
 						const FBox RuntimeBounds = ActorDescView.GetRuntimeBounds();
-						check(RuntimeBounds.IsValid);
-
-						ContainerCollectionInstanceDescriptor.Bounds += RuntimeBounds;
+						// Test if RuntimeBounds is valid because GetIsSpatiallyLoaded() is affected by a valid ParentView
+						// So the RuntimeBounds can be invalid in the case where its a non-spatial with a spatial parent.
+						if (RuntimeBounds.IsValid)
+						{
+							ContainerCollectionInstanceDescriptor.Bounds += RuntimeBounds;
+						}
 					}
 				});
 
