[World Partition] Fix old spatial runtime hash spherical circle intersection test for small angles.

diff --git a/Engine/Source/Runtime/Engine/Private/WorldPartition/WorldPartitionStreamingSource.cpp b/Engine/Source/Runtime/Engine/Private/WorldPartition/WorldPartitionStreamingSource.cpp
--- a/Engine/Source/Runtime/Engine/Private/WorldPartition/WorldPartitionStreamingSource.cpp
+++ b/Engine/Source/Runtime/Engine/Private/WorldPartition/WorldPartitionStreamingSource.cpp
@@ -164,34 +164,6 @@ void FWorldPartitionStreamingSource::UpdateHash()
 	Hash3D = HashBuilder.GetHash();
 };
 
-FORCEINLINE static bool IsClockWise(const FVector2D& V1, const FVector2D& V2)
-{
-	return (V1 ^ V2) < 0;
-}
-
-FORCEINLINE static bool IsVectorInsideVectorPair(const FVector2D& TestVector, const FVector2D& V1, const FVector2D& V2)
-{
-	return IsClockWise(V1, TestVector) && !IsClockWise(V2, TestVector);
-}
-
-FORCEINLINE static bool IsPointInsideSector(const FVector2D& TestPoint, const FVector2D& SectorCenter, float SectorRadiusSquared, const FVector2D& SectorStart, const FVector2D& SectorEnd, float SectorAngle)
-{
-	const FVector2D TestVector = TestPoint - SectorCenter;
-	if (TestVector.SizeSquared() > SectorRadiusSquared)
-	{
-		return false;
-	}
-
-	if (SectorAngle <= 180.0f)
-	{
-		return IsVectorInsideVectorPair(TestVector, SectorStart, SectorEnd);
-	}
-	else
-	{
-		return !IsVectorInsideVectorPair(TestVector, SectorEnd, SectorStart);
-	}
-}
-
 TArray<TPair<FVector, FVector>> FSphericalSector::BuildDebugMesh() const
 {
 	TArray<TPair<FVector, FVector>> Segments;
@@ -239,36 +211,78 @@ TArray<TPair<FVector, FVector>> FSphericalSector::BuildDebugMesh() const
 
 bool FSphericalSector::IntersectsBox(const FBox2D& InBox) const
 {
-	const FVector ScaledAxis = FVector(FVector2D(Axis), 0).GetSafeNormal() * Radius;
-	const FVector SectorStart = FRotator(0, 0.5f * Angle, 0).RotateVector(ScaledAxis);
-	const FVector SectorEnd = FRotator(0, -0.5f * Angle, 0).RotateVector(ScaledAxis);
-	const FVector2D SectorCenter(Center);
-	const FVector2D SectorStartVector = FVector2D(SectorStart);
-	const FVector2D SectorEndVector = FVector2D(SectorEnd);
-	const float SectorRadiusSquared = Radius * Radius;
-
-	// Test whether any cell corners are inside sector
-	if (IsPointInsideSector(FVector2D(InBox.Min.X, InBox.Min.Y), SectorCenter, SectorRadiusSquared, SectorStartVector, SectorEndVector, Angle) ||
-		IsPointInsideSector(FVector2D(InBox.Max.X, InBox.Min.Y), SectorCenter, SectorRadiusSquared, SectorStartVector, SectorEndVector, Angle) ||
-		IsPointInsideSector(FVector2D(InBox.Max.X, InBox.Max.Y), SectorCenter, SectorRadiusSquared, SectorStartVector, SectorEndVector, Angle) ||
-		IsPointInsideSector(FVector2D(InBox.Min.X, InBox.Max.Y), SectorCenter, SectorRadiusSquared, SectorStartVector, SectorEndVector, Angle))
+	// First check if the box intersects the radius
+	const FVector2D ClosestPoint = FVector2D::Max(InBox.Min, FVector2D::Min(FVector2D(Center), InBox.Max));
+	if ((ClosestPoint - FVector2D(Center)).SizeSquared() > FMath::Square(Radius))
 	{
-		return true;
+		return false;
 	}
 
-	// Test whether any sector point lies inside the cell bounds
-	if (InBox.IsInside(SectorCenter) ||
-		InBox.IsInside(SectorCenter + SectorStartVector) ||
-		InBox.IsInside(SectorCenter + SectorEndVector))
+	if (Angle < 360.0)
 	{
-		return true;
-	}
+		const FVector2D CenterToMinXMinY((FVector2D(InBox.Min.X, InBox.Min.Y) - FVector2D(Center)).GetSafeNormal());
+		const FVector2D CenterToMaxXMinY((FVector2D(InBox.Max.X, InBox.Min.Y) - FVector2D(Center)).GetSafeNormal());
+		const FVector2D CenterToMaxXMaxY((FVector2D(InBox.Max.X, InBox.Max.Y) - FVector2D(Center)).GetSafeNormal());
+		const FVector2D CenterToMinXMaxY((FVector2D(InBox.Min.X, InBox.Max.Y) - FVector2D(Center)).GetSafeNormal());
 
-	// Test whether closest point on cell from center is inside sector
-	if (IsPointInsideSector(InBox.GetClosestPointTo(SectorCenter), SectorCenter, SectorRadiusSquared, SectorStartVector, SectorEndVector, Angle))
-	{
-		return true;
+		if (Angle <= 180.0f)
+		{
+			const FVector2D::FReal SinHalfAngle = FMath::Sin(Angle * 0.5f * UE_PI / 180.0);
+			const FVector2D::FReal SinAngleMinXMinY = FVector2D::CrossProduct(FVector2D(Axis), CenterToMinXMinY);
+			const FVector2D::FReal SinAngleMaxXMinY = FVector2D::CrossProduct(FVector2D(Axis), CenterToMaxXMinY);
+			const FVector2D::FReal SinAngleMaxXMaxY = FVector2D::CrossProduct(FVector2D(Axis), CenterToMaxXMaxY);
+			const FVector2D::FReal SinAngleMinXMaxY = FVector2D::CrossProduct(FVector2D(Axis), CenterToMinXMaxY);
+
+			// Reject completly left
+			if (SinAngleMinXMinY < -SinHalfAngle && SinAngleMaxXMinY < -SinHalfAngle && SinAngleMaxXMaxY < -SinHalfAngle && SinAngleMinXMaxY < -SinHalfAngle)
+			{
+				return false;
+			}
+
+			// Reject completly right
+			if (SinAngleMinXMinY > SinHalfAngle && SinAngleMaxXMinY > SinHalfAngle && SinAngleMaxXMaxY > SinHalfAngle && SinAngleMinXMaxY > SinHalfAngle)
+			{
+				return false;
+			}
+
+			const FVector2D::FReal CosAngleMinXMinY = FVector2D::DotProduct(FVector2D(Axis), CenterToMinXMinY);
+			const FVector2D::FReal CosAngleMaxXMinY = FVector2D::DotProduct(FVector2D(Axis), CenterToMaxXMinY);
+			const FVector2D::FReal CosAngleMaxXMaxY = FVector2D::DotProduct(FVector2D(Axis), CenterToMaxXMaxY);
+			const FVector2D::FReal CosAngleMinXMaxY = FVector2D::DotProduct(FVector2D(Axis), CenterToMinXMaxY);
+
+			// Reject backward half circle
+			if (CosAngleMinXMinY < 0 && CosAngleMaxXMinY < 0 && CosAngleMaxXMaxY < 0 && CosAngleMinXMaxY < 0)
+			{
+				return false;
+			}
+		}
+		else
+		{
+			const FVector2D::FReal CosAngleMinXMinY = FVector2D::DotProduct(FVector2D(Axis),CenterToMinXMinY);
+			const FVector2D::FReal CosAngleMaxXMinY = FVector2D::DotProduct(FVector2D(Axis),CenterToMaxXMinY);
+			const FVector2D::FReal CosAngleMaxXMaxY = FVector2D::DotProduct(FVector2D(Axis),CenterToMaxXMaxY);
+			const FVector2D::FReal CosAngleMinXMaxY = FVector2D::DotProduct(FVector2D(Axis),CenterToMinXMaxY);
+
+			// Include forward half circle
+			if (CosAngleMinXMinY >= 0 || CosAngleMaxXMinY >= 0 || CosAngleMaxXMaxY >= 0 || CosAngleMinXMaxY >= 0)
+			{
+				return true;
+			}
+
+			const FVector2D::FReal InvInSinHalfAngle = FMath::Sin((360.0 - Angle) * 0.5 * UE_PI / 180.0);
+			const FVector2D::FReal InvSinAngleMinXMinY = FVector2D::CrossProduct(-FVector2D(Axis), CenterToMinXMinY);
+			const FVector2D::FReal InvSinAngleMaxXMinY = FVector2D::CrossProduct(-FVector2D(Axis), CenterToMaxXMinY);
+			const FVector2D::FReal InvSinAngleMaxXMaxY = FVector2D::CrossProduct(-FVector2D(Axis), CenterToMaxXMaxY);
+			const FVector2D::FReal InvSinAngleMinXMaxY = FVector2D::CrossProduct(-FVector2D(Axis), CenterToMinXMaxY);
+
+			// Reject completly inside inverted axis test
+			if (InvSinAngleMinXMinY > -InvInSinHalfAngle && InvSinAngleMaxXMinY > -InvInSinHalfAngle && InvSinAngleMaxXMaxY > -InvInSinHalfAngle && InvSinAngleMinXMaxY > -InvInSinHalfAngle &&
+				InvSinAngleMinXMinY < InvInSinHalfAngle && InvSinAngleMaxXMinY < InvInSinHalfAngle && InvSinAngleMaxXMaxY < InvInSinHalfAngle && InvSinAngleMinXMaxY < InvInSinHalfAngle)
+			{
+				return false;
+			}
+		}
 	}
 
-	return false;
+	return true;
 }
\ No newline at end of file
