Skeletal Mesh: Add skin weight profile layers.

Current limitations:
- Only two layers allowed.
- Active layers cannot have overlapping weights that differ from the base.

diff --git a/Engine/Plugins/Animation/DeformerGraph/Source/OptimusCore/Private/DataInterfaces/OptimusDataInterfaceAdvancedSkeleton.cpp b/Engine/Plugins/Animation/DeformerGraph/Source/OptimusCore/Private/DataInterfaces/OptimusDataInterfaceAdvancedSkeleton.cpp
--- a/Engine/Plugins/Animation/DeformerGraph/Source/OptimusCore/Private/DataInterfaces/OptimusDataInterfaceAdvancedSkeleton.cpp
+++ b/Engine/Plugins/Animation/DeformerGraph/Source/OptimusCore/Private/DataInterfaces/OptimusDataInterfaceAdvancedSkeleton.cpp
@@ -581,7 +581,8 @@ void UOptimusAdvancedSkeletonDataProvider::ComputeBoneTransformsForLayeredSkinni
 		FSkinWeightVertexBuffer const* WeightBuffer = InLodRenderData->GetSkinWeightVertexBuffer();
 		if (InLodRenderData->SkinWeightProfilesData.ContainsProfile(SkinWeightProfile))
 		{
-			WeightBuffer = InLodRenderData->SkinWeightProfilesData.GetOverrideBuffer(SkinWeightProfile);
+			const FSkinWeightProfileStack ProfileStack{SkinWeightProfile};
+			WeightBuffer = InLodRenderData->SkinWeightProfilesData.GetOverrideBuffer(ProfileStack);
 		}
 
 		for (uint32 VertexIndex = 0; VertexIndex < WeightBuffer->GetNumVertices(); VertexIndex++)
@@ -700,7 +701,8 @@ FComputeDataProviderRenderProxy* UOptimusAdvancedSkeletonDataProvider::GetRender
 				if (LodRenderData.SkinWeightProfilesData.ContainsProfile(SkinWeightProfile))
 				{
 					// Retrieve this profile's skin weight buffer
-					FSkinWeightVertexBuffer* Buffer = LodRenderData.SkinWeightProfilesData.GetOverrideBuffer(SkinWeightProfile);
+					const FSkinWeightProfileStack ProfileStack{SkinWeightProfile};
+					FSkinWeightVertexBuffer* Buffer = LodRenderData.SkinWeightProfilesData.GetOverrideBuffer(ProfileStack);
 						
 					if (Buffer)
 					{
@@ -710,7 +712,7 @@ FComputeDataProviderRenderProxy* UOptimusAdvancedSkeletonDataProvider::GetRender
 					{
 						TWeakObjectPtr<USkinnedMeshComponent> WeakComponent = SkeletalMesh;
 						TWeakObjectPtr<UOptimusAdvancedSkeletonDataProvider> WeakThis = this;
-						FRequestFinished Callback = [WeakComponent, WeakThis](TWeakObjectPtr<USkeletalMesh> WeakMesh, FName ProfileName)
+						FRequestFinished Callback = [WeakComponent, WeakThis](TWeakObjectPtr<USkeletalMesh> WeakMesh, FSkinWeightProfileStack ProfileStack)
 						{
 							// Ensure that the request objects are still valid
 							if (WeakMesh.IsValid() && WeakComponent.IsValid() && WeakThis.IsValid())
@@ -724,7 +726,7 @@ FComputeDataProviderRenderProxy* UOptimusAdvancedSkeletonDataProvider::GetRender
 										FSkinWeightProfilesData& SkinweightData = LODRenderData.SkinWeightProfilesData;
 
 										// Retrieve this profile's skin weight buffer
-										FSkinWeightVertexBuffer* Buffer = SkinweightData.GetOverrideBuffer(ProfileName);
+										FSkinWeightVertexBuffer* Buffer = SkinweightData.GetOverrideBuffer(ProfileStack);
 										if (ensure(Buffer))
 										{
 											WeakThis->bSkinWeightBufferReady = true;
@@ -737,7 +739,7 @@ FComputeDataProviderRenderProxy* UOptimusAdvancedSkeletonDataProvider::GetRender
 						// Put in a skin weight profile request
 						if (FSkinWeightProfileManager* Manager = FSkinWeightProfileManager::Get(SkeletalMesh->GetWorld()))
 						{
-							Manager->RequestSkinWeightProfile(SkinWeightProfile, SkeletalMesh->GetSkinnedAsset(), this, Callback);
+							Manager->RequestSkinWeightProfileStack(ProfileStack, SkeletalMesh->GetSkinnedAsset(), this, Callback);
 						}
 					}
 				}	
@@ -868,7 +870,8 @@ bool FOptimusAdvancedSkeletonDataProviderProxy::IsValid(FValidationData const& I
 	FSkinWeightVertexBuffer const* WeightBuffer = LodRenderData->GetSkinWeightVertexBuffer();
 	if (LodRenderData->SkinWeightProfilesData.ContainsProfile(SkinWeightProfile))
 	{
-		WeightBuffer = LodRenderData->SkinWeightProfilesData.GetOverrideBuffer(SkinWeightProfile);
+		const FSkinWeightProfileStack ProfileStack{SkinWeightProfile};
+		WeightBuffer = LodRenderData->SkinWeightProfilesData.GetOverrideBuffer(ProfileStack);
 	}
 	
 	if (WeightBuffer == nullptr)
@@ -967,7 +970,8 @@ void FOptimusAdvancedSkeletonDataProviderProxy::GatherPermutations(FPermutationD
 		FSkinWeightVertexBuffer const* WeightBuffer = LodRenderData->GetSkinWeightVertexBuffer();
 		if (LodRenderData->SkinWeightProfilesData.ContainsProfile(SkinWeightProfile))
 		{
-			WeightBuffer = LodRenderData->SkinWeightProfilesData.GetOverrideBuffer(SkinWeightProfile);
+			const FSkinWeightProfileStack ProfileStack{SkinWeightProfile};
+			WeightBuffer = LodRenderData->SkinWeightProfilesData.GetOverrideBuffer(ProfileStack);
 		}
 		
 		check(WeightBuffer != nullptr);
@@ -1006,7 +1010,8 @@ void FOptimusAdvancedSkeletonDataProviderProxy::GatherDispatchData(FDispatchData
 		FSkinWeightVertexBuffer const* WeightBuffer = LodRenderData->GetSkinWeightVertexBuffer();
 		if (LodRenderData->SkinWeightProfilesData.ContainsProfile(SkinWeightProfile))
 		{
-			WeightBuffer = LodRenderData->SkinWeightProfilesData.GetOverrideBuffer(SkinWeightProfile);
+			const FSkinWeightProfileStack ProfileStack{SkinWeightProfile};
+			WeightBuffer = LodRenderData->SkinWeightProfilesData.GetOverrideBuffer(ProfileStack);
 		}
 		check(WeightBuffer != nullptr);
 		
diff --git a/Engine/Source/Runtime/Engine/Classes/Components/SkinnedMeshComponent.h b/Engine/Source/Runtime/Engine/Classes/Components/SkinnedMeshComponent.h
--- a/Engine/Source/Runtime/Engine/Classes/Components/SkinnedMeshComponent.h
+++ b/Engine/Source/Runtime/Engine/Classes/Components/SkinnedMeshComponent.h
@@ -16,6 +16,7 @@
 #include "LODSyncInterface.h"
 #include "BoneContainer.h"
 #include "ClothingSystemRuntimeTypes.h"
+#include "Animation/SkinWeightProfile.h"
 #include "SkinnedMeshComponent.generated.h"
 
 enum class ESkinCacheUsage : uint8;
@@ -34,6 +35,7 @@ class UMorphTarget;
 class USkinnedAsset;
 struct FExternalMorphSet;
 struct FExternalMorphWeightData;
+struct FSkinWeightProfileStack;
 
 namespace Nanite
 {
@@ -131,6 +133,16 @@ namespace EBoneSpaces
 	};
 }
 
+
+/** Values for specifying which layer a skin weight profile is applied at.
+ */
+UENUM(BlueprintType)
+enum class ESkinWeightProfileLayer : uint8
+{
+	Primary,		/** Primary skin weight profile layer */
+	Secondary,		/** Secondary skin weight profile layer */
+};
+
 /** WeightIndex is an into the MorphTargetWeights array */
 using FMorphTargetWeightMap = TMap<const UMorphTarget* /* MorphTarget */, int32 /* WeightIndex */>;
 
@@ -1191,7 +1203,7 @@ public:
 	ENGINE_API virtual bool DoesSocketExist(FName InSocketName) const override;
 	ENGINE_API virtual bool HasAnySockets() const override;
 	ENGINE_API virtual void QuerySupportedSockets(TArray<FComponentSocketDescription>& OutSockets) const override;
-	ENGINE_API virtual bool UpdateOverlapsImpl(const TOverlapArrayView* PendingOverlaps=NULL, bool bDoNotifies=true, const TOverlapArrayView* OverlapsAtEndLocation=NULL) override;
+	ENGINE_API virtual bool UpdateOverlapsImpl(const TOverlapArrayView* PendingOverlaps=nullptr, bool bDoNotifies=true, const TOverlapArrayView* OverlapsAtEndLocation=nullptr) override;
 	//~ End USceneComponent Interface
 
 	//~ Begin UPrimitiveComponent Interface
@@ -1372,51 +1384,52 @@ public:
 	UFUNCTION(BlueprintCallable, Category = "Components|SkinnedMesh")
 	ENGINE_API void ClearSkinWeightOverride(int32 LODIndex);
 
-	/** Setup an override Skin Weight Profile for this component */
+	/** Set up an override skin weight profile for this component on the given layer.
+	 *  The values from the secondary layer (if set to have a profile) are applied first, followed by the values from the primary layer.
+	 *  Since skin weight profiles are stored as sparse data, where only weight values different from the base are kept in storage, it's
+	 *  possible to set up layers such that they don't interfere with one another.
+	 */
 	UFUNCTION(BlueprintCallable, Category = "Components|SkinnedMesh")
-	ENGINE_API bool SetSkinWeightProfile(FName InProfileName);
+	ENGINE_API bool SetSkinWeightProfile(FName InProfileName, ESkinWeightProfileLayer InLayer = ESkinWeightProfileLayer::Primary);
 
-	/** Clear the Skin Weight Profile from this component, in case it is set */
+	/** Clear the skin weight profile from the given layer on this component, in case it is set. If no profile is set for the layer,
+	 *  then this call does nothing. */
 	UFUNCTION(BlueprintCallable, Category = "Components|SkinnedMesh")
-	ENGINE_API void ClearSkinWeightProfile();
+	ENGINE_API void ClearSkinWeightProfile(ESkinWeightProfileLayer InLayer = ESkinWeightProfileLayer::Primary);
 
+	/** Clear the skin Weight Profile from all layers on this component. If no profiles are set for any layer, then this call does nothing. */
+	UFUNCTION(BlueprintCallable, Category = "Components|SkinnedMesh")
+	ENGINE_API void ClearAllSkinWeightProfiles();
+	
 	/** Unload a Skin Weight Profile's skin weight buffer (if created) */
 	UFUNCTION(BlueprintCallable, Category = "Components|SkinnedMesh")
 	ENGINE_API void UnloadSkinWeightProfile(FName InProfileName);
 
-	/** Return the name of the Skin Weight Profile that is currently set otherwise returns 'None' */
+	/** Return the name of the skin weight profile that is currently set on the given layer, otherwise returns 'None' */
 	UFUNCTION(BlueprintCallable, Category = "Components|SkinnedMesh")
-	ENGINE_API FName GetCurrentSkinWeightProfileName() const;
+	ENGINE_API FName GetCurrentSkinWeightProfileName(ESkinWeightProfileLayer InLayer = ESkinWeightProfileLayer::Primary) const;
 
-	/** Check whether or not a Skin Weight Profile is currently set */
+	/** Return the names of the skin weight profiles for all the layers */
 	UFUNCTION(BlueprintCallable, Category = "Components|SkinnedMesh")
-	ENGINE_API bool IsUsingSkinWeightProfile() const;
-
-	UE_DEPRECATED(4.26, "GetVertexOffsetUsage() has been deprecated. Support will be dropped in the future.")
-	UFUNCTION(BlueprintCallable, Category = "Components|SkinnedMesh")
-	int32 GetVertexOffsetUsage(int32 LODIndex) const { return 0; }
-
-	UE_DEPRECATED(4.26, "SetVertexOffsetUsage() has been deprecated. Support will be dropped in the future.")
-	UFUNCTION(BlueprintCallable, Category = "Components|SkinnedMesh")
-	void SetVertexOffsetUsage(int32 LODIndex, int32 Usage) {}
-
-	UE_DEPRECATED(4.26, "SetPreSkinningOffsets() has been deprecated. Support will be dropped in the future.")
-	UFUNCTION(BlueprintCallable, Category = "Components|SkinnedMesh")
-	void SetPreSkinningOffsets(int32 LODIndex, TArray<FVector> Offsets) {}
-
-	UE_DEPRECATED(4.26, "SetPostSkinningOffsets() has been deprecated. Support will be dropped in the future.")
+	ENGINE_API TArray<FName> GetCurrentSkinWeightProfileLayerNames() const;
+	
+	/** Check whether a skin weight profile is currently set on any layer. */
 	UFUNCTION(BlueprintCallable, Category = "Components|SkinnedMesh")
-	void SetPostSkinningOffsets(int32 LODIndex, TArray<FVector> Offsets) {}
+	ENGINE_API bool IsUsingSkinWeightProfile() const;
 
-	/** Check whether or not a Skin Weight Profile is currently pending load / create */
-	bool IsSkinWeightProfilePending() const { return bSkinWeightProfilePending == 1; }
+	/** Check whether a skin weight profile is currently pending load / create */
+	bool IsSkinWeightProfilePending() const { return bSkinWeightProfilePending; }
 
 	/** Queues an update of the Skin Weight Buffer used by the current MeshObject */
 	ENGINE_API void UpdateSkinWeightOverrideBuffer();
-protected:	
+protected:
+	/** Set the currently active skin weight profile stack to the given stack. If the stack is empty, then the skin weight profile buffer
+	  * will be set to the base buffer for this mesh.
+	  */
+	bool SetSkinWeightProfileStack(const FSkinWeightProfileStack& InProfileStack);
 
-	/** Name of currently set up Skin Weight profile, otherwise is 'none' */
-	FName CurrentSkinWeightProfileName;
+	/** Name of currently set up Skin Weight layers, all set to NAME_None for no override */
+	FName CurrentSkinWeightProfileLayers[2] = {NAME_None};
 public:
 	/** Returns skin weight vertex buffer to use for specific LOD (will look at override) */
 	ENGINE_API FSkinWeightVertexBuffer* GetSkinWeightBuffer(int32 LODIndex) const;
@@ -1442,7 +1455,7 @@ public:
 	 * @param TickFunction Supplied as non null if we are running in a tick, allows us to create graph tasks for parallelism
 	 * 
 	 */
-	ENGINE_API virtual void RefreshBoneTransforms(FActorComponentTickFunction* TickFunction = NULL) PURE_VIRTUAL(USkinnedMeshComponent::RefreshBoneTransforms, );
+	ENGINE_API virtual void RefreshBoneTransforms(FActorComponentTickFunction* TickFunction = nullptr) PURE_VIRTUAL(USkinnedMeshComponent::RefreshBoneTransforms, );
 
 protected:
 	/** 
@@ -1475,7 +1488,7 @@ public:
 	FOnTickPose OnTickPose;
 
 	/** 
-	 * Update Follower Component. This gets called when LeaderPoseComponent!=NULL
+	 * Update Follower Component. This gets called when LeaderPoseComponent!=nullptr
 	 * 
 	 */
 	ENGINE_API virtual void UpdateFollowerComponent();
@@ -1723,13 +1736,13 @@ public:
 	 * @param InSocketName	The name of the socket to find
 	 * @param OutBoneIndex	The socket bone index in this skeletal mesh, or INDEX_NONE if the socket is not found or not a bone-relative socket
 	 * @param OutTransform	The socket local transform, or identity if the socket is not found.
-	 * @return SkeletalMeshSocket of named socket on the skeletal mesh component, or NULL if not found.
+	 * @return SkeletalMeshSocket of named socket on the skeletal mesh component, or nullptr if not found.
 	 */
 	ENGINE_API class USkeletalMeshSocket const* GetSocketInfoByName(FName InSocketName, FTransform& OutTransform, int32& OutBoneIndex) const;
 
 	/**
 	 * @param InSocketName	The name of the socket to find
-	 * @return SkeletalMeshSocket of named socket on the skeletal mesh component, or NULL if not found.
+	 * @return SkeletalMeshSocket of named socket on the skeletal mesh component, or nullptr if not found.
 	 */
 	ENGINE_API class USkeletalMeshSocket const* GetSocketByName( FName InSocketName ) const;
 
@@ -1864,7 +1877,7 @@ public:
 	 *
 	 * @return the name of the bone that was found, or 'None' if no bone was found
 	 */
-	ENGINE_API FName FindClosestBone(FVector TestLocation, FVector* BoneLocation = NULL, float IgnoreScale = 0.f, bool bRequirePhysicsAsset = false) const;
+	ENGINE_API FName FindClosestBone(FVector TestLocation, FVector* BoneLocation = nullptr, float IgnoreScale = 0.f, bool bRequirePhysicsAsset = false) const;
 
 	/** finds the closest bone to the given location
 	*
@@ -1883,7 +1896,7 @@ public:
 	 *
 	 * @param MorphTargetName Name of MorphTarget to look for.
 	 *
-	 * @return Pointer to found MorphTarget. Returns NULL if could not find target with that name.
+	 * @return Pointer to found MorphTarget. Returns nullptr if could not find target with that name.
 	 */
 	ENGINE_API virtual class UMorphTarget* FindMorphTarget( FName MorphTargetName ) const;
 
diff --git a/Engine/Source/Runtime/Engine/Internal/Streaming/SkeletalMeshUpdate.h b/Engine/Source/Runtime/Engine/Internal/Streaming/SkeletalMeshUpdate.h
--- a/Engine/Source/Runtime/Engine/Internal/Streaming/SkeletalMeshUpdate.h
+++ b/Engine/Source/Runtime/Engine/Internal/Streaming/SkeletalMeshUpdate.h
@@ -15,6 +15,9 @@ SkeletalMeshUpdate.h: Helpers to stream in and out skeletal mesh LODs.
 #include "Serialization/BulkData.h"
 #include "RayTracingGeometry.h"
 
+
+struct FSkinWeightProfileStack;
+
 /**
 * A context used to update or proceed with the next update step.
 * The mesh and render data references could be stored in the update object
@@ -84,7 +87,7 @@ protected:
 		FSkinWeightRHIInfo SkinWeightVertexBuffer;
 		FBufferRHIRef ClothVertexBuffer;
 		FBufferRHIRef IndexBuffer;
-		TArray<TPair<FName, FSkinWeightRHIInfo>> AltSkinWeightVertexBuffers;
+		TArray<TPair<FSkinWeightProfileStack, FSkinWeightRHIInfo>> AltSkinWeightVertexBuffers;
 		FSkeletalMeshHalfEdgeBuffer::FRHIInfo HalfEdgeBuffer;
 
 		void CreateFromCPUData(FRHICommandListBase& RHICmdList, FSkeletalMeshLODRenderData& LODResource);
diff --git a/Engine/Source/Runtime/Engine/Private/Animation/SkinWeightProfile.cpp b/Engine/Source/Runtime/Engine/Private/Animation/SkinWeightProfile.cpp
--- a/Engine/Source/Runtime/Engine/Private/Animation/SkinWeightProfile.cpp
+++ b/Engine/Source/Runtime/Engine/Private/Animation/SkinWeightProfile.cpp
@@ -246,23 +246,17 @@ void FSkinWeightProfilesData::Init(FSkinWeightVertexBuffer* InBaseBuffer)
 FSkinWeightProfilesData::~FSkinWeightProfilesData()
 {
 	ReleaseResources();
-
-	BaseBuffer = nullptr;
-	DefaultOverrideSkinWeightBuffer = nullptr;
-
-	bDefaultOverridden = false;
-	bStaticOverridden = false;
-	DefaultProfileName = NAME_None;
 }
 
 FSkinWeightProfilesData::FOnPickOverrideSkinWeightProfile FSkinWeightProfilesData::OnPickOverrideSkinWeightProfile;
+
 #if !WITH_EDITOR
 void FSkinWeightProfilesData::OverrideBaseBufferSkinWeightData(USkeletalMesh* Mesh, int32 LODIndex)
 {
 	if (GSkinWeightProfilesLoadByDefaultMode == 1)
 	{
 		const TArray<FSkinWeightProfileInfo>& Profiles = Mesh->GetSkinWeightProfiles();
-		// Try and find a default buffer and whether or not it is set for this LOD index 
+		// Try and find a default buffer and whether it is set for this LOD index 
 		int32 DefaultProfileIndex = INDEX_NONE;
 
 		// Setup to not apply any skin weight profiles at this LOD level
@@ -281,19 +275,18 @@ void FSkinWeightProfilesData::OverrideBaseBufferSkinWeightData(USkeletalMesh* Me
 			});
 		}
 
-		bool bProfileSet = false;
 		// If we found a profile try and find the override skin weights and apply if found
 		if (DefaultProfileIndex != INDEX_NONE)
 		{
-			const FRuntimeSkinWeightProfileData* ProfilePtr = OverrideData.Find(Profiles[DefaultProfileIndex].Name);
-			if (ProfilePtr)
+			const FName ProfileName = Profiles[DefaultProfileIndex].Name;
+			if (const FRuntimeSkinWeightProfileData* ProfilePtr = OverrideData.Find(ProfileName))
 			{
 				ProfilePtr->ApplyDefaultOverride(BaseBuffer);
 			}
 
 			bDefaultOverridden = true;
 			bStaticOverridden = true;
-			DefaultProfileName = Profiles[DefaultProfileIndex].Name;
+			DefaultProfileStack = FSkinWeightProfileStack{ProfileName};
 		}
 	}
 }
@@ -329,15 +322,17 @@ void FSkinWeightProfilesData::SetDynamicDefaultSkinWeightProfile(USkeletalMesh*
 			return (ProfileInfo.DefaultProfile.Default && LODIndex >= ProfileInfo.DefaultProfileFromLODIndex.Default);
 		});
 
-		bool bProfileSet = false;
 		// If we found a profile try and find the override skin weights and apply if found
 		if (DefaultProfileIndex != INDEX_NONE)
 		{
+			const FName& ProfileName = Profiles[DefaultProfileIndex].Name;
+			const FSkinWeightProfileStack ProfileStack(ProfileName);
+			
 			const bool bNoDefaultProfile = DefaultOverrideSkinWeightBuffer == nullptr;
-			const bool bDifferentDefaultProfile = bNoDefaultProfile && (!bDefaultOverridden || DefaultProfileName != Profiles[DefaultProfileIndex].Name);
+			const bool bDifferentDefaultProfile = bNoDefaultProfile && (!bDefaultOverridden || DefaultProfileStack != ProfileStack);
 			if (bNoDefaultProfile || bDifferentDefaultProfile)
 			{
-				if (GetOverrideBuffer(Profiles[DefaultProfileIndex].Name) == nullptr)
+				if (GetOverrideBuffer(ProfileStack) == nullptr)
 				{
 					if (bSerialization)
 					{
@@ -346,23 +341,18 @@ void FSkinWeightProfilesData::SetDynamicDefaultSkinWeightProfile(USkeletalMesh*
 						
 						if (ensure(BaseBufferData))
 						{
-							FSkinWeightVertexBuffer* OverrideBuffer = nullptr;
-							
-							const FName& ProfileName = Profiles[DefaultProfileIndex].Name;
-							const FRuntimeSkinWeightProfileData* ProfilePtr = OverrideData.Find(ProfileName);
-							if (ProfilePtr)
+							if (const FRuntimeSkinWeightProfileData* ProfilePtr = OverrideData.Find(ProfileName))
 							{
 								FSkinnedMeshComponentUpdateSkinWeightsContext Context(Mesh);
 
-								OverrideBuffer = new FSkinWeightVertexBuffer();
-								OverrideBuffer->CopyMetaData(*BaseBuffer);
-								ProfileNameToBuffer.Add(ProfileName, OverrideBuffer);
-								
-								ProfilePtr->ApplyOverrides(OverrideBuffer, BaseBufferData, BaseBuffer->GetNumVertices());
+								FSkinWeightVertexBuffer* OverrideBuffer = new FSkinWeightVertexBuffer();
+								ProfileStackToBuffer.Add(ProfileStack, OverrideBuffer);
+
+								ApplyOverrideProfileStack(ProfileStack, OverrideBuffer);
 
 								DefaultOverrideSkinWeightBuffer = OverrideBuffer;
 								bDefaultOverridden = true;
-								DefaultProfileName = Profiles[DefaultProfileIndex].Name;
+								DefaultProfileStack = ProfileStack;
 								
 #if RHI_ENABLE_RESOURCE_INFO
 								const FName OwnerName(USkinnedAsset::GetLODPathName(Mesh, LODIndex));
@@ -375,14 +365,14 @@ void FSkinWeightProfilesData::SetDynamicDefaultSkinWeightProfile(USkeletalMesh*
 					else
 					{
 						FSkinWeightProfilesData* DataPtr = this;
-						FRequestFinished Callback = [DataPtr](TWeakObjectPtr<USkeletalMesh> WeakMesh, FName ProfileName)
+						FRequestFinished Callback = [DataPtr](TWeakObjectPtr<USkeletalMesh> WeakMesh, FSkinWeightProfileStack ProfileStackRequested)
 						{
 							if (WeakMesh.IsValid())
 							{
 								FSkinnedMeshComponentRecreateRenderStateContext RecreateState(WeakMesh.Get());
 								DataPtr->bDefaultOverridden = true;
-								DataPtr->DefaultProfileName = ProfileName;
-								DataPtr->SetupDynamicDefaultSkinweightProfile();								
+								DataPtr->DefaultProfileStack = ProfileStackRequested;
+								DataPtr->SetupDynamicDefaultSkinWeightProfile();
 							}
 						};
 
@@ -401,7 +391,7 @@ void FSkinWeightProfilesData::SetDynamicDefaultSkinWeightProfile(USkeletalMesh*
 						{
 							if (FSkinWeightProfileManager* Manager = FSkinWeightProfileManager::Get(World))
 							{
-								Manager->RequestSkinWeightProfile(Profiles[DefaultProfileIndex].Name, Mesh, Mesh, Callback, LODIndex);
+								Manager->RequestSkinWeightProfileStack(ProfileStack, Mesh, Mesh, Callback, LODIndex);
 							}
 						}
 					}
@@ -409,9 +399,9 @@ void FSkinWeightProfilesData::SetDynamicDefaultSkinWeightProfile(USkeletalMesh*
 				else
 				{
 					bDefaultOverridden = true;
-					DefaultProfileName = Profiles[DefaultProfileIndex].Name;
+					DefaultProfileStack = ProfileStack;
 
-					SetupDynamicDefaultSkinweightProfile();
+					SetupDynamicDefaultSkinWeightProfile();
 				}
 			}
 		}
@@ -432,21 +422,21 @@ void FSkinWeightProfilesData::ClearDynamicDefaultSkinWeightProfile(USkeletalMesh
 		{
 #if !WITH_EDITOR
 			// Only release when not in Editor, as any other viewport / editor could be relying on this buffer
-			ReleaseBuffer(DefaultProfileName, true);
+			ReleaseBuffer(DefaultProfileStack, true);
 #endif // !WITH_EDITOR
 			DefaultOverrideSkinWeightBuffer = nullptr;
 		}
 
 		bDefaultOverridden = false;
-		DefaultProfileName = NAME_None;		
+		DefaultProfileStack = {};		
 	}
 }
 
-void FSkinWeightProfilesData::SetupDynamicDefaultSkinweightProfile()
+void FSkinWeightProfilesData::SetupDynamicDefaultSkinWeightProfile()
 {
-	if (ProfileNameToBuffer.Contains(DefaultProfileName) && bDefaultOverridden && !bStaticOverridden)
+	if (ProfileStackToBuffer.Contains(DefaultProfileStack) && bDefaultOverridden && !bStaticOverridden)
 	{
-		DefaultOverrideSkinWeightBuffer = ProfileNameToBuffer.FindChecked(DefaultProfileName);
+		DefaultOverrideSkinWeightBuffer = ProfileStackToBuffer.FindChecked(DefaultProfileStack);
 	}
 }
 
@@ -455,14 +445,16 @@ bool FSkinWeightProfilesData::ContainsProfile(const FName& ProfileName) const
 	return OverrideData.Contains(ProfileName);
 }
 
-FSkinWeightVertexBuffer* FSkinWeightProfilesData::GetOverrideBuffer(const FName& ProfileName) const
+FSkinWeightVertexBuffer* FSkinWeightProfilesData::GetOverrideBuffer(const FSkinWeightProfileStack& InProfileStack) const
 {
 	SCOPED_NAMED_EVENT(FSkinWeightProfilesData_GetOverrideBuffer, FColor::Red);
 
+	FSkinWeightProfileStack ProfileStack{InProfileStack.Normalized()};
+	
 	// In case we have overridden the default skin weight buffer we do not need to create an override buffer, if it was statically overridden we cannot load any other profile
-	if (bDefaultOverridden && (ProfileName == DefaultProfileName || bStaticOverridden))
+	if (bDefaultOverridden && (ProfileStack == DefaultProfileStack || bStaticOverridden))
 	{	
-		if (bStaticOverridden && ProfileName != DefaultProfileName)
+		if (bStaticOverridden && ProfileStack != DefaultProfileStack)
 		{
 			UE_LOG(LogSkeletalMesh, Error, TEXT("Skeletal Mesh has overridden the default Skin Weights buffer during serialization, cannot set any other skin weight profile."));
 		}	
@@ -473,20 +465,20 @@ FSkinWeightVertexBuffer* FSkinWeightProfilesData::GetOverrideBuffer(const FName&
 	if (BaseBuffer)
 	{
 		check(BaseBuffer->GetNumVertices() > 0);
-		if (ProfileNameToBuffer.Contains(ProfileName))
+
+		if (FSkinWeightVertexBuffer* const* BufferPtr = ProfileStackToBuffer.Find(ProfileStack))
 		{
-			FSkinWeightVertexBuffer* Buffer = ProfileNameToBuffer.FindChecked(ProfileName);
-			return Buffer;
-		}
+			return *BufferPtr;
 		}
+	}
 
 	return nullptr;
 }
 
 
-bool FSkinWeightProfilesData::ContainsOverrideBuffer(const FName& ProfileName) const
+bool FSkinWeightProfilesData::ContainsOverrideBuffer(const FSkinWeightProfileStack& InProfileStack) const
 {
-	return ProfileNameToBuffer.Contains(ProfileName);
+	return ProfileStackToBuffer.Contains(InProfileStack.Normalized());
 }
 
 
@@ -501,12 +493,14 @@ FRuntimeSkinWeightProfileData& FSkinWeightProfilesData::AddOverrideData(const FN
 }
 
 
-void FSkinWeightProfilesData::ReleaseBuffer(const FName& ProfileName, bool bForceRelease /*= false*/)
+void FSkinWeightProfilesData::ReleaseBuffer(const FSkinWeightProfileStack& InProfileStack, bool bForceRelease /*= false*/)
 {
-	if (ProfileNameToBuffer.Contains(ProfileName) && (!bDefaultOverridden || ProfileName != DefaultProfileName || bForceRelease))
+	FSkinWeightProfileStack ProfileStack{InProfileStack.Normalized()};
+	
+	if (ProfileStackToBuffer.Contains(ProfileStack) && (!bDefaultOverridden || ProfileStack != DefaultProfileStack || bForceRelease))
 	{
 		FSkinWeightVertexBuffer* Buffer = nullptr;
-		ProfileNameToBuffer.RemoveAndCopyValue(ProfileName, Buffer);
+		ProfileStackToBuffer.RemoveAndCopyValue(ProfileStack, Buffer);
 
 		if (Buffer)
 		{
@@ -524,15 +518,15 @@ void FSkinWeightProfilesData::ReleaseBuffer(const FName& ProfileName, bool bForc
 void FSkinWeightProfilesData::ReleaseResources()
 {
 	TArray<FSkinWeightVertexBuffer*> Buffers;
-	ProfileNameToBuffer.GenerateValueArray(Buffers);
-	ProfileNameToBuffer.Empty();
+	ProfileStackToBuffer.GenerateValueArray(Buffers);
+	ProfileStackToBuffer.Empty();
 
 	// Never release a default _dynamic_ buffer
 	if (bDefaultOverridden && !bStaticOverridden)
 	{
 		ensure(DefaultOverrideSkinWeightBuffer != nullptr);
 		Buffers.Remove(DefaultOverrideSkinWeightBuffer);
-		ProfileNameToBuffer.Add(DefaultProfileName, DefaultOverrideSkinWeightBuffer);
+		ProfileStackToBuffer.Add(DefaultProfileStack, DefaultOverrideSkinWeightBuffer);
 	}
 
 	Buffers.Remove(nullptr);
@@ -556,9 +550,9 @@ void FSkinWeightProfilesData::ReleaseResources()
 SIZE_T FSkinWeightProfilesData::GetResourcesSize() const
 {
 	SIZE_T SummedSize = 0;
-	for (const TPair<FName, FSkinWeightVertexBuffer*>& Pair : ProfileNameToBuffer)
+	for (const TPair<FSkinWeightProfileStack, FSkinWeightVertexBuffer*>& Item : ProfileStackToBuffer)
 	{
-		SummedSize += Pair.Value->GetVertexDataSize();
+		SummedSize += Item.Value->GetVertexDataSize();
 	}
 
 	return SummedSize;
@@ -567,9 +561,9 @@ SIZE_T FSkinWeightProfilesData::GetResourcesSize() const
 SIZE_T FSkinWeightProfilesData::GetCPUAccessMemoryOverhead() const
 {
 	SIZE_T Result = 0;
-	for (typename TMap<FName, FSkinWeightVertexBuffer*>::TConstIterator It(ProfileNameToBuffer); It; ++It)
+	for (const TPair<FSkinWeightProfileStack, FSkinWeightVertexBuffer*>& Item : ProfileStackToBuffer)
 	{
-		Result += It->Value->GetNeedsCPUAccess() ? It->Value->GetVertexDataSize() : 0;
+		Result += Item.Value->GetNeedsCPUAccess() ? Item.Value->GetVertexDataSize() : 0;
 	}
 	return Result;
 }
@@ -595,25 +589,26 @@ void FSkinWeightProfilesData::SerializeMetaData(FArchive& Ar)
 
 void FSkinWeightProfilesData::ReleaseCPUResources()
 {
-	for (TMap<FName, FRuntimeSkinWeightProfileData>::TIterator It(OverrideData); It; ++It)
+	for (TPair<FName, FRuntimeSkinWeightProfileData>& Item: OverrideData)
 	{
-		It->Value = FRuntimeSkinWeightProfileData();
+		Item.Value = FRuntimeSkinWeightProfileData();
 	}
 
 	ResetGPUReadback();
 }
 
-void FSkinWeightProfilesData::CreateRHIBuffers(FRHICommandListBase& RHICmdList, TArray<TPair<FName, FSkinWeightRHIInfo>>& OutBuffers)
+void FSkinWeightProfilesData::CreateRHIBuffers(FRHICommandListBase& RHICmdList, TArray<TPair<FSkinWeightProfileStack, FSkinWeightRHIInfo>>& OutBuffers)
 {
-	const int32 NumActiveProfiles = ProfileNameToBuffer.Num();
+	const int32 NumActiveProfiles = ProfileStackToBuffer.Num();
 	check(BaseBuffer || !NumActiveProfiles);
 	OutBuffers.Empty(NumActiveProfiles);
-	for (TMap<FName, FSkinWeightVertexBuffer*>::TIterator It(ProfileNameToBuffer); It; ++It)
+	for (TPair<FSkinWeightProfileStack, FSkinWeightVertexBuffer*>& Item : ProfileStackToBuffer)
 	{
-		const FName& ProfileName = It->Key;
-		FSkinWeightVertexBuffer* OverrideBuffer = It->Value;
-		ApplyOverrideProfile(OverrideBuffer, ProfileName);
-		OutBuffers.Emplace(ProfileName, OverrideBuffer->CreateRHIBuffer(RHICmdList));
+		const FSkinWeightProfileStack& ProfileStack = Item.Key;
+		FSkinWeightVertexBuffer* OverrideBuffer = Item.Value;
+		ApplyOverrideProfileStack(ProfileStack, OverrideBuffer);
+
+		OutBuffers.Emplace(ProfileStack, OverrideBuffer->CreateRHIBuffer(RHICmdList));
 	}
 }
 
@@ -689,11 +684,13 @@ void FSkinWeightProfilesData::ResetGPUReadback()
 	ReadbackData.ReadbackFinishedFrameIndex = INDEX_NONE;
 }
 
-void FSkinWeightProfilesData::InitialiseProfileBuffer(const FName& ProfileName)
+void FSkinWeightProfilesData::InitialiseProfileBuffer(const FSkinWeightProfileStack& InProfileStack)
 {
 	LLM_SCOPE(ELLMTag::SkeletalMesh);
 
-	if (ProfileNameToBuffer.Contains(ProfileName))
+	check(InProfileStack == InProfileStack.Normalized());
+	
+	if (ProfileStackToBuffer.Contains(InProfileStack))
 	{
 		// Maybe it has already been previously set-up and not unloaded yet
 		return;
@@ -716,19 +713,14 @@ void FSkinWeightProfilesData::InitialiseProfileBuffer(const FName& ProfileName)
 		if (ensure(BaseBufferData))
 		{
 			FSkinWeightVertexBuffer* OverrideBuffer = new FSkinWeightVertexBuffer();
-			OverrideBuffer->CopyMetaData(*BaseBuffer);
-			ProfileNameToBuffer.Add(ProfileName, OverrideBuffer);
-
 			OverrideBuffer->SetNeedsCPUAccess(BaseBuffer->GetNeedsCPUAccess());
-			
-			const FRuntimeSkinWeightProfileData* ProfilePtr = OverrideData.Find(ProfileName);
-			if (ProfilePtr)
-			{
-				ProfilePtr->ApplyOverrides(OverrideBuffer, BaseBufferData, BaseBuffer->GetNumVertices());
-			}
 
+			ProfileStackToBuffer.Add(InProfileStack, OverrideBuffer);
+
+			ApplyOverrideProfileStack(InProfileStack, OverrideBuffer, BaseBufferData);
+			
 #if RHI_ENABLE_RESOURCE_INFO
-			const FName OwnerName = FName(ProfileName.ToString() + TEXT("_FSkinWeightProfilesData"));
+			const FName OwnerName = FName(InProfileStack.GetUniqueId() + TEXT("_FSkinWeightProfilesData"));
 			OverrideBuffer->SetOwnerName(OwnerName);
 #endif
 			OverrideBuffer->BeginInitResources();
@@ -736,71 +728,87 @@ void FSkinWeightProfilesData::InitialiseProfileBuffer(const FName& ProfileName)
 	}
 }
 
-void FSkinWeightProfilesData::ApplyOverrideProfile(FSkinWeightVertexBuffer* OverrideBuffer, const FName& ProfileName)
+void FSkinWeightProfilesData::ApplyOverrideProfileStack(
+	const FSkinWeightProfileStack& InProfileStack,
+	FSkinWeightVertexBuffer* OverrideBuffer,
+	const uint8* BaseBufferData
+	)
 {
+	if (!BaseBufferData)
+	{
+		BaseBufferData = BaseBuffer->GetDataVertexBuffer()->GetWeightData();
+	}
+	
 	OverrideBuffer->CopyMetaData(*BaseBuffer);
-
-	const FRuntimeSkinWeightProfileData* ProfilePtr = OverrideData.Find(ProfileName);
-	if (ProfilePtr)
+	OverrideBuffer->CopySkinWeightRawDataFromBuffer(BaseBufferData, BaseBuffer->GetNumVertices());
+	
+	for (int32 LayerIndex = 0; LayerIndex < FSkinWeightProfileStack::MaxLayerCount; ++LayerIndex)
 	{
-		ProfilePtr->ApplyOverrides(OverrideBuffer, BaseBuffer->GetDataVertexBuffer()->GetWeightData(), BaseBuffer->GetNumVertices());
-	}	
+		if (FName ProfileName = InProfileStack.Layers[LayerIndex];
+			!ProfileName.IsNone())
+		{
+			if (const FRuntimeSkinWeightProfileData* ProfilePtr = OverrideData.Find(ProfileName))
+			{
+				// Only copy the base buffer's weights when applying the first layer.
+				ProfilePtr->ApplyOverrides(OverrideBuffer);
+			}	
+		}
+	}
+
 }
 
-void FSkinWeightProfilesData::InitRHIForStreaming(const TArray<TPair<FName, FSkinWeightRHIInfo>>& IntermediateBuffers, FRHIResourceReplaceBatcher& Batcher)
+void FSkinWeightProfilesData::InitRHIForStreaming(
+	const TArray<TPair<FSkinWeightProfileStack, FSkinWeightRHIInfo>>& IntermediateBuffers, 
+	FRHIResourceReplaceBatcher& Batcher
+	)
 {
 	for (int32 Idx = 0; Idx < IntermediateBuffers.Num(); ++Idx)
 	{
-		const FName& ProfileName = IntermediateBuffers[Idx].Key;
+		const FSkinWeightProfileStack& ProfileStack = IntermediateBuffers[Idx].Key;
 		const FSkinWeightRHIInfo& IntermediateBuffer = IntermediateBuffers[Idx].Value;
-		ProfileNameToBuffer.FindChecked(ProfileName)->InitRHIForStreaming(IntermediateBuffer, Batcher);
+		ProfileStackToBuffer.FindChecked(ProfileStack)->InitRHIForStreaming(IntermediateBuffer, Batcher);
 	}
 }
 
 void FSkinWeightProfilesData::ReleaseRHIForStreaming(FRHIResourceReplaceBatcher& Batcher)
 {
-	for (TMap<FName, FSkinWeightVertexBuffer*>::TIterator It(ProfileNameToBuffer); It; ++It)
+	for (TPair<FSkinWeightProfileStack, FSkinWeightVertexBuffer*>& Item: ProfileStackToBuffer)
 	{
-		It->Value->ReleaseRHIForStreaming(Batcher);
+		Item.Value->ReleaseRHIForStreaming(Batcher);
 	}
 }
 
-void FRuntimeSkinWeightProfileData::ApplyOverrides(FSkinWeightVertexBuffer* OverrideBuffer, const uint8* DataBuffer, const int32 NumVerts) const
+void FRuntimeSkinWeightProfileData::ApplyOverrides(FSkinWeightVertexBuffer* OverrideBuffer) const
 {
-	if (DataBuffer)
+	if (OverrideBuffer)
 	{
-		if (OverrideBuffer)
-		{
-			OverrideBuffer->CopySkinWeightRawDataFromBuffer(DataBuffer, NumVerts);
+		uint8* TargetSkinWeightData = OverrideBuffer->GetDataVertexBuffer()->GetWeightData();
+		
+		const uint8 VertexStride = OverrideBuffer->GetConstantInfluencesVertexStride();
+		const uint8 WeightDataOffset = OverrideBuffer->GetBoneIndexByteSize() * OverrideBuffer->GetMaxBoneInfluences();
 
-			uint8* TargetSkinWeightData = OverrideBuffer->GetDataVertexBuffer()->GetWeightData();
+		// Apply overrides
+		for (auto VertexIndexOverridePair : VertexIndexToInfluenceOffset)
+		{
+			const uint32 VertexIndex = VertexIndexOverridePair.Key;
+			const uint32 InfluenceOffset = VertexIndexOverridePair.Value;
 			
-			const uint8 VertexStride = OverrideBuffer->GetConstantInfluencesVertexStride();
-			const uint8 WeightDataOffset = OverrideBuffer->GetBoneIndexByteSize() * OverrideBuffer->GetMaxBoneInfluences();
-
-			// Apply overrides
-			for (auto VertexIndexOverridePair : VertexIndexToInfluenceOffset)
-			{
-				const uint32 VertexIndex = VertexIndexOverridePair.Key;
-				const uint32 InfluenceOffset = VertexIndexOverridePair.Value;
-				
-				uint8* BoneData = TargetSkinWeightData + (VertexIndex * VertexStride);
-				uint8* WeightData = BoneData + WeightDataOffset;
+			uint8* BoneData = TargetSkinWeightData + (VertexIndex * VertexStride);
+			uint8* WeightData = BoneData + WeightDataOffset;
 
 #if !UE_BUILD_SHIPPING
-				uint32 VertexOffset = 0;
-				uint32 VertexInfluenceCount = 0;
-				OverrideBuffer->GetVertexInfluenceOffsetCount(VertexIndex, VertexOffset, VertexInfluenceCount);
-				check(NumWeightsPerVertex <= VertexInfluenceCount);
-				check((void*)(((uint8*)TargetSkinWeightData) + VertexOffset) == (void*)BoneData);
-				check(b16BitBoneIndices == OverrideBuffer->Use16BitBoneIndex());
+			uint32 VertexOffset = 0;
+			uint32 VertexInfluenceCount = 0;
+			OverrideBuffer->GetVertexInfluenceOffsetCount(VertexIndex, VertexOffset, VertexInfluenceCount);
+			check(NumWeightsPerVertex <= VertexInfluenceCount);
+			check((void*)(((uint8*)TargetSkinWeightData) + VertexOffset) == (void*)BoneData);
+			check(b16BitBoneIndices == OverrideBuffer->Use16BitBoneIndex());
 #endif
-				// BoneIDs either contains FBoneIndexType entries spanning (2) uint8 values, or single uint8 bone indices (1)
-				const uint32 BoneIndexByteSize = OverrideBuffer->GetBoneIndexByteSize();
-				const uint32 BoneWeightByteSize = OverrideBuffer->GetBoneWeightByteSize();
-				FMemory::Memcpy(BoneData, &BoneIDs[InfluenceOffset * NumWeightsPerVertex * BoneIndexByteSize], BoneIndexByteSize * NumWeightsPerVertex);
-				FMemory::Memcpy(WeightData, &BoneWeights[InfluenceOffset * NumWeightsPerVertex * BoneWeightByteSize], BoneWeightByteSize * NumWeightsPerVertex);
-			}
+			// BoneIDs either contains FBoneIndexType entries spanning (2) uint8 values, or single uint8 bone indices (1)
+			const uint32 BoneIndexByteSize = OverrideBuffer->GetBoneIndexByteSize();
+			const uint32 BoneWeightByteSize = OverrideBuffer->GetBoneWeightByteSize();
+			FMemory::Memcpy(BoneData, &BoneIDs[InfluenceOffset * NumWeightsPerVertex * BoneIndexByteSize], BoneIndexByteSize * NumWeightsPerVertex);
+			FMemory::Memcpy(WeightData, &BoneWeights[InfluenceOffset * NumWeightsPerVertex * BoneWeightByteSize], BoneWeightByteSize * NumWeightsPerVertex);
 		}
 	}
 }
diff --git a/Engine/Source/Runtime/Engine/Private/Animation/SkinWeightProfileManager.cpp b/Engine/Source/Runtime/Engine/Private/Animation/SkinWeightProfileManager.cpp
--- a/Engine/Source/Runtime/Engine/Private/Animation/SkinWeightProfileManager.cpp
+++ b/Engine/Source/Runtime/Engine/Private/Animation/SkinWeightProfileManager.cpp
@@ -102,14 +102,20 @@ void FSkinWeightProfileManager::OnShutdown()
 	WorldManagers.Empty();
 }
 
-void FSkinWeightProfileManager::RequestSkinWeightProfile(FName InProfileName, USkinnedAsset* SkinnedAsset, UObject* Requester, FRequestFinished& Callback, int32 LODIndex /*= INDEX_NONE*/)
+void FSkinWeightProfileManager::RequestSkinWeightProfileStack(
+	FSkinWeightProfileStack InProfileStack, 
+	USkinnedAsset* SkinnedAsset, 
+	UObject* Requester,
+	FRequestFinished& Callback,
+	int32 LODIndex /*= INDEX_NONE*/
+	)
 {
 	// Make sure we have an actual skeletal mesh
 	if (USkeletalMesh* const Mesh = Cast<USkeletalMesh>(SkinnedAsset))
 	{
 		// Setup a request structure
 		FSetProfileRequest ProfileRequest;
-		ProfileRequest.ProfileName = InProfileName;
+		ProfileRequest.ProfileStack = InProfileStack.Normalized();
 		ProfileRequest.Callback = Callback;
 		ProfileRequest.WeakSkeletalMesh = Mesh;
 		ProfileRequest.IdentifyingObject = Requester;
@@ -184,7 +190,7 @@ void FSkinWeightProfileManager::DoTick(float DeltaTime, ENamedThreads::Type Curr
 			bool bRemove = Request.IdentifyingObject.IsStale() || Request.WeakSkeletalMesh.IsStale() ||
 				CanceledRequest.ContainsByPredicate([Request](FSetProfileRequest& B)
 			{
-				return Request.IdentifyingObject == B.IdentifyingObject && Request.WeakSkeletalMesh == B.WeakSkeletalMesh && Request.ProfileName == B.ProfileName;
+				return Request.IdentifyingObject == B.IdentifyingObject && Request.WeakSkeletalMesh == B.WeakSkeletalMesh && Request.ProfileStack == B.ProfileStack;
 			});
 
 			const USkeletalMesh* SkeletalMesh = Request.WeakSkeletalMesh.Get();
@@ -197,7 +203,7 @@ void FSkinWeightProfileManager::DoTick(float DeltaTime, ENamedThreads::Type Curr
 				bool bAllBuffersReady = true;
 				for (const int32 LODIndex : Request.LODIndices)
 				{
-					if (!RenderData->LODRenderData[LODIndex].SkinWeightProfilesData.ContainsOverrideBuffer(Request.ProfileName))
+					if (!RenderData->LODRenderData[LODIndex].SkinWeightProfilesData.ContainsOverrideBuffer(Request.ProfileStack))
 					{
 						bAllBuffersReady = false;
 						break;
@@ -209,7 +215,7 @@ void FSkinWeightProfileManager::DoTick(float DeltaTime, ENamedThreads::Type Curr
 					if (Request.IdentifyingObject.IsValid())
 					{
 						UE_LOG(LogSkinWeightProfileManager, Display, TEXT("Callback [%s | %s]"), *SkeletalMesh->GetName(), *Request.IdentifyingObject->GetName());
-						Request.Callback(Request.WeakSkeletalMesh, Request.ProfileName);
+						Request.Callback(Request.WeakSkeletalMesh, Request.ProfileStack);
 						bRemove = true;
 					}
 				}
@@ -335,7 +341,7 @@ void FSkinWeightProfileManagerAsyncTask::DoTask(ENamedThreads::Type CurrentThrea
 					
 					if (!!GAllowCPU && LODRenderData.SkinWeightVertexBuffer.GetNeedsCPUAccess())
 					{
-						SkinweightData.InitialiseProfileBuffer(Request.ProfileName);
+						SkinweightData.InitialiseProfileBuffer(Request.ProfileStack);
 					}
 					else
 					{
@@ -351,7 +357,7 @@ void FSkinWeightProfileManagerAsyncTask::DoTask(ENamedThreads::Type CurrentThrea
 							}
 							else if (SkinweightData.IsDataReadbackFinished())
 							{						
-								SkinweightData.InitialiseProfileBuffer(Request.ProfileName);
+								SkinweightData.InitialiseProfileBuffer(Request.ProfileStack);
 							}
 						}
 					}
diff --git a/Engine/Source/Runtime/Engine/Private/Components/SkinnedMeshComponent.cpp b/Engine/Source/Runtime/Engine/Private/Components/SkinnedMeshComponent.cpp
--- a/Engine/Source/Runtime/Engine/Private/Components/SkinnedMeshComponent.cpp
+++ b/Engine/Source/Runtime/Engine/Private/Components/SkinnedMeshComponent.cpp
@@ -481,8 +481,6 @@ USkinnedMeshComponent::USkinnedMeshComponent(const FObjectInitializer& ObjectIni
 	bDrawDebugSkeleton = false;
 #endif
 
-	CurrentSkinWeightProfileName = NAME_None;
-
 	PRAGMA_DISABLE_DEPRECATION_WARNINGS
 	PredictedLODLevel = 0;
 	PRAGMA_ENABLE_DEPRECATION_WARNINGS
@@ -4881,133 +4879,36 @@ void USkinnedMeshComponent::ClearSkinWeightOverride(int32 LODIndex)
 	}
 }
 
-bool USkinnedMeshComponent::SetSkinWeightProfile(FName InProfileName)
-{
-	bool bContainsProfile = false;
-
-	if (FSkeletalMeshRenderData* SkelMeshRenderData = GetSkeletalMeshRenderData())
-	{
-		// Ensure the LOD infos array is initialized
-		InitLODInfos();
-		for (int32 LODIndex = 0; LODIndex < LODInfo.Num(); ++LODIndex)
-        {
-			// Check whether or not setting a profile is allow for this LOD index
-			if (LODIndex > GSkinWeightProfilesAllowedFromLOD)
-			{
-				FSkeletalMeshLODRenderData& RenderData = SkelMeshRenderData->LODRenderData[LODIndex];
-
-				bContainsProfile |= RenderData.SkinWeightProfilesData.ContainsProfile(InProfileName);
-
-				// Retrieve this profile's skin weight buffer
-				FSkinWeightVertexBuffer* Buffer = RenderData.SkinWeightProfilesData.GetOverrideBuffer(InProfileName);
-        
-				FSkelMeshComponentLODInfo& Info = LODInfo[LODIndex];
-				Info.OverrideProfileSkinWeights = Buffer;
-                
-				if (Buffer != nullptr)
-				{
-					bSkinWeightProfileSet = true;
-				}
-			}
-        }
-
-		if (bContainsProfile)
-		{
-			CurrentSkinWeightProfileName = InProfileName;
-
-			if (bSkinWeightProfileSet)
-			{
-				UpdateSkinWeightOverrideBuffer();
-			}
-			else 
-			{
-				TWeakObjectPtr<USkinnedMeshComponent> WeakComponent = this;
-				FRequestFinished Callback = [WeakComponent](TWeakObjectPtr<USkeletalMesh> WeakMesh, FName ProfileName)
-				{
-					// Ensure that the request objects are still valid
-					if (WeakMesh.IsValid() && WeakComponent.IsValid())
-					{
-						USkinnedMeshComponent* Component = WeakComponent.Get();
-						Component->InitLODInfos();
-
-						Component->bSkinWeightProfilePending = false;
-						Component->bSkinWeightProfileSet = true;
-
-						if (FSkeletalMeshRenderData * RenderData = WeakMesh->GetResourceForRendering())
-						{
-							const int32 NumLODs = RenderData->LODRenderData.Num();
-							for (int32 Index = 0; Index < NumLODs; ++Index)
-							{
-								FSkeletalMeshLODRenderData& LODRenderData = RenderData->LODRenderData[Index];
-								FSkinWeightProfilesData& SkinweightData = LODRenderData.SkinWeightProfilesData;
 
-								// Check whether or not setting a profile is allow for this LOD index
-								if (Index > GSkinWeightProfilesAllowedFromLOD)
-								{
-									// Retrieve this profile's skin weight buffer
-									FSkinWeightVertexBuffer* Buffer = SkinweightData.GetOverrideBuffer(ProfileName);
-									FSkelMeshComponentLODInfo& Info = Component->LODInfo[Index];
-									Info.OverrideProfileSkinWeights = Buffer;
-								}
-							}
-
-							Component->UpdateSkinWeightOverrideBuffer();
-						}
-					}
-				};
-
-				// Put in a skin weight profile request
-				if (FSkinWeightProfileManager* Manager = FSkinWeightProfileManager::Get(GetWorld()))
-				{
-					Manager->RequestSkinWeightProfile(InProfileName, GetSkinnedAsset(), this, Callback);
-					bSkinWeightProfilePending = true;
-				}
-			}
-		}
-	}
-
-	return bContainsProfile;
+bool USkinnedMeshComponent::SetSkinWeightProfile(
+	FName InProfileName,
+	ESkinWeightProfileLayer InLayer
+	)
+{
+	FSkinWeightProfileStack ProfileStack{CurrentSkinWeightProfileLayers};
+	ProfileStack[static_cast<int32>(InLayer)] = InProfileName;
+	return SetSkinWeightProfileStack(ProfileStack);
 }
 
-void USkinnedMeshComponent::ClearSkinWeightProfile()
+void USkinnedMeshComponent::ClearSkinWeightProfile(
+	ESkinWeightProfileLayer InLayer
+	)
 {
-	if (FSkeletalMeshRenderData* SkelMeshRenderData = GetSkeletalMeshRenderData())
-	{	
-		bool bCleared = false;
-
-		if (bSkinWeightProfileSet)
-		{
-			InitLODInfos();
-			// Clear skin weight buffer set for all of the LODs
-			for (int32 LODIndex = 0; LODIndex < LODInfo.Num(); ++LODIndex)
-			{
-				FSkelMeshComponentLODInfo& Info = LODInfo[LODIndex];
-				bCleared |= (Info.OverrideProfileSkinWeights != nullptr);
-				Info.OverrideProfileSkinWeights = nullptr;
-			}
-
-			if (bCleared)
-			{
-				UpdateSkinWeightOverrideBuffer();
-			}
-		}
+	FSkinWeightProfileStack ProfileStack{CurrentSkinWeightProfileLayers};
+	ProfileStack[static_cast<int32>(InLayer)] = NAME_None;
+	SetSkinWeightProfileStack(ProfileStack);
+}
 
-		if (bSkinWeightProfilePending)
-		{
-			if (FSkinWeightProfileManager * Manager = FSkinWeightProfileManager::Get(GetWorld()))
-			{
-				Manager->CancelSkinWeightProfileRequest(this);
-			}
-		}
-	}
 
-	bSkinWeightProfilePending = false;
-	bSkinWeightProfileSet = false;
-	CurrentSkinWeightProfileName = NAME_None;
+void USkinnedMeshComponent::ClearAllSkinWeightProfiles()
+{
+	SetSkinWeightProfileStack(FSkinWeightProfileStack());
 }
 
+
 void USkinnedMeshComponent::UnloadSkinWeightProfile(FName InProfileName)
 {
+	#if 0
 	if (FSkeletalMeshRenderData* SkelMeshRenderData = GetSkeletalMeshRenderData())
 	{
 		if (LODInfo.Num())
@@ -5050,6 +4951,7 @@ void USkinnedMeshComponent::UnloadSkinWeightProfile(FName InProfileName)
 		bSkinWeightProfileSet = false;
 		CurrentSkinWeightProfileName = NAME_None;
 	}
+#endif
 }
 
 bool USkinnedMeshComponent::IsUsingSkinWeightProfile() const
@@ -5070,11 +4972,20 @@ bool USkinnedMeshComponent::IsUsingSkinWeightProfile() const
 	return false;
 }
 
-FName USkinnedMeshComponent::GetCurrentSkinWeightProfileName() const
+FName USkinnedMeshComponent::GetCurrentSkinWeightProfileName(
+	ESkinWeightProfileLayer InLayer
+	) const
 {
+	const int32 LayerIndex = static_cast<int32>(InLayer);
+	return GetCurrentSkinWeightProfileLayerNames()[LayerIndex];
+}
+
+TArray<FName> USkinnedMeshComponent::GetCurrentSkinWeightProfileLayerNames() const
+{
+	FSkinWeightProfileStack ProfileStack;
 	if (bSkinWeightProfileSet)
 	{
-		return CurrentSkinWeightProfileName;
+		ProfileStack = FSkinWeightProfileStack{CurrentSkinWeightProfileLayers};
 	}
 
 	if (USkinnedAsset* Asset = GetSkinnedAsset())
@@ -5083,12 +4994,18 @@ FName USkinnedMeshComponent::GetCurrentSkinWeightProfileName() const
 		{
 			if( ProfileData->IsDefaultOverridden() || ProfileData->IsStaticOverridden())
 			{
-				return ProfileData->GetDefaultProfileName();
+				ProfileStack = ProfileData->GetDefaultProfileStack();
 			}
 		}
 	}
 
-	return NAME_None;
+	TArray<FName> ProfileNames;
+	ProfileNames.SetNum(FSkinWeightProfileStack::MaxLayerCount);
+	for (int32 LayerIndex = 0; LayerIndex < FSkinWeightProfileStack::MaxLayerCount; ++LayerIndex)
+	{
+		ProfileNames[LayerIndex] = ProfileStack[LayerIndex];
+	}
+	return ProfileNames;
 }
 
 void USkinnedMeshComponent::UpdateSkinWeightOverrideBuffer()
@@ -5104,6 +5021,146 @@ void USkinnedMeshComponent::UpdateSkinWeightOverrideBuffer()
 	}
 }
 
+bool USkinnedMeshComponent::SetSkinWeightProfileStack(const FSkinWeightProfileStack& InProfileStack)
+{
+	// If we're not actually changing anything, then just return and say we did, as if the operation succeeded.
+	if (InProfileStack.Normalized() == FSkinWeightProfileStack{CurrentSkinWeightProfileLayers}.Normalized())
+	{
+		return true;
+	}
+
+	bool bChanged = false;
+	
+	if (FSkeletalMeshRenderData* SkelMeshRenderData = GetSkeletalMeshRenderData())
+	{
+		if (InProfileStack.IsEmpty())
+		{
+			if (bSkinWeightProfileSet)
+			{
+				bool bCleared = false;
+				
+				InitLODInfos();
+				// Clear skin weight buffer set for all of the LODs
+				for (int32 LODIndex = 0; LODIndex < LODInfo.Num(); ++LODIndex)
+				{
+					FSkelMeshComponentLODInfo& Info = LODInfo[LODIndex];
+					bCleared |= (Info.OverrideProfileSkinWeights != nullptr);
+					Info.OverrideProfileSkinWeights = nullptr;
+				}
+
+				if (bCleared)
+				{
+					UpdateSkinWeightOverrideBuffer();
+				}
+			}
+
+			// If there was a change request in flight, cancel it now so that we don't get updated with a stale request.
+			if (bSkinWeightProfilePending)
+			{
+				if (FSkinWeightProfileManager* Manager = FSkinWeightProfileManager::Get(GetWorld()))
+				{
+					Manager->CancelSkinWeightProfileRequest(this);
+				}
+			}
+
+			bChanged = true;
+		}
+		else
+		{
+			// Ensure the LOD infos array is initialized
+			bool bHaveExistingBuffer = false;
+			
+			InitLODInfos();
+			for (int32 LODIndex = 0; LODIndex < LODInfo.Num(); ++LODIndex)
+	        {
+				// Check whether setting a profile is allowed for this LOD index
+				if (LODIndex > GSkinWeightProfilesAllowedFromLOD)
+				{
+					FSkeletalMeshLODRenderData& RenderData = SkelMeshRenderData->LODRenderData[LODIndex];
+
+					for (int32 LayerIndex = 0; LayerIndex < FSkinWeightProfileStack::MaxLayerCount; ++LayerIndex)
+					{
+						bChanged |= RenderData.SkinWeightProfilesData.ContainsProfile(InProfileStack[LayerIndex]);
+					}
+
+					// Retrieve this profile's skin weight buffer
+					FSkinWeightVertexBuffer* Buffer = RenderData.SkinWeightProfilesData.GetOverrideBuffer(InProfileStack);
+	        
+					FSkelMeshComponentLODInfo& Info = LODInfo[LODIndex];
+					Info.OverrideProfileSkinWeights = Buffer;
+	                
+					bHaveExistingBuffer = (Buffer != nullptr);
+				}
+	        }
+
+			if (bChanged)
+			{
+				if (bHaveExistingBuffer)
+				{
+					UpdateSkinWeightOverrideBuffer();
+					bSkinWeightProfileSet = true;
+				}
+				else 
+				{
+					TWeakObjectPtr<USkinnedMeshComponent> WeakComponent = this;
+					
+					FRequestFinished Callback = [WeakComponent](TWeakObjectPtr<USkeletalMesh> WeakMesh, FSkinWeightProfileStack ProfileStack)
+					{
+						// Ensure that the request objects are still valid
+						if (WeakMesh.IsValid() && WeakComponent.IsValid())
+						{
+							USkinnedMeshComponent* Component = WeakComponent.Get();
+							Component->InitLODInfos();
+
+							Component->bSkinWeightProfilePending = false;
+							Component->bSkinWeightProfileSet = true;
+
+							if (FSkeletalMeshRenderData * RenderData = WeakMesh->GetResourceForRendering())
+							{
+								const int32 NumLODs = RenderData->LODRenderData.Num();
+								for (int32 Index = 0; Index < NumLODs; ++Index)
+								{
+									FSkeletalMeshLODRenderData& LODRenderData = RenderData->LODRenderData[Index];
+									FSkinWeightProfilesData& SkinweightData = LODRenderData.SkinWeightProfilesData;
+
+									// Check whether or not setting a profile is allow for this LOD index
+									if (Index > GSkinWeightProfilesAllowedFromLOD)
+									{
+										// Retrieve this profile's skin weight buffer
+										FSkinWeightVertexBuffer* Buffer = SkinweightData.GetOverrideBuffer(ProfileStack);
+										FSkelMeshComponentLODInfo& Info = Component->LODInfo[Index];
+										Info.OverrideProfileSkinWeights = Buffer;
+									}
+								}
+
+								Component->UpdateSkinWeightOverrideBuffer();
+							}
+						}
+					};
+
+					// Put in a skin weight profile request
+					if (FSkinWeightProfileManager* Manager = FSkinWeightProfileManager::Get(GetWorld()))
+					{
+						Manager->RequestSkinWeightProfileStack(InProfileStack, GetSkinnedAsset(), this, Callback);
+						bSkinWeightProfilePending = true;
+					}
+				}
+			}
+		}
+	}
+
+	if (InProfileStack.IsEmpty())
+	{
+		bSkinWeightProfilePending = false;
+		bSkinWeightProfileSet = false;
+	}
+
+	// Store the current layers.
+	InProfileStack.CopyIntoArray(CurrentSkinWeightProfileLayers);
+
+	return bChanged;
+}
+
 void USkinnedMeshComponent::ReleaseUpdateRateParams()
 {
 	FAnimUpdateRateManager::CleanupUpdateRateParametersRef(this);
diff --git a/Engine/Source/Runtime/Engine/Public/Animation/SkinWeightProfile.h b/Engine/Source/Runtime/Engine/Public/Animation/SkinWeightProfile.h
--- a/Engine/Source/Runtime/Engine/Public/Animation/SkinWeightProfile.h
+++ b/Engine/Source/Runtime/Engine/Public/Animation/SkinWeightProfile.h
@@ -2,15 +2,15 @@
 
 #pragma once
 
-#include "Rendering/SkinWeightVertexBuffer.h"
-#include "UObject/PerPlatformProperties.h"
-#include "UObject/NameTypes.h"
-#include "Misc/CoreStats.h"
-#include "RenderingThread.h"
-#include "HAL/UnrealMemory.h"
 #include "BoneIndices.h"
+#include "Misc/StringBuilder.h"
 #include "RHIGPUReadback.h"
+#include "Rendering/SkinWeightVertexBuffer.h"
+#include "RenderingThread.h"
+#include "String/Join.h"
 #include "Templates/UniquePtr.h"
+#include "UObject/NameTypes.h"
+#include "UObject/PerPlatformProperties.h"
 
 #include "SkinWeightProfile.generated.h"
 
@@ -95,7 +95,7 @@ struct FRuntimeSkinWeightProfileData
 		friend FArchive& operator<<(FArchive& Ar, FSkinWeightOverrideInfo& OverrideInfo);
 	};
 
-	void ApplyOverrides(FSkinWeightVertexBuffer* OverrideBuffer, const uint8* DataBuffer, const int32 NumVerts) const;	
+	void ApplyOverrides(FSkinWeightVertexBuffer* OverrideBuffer) const;	
 	void ApplyDefaultOverride(FSkinWeightVertexBuffer* Buffer) const;
 
 #if WITH_EDITORONLY_DATA
@@ -121,13 +121,119 @@ struct FSkinweightReadbackData
 {
 	TUniquePtr<FRHIGPUBufferReadback> BufferReadback;
 	TArray<uint8> ReadbackData;
-	uint32 ReadbackFinishedFrameIndex;
+	uint32 ReadbackFinishedFrameIndex = std::numeric_limits<uint32>::max();
 };
 
+/** An identifier to identify a skin weight profile layer stack. 
+  * Currently limited to two layers, but can be extended if needed.
+  */
+struct FSkinWeightProfileStack 
+{
+	static constexpr int32 MaxLayerCount = 2;
+	TStaticArray<FName, MaxLayerCount> Layers;
+	
+	FSkinWeightProfileStack() :
+		Layers{InPlace, NAME_None}
+	{}
+
+	// A fancy way of specifying a function with FName argument count between 1 and MaxLayerCount
+	template <
+		typename... ArgTypes
+		UE_REQUIRES((sizeof...(ArgTypes) > 0 && sizeof...(ArgTypes) <= MaxLayerCount) && UE::Core::Private::TCanBeConvertedToFromAll_V<FName, ArgTypes...>)
+	>
+	explicit FSkinWeightProfileStack(ArgTypes&&... Args)
+		: Layers{Forward<ArgTypes>(Args)...}
+	{
+		// Make sure that we mark any remaining elements are NAME_None to avoid corrupted names.  
+		for (int32 LayerIndex = sizeof...(ArgTypes); LayerIndex < MaxLayerCount; ++LayerIndex)
+		{
+			Layers[LayerIndex] = NAME_None;
+		}
+	}
+
+	explicit FSkinWeightProfileStack(const FName InProfileNames[MaxLayerCount])
+	{
+		for (int32 LayerIndex = 0; LayerIndex < MaxLayerCount; ++LayerIndex)
+		{
+			Layers[LayerIndex] = InProfileNames[LayerIndex];
+		}
+	}
+
+	bool operator==(const FSkinWeightProfileStack& InProfileStack) const
+	{
+		return Layers == InProfileStack.Layers;
+	}
+	bool operator!=(const FSkinWeightProfileStack& InProfileStack) const
+	{
+		return Layers != InProfileStack.Layers;
+	}
+
+	FName operator[](int32 InLayerIndex) const
+	{
+		return Layers[InLayerIndex];
+	}
+
+	FName& operator[](int32 InLayerIndex)
+	{
+		return Layers[InLayerIndex];
+	}
+
+	bool IsEmpty() const
+	{
+		for (int32 LayerIndex = 0; LayerIndex < MaxLayerCount; ++LayerIndex)
+		{
+			if (!Layers[LayerIndex].IsNone())
+			{
+				return false;
+			}
+		}
+		return true;
+	}
+
+	void CopyIntoArray(FName OutProfileNames[MaxLayerCount]) const
+	{
+		for (int32 LayerIndex = 0; LayerIndex < MaxLayerCount; ++LayerIndex)
+		{
+			OutProfileNames[LayerIndex] = Layers[LayerIndex];
+		}
+	}
+	
+	/** Returns a normalized layer stack, such that any empty layers are collapsed. */
+	FSkinWeightProfileStack Normalized() const
+	{
+		FSkinWeightProfileStack NormalizedLayers;
+		
+		for (int32 ReadLayer = 0, WriteLayer = 0; ReadLayer < MaxLayerCount; ++ReadLayer)
+		{
+			if (!Layers[ReadLayer].IsNone())
+			{
+				NormalizedLayers.Layers[WriteLayer++] = Layers[ReadLayer];
+			}
+		}
+		return NormalizedLayers;
+	}
+	
+	FString GetUniqueId() const
+	{
+		TStringBuilder<256> IdString(InPlace, UE::String::JoinBy(Layers, [](FName InLayerName) { return InLayerName.ToString(); }, TEXT("-")));
+		return IdString.ToString();
+	}
+
+	friend uint32 GetTypeHash(const FSkinWeightProfileStack& InProfileStack)
+	{
+		uint32 Hash = 0;
+		for (int32 LayerIndex = 0; LayerIndex < MaxLayerCount; ++LayerIndex)
+		{
+			Hash = HashCombineFast(Hash, GetTypeHash(InProfileStack[LayerIndex]));
+		}
+		return Hash;
+	}
+};
+
+
 /** Runtime structure for keeping track of skin weight profile(s) and the associated buffer */
 struct FSkinWeightProfilesData
 {
-	FSkinWeightProfilesData() : BaseBuffer(nullptr), DefaultOverrideSkinWeightBuffer(nullptr), bDefaultOverridden(false), bStaticOverridden(false), DefaultProfileName(NAME_None) {}
 	ENGINE_API void Init(FSkinWeightVertexBuffer* InBaseBuffer);
 
 	ENGINE_API ~FSkinWeightProfilesData();
@@ -141,59 +247,65 @@ struct FSkinWeightProfilesData
 #endif 
 	ENGINE_API void SetDynamicDefaultSkinWeightProfile(USkeletalMesh* Mesh, int32 LODIndex, bool bSerialization = false);	
 	ENGINE_API void ClearDynamicDefaultSkinWeightProfile(USkeletalMesh* Mesh, int32 LODIndex);
-	ENGINE_API void SetupDynamicDefaultSkinweightProfile();
+	ENGINE_API void SetupDynamicDefaultSkinWeightProfile();
 	FSkinWeightVertexBuffer* GetDefaultOverrideBuffer() const { return DefaultOverrideSkinWeightBuffer; }
 
-	ENGINE_API bool ContainsProfile(const FName& ProfileName) const;
-	ENGINE_API FSkinWeightVertexBuffer* GetOverrideBuffer(const FName& ProfileName) const;
-	ENGINE_API bool ContainsOverrideBuffer(const FName& ProfileName) const;
+	// Buffer lookup for layered profiles.
+	ENGINE_API FSkinWeightVertexBuffer* GetOverrideBuffer(const FSkinWeightProfileStack& InProfileStack) const;
+	bool ContainsOverrideBuffer(const FSkinWeightProfileStack& InProfileStack) const;
 	
+	// Lookups for individual profiles.
+	ENGINE_API bool ContainsProfile(const FName& ProfileName) const;
 	ENGINE_API const FRuntimeSkinWeightProfileData* GetOverrideData(const FName& ProfileName) const;
 	ENGINE_API FRuntimeSkinWeightProfileData& AddOverrideData(const FName& ProfileName);
-	
-	ENGINE_API void ReleaseBuffer(const FName& ProfileName, bool bForceRelease = false);
+
 	ENGINE_API void ReleaseResources();
 
 	ENGINE_API SIZE_T GetResourcesSize() const;
 	ENGINE_API SIZE_T GetCPUAccessMemoryOverhead() const;
-
+ 
 	friend FArchive& operator<<(FArchive& Ar, FSkinWeightProfilesData& OverrideData);
 
 	ENGINE_API void SerializeMetaData(FArchive& Ar);
 
 	ENGINE_API void ReleaseCPUResources();
 
-	ENGINE_API void CreateRHIBuffers(FRHICommandListBase& RHICmdList, TArray<TPair<FName, FSkinWeightRHIInfo>>& OutBuffers);
+	ENGINE_API void CreateRHIBuffers(FRHICommandListBase& RHICmdList, TArray<TPair<FSkinWeightProfileStack, FSkinWeightRHIInfo>>& OutBuffers);
 
-	ENGINE_API void InitRHIForStreaming(const TArray<TPair<FName, FSkinWeightRHIInfo>>& IntermediateBuffers, FRHIResourceReplaceBatcher& Batcher);
+	ENGINE_API void InitRHIForStreaming(const TArray<TPair<FSkinWeightProfileStack, FSkinWeightRHIInfo>>& IntermediateBuffers, FRHIResourceReplaceBatcher& Batcher);
 	ENGINE_API void ReleaseRHIForStreaming(FRHIResourceReplaceBatcher& Batcher);
 
-	ENGINE_API bool IsPendingReadback() const;
-	ENGINE_API void EnqueueGPUReadback();
-	ENGINE_API bool IsGPUReadbackFinished() const;
-	ENGINE_API void EnqueueDataReadback();
-	ENGINE_API bool IsDataReadbackPending() const;
-	ENGINE_API bool IsDataReadbackFinished() const;
-	ENGINE_API void ResetGPUReadback();
-	ENGINE_API void InitialiseProfileBuffer(const FName& ProfileName);
-
 	ENGINE_API bool IsDefaultOverridden() const { return bDefaultOverridden; }
 	ENGINE_API bool IsStaticOverridden() const { return bStaticOverridden; }
-	ENGINE_API FName GetDefaultProfileName() const { return DefaultProfileName; }
-protected:
-	ENGINE_API void ApplyOverrideProfile(FSkinWeightVertexBuffer* OverrideBuffer, const FName& ProfileName);
-
-	FSkinWeightVertexBuffer* BaseBuffer;
-	FSkinWeightVertexBuffer* DefaultOverrideSkinWeightBuffer;
-
-	TMap<FName, FSkinWeightVertexBuffer*> ProfileNameToBuffer;
+	ENGINE_API FSkinWeightProfileStack GetDefaultProfileStack() const { return DefaultProfileStack; }
+	
+private:
+	friend class FSkinWeightProfileManager;
+	friend class FSkinWeightProfileManagerAsyncTask;
+	
+	void InitialiseProfileBuffer(const FSkinWeightProfileStack& InProfileStack);
+	void ReleaseBuffer(const FSkinWeightProfileStack& InProfileStack, bool bForceRelease = false);
+	
+	void ApplyOverrideProfileStack(const FSkinWeightProfileStack& InProfileStack, FSkinWeightVertexBuffer* OverrideBuffer, const uint8* BaseBufferData = nullptr);
+	bool IsPendingReadback() const;
+	void EnqueueGPUReadback();
+	bool IsGPUReadbackFinished() const;
+	void EnqueueDataReadback();
+	bool IsDataReadbackPending() const;
+	bool IsDataReadbackFinished() const;
+	void ResetGPUReadback();
+
+	FSkinWeightVertexBuffer* BaseBuffer = nullptr;
+	FSkinWeightVertexBuffer* DefaultOverrideSkinWeightBuffer = nullptr;
+
+	TMap<FSkinWeightProfileStack, FSkinWeightVertexBuffer*> ProfileStackToBuffer;
+	
 	TMap<FName, FRuntimeSkinWeightProfileData> OverrideData;
 
-	bool bDefaultOverridden;
-	bool bStaticOverridden;
-	FName DefaultProfileName;
+	bool bDefaultOverridden = false;
+	bool bStaticOverridden = false;
+	FSkinWeightProfileStack DefaultProfileStack;
 
-protected:
 	FSkinweightReadbackData ReadbackData;
 };
 
diff --git a/Engine/Source/Runtime/Engine/Public/Animation/SkinWeightProfileManager.h b/Engine/Source/Runtime/Engine/Public/Animation/SkinWeightProfileManager.h
--- a/Engine/Source/Runtime/Engine/Public/Animation/SkinWeightProfileManager.h
+++ b/Engine/Source/Runtime/Engine/Public/Animation/SkinWeightProfileManager.h
@@ -2,30 +2,29 @@
 
 #pragma once
 
-#include "Stats/Stats.h"
-#include "UObject/WeakObjectPtrTemplates.h"
 #include "Components/SkinnedMeshComponent.h"
+#include "Engine/EngineBaseTypes.h"
 #include "Engine/SkeletalMesh.h"
-#include "RHIGPUReadback.h"
-#include "Tickable.h"
+#include "Engine/World.h"
+#include "SkinWeightProfile.h"
+#include "Stats/Stats.h"
 #include "Templates/Function.h"
-#include "Engine/EngineBaseTypes.h"
-#include "UObject/GCObject.h"
+#include "Tickable.h"
 #include "UObject/ObjectMacros.h"
-#include "Engine/World.h"
+#include "UObject/WeakObjectPtrTemplates.h"
 
 #include "SkinWeightProfileManager.generated.h"
 
 class FSkinWeightProfileManager;
 class UWorld;
 
-typedef TFunction<void(TWeakObjectPtr<USkeletalMesh> WeakMesh, FName ProfileName)> FRequestFinished;
+typedef TFunction<void(TWeakObjectPtr<USkeletalMesh> WeakMesh, FSkinWeightProfileStack ProfileStack)> FRequestFinished;
 
 /** Describes a single skin weight profile request */
 struct FSetProfileRequest
 {
-	/** Name of the skin weight profile to be loaded */
-	FName ProfileName;
+	/** Name of the skin weight profile stack to be loaded. Must be normalized, see FSkinWeightProfileStack::Normalized */
+	FSkinWeightProfileStack ProfileStack;
 	/** LOD Indices to load the profile for */
 	TArray<int32> LODIndices;
 	/** Called when the profile request has finished and data is ready (called from GT only) */
@@ -38,12 +37,12 @@ struct FSetProfileRequest
 
 	friend bool operator==(const FSetProfileRequest& A, const FSetProfileRequest& B)
 	{
-		return A.ProfileName == B.ProfileName && A.WeakSkeletalMesh == B.WeakSkeletalMesh && A.IdentifyingObject == B.IdentifyingObject;
+		return A.ProfileStack == B.ProfileStack && A.WeakSkeletalMesh == B.WeakSkeletalMesh && A.IdentifyingObject == B.IdentifyingObject;
 	}
 
 	friend uint32 GetTypeHash(FSetProfileRequest A)
 	{
-		return HashCombine(GetTypeHash(A.ProfileName), GetTypeHash(A.WeakSkeletalMesh));
+		return HashCombine(GetTypeHash(A.ProfileStack), GetTypeHash(A.WeakSkeletalMesh));
 	}
 };
 
@@ -124,9 +123,8 @@ public:
 	static ENGINE_API FSkinWeightProfileManager* Get(UWorld* World);
 
 	FSkinWeightProfileManager(UWorld* InWorld);
-	virtual ~FSkinWeightProfileManager() {}
 
-	void ENGINE_API RequestSkinWeightProfile(FName InProfileName, USkinnedAsset* SkinnedAsset, UObject* Requester, FRequestFinished& Callback, int32 LODIndex = INDEX_NONE);
+	void ENGINE_API RequestSkinWeightProfileStack(FSkinWeightProfileStack InProfileStack, USkinnedAsset* SkinnedAsset, UObject* Requester, FRequestFinished& Callback, int32 LODIndex = INDEX_NONE);
 	void CancelSkinWeightProfileRequest(UObject* Requester);
 	
 	void DoTick(float DeltaTime, ENamedThreads::Type CurrentThread, const FGraphEventRef& MyCompletionGraphEvent);
