Fixed Build NotForLicensee Tools Mac - error: 'GetHLODProxyComponents

diff --git a/Engine/Source/Runtime/Engine/Classes/Components/ActorComponent.h b/Engine/Source/Runtime/Engine/Classes/Components/ActorComponent.h
--- a/Engine/Source/Runtime/Engine/Classes/Components/ActorComponent.h
+++ b/Engine/Source/Runtime/Engine/Classes/Components/ActorComponent.h
@@ -1181,7 +1181,7 @@ public:
 	/** Override to specify that a component is relevant to the HLOD generation. */
 	virtual bool IsHLODRelevant() const { return false; }
 
-#if UE_EDITOR
+#if WITH_EDITOR
 	/** Override to provide a custom set of components to be used in place of the current component for the HLOD generation. */
 	virtual TArray<UActorComponent*> GetHLODProxyComponents() const { return TArray<UActorComponent*>(); }
 #endif
