Add a new UFUNCTION specifier "Remote" that turns the function into a unicast RPC that can be invoked on either clients or servers and executed on the remote side of the net connection.
These functions will only execute remotely and only on net-owned actors.
Added EngineTest Net.RPC.Execution for the various combinations of RPC execution rules.
NetTestMultiServer now uses a Remote function instead of duplicated Client and Server functions.

diff --git a/Engine/Source/Programs/Shared/EpicGames.UHT/Specifiers/UhtFunctionSpecifiers.cs b/Engine/Source/Programs/Shared/EpicGames.UHT/Specifiers/UhtFunctionSpecifiers.cs
--- a/Engine/Source/Programs/Shared/EpicGames.UHT/Specifiers/UhtFunctionSpecifiers.cs
+++ b/Engine/Source/Programs/Shared/EpicGames.UHT/Specifiers/UhtFunctionSpecifiers.cs
@@ -237,6 +237,29 @@ namespace EpicGames.UHT.Parsers
 			}
 		}
 
+		[UhtSpecifier(Extends = UhtTableNames.Function, ValueType = UhtSpecifierValueType.OptionalString)]
+		private static void RemoteSpecifier(UhtSpecifierContext specifierContext, StringView? value)
+		{
+			UhtFunction function = (UhtFunction)specifierContext.Type;
+
+			if (function.FunctionFlags.HasAnyFlags(EFunctionFlags.BlueprintEvent))
+			{
+				specifierContext.MessageSite.LogError("BlueprintImplementableEvent or BlueprintNativeEvent functions cannot be declared as Client, Server, or Remote");
+			}
+
+			if (function.FunctionFlags.HasAnyFlags(EFunctionFlags.Exec))
+			{
+				specifierContext.MessageSite.LogError("Exec functions cannot be replicated!");
+			}
+
+			function.FunctionFlags |= EFunctionFlags.Net;
+
+			if (value != null)
+			{
+				function.CppImplName = ((StringView)value).ToString();
+			}
+		}
+
 		[UhtSpecifier(Extends = UhtTableNames.Function, ValueType = UhtSpecifierValueType.OptionalEqualsKeyValuePairList)]
 		private static void ServiceRequestSpecifier(UhtSpecifierContext specifierContext, List<KeyValuePair<StringView, StringView>> value)
 		{
diff --git a/Engine/Source/Runtime/Engine/Private/Actor.cpp b/Engine/Source/Runtime/Engine/Private/Actor.cpp
--- a/Engine/Source/Runtime/Engine/Private/Actor.cpp
+++ b/Engine/Source/Runtime/Engine/Private/Actor.cpp
@@ -9,6 +9,7 @@
 #include "Engine/NetConnection.h"
 #include "GameFramework/DamageType.h"
 #include "GameFramework/WorldSettings.h"
+#include "Net/Core/Misc/NetContext.h"
 #include "Net/Core/PropertyConditions/PropertyConditions.h"
 #include "TimerManager.h"
 #include "GameFramework/Pawn.h"
@@ -5069,18 +5070,29 @@ int32 AActor::GetFunctionCallspace( UFunction* Function, FFrame* Stack )
 		}
 	}
 
-	// if we are the server, and it's not a send-to-client function,
-	if (bIsServer && !(Function->FunctionFlags & FUNC_NetClient))
+	if (Function->FunctionFlags & (FUNC_NetServer | FUNC_NetClient | FUNC_NetMulticast))
 	{
-		// don't replicate
-		DEBUG_CALLSPACE(TEXT("GetFunctionCallspace Server calling Server function: %s %s"), *Function->GetName(), FunctionCallspace::ToString(Callspace));
-		return Callspace;
+		// Handle uni-directional RPCs
+
+		// if we are the server, and it's not a send-to-client function,
+		if (bIsServer && !(Function->FunctionFlags & FUNC_NetClient))
+		{
+			// don't replicate
+			DEBUG_CALLSPACE(TEXT("GetFunctionCallspace Server calling Server function: %s %s"), *Function->GetName(), FunctionCallspace::ToString(Callspace));
+			return Callspace;
+		}
+		// if we aren't the server, and it's not a send-to-server function,
+		if (!bIsServer && !(Function->FunctionFlags & FUNC_NetServer))
+		{
+			// don't replicate
+			DEBUG_CALLSPACE(TEXT("GetFunctionCallspace Client calling Client function: %s %s"), *Function->GetName(), FunctionCallspace::ToString(Callspace));
+			return Callspace;
+		}
 	}
-	// if we aren't the server, and it's not a send-to-server function,
-	if (!bIsServer && !(Function->FunctionFlags & FUNC_NetServer))
+	else if (UE::Net::FNetContext::IsInsideNetRPC())
 	{
-		// don't replicate
-		DEBUG_CALLSPACE(TEXT("GetFunctionCallspace Client calling Client function: %s %s"), *Function->GetName(), FunctionCallspace::ToString(Callspace));
+		// Received a function with FUNC_Net but not Client, Server, or Multicast (Remote specifier) - can be sent bi-directionally
+		DEBUG_CALLSPACE(TEXT("GetFunctionCallspace Received a Remote function, calling: %s %s"), *Function->GetName(), FunctionCallspace::ToString(Callspace));
 		return Callspace;
 	}
 
@@ -5100,7 +5112,14 @@ int32 AActor::GetFunctionCallspace( UFunction* Function, FFrame* Stack )
 					DEBUG_CALLSPACE(TEXT("GetFunctionCallspace Client without owner absorbed %s"), *Function->GetName());
 					return FunctionCallspace::Absorbed;
 				}
-				
+
+				if ((Function->FunctionFlags & (FUNC_NetServer | FUNC_NetClient | FUNC_NetMulticast)) == 0)
+				{
+					// Calling a Remote function (client or server) with no owning player, absorb these.
+					UE_LOG(LogNet, Error, TEXT("GetFunctionCallspace calling a Remote function without owner, absorbing: %s"), *Function->GetName());
+					return FunctionCallspace::Absorbed;
+				}
+
 				// Role authority object calling a client RPC locally (ie AI owned objects)
 				DEBUG_CALLSPACE(TEXT("GetFunctionCallspace authority non client owner %s %s"), *Function->GetName(), FunctionCallspace::ToString(Callspace));
 				return Callspace;
diff --git a/Engine/Source/Runtime/Engine/Private/DataReplication.cpp b/Engine/Source/Runtime/Engine/Private/DataReplication.cpp
--- a/Engine/Source/Runtime/Engine/Private/DataReplication.cpp
+++ b/Engine/Source/Runtime/Engine/Private/DataReplication.cpp
@@ -1281,9 +1281,16 @@ bool FObjectReplicator::ReceivedRPC(FNetBitReader& Reader, const FReplicationFla
 		HANDLE_INCOMPATIBLE_RPC
 	}
 
-	if ((Function->FunctionFlags & (bIsServer ? FUNC_NetServer : (FUNC_NetClient | FUNC_NetMulticast))) == 0)
+	// If NetServer, NetClient, or NetMulticast flags are present, filter the RPC based on them. If not, accept the RPC for servers and clients.
+	if ((Function->FunctionFlags & FUNC_NetServer) && !bIsServer)
 	{
-		UE_LOG(LogRep, Error, TEXT("Rejected RPC function due to access rights. Object: %s, Function: %s"), *Object->GetFullName(), *FunctionName.ToString());
+		UE_LOG(LogRep, Error, TEXT("Rejected server RPC function due to access rights. Object: %s, Function: %s"), *Object->GetFullName(), *FunctionName.ToString());
+		HANDLE_INCOMPATIBLE_RPC
+	}
+
+	if ((Function->FunctionFlags & (FUNC_NetClient | FUNC_NetMulticast)) && bIsServer)
+	{
+		UE_LOG(LogRep, Error, TEXT("Rejected client RPC function due to access rights. Object: %s, Function: %s"), *Object->GetFullName(), *FunctionName.ToString());
 		HANDLE_INCOMPATIBLE_RPC
 	}
 
diff --git a/Engine/Source/Runtime/Experimental/Iris/Core/Private/Iris/ReplicationSystem/NetBlob/NetBlobManager.cpp b/Engine/Source/Runtime/Experimental/Iris/Core/Private/Iris/ReplicationSystem/NetBlob/NetBlobManager.cpp
--- a/Engine/Source/Runtime/Experimental/Iris/Core/Private/Iris/ReplicationSystem/NetBlob/NetBlobManager.cpp
+++ b/Engine/Source/Runtime/Experimental/Iris/Core/Private/Iris/ReplicationSystem/NetBlob/NetBlobManager.cpp
@@ -182,9 +182,16 @@ bool FNetBlobManager::SendUnicastRPC(uint32 ConnectionId, const FSendRPCContext&
 		return false;
 	}
 
-	if ((Context.Function->FunctionFlags & (bIsServer ? (FUNC_NetClient | FUNC_NetMulticast) : FUNC_NetServer)) == 0)
+	// If NetServer, NetClient, or NetMulticast flags are present, filter the RPC based on them. If not, send the RPC in either directon.
+	if ((Context.Function->FunctionFlags & FUNC_NetServer) && bIsServer)
 	{
-		checkf(false, TEXT("Trying to call RPC %s in the wrong direction."), ToCStr(Context.Function->GetName()));
+		checkf(false, TEXT("Trying to call server RPC %s in the wrong direction."), ToCStr(Context.Function->GetName()));
+		return true;
+	}
+
+	if ((Context.Function->FunctionFlags & (FUNC_NetClient | FUNC_NetMulticast)) && !bIsServer)
+	{
+		checkf(false, TEXT("Trying to call client RPC %s in the wrong direction."), ToCStr(Context.Function->GetName()));
 		return true;
 	}
 
diff --git a/Engine/Source/Runtime/Experimental/Iris/Core/Private/Iris/ReplicationSystem/NetBlob/NetRPC.cpp b/Engine/Source/Runtime/Experimental/Iris/Core/Private/Iris/ReplicationSystem/NetBlob/NetRPC.cpp
--- a/Engine/Source/Runtime/Experimental/Iris/Core/Private/Iris/ReplicationSystem/NetBlob/NetRPC.cpp
+++ b/Engine/Source/Runtime/Experimental/Iris/Core/Private/Iris/ReplicationSystem/NetBlob/NetRPC.cpp
@@ -481,12 +481,19 @@ void FNetRPC::CallFunction(FNetRPCCallContext& CallContext)
 		return;
 	}
 
+	if ((Function->FunctionFlags & FUNC_Net) == 0)
+	{
+		UE_LOG(LogIrisRpc, Error, TEXT("Rejected %s function %s due to it not being a Net function. %s : %s"), (Function->FunctionFlags & FUNC_NetReliable ? TEXT("reliable") : TEXT("unreliable")), ToCStr(Function->GetName()), *NetObjectReference.ToString(), *Object->GetFullName());
+		Context.SetError(NetError_FunctionCallNotAllowed);
+		return;
+	}
+
 	const bool bIsServer = ReplicationSystem->IsServer();
 	if (bIsServer)
 	{
-		if ((Function->FunctionFlags & FUNC_NetServer) == 0)
+		if ((Function->FunctionFlags & (FUNC_NetClient | FUNC_NetMulticast)) != 0)
 		{
-			UE_LOG(LogIrisRpc, Error, TEXT("Rejected %s RPC function %s due to access rights. %s : %s"), (Function->FunctionFlags & FUNC_NetReliable ? TEXT("reliable") : TEXT("unreliable")), ToCStr(Function->GetName()), *NetObjectReference.ToString(), *Object->GetFullName());
+			UE_LOG(LogIrisRpc, Error, TEXT("Rejected %s client RPC function %s due to this being the server. %s : %s"), (Function->FunctionFlags & FUNC_NetReliable ? TEXT("reliable") : TEXT("unreliable")), ToCStr(Function->GetName()), *NetObjectReference.ToString(), *Object->GetFullName());
 			Context.SetError(NetError_FunctionCallNotAllowed);
 			return;
 		}
@@ -501,9 +508,9 @@ void FNetRPC::CallFunction(FNetRPCCallContext& CallContext)
 	}
 	else
 	{
-		if ((Function->FunctionFlags & (FUNC_NetClient | FUNC_NetMulticast)) == 0)
+		if ((Function->FunctionFlags & FUNC_NetServer) != 0)
 		{
-			UE_LOG(LogIrisRpc, Error, TEXT("Rejected %s RPC function %s due to access rights. %s : %s"), (Function->FunctionFlags & FUNC_NetReliable ? TEXT("reliable") : TEXT("unreliable")), ToCStr(Function->GetName()), *NetObjectReference.ToString(), *Object->GetFullName());
+			UE_LOG(LogIrisRpc, Error, TEXT("Rejected %s server RPC function %s due to this being the client. %s : %s"), (Function->FunctionFlags & FUNC_NetReliable ? TEXT("reliable") : TEXT("unreliable")), ToCStr(Function->GetName()), *NetObjectReference.ToString(), *Object->GetFullName());
 			return;
 		}
 	}
