Fix MMV (for ES31 projects) not working with -emulatestereo.

- Also fixes UE-220113 (no image for the right eye when previewing a HDR project) by auto-enabling subpass tonemapper for mobile MMV rendering.

diff --git a/Engine/Shaders/Private/ScreenPixelShader.usf b/Engine/Shaders/Private/ScreenPixelShader.usf
--- a/Engine/Shaders/Private/ScreenPixelShader.usf
+++ b/Engine/Shaders/Private/ScreenPixelShader.usf
@@ -10,7 +10,11 @@
 	#define SCREEN_PS_FROM_SLICE0   0
 #endif
 
-#if SCREEN_PS_FROM_SLICE0
+#ifndef SCREEN_PS_UNWRAP_SLICES
+	#define SCREEN_PS_UNWRAP_SLICES 0
+#endif
+
+#if (SCREEN_PS_FROM_SLICE0 || SCREEN_PS_UNWRAP_SLICES)
 	Texture2DArray InTexture;
 
 	// if source is an array, we only need the first slice
@@ -34,6 +38,17 @@ void Main(
 	OutColor = TexOrArr2DSample(InTexture, InTextureSampler, Input.UV);
 }
 
+#if SCREEN_PS_UNWRAP_SLICES
+void MainUnwrap(
+	FScreenVertexOutput Input,
+	out float4 OutColor : SV_Target0
+	)
+{
+	// for the first half of the screen, sample from slice 0, for the second, from slice 1
+	OutColor = InTexture.Sample(InTextureSampler, float3(frac(Input.UV.x * 2.0), Input.UV.y, (float)((int)(2 * Input.UV.x))) );
+}
+#endif
+
 void MainInvertAlpha(
 	FScreenVertexOutput Input,
 	out float4 OutColor : SV_Target0
diff --git a/Engine/Source/Runtime/Engine/Private/SceneTexturesConfig.cpp b/Engine/Source/Runtime/Engine/Private/SceneTexturesConfig.cpp
--- a/Engine/Source/Runtime/Engine/Private/SceneTexturesConfig.cpp
+++ b/Engine/Source/Runtime/Engine/Private/SceneTexturesConfig.cpp
@@ -343,9 +343,9 @@ void FSceneTexturesConfig::Init(const FSceneTexturesConfigInitSettings& InitSett
 
 	if (ShadingPath == EShadingPath::Mobile)
 	{
-		bRequiresDepthAux = MobileRequiresSceneDepthAux(ShaderPlatform) && !IsMobileTonemapSubpassEnabled(ShaderPlatform);
+		bRequiresDepthAux = MobileRequiresSceneDepthAux(ShaderPlatform) && !IsMobileTonemapSubpassEnabled(ShaderPlatform, InitSettings.bRequireMultiView);
 		bPreciseDepthAux = bPreciseDepthAux || MobileRequiresPreciseSceneDepthAux(ShaderPlatform);
-		bCustomResolveSubpass = IsMobileTonemapSubpassEnabledInline(ShaderPlatform, NumSamples);
+		bCustomResolveSubpass = IsMobileTonemapSubpassEnabledInline(ShaderPlatform, InitSettings.bRequireMultiView, NumSamples);
 	}
 }
 
diff --git a/Engine/Source/Runtime/Engine/Private/SceneUtils.cpp b/Engine/Source/Runtime/Engine/Private/SceneUtils.cpp
--- a/Engine/Source/Runtime/Engine/Private/SceneUtils.cpp
+++ b/Engine/Source/Runtime/Engine/Private/SceneUtils.cpp
@@ -49,16 +49,16 @@ bool IsMobilePropagateAlphaEnabled(EShaderPlatform Platform)
 	return IsMobilePlatform(Platform) && (FPlatformMisc::GetMobilePropagateAlphaSetting() > 0);
 }
 
-ENGINE_API bool IsMobileTonemapSubpassEnabled(EShaderPlatform Platform)
+ENGINE_API bool IsMobileTonemapSubpassEnabled(EShaderPlatform Platform, bool bMultiViewRendering)
 {
 	static auto* MobileTonemapSubpassPathCvar = IConsoleManager::Get().FindTConsoleVariableDataInt(TEXT("r.Mobile.TonemapSubpass"));
-	return (MobileTonemapSubpassPathCvar && (MobileTonemapSubpassPathCvar->GetValueOnAnyThread() == 1)) && IsMobileHDR() && !IsMobileDeferredShadingEnabled(Platform);
+	return ((MobileTonemapSubpassPathCvar && (MobileTonemapSubpassPathCvar->GetValueOnAnyThread() == 1)) || bMultiViewRendering) && IsMobileHDR() && !IsMobileDeferredShadingEnabled(Platform);
 }
 
-ENGINE_API bool IsMobileTonemapSubpassEnabledInline(EShaderPlatform Platform, uint32 NumMSAASamples)
+ENGINE_API bool IsMobileTonemapSubpassEnabledInline(EShaderPlatform Platform, bool bMultiViewRendering, uint32 NumMSAASamples)
 {
 	// As of UE 5.4 only vulkan supports inline (single pass) tonemap
-	return IsMobileTonemapSubpassEnabled(Platform) && IsVulkanPlatform(Platform) && (GRHISupportsMSAAShaderResolve || NumMSAASamples <= 1u);
+	return IsMobileTonemapSubpassEnabled(Platform, bMultiViewRendering) && IsVulkanPlatform(Platform) && (GRHISupportsMSAAShaderResolve || NumMSAASamples <= 1u);
 }
 
 ENGINE_API bool IsMobileColorsRGB()
diff --git a/Engine/Source/Runtime/Engine/Private/ScreenRendering.cpp b/Engine/Source/Runtime/Engine/Private/ScreenRendering.cpp
--- a/Engine/Source/Runtime/Engine/Private/ScreenRendering.cpp
+++ b/Engine/Source/Runtime/Engine/Private/ScreenRendering.cpp
@@ -12,6 +12,7 @@ TGlobalResource<FScreenVertexDeclaration> GScreenVertexDeclaration;
 // Shader implementations.
 IMPLEMENT_SHADER_TYPE(, FScreenPS, TEXT("/Engine/Private/ScreenPixelShader.usf"), TEXT("Main"), SF_Pixel);
 IMPLEMENT_SHADER_TYPE(, FScreenFromSlice0PS, TEXT("/Engine/Private/ScreenPixelShader.usf"), TEXT("Main"), SF_Pixel);
+IMPLEMENT_SHADER_TYPE(, FScreenUnwrapSlicesPS, TEXT("/Engine/Private/ScreenPixelShader.usf"), TEXT("MainUnwrap"), SF_Pixel);
 IMPLEMENT_SHADER_TYPE(,FScreenPSInvertAlpha,TEXT("/Engine/Private/ScreenPixelShader.usf"),TEXT("MainInvertAlpha"),SF_Pixel);
 IMPLEMENT_SHADER_TYPE(,FScreenPSsRGBSource, TEXT("/Engine/Private/ScreenPixelShader.usf"), TEXT("MainsRGBSource"), SF_Pixel);
 IMPLEMENT_SHADER_TYPE(,FScreenPSMipLevel, TEXT("/Engine/Private/ScreenPixelShader.usf"), TEXT("MainMipLevel"), SF_Pixel);
diff --git a/Engine/Source/Runtime/Engine/Private/UnrealEngine.cpp b/Engine/Source/Runtime/Engine/Private/UnrealEngine.cpp
--- a/Engine/Source/Runtime/Engine/Private/UnrealEngine.cpp
+++ b/Engine/Source/Runtime/Engine/Private/UnrealEngine.cpp
@@ -250,6 +250,10 @@ UnrealEngine.cpp: Implements the UEngine class and helpers.
 #include "UObject/UObjectThreadContext.h"
 #include "UObject/OverridableManager.h"
 
+// MMV support in FakeStereoRenderingDevice
+#include "StereoRenderTargetManager.h"
+#include "ScreenRendering.h"
+
 #if WITH_DUMPGPU
 #include "RenderGraphBuilder.h"
 #endif
@@ -3770,10 +3774,11 @@ bool UEngine::UseSound() const
 {
 	return AudioDeviceManager != nullptr;
 }
+
 /**
 * A fake stereo rendering device used to test stereo rendering without an attached device.
 */
-class FFakeStereoRenderingDevice : public IStereoRendering
+class FFakeStereoRenderingDevice : public IStereoRendering, public IStereoRenderTargetManager
 {
 public:
 	FFakeStereoRenderingDevice(int ViewportWidth = 640, int ViewportHeight = 480, int RequestedNumViews = 2) 
@@ -3781,6 +3786,7 @@ public:
 		, Width(ViewportWidth)
 		, Height(ViewportHeight)
 		, NumViews(RequestedNumViews)
+		, bLayeredRTs(true)
 	{
 		static TAutoConsoleVariable<float> CVarEmulateStereoFOV(TEXT("r.StereoEmulationFOV"), 0, TEXT("FOV in degrees, of the imaginable HMD for stereo emulation"), ECVF_ReadOnly);
 		static TAutoConsoleVariable<int32> CVarEmulateStereoWidth(TEXT("r.StereoEmulationWidth"), 0, TEXT("Width of the imaginable HMD for stereo emulation"), ECVF_ReadOnly);
@@ -3806,6 +3812,9 @@ public:
 		{
 			NumViews = FMath::Clamp(V, 1, 32);
 		}
+
+		// mobile rendering path does not use side-by-side normally
+		bLayeredRTs = GetFeatureLevelShadingPath(GMaxRHIFeatureLevel) == EShadingPath::Mobile;
 	}
 
 	virtual ~FFakeStereoRenderingDevice() {}
@@ -3909,9 +3918,53 @@ public:
 	virtual void RenderTexture_RenderThread(FRHICommandListImmediate& RHICmdList, FRHITexture* BackBuffer, FRHITexture* SrcTexture, FVector2D WindowSize) const override
 	{
 		check(IsInRenderingThread());
+		checkf(bLayeredRTs, TEXT("Non-layered stereo can use the backbuffer directly"));
 
 		FRHIRenderPassInfo RPInfo(BackBuffer, ERenderTargetActions::Clear_Store);
 		RHICmdList.BeginRenderPass(RPInfo, TEXT("RenderTexture_RenderThread"));
+
+			FGraphicsPipelineStateInitializer GraphicsPSOInit;
+			RHICmdList.ApplyCachedRenderTargets(GraphicsPSOInit);
+			GraphicsPSOInit.BlendState = TStaticBlendState<>::GetRHI();
+			GraphicsPSOInit.RasterizerState = TStaticRasterizerState<>::GetRHI();
+			GraphicsPSOInit.DepthStencilState = TStaticDepthStencilState<false, CF_Always>::GetRHI();
+
+			FGlobalShaderMap* ShaderMap = GetGlobalShaderMap(GMaxRHIShaderPlatform);
+
+			TShaderMapRef<FScreenVS> VertexShader(ShaderMap);
+			TShaderMapRef<FScreenUnwrapSlicesPS> PixelShader(ShaderMap);
+
+			GraphicsPSOInit.BoundShaderState.VertexDeclarationRHI = GFilterVertexDeclaration.VertexDeclarationRHI;
+			GraphicsPSOInit.BoundShaderState.VertexShaderRHI = VertexShader.GetVertexShader();
+			GraphicsPSOInit.BoundShaderState.PixelShaderRHI = PixelShader.GetPixelShader();
+			GraphicsPSOInit.PrimitiveType = PT_TriangleList;
+
+			SetGraphicsPipelineState(RHICmdList, GraphicsPSOInit, 0);
+
+			RHICmdList.Transition(FRHITransitionInfo(SrcTexture, ERHIAccess::Unknown, ERHIAccess::SRVMask));
+
+			SetShaderParametersLegacyPS(RHICmdList, PixelShader, TStaticSamplerState<SF_Point>::GetRHI(), SrcTexture);
+
+			const int SourceWidth = Width;
+			const int SourceHeight = Height;
+			const int TargetWidth = Width;
+			const int TargetHeight = Height;
+
+			const int TargetX = 0;
+			const int TargetY = 0;
+
+			IRendererModule& RendererModule = FModuleManager::GetModuleChecked<IRendererModule>(FName("Renderer"));
+			RendererModule.DrawRectangle(
+				RHICmdList,
+				(float)TargetX, (float)TargetY,
+				(float)SourceWidth, (float)SourceHeight,
+				0.f, 0.f,
+				(float)SourceWidth, (float)SourceHeight,
+				FIntPoint(TargetWidth, TargetHeight),
+				FIntPoint(TargetWidth, TargetHeight),
+				VertexShader,
+				EDRF_UseTriangleOptimization);
+
 		RHICmdList.EndRenderPass();
 
 		const uint32 ViewportWidth = BackBuffer->GetSizeX();
@@ -3919,9 +3972,80 @@ public:
 		RHICmdList.SetViewport( 0,0,0,ViewportWidth, ViewportHeight, 1.0f );
 	}
 
-	float FOVInDegrees;		// max(HFOV, VFOV) in degrees of imaginable HMD
-	int32 Width, Height;	// resolution of imaginable HMD
-	int32 NumViews;			// views of imaginable HMD
+	virtual IStereoRenderTargetManager* GetRenderTargetManager() override { return this; }
+
+	virtual bool ShouldUseSeparateRenderTarget() const { return bLayeredRTs && IsStereoEnabled(); };
+
+	virtual void UpdateViewport(bool bUseSeparateRenderTarget, const class FViewport& Viewport, class SViewport* ViewportWidget = nullptr) override {};
+
+	virtual bool AllocateRenderTargetTextures(uint32 SizeX, uint32 SizeY, uint8 Format, uint32 NumLayers, ETextureCreateFlags Flags, ETextureCreateFlags TargetableTextureFlags, TArray<FTextureRHIRef>& OutTargetableTextures, TArray<FTextureRHIRef>& OutShaderResourceTextures, uint32 NumSamples = 1)
+	{ 
+		check(IsInRenderingThread());
+		checkf(bLayeredRTs, TEXT("Non-layered stereo can use the backbuffer directly"));
+
+		// We're only creating a 1x target here, but we don't know whether it'll be the targeted texture
+		// or the resolve texture. Because of this, we unify the input flags.
+		ETextureCreateFlags UnifiedCreateFlags = Flags | TargetableTextureFlags | TexCreate_Dynamic;
+
+		// We need to ensure we can sample from the texture in CopyTexture
+		UnifiedCreateFlags |= TexCreate_ShaderResource;
+
+		// We assume this could be used as a resolve target
+		//UnifiedCreateFlags |= TexCreate_ResolveTargetable;
+
+		// Some render APIs require us to present in RT layouts/configs,
+		// so even if app won't use this texture as RT, we need the flag.
+		UnifiedCreateFlags |= TexCreate_RenderTargetable;
+
+		UnifiedCreateFlags |= TexCreate_SRGB;
+		ETextureCreateFlags AuxiliaryCreateFlags = ETextureCreateFlags::None;
+
+		// support foveation?
+		//if (FBFoveationImageGenerator && FBFoveationImageGenerator->IsFoveationExtensionEnabled())
+		//{
+		//	AuxiliaryCreateFlags |= TexCreate_Foveation;
+		//}
+
+		FClearValueBinding ClearColor = FClearValueBinding::Transparent;
+
+		uint8 ActualFormat = Format;
+		FIntPoint Extent = FIntPoint(Width, Height);
+
+		const TCHAR* EmulateStereoSwapchainDebugName = TEXT("EmulateStereoSwapchainImage");
+		FRHITextureCreateDesc SwapchainImageDesc = true ?
+			FRHITextureCreateDesc::Create2DArray(EmulateStereoSwapchainDebugName).SetArraySize(2) :
+			FRHITextureCreateDesc::Create2D(EmulateStereoSwapchainDebugName);
+
+		SwapchainImageDesc.SetExtent(Width, Height);
+		SwapchainImageDesc.SetFormat(static_cast<EPixelFormat>(ActualFormat));
+		SwapchainImageDesc.SetFlags(UnifiedCreateFlags);
+
+		// this will be owned by the caller
+		FTextureRHIRef NewImage = RHICreateTexture(SwapchainImageDesc);
+
+		OutTargetableTextures.Add(NewImage);
+		OutShaderResourceTextures = OutTargetableTextures;
+
+		return true;
+
+	}
+
+	virtual void CalculateRenderTargetSize(const class FViewport& Viewport, uint32& InOutSizeX, uint32& InOutSizeY)
+	{ 
+		InOutSizeX = Width;
+		InOutSizeY = Height;
+	}
+
+	virtual bool NeedReAllocateViewportRenderTarget(const class FViewport& Viewport) { return false; };
+
+	/** max(HFOV, VFOV) in degrees of imaginable HMD */
+	float FOVInDegrees;
+	/** resolution of imaginable HMD (per eye) */
+	int32 Width, Height;
+	/** views of imaginable HMD */
+	int32 NumViews;
+	/** Whether the rendertargets are layered (like in MMV case) or not (side-by-side). */
+	bool bLayeredRTs;
 };
 
 bool UEngine::InitializeHMDDevice()
diff --git a/Engine/Source/Runtime/Engine/Public/SceneUtils.h b/Engine/Source/Runtime/Engine/Public/SceneUtils.h
--- a/Engine/Source/Runtime/Engine/Public/SceneUtils.h
+++ b/Engine/Source/Runtime/Engine/Public/SceneUtils.h
@@ -60,8 +60,8 @@ inline bool IsTemporalAccumulationBasedMethod(EAntiAliasingMethod AntiAliasingMe
 
 /** True if Alpha Propagate is enabled for the mobile renderer. */
 ENGINE_API bool IsMobilePropagateAlphaEnabled(EShaderPlatform Platform);
-ENGINE_API bool IsMobileTonemapSubpassEnabled(EShaderPlatform Platform);
-ENGINE_API bool IsMobileTonemapSubpassEnabledInline(EShaderPlatform Platform, uint32 NumMSAASamples);
+ENGINE_API bool IsMobileTonemapSubpassEnabled(EShaderPlatform Platform, bool bMultiViewRendering);
+ENGINE_API bool IsMobileTonemapSubpassEnabledInline(EShaderPlatform Platform, bool bMultiViewRendering, uint32 NumMSAASamples);
 ENGINE_API bool IsMobileColorsRGB();
 
 ENGINE_API EAntiAliasingMethod GetDefaultAntiAliasingMethod(const FStaticFeatureLevel InFeatureLevel);
diff --git a/Engine/Source/Runtime/Engine/Public/ScreenRendering.h b/Engine/Source/Runtime/Engine/Public/ScreenRendering.h
--- a/Engine/Source/Runtime/Engine/Public/ScreenRendering.h
+++ b/Engine/Source/Runtime/Engine/Public/ScreenRendering.h
@@ -122,6 +122,37 @@ public:
 	FScreenFromSlice0PS() {}
 };
 
+/**
+ * A pixel shader for rendering side-by-side textures from the array
+ */
+class FScreenUnwrapSlicesPS : public FScreenPS
+{
+	DECLARE_EXPORTED_SHADER_TYPE(FScreenUnwrapSlicesPS, Global, ENGINE_API);
+public:
+
+	static bool ShouldCompilePermutation(const FGlobalShaderPermutationParameters& Parameters)
+	{
+		// this is a desktop-only shader so far (really only needed in the editor)
+		if ((PLATFORM_DESKTOP) && FScreenPS::ShouldCompilePermutation(Parameters))
+		{
+			UE::StereoRenderUtils::FStereoShaderAspects Aspects(Parameters.Platform);
+			return Aspects.IsMobileMultiViewEnabled();
+		}
+		return false;
+	}
+
+	static void ModifyCompilationEnvironment(const FGlobalShaderPermutationParameters& Parameters, FShaderCompilerEnvironment& OutEnvironment)
+	{
+		FScreenPS::ModifyCompilationEnvironment(Parameters, OutEnvironment);
+		OutEnvironment.SetDefine(TEXT("SCREEN_PS_UNWRAP_SLICES"), 1);
+	}
+
+	FScreenUnwrapSlicesPS(const ShaderMetaType::CompiledShaderInitializerType& Initializer) :
+		FScreenPS(Initializer)
+	{
+	}
+	FScreenUnwrapSlicesPS() {}
+};
 
 /**
  * A pixel shader for rendering a textured screen element.
diff --git a/Engine/Source/Runtime/Renderer/Private/CompositionLighting/PostProcessDeferredDecals.cpp b/Engine/Source/Runtime/Renderer/Private/CompositionLighting/PostProcessDeferredDecals.cpp
--- a/Engine/Source/Runtime/Renderer/Private/CompositionLighting/PostProcessDeferredDecals.cpp
+++ b/Engine/Source/Runtime/Renderer/Private/CompositionLighting/PostProcessDeferredDecals.cpp
@@ -550,7 +550,7 @@ void CollectDeferredDecalPassPSOInitializers(
 	{
 		// subpass info set during the submission of the draws in a mobile renderer
 		GraphicsPSOInit.SubpassIndex = 1; // all decals use second sub-pass on mobile
-		GraphicsPSOInit.SubpassHint = GetSubpassHint(SceneTexturesConfig.ShaderPlatform, SceneTexturesConfig.bIsUsingGBuffers, SceneTexturesConfig.NumSamples);
+		GraphicsPSOInit.SubpassHint = GetSubpassHint(SceneTexturesConfig.ShaderPlatform, SceneTexturesConfig.bIsUsingGBuffers, SceneTexturesConfig.bRequireMultiView, SceneTexturesConfig.NumSamples);
 	}
 		
 	const auto AddDeferredDecalPSO = [&](bool bInsideDecal,	bool bReverseHanded, bool bReverseCulling, bool bDecalUsesStencil)
diff --git a/Engine/Source/Runtime/Renderer/Private/CompositionLighting/PostProcessMeshDecals.cpp b/Engine/Source/Runtime/Renderer/Private/CompositionLighting/PostProcessMeshDecals.cpp
--- a/Engine/Source/Runtime/Renderer/Private/CompositionLighting/PostProcessMeshDecals.cpp
+++ b/Engine/Source/Runtime/Renderer/Private/CompositionLighting/PostProcessMeshDecals.cpp
@@ -393,7 +393,7 @@ void FMeshDecalMeshProcessor::CollectDeferredDecalMeshPSOInitializers(
 	{
 		// subpass info set during the submission of the draws in a mobile renderer
 		SubpassIndex = 1; // all decals use second sub-pass on mobile
-		SubpassHint = GetSubpassHint(SceneTexturesConfig.ShaderPlatform, SceneTexturesConfig.bIsUsingGBuffers, SceneTexturesConfig.NumSamples);
+		SubpassHint = GetSubpassHint(SceneTexturesConfig.ShaderPlatform, SceneTexturesConfig.bIsUsingGBuffers, SceneTexturesConfig.bRequireMultiView, SceneTexturesConfig.NumSamples);
 	}
 
 	AddGraphicsPipelineStateInitializer(
diff --git a/Engine/Source/Runtime/Renderer/Private/MeshPassProcessor.cpp b/Engine/Source/Runtime/Renderer/Private/MeshPassProcessor.cpp
--- a/Engine/Source/Runtime/Renderer/Private/MeshPassProcessor.cpp
+++ b/Engine/Source/Runtime/Renderer/Private/MeshPassProcessor.cpp
@@ -1554,7 +1554,7 @@ void ApplyTargetsInfo(FGraphicsPipelineStateInitializer& GraphicsPSOInit, const
 	GraphicsPSOInit.bHasFragmentDensityAttachment = RenderTargetsInfo.bHasFragmentDensityAttachment;
 }
 
-ESubpassHint GetSubpassHint(const FStaticShaderPlatform Platform, bool bIsUsingGBuffers, uint32 NumSamples)
+ESubpassHint GetSubpassHint(const FStaticShaderPlatform Platform, bool bIsUsingGBuffers, bool bMultiViewRendering, uint32 NumSamples)
 {
 	ESubpassHint SubpassHint = ESubpassHint::None;
 	
@@ -1566,7 +1566,7 @@ ESubpassHint GetSubpassHint(const FStaticShaderPlatform Platform, bool bIsUsingG
 		}
 		else
 		{
-			SubpassHint = IsMobileTonemapSubpassEnabledInline(Platform, NumSamples) ? ESubpassHint::CustomResolveSubpass : ESubpassHint::DepthReadSubpass;
+			SubpassHint = IsMobileTonemapSubpassEnabledInline(Platform, bMultiViewRendering, NumSamples) ? ESubpassHint::CustomResolveSubpass : ESubpassHint::DepthReadSubpass;
 		}
 	}
 
diff --git a/Engine/Source/Runtime/Renderer/Private/MobileBasePass.cpp b/Engine/Source/Runtime/Renderer/Private/MobileBasePass.cpp
--- a/Engine/Source/Runtime/Renderer/Private/MobileBasePass.cpp
+++ b/Engine/Source/Runtime/Renderer/Private/MobileBasePass.cpp
@@ -926,7 +926,7 @@ void FMobileBasePassMeshProcessor::CollectPSOInitializersForLMPolicy(
 
 	// subpass info set during the submission of the draws in mobile deferred renderer.
 	uint8 SubpassIndex = bTranslucentBasePass ? (bDeferredShading ? 2 : 1) : 0;
-	ESubpassHint SubpassHint = GetSubpassHint(GMaxRHIShaderPlatform, bDeferredShading, RenderTargetsInfo.NumSamples);
+	ESubpassHint SubpassHint = GetSubpassHint(GMaxRHIShaderPlatform, bDeferredShading, RenderTargetsInfo.MultiViewCount > 1, RenderTargetsInfo.NumSamples);
 
 	AddGraphicsPipelineStateInitializer(
 		VertexFactoryData,
diff --git a/Engine/Source/Runtime/Renderer/Private/MobileShadingRenderer.cpp b/Engine/Source/Runtime/Renderer/Private/MobileShadingRenderer.cpp
--- a/Engine/Source/Runtime/Renderer/Private/MobileShadingRenderer.cpp
+++ b/Engine/Source/Runtime/Renderer/Private/MobileShadingRenderer.cpp
@@ -304,8 +304,8 @@ FMobileSceneRenderer::FMobileSceneRenderer(const FSceneViewFamily* InViewFamily,
 
 	NumMSAASamples = GetDefaultMSAACount(ERHIFeatureLevel::ES3_1);
 	// As of UE 5.4 only vulkan supports inline (single pass) tonemap
-	bTonemapSubpass = IsMobileTonemapSubpassEnabled(ShaderPlatform) && ViewFamily.bResolveScene && GetRendererOutput() != FSceneRenderer::ERendererOutput::DepthPrepassOnly;
-	bTonemapSubpassInline = IsMobileTonemapSubpassEnabledInline(ShaderPlatform, NumMSAASamples) && bTonemapSubpass;
+	bTonemapSubpass = IsMobileTonemapSubpassEnabled(ShaderPlatform, ViewFamily.bRequireMultiView) && ViewFamily.bResolveScene && GetRendererOutput() != FSceneRenderer::ERendererOutput::DepthPrepassOnly;
+	bTonemapSubpassInline = IsMobileTonemapSubpassEnabledInline(ShaderPlatform, ViewFamily.bRequireMultiView, NumMSAASamples) && bTonemapSubpass;
 	bRequiresSceneDepthAux = MobileRequiresSceneDepthAux(ShaderPlatform) && !bTonemapSubpass;
 }
 
diff --git a/Engine/Source/Runtime/Renderer/Private/SceneRendering.cpp b/Engine/Source/Runtime/Renderer/Private/SceneRendering.cpp
--- a/Engine/Source/Runtime/Renderer/Private/SceneRendering.cpp
+++ b/Engine/Source/Runtime/Renderer/Private/SceneRendering.cpp
@@ -4329,9 +4329,13 @@ void FSceneRenderer::OnRenderFinish(FRDGBuilder& GraphBuilder, FRDGTextureRef Vi
 									{
 										Technique = TEXT("Multi-viewport");
 									}
+									else if (bIsMobileMultiViewEnabled)
+									{
+										Technique = TEXT("Multi-view (mobile, fallback)");
+									}
 									else
 									{
-										Technique = TEXT("Instanced, clip planes");
+										Technique = TEXT("Instanced, clip planes (deprecated, if you see this, it must be a bug)");
 									}
 								}
 								else if (bIsMobileMultiViewEnabled)
diff --git a/Engine/Source/Runtime/Renderer/Public/MeshPassProcessor.h b/Engine/Source/Runtime/Renderer/Public/MeshPassProcessor.h
--- a/Engine/Source/Runtime/Renderer/Public/MeshPassProcessor.h
+++ b/Engine/Source/Runtime/Renderer/Public/MeshPassProcessor.h
@@ -2405,7 +2405,7 @@ RENDERER_API extern void AddRenderTargetInfo(EPixelFormat PixelFormat, ETextureC
 RENDERER_API extern void SetupDepthStencilInfo(EPixelFormat DepthStencilFormat, ETextureCreateFlags DepthStencilCreateFlags, ERenderTargetLoadAction DepthTargetLoadAction, ERenderTargetLoadAction StencilTargetLoadAction, FExclusiveDepthStencil DepthStencilAccess, FGraphicsPipelineRenderTargetsInfo& RenderTargetsInfo);
 RENDERER_API extern void SetupGBufferRenderTargetInfo(const FSceneTexturesConfig& SceneTexturesConfig, FGraphicsPipelineRenderTargetsInfo& RenderTargetsInfo, bool bSetupDepthStencil);
 RENDERER_API extern void ApplyTargetsInfo(FGraphicsPipelineStateInitializer& GraphicsPSOInit, const FGraphicsPipelineRenderTargetsInfo& RenderTargetsInfo);
-RENDERER_API extern ESubpassHint GetSubpassHint(const FStaticShaderPlatform Platform, bool bIsUsingGBuffers, uint32 NumSamples);
+RENDERER_API extern ESubpassHint GetSubpassHint(const FStaticShaderPlatform Platform, bool bIsUsingGBuffers, bool bMultiViewRendering, uint32 NumSamples);
 
 inline FMeshDrawCommandSortKey CalculateMeshStaticSortKey(const TShaderRef<FMeshMaterialShader>& VertexShader, const TShaderRef<FMeshMaterialShader>& PixelShader)
 {
