[Physics] Fix using wrong cvar in resim logic.

diff --git a/Engine/Source/Runtime/Engine/Public/Physics/NetworkPhysicsSettingsComponent.h b/Engine/Source/Runtime/Engine/Public/Physics/NetworkPhysicsSettingsComponent.h
--- a/Engine/Source/Runtime/Engine/Public/Physics/NetworkPhysicsSettingsComponent.h
+++ b/Engine/Source/Runtime/Engine/Public/Physics/NetworkPhysicsSettingsComponent.h
@@ -250,7 +250,7 @@ struct FNetworkPhysicsSettingsResimulation
 	// Overrides Project Settings -> Physics -> Replication -> Physics Prediction -> Resimulation Error Threshold -- Distance that the object is allowed to desync from the server before triggering a resimulation, within this distance runtime correction can be performed if RuntimeCorrectionEnabled is true.
 	UPROPERTY(config, EditDefaultsOnly, Category = "Overrides", Meta = (EditCondition = "bOverrideResimulationErrorThreshold"))
 	uint32 ResimulationErrorThreshold = 10;
-	uint32 GetResimulationErrorThreshold(uint32 DefaultValue) { return bOverrideRotStabilityMultiplier ? ResimulationErrorThreshold : DefaultValue; }
+	uint32 GetResimulationErrorThreshold(uint32 DefaultValue) { return bOverrideResimulationErrorThreshold ? ResimulationErrorThreshold : DefaultValue; }
 
 	UPROPERTY(config, EditDefaultsOnly, Category = "Overrides", Meta = (InlineEditConditionToggle))
 	uint32 bOverrideRuntimeCorrectionEnabled : 1;
