[Shaders] Set USING_VERTEX_SHADER_LAYER depending on DDPI values not hard coded in shader backend.

This also reverts the DDPI information from CL 35396864, since D3D does in fact support vertex shader layers (except for D3D11.2 or older).

#rnx

diff --git a/Engine/Config/Windows/DataDrivenPlatformInfo.ini b/Engine/Config/Windows/DataDrivenPlatformInfo.ini
--- a/Engine/Config/Windows/DataDrivenPlatformInfo.ini
+++ b/Engine/Config/Windows/DataDrivenPlatformInfo.ini
@@ -124,7 +124,7 @@ bSupportsMeshShadersTier1=true
 bSupportsMeshShadersWithClipDistance=true
 bSupportsHairStrandGeometry=true
 bSupportsFFTBloom=true
-bSupportsVertexShaderLayer=false
+bSupportsVertexShaderLayer=true
 bSupportsInstancedStereo=true
 bSupportsCapsuleShadows=true
 bSupportsPercentageCloserShadows=true
@@ -160,7 +160,7 @@ bSupportsDistanceFields = false
 bSupportsComputeFramework=true
 bSupportsByteBufferComputeShaders = true
 bSupportsInstancedStereo=true
-bSupportsVertexShaderLayer=false
+bSupportsVertexShaderLayer=true
 bSupportsManualVertexFetch=false
 bSupportsSceneDataCompressedTransforms=true
 bSupportsDebugViewShaders=true
diff --git a/Engine/Source/Developer/Apple/MetalShaderFormat/Private/MetalShaderFormat.cpp b/Engine/Source/Developer/Apple/MetalShaderFormat/Private/MetalShaderFormat.cpp
--- a/Engine/Source/Developer/Apple/MetalShaderFormat/Private/MetalShaderFormat.cpp
+++ b/Engine/Source/Developer/Apple/MetalShaderFormat/Private/MetalShaderFormat.cpp
@@ -376,12 +376,10 @@ public:
 		else if (Input.ShaderFormat == NAME_SF_METAL_SM5)
 		{
 			Input.Environment.SetDefine(TEXT("METAL_SM5_PROFILE"), 1);
-			Input.Environment.SetDefine(TEXT("USING_VERTEX_SHADER_LAYER"), 1);
 		}
 		else if (Input.ShaderFormat == NAME_SF_METAL_SM6)
 		{
 			Input.Environment.SetDefine(TEXT("METAL_SM6_PROFILE"), 1);
-			Input.Environment.SetDefine(TEXT("USING_VERTEX_SHADER_LAYER"), 1);
 		}
 		else if (Input.ShaderFormat == NAME_SF_METAL_MRT_MAC)
 		{
diff --git a/Engine/Source/Runtime/Engine/Private/ShaderCompiler/ShaderCompiler.cpp b/Engine/Source/Runtime/Engine/Private/ShaderCompiler/ShaderCompiler.cpp
--- a/Engine/Source/Runtime/Engine/Private/ShaderCompiler/ShaderCompiler.cpp
+++ b/Engine/Source/Runtime/Engine/Private/ShaderCompiler/ShaderCompiler.cpp
@@ -3840,6 +3840,8 @@ void GlobalBeginCompileShader(
 
 	SET_SHADER_DEFINE(Input.Environment, COLORED_LIGHT_FUNCTION_ATLAS, GetLightFunctionAtlasFormat() > 0 ? 1 : 0);
 
+	SET_SHADER_DEFINE(Input.Environment, USING_VERTEX_SHADER_LAYER, RHISupportsVertexShaderLayer(EShaderPlatform(Target.Platform)) ? 1 : 0);
+
 	SET_SHADER_DEFINE(Input.Environment, PLATFORM_SUPPORTS_SHADER_ROOT_CONSTANTS, RHISupportsShaderRootConstants(EShaderPlatform(Target.Platform)) ? 1 : 0);
 	SET_SHADER_DEFINE(Input.Environment, PLATFORM_SUPPORTS_SHADER_BUNDLE_DISPATCH, RHISupportsShaderBundleDispatch(EShaderPlatform(Target.Platform)) ? 1 : 0);
 	SET_SHADER_DEFINE(Input.Environment, PLATFORM_SUPPORTS_RENDERTARGET_WRITE_MASK, RHISupportsRenderTargetWriteMask(EShaderPlatform(Target.Platform)) ? 1 : 0);
