Support exporting Nanite skeletal mesh assets to FBX (fixes a crash when trying to force CPU skinning in order to fetch out the final vertices).

diff --git a/Engine/Source/Runtime/Engine/Private/Components/SkinnedMeshComponent.cpp b/Engine/Source/Runtime/Engine/Private/Components/SkinnedMeshComponent.cpp
--- a/Engine/Source/Runtime/Engine/Private/Components/SkinnedMeshComponent.cpp
+++ b/Engine/Source/Runtime/Engine/Private/Components/SkinnedMeshComponent.cpp
@@ -1011,7 +1011,7 @@ void USkinnedMeshComponent::CreateRenderState_Concurrent(FRegisterComponentConte
 					MeshObject = MeshObjectFactory(MeshObjectFactoryUserData, this, SkelMeshRenderData, SceneFeatureLevel);
 				}
 
-				if (!MeshObject && ShouldNaniteSkin())
+				if (!MeshObject && ShouldNaniteSkin() && !ShouldCPUSkin() /* Needed for calls to GetCPUSkinnedVertices() */)
 				{
 					FSkeletalMeshObjectNanite* NaniteMeshObject = ::new FSkeletalMeshObjectNanite(this, SkelMeshRenderData, SceneFeatureLevel);
 					if (NaniteMeshObject->HasValidMaterials())
