[AutoRTFM] Add some tests for new transactional logic.

diff --git a/Engine/Source/Runtime/Engine/Private/Tests/AutoRTFM/AutoRTFMActorComponentTests.cpp b/Engine/Source/Runtime/Engine/Private/Tests/AutoRTFM/AutoRTFMActorComponentTests.cpp
--- a/Engine/Source/Runtime/Engine/Private/Tests/AutoRTFM/AutoRTFMActorComponentTests.cpp
+++ b/Engine/Source/Runtime/Engine/Private/Tests/AutoRTFM/AutoRTFMActorComponentTests.cpp
@@ -19,6 +19,7 @@
 #include "Physics/Experimental/PhysScene_Chaos.h"
 #include "PhysicsProxy/SingleParticlePhysicsProxy.h"
 #include "Rendering/SkeletalMeshRenderData.h"
+#include "UObject/UObjectArray.h"
 #include "UObject/UObjectGlobals.h"
 
 #if WITH_DEV_AUTOMATION_TESTS
@@ -439,6 +440,33 @@ AUTORTFM_ACTOR_COMPONENT_TEST(CreateComponentFromTemplate)
 	TEST_CHECK_TRUE(AutoRTFM::ETransactionResult::AbortedByRequest == Result);
 }
 
+// Test aborting a call to UObject::GetArchetype().
+// See: SOL-7024
+AUTORTFM_ACTOR_COMPONENT_TEST(GetArchetype)
+{
+	AutoRTFM::ETransactionResult Result = AutoRTFM::Transact([&]
+		{
+			Actor->GetArchetype();
+			AutoRTFM::AbortTransaction();
+		});
+
+	TEST_CHECK_TRUE(AutoRTFM::ETransactionResult::AbortedByRequest == Result);
+}
+
+// Test aborting a call to FUObjectArray::CloseDisregardForGC().
+// See: SOL-7027
+AUTORTFM_ACTOR_COMPONENT_TEST(CloseDisregardForGC)
+{
+	FUObjectArray ObjectArray;
+	AutoRTFM::ETransactionResult Result = AutoRTFM::Transact([&]
+		{
+			ObjectArray.CloseDisregardForGC();
+			AutoRTFM::AbortTransaction();
+		});
+
+	TEST_CHECK_TRUE(AutoRTFM::ETransactionResult::AbortedByRequest == Result);
+}
+
 }  // anonymous namespace
 
 #undef TEST_CHECK_FALSE
