Call QueuePendingRayTracingGeometryUpdate(...) using Context.RHICmdList instead of immediate command list.
- makes it possible to call GetDynamicRayTracingInstances(...) in parallel tasks.

diff --git a/Engine/Source/Runtime/Engine/Private/SkeletalMesh.cpp b/Engine/Source/Runtime/Engine/Private/SkeletalMesh.cpp
--- a/Engine/Source/Runtime/Engine/Private/SkeletalMesh.cpp
+++ b/Engine/Source/Runtime/Engine/Private/SkeletalMesh.cpp
@@ -7229,7 +7229,7 @@ void FSkeletalMeshSceneProxy::GetDynamicRayTracingInstances(FRayTracingMaterialG
 		return;
 	}
 
-	MeshObject->QueuePendingRayTracingGeometryUpdate(FRHICommandListImmediate::Get());
+	MeshObject->QueuePendingRayTracingGeometryUpdate(Context.RHICmdList);
 
 	FRayTracingGeometry* RayTracingGeometry = MeshObject->GetRayTracingGeometry();
 
