Substrate - made it easier to read and understand function parameter large list.

#rnx

diff --git a/Engine/Source/Runtime/Engine/Private/Materials/HLSLMaterialTranslator.cpp b/Engine/Source/Runtime/Engine/Private/Materials/HLSLMaterialTranslator.cpp
--- a/Engine/Source/Runtime/Engine/Private/Materials/HLSLMaterialTranslator.cpp
+++ b/Engine/Source/Runtime/Engine/Private/Materials/HLSLMaterialTranslator.cpp
@@ -13537,7 +13537,33 @@ int32 FHLSLMaterialTranslator::SubstrateSlabBSDF(
 		}
 
 		return AddCodeChunk(
-			MCT_Substrate, TEXT("Parameters.%s.PromoteParameterBlendedBSDFToOperator(GetSubstrateSlabBSDF(Parameters.SubstratePixelFootprint, %s, %s, %s, %s, %s, %s, %s, %s, %s, %s, %s, %s, %s, %s, %s, %s, %s, %s, %s, %s, %s, %s, %s, %s, Parameters.%s.Types) /* Normal = %s ; Tangent = %s ; Thickness = %s */, %u, %u, %u, %u)"),
+			MCT_Substrate, TEXT("Parameters.%s.PromoteParameterBlendedBSDFToOperator(GetSubstrateSlabBSDF(Parameters.SubstratePixelFootprint,\n\
+				/*DiffuseAlbedo*/                    %s,\n\
+				/*F0*/                               %s,\n\
+				/*F90*/                              %s,\n\
+				/*Roughness*/                        %s,\n\
+				/*Anisotropy*/                       %s,\n\
+				/*SSSProfileId*/                     %s,\n\
+				/*SSSMFP*/                           %s,\n\
+				/*SSSMFPScale*/                      %s,\n\
+				/*SSSPhaseAniso*/                    %s,\n\
+				/*UseSSSDiffusion*/                  %s,\n\
+				/*EmissiveColor*/                    %s,\n\
+				/*SecondRoughness*/                  %s,\n\
+				/*SecondRoughnessWeight*/            %s,\n\
+				/*SecondRoughnessAsSimpleClearCoat*/ %s,\n\
+				/*FuzzAmount*/                       %s,\n\
+				/*FuzzColor*/                        %s,\n\
+				/*FuzzRoughness*/                    %s,\n\
+				/*GlintValue*/                       %s,\n\
+				/*GlintUV*/                          %s,\n\
+				/*SpecularProfileId*/                %s,\n\
+				/*Thickness*/                        %s,\n\
+				/*IsThin*/                           %s,\n\
+				/*IsAtBottom*/                       %s,\n\
+				/*LocalBasisIndex*/                  %s,\n\
+				Parameters.%s.Types)\n\
+				 /* Normal = %s ; Tangent = %s ; Thickness = %s */, %u, %u, %u, %u)"),
 			*GetParametersSubstrateTreeName(CurrentSubstrateCompilationContext),
 			*SubstrateGetCastParameterCode(DiffuseAlbedo,			MCT_Float3),
 			*SubstrateGetCastParameterCode(F0,						MCT_Float3),
@@ -13575,7 +13601,32 @@ int32 FHLSLMaterialTranslator::SubstrateSlabBSDF(
 	}
 	
 	return AddCodeChunk(
-		MCT_Substrate, TEXT("GetSubstrateSlabBSDF(Parameters.SubstratePixelFootprint, %s, %s, %s, %s, %s, %s, %s, %s, %s, %s, %s, %s, %s, %s, %s, %s, %s, %s, %s, %s, %s, %s, %s, %s, Parameters.%s.Types) /* Normal = %s ; Tangent = %s ; Thickness = %s */"),
+		MCT_Substrate, TEXT("GetSubstrateSlabBSDF(Parameters.SubstratePixelFootprint,\n\
+			/*DiffuseAlbedo*/                    %s,\n\
+			/*F0*/                               %s,\n\
+			/*F90*/                              %s,\n\
+			/*Roughness*/                        %s,\n\
+			/*Anisotropy*/                       %s,\n\
+			/*SSSProfileId*/                     %s,\n\
+			/*SSSMFP*/                           %s,\n\
+			/*SSSMFPScale*/                      %s,\n\
+			/*SSSPhaseAniso*/                    %s,\n\
+			/*UseSSSDiffusion*/                  %s,\n\
+			/*EmissiveColor*/                    %s,\n\
+			/*SecondRoughness*/                  %s,\n\
+			/*SecondRoughnessWeight*/            %s,\n\
+			/*SecondRoughnessAsSimpleClearCoat*/ %s,\n\
+			/*FuzzAmount*/                       %s,\n\
+			/*FuzzColor*/                        %s,\n\
+			/*FuzzRoughness*/                    %s,\n\
+			/*GlintValue*/                       %s,\n\
+			/*GlintUV*/                          %s,\n\
+			/*SpecularProfileId*/                %s,\n\
+			/*Thickness*/                        %s,\n\
+			/*IsThin*/                           %s,\n\
+			/*IsAtBottom*/                       %s,\n\
+			/*LocalBasisIndex*/                  %s,\n\
+			Parameters.%s.Types) /* Normal = %s ; Tangent = %s ; Thickness = %s */"),
 		*SubstrateGetCastParameterCode(DiffuseAlbedo,			MCT_Float3),
 		*SubstrateGetCastParameterCode(F0,						MCT_Float3),
 		*SubstrateGetCastParameterCode(F90,					MCT_Float3),
@@ -13658,7 +13709,33 @@ int32 FHLSLMaterialTranslator::SubstrateConversionFromLegacy(
 		}
 
 		return AddCodeChunk(
-			MCT_Substrate, TEXT("Parameters.%s.PromoteParameterBlendedBSDFToOperator(SubstrateConvertLegacyMaterial%s(Parameters.SubstratePixelFootprint, %s, %s, %s, %s, %s, %s, %s, %s, %s, %s, %s, %s, %s, %s, %s, %s, %s, %s, %s, %s, %s, %s, %s, %s, Parameters.%s.Types) /* Normal = %s ; Tangent = %s ; ClearCoat_Normal = %s ; ClearCoat_Tangent = %s */, %u, %u, %u, %u)"),
+			MCT_Substrate, TEXT("Parameters.%s.PromoteParameterBlendedBSDFToOperator(SubstrateConvertLegacyMaterial%s(Parameters.SubstratePixelFootprint,\n\
+				/*BaseColor*/                         %s,\n\
+				/*Specular*/                          %s,\n\
+				/*Metallic*/                          %s,\n\
+				/*Roughness*/                         %s,\n\
+				/*Anisotropy*/                        %s,\n\
+				/*SubSurfaceColor*/                   %s,\n\
+				/*SubSurfaceProfileId*/               %s,\n\
+				/*ClearCoat*/                         %s,\n\
+				/*ClearCoatRoughness*/                %s,\n\
+				/*EmissiveColor*/                     %s,\n\
+				/*Opacity*/                           %s,\n\
+				/*ThinTranslucentTransmittanceColor*/ %s,\n\
+				/*ThinTranslucentSurfaceCoverage*/    %s,\n\
+				/*WaterScatteringCoefficients*/       %s,\n\
+				/*WaterAbsorptionCoefficients*/       %s,\n\
+				/*WaterPhaseG*/                       %s,\n\
+				/*ColorScaleBehindWater*/             %s,\n\
+				/*InShadingModel*/                    %s,\n\
+				/*Normal*/                            %s,\n\
+				/*Tangent*/                           %s,\n\
+				/*ClearCoat_Normal*/                  %s,\n\
+				/*CustomTangent_Tangent*/             %s,\n\
+				/*LocalBasisIndex*/                   %s,\n\
+				/*ClearCoat_SharedLocalBasisIndex*/   %s,\n\
+				Parameters.%s.Types)\n\
+				/* Normal = %s ; Tangent = %s ; ClearCoat_Normal = %s ; ClearCoat_Tangent = %s */, %u, %u, %u, %u)"),
 			*GetParametersSubstrateTreeName(CurrentSubstrateCompilationContext),
 			bHasDynamicShadingModels ? TEXT("Dynamic") : TEXT("Static"),
 			*SubstrateGetCastParameterCode(BaseColor,							MCT_Float3),
@@ -13701,7 +13778,33 @@ int32 FHLSLMaterialTranslator::SubstrateConversionFromLegacy(
 	}
 
 	return AddCodeChunk(
-		MCT_Substrate, TEXT("SubstrateConvertLegacyMaterial%s(Parameters.SubstratePixelFootprint, %s, %s, %s, %s, %s, %s, %s, %s, %s, %s, %s, %s, %s, %s, %s, %s, %s, %s, %s, %s, %s, %s, %s, %s, Parameters.%s.Types) /* Normal = %s ; Tangent = %s ; ClearCoat_Normal = %s ; ClearCoat_Tangent = %s */"),
+		MCT_Substrate, TEXT("SubstrateConvertLegacyMaterial%s(Parameters.SubstratePixelFootprint,\n\
+			/*BaseColor*/                         %s,\n\
+			/*Specular*/                          %s,\n\
+			/*Metallic*/                          %s,\n\
+			/*Roughness*/                         %s,\n\
+			/*Anisotropy*/                        %s,\n\
+			/*SubSurfaceColor*/                   %s,\n\
+			/*SubSurfaceProfileId*/               %s,\n\
+			/*ClearCoat*/                         %s,\n\
+			/*ClearCoatRoughness*/                %s,\n\
+			/*EmissiveColor*/                     %s,\n\
+			/*Opacity*/                           %s,\n\
+			/*ThinTranslucentTransmittanceColor*/ %s,\n\
+			/*ThinTranslucentSurfaceCoverage*/    %s,\n\
+			/*WaterScatteringCoefficients*/       %s,\n\
+			/*WaterAbsorptionCoefficients*/       %s,\n\
+			/*WaterPhaseG*/                       %s,\n\
+			/*ColorScaleBehindWater*/             %s,\n\
+			/*InShadingModel*/                    %s,\n\
+			/*Normal*/                            %s,\n\
+			/*Tangent*/                           %s,\n\
+			/*ClearCoat_Normal*/                  %s,\n\
+			/*CustomTangent_Tangent*/             %s,\n\
+			/*LocalBasisIndex*/                   %s,\n\
+			/*ClearCoat_SharedLocalBasisIndex*/   %s,\n\
+			Parameters.%s.Types)\n\
+			/* Normal = %s ; Tangent = %s ; ClearCoat_Normal = %s ; ClearCoat_Tangent = %s */"),
 		bHasDynamicShadingModels ? TEXT("Dynamic") : TEXT("Static"),
 		*SubstrateGetCastParameterCode(BaseColor,							MCT_Float3),
 		*SubstrateGetCastParameterCode(Specular,							MCT_Float),
