Fixed GravityGroupIndex not being initialized for body instances

diff --git a/Engine/Source/Runtime/Engine/Private/PhysicsEngine/BodyInstance.cpp b/Engine/Source/Runtime/Engine/Private/PhysicsEngine/BodyInstance.cpp
--- a/Engine/Source/Runtime/Engine/Private/PhysicsEngine/BodyInstance.cpp
+++ b/Engine/Source/Runtime/Engine/Private/PhysicsEngine/BodyInstance.cpp
@@ -372,6 +372,7 @@ FBodyInstance::FBodyInstance()
 	, CustomDOFPlaneNormal(FVector::ZeroVector)
 	, COMNudge(ForceInit)
 	, MassScale(1.f)
+	, GravityGroupIndex(0)
 	, InertiaTensorScale(1.f)
 	, DOFConstraint(nullptr)
 	, WeldParent(nullptr)
