Fixed linux compilation issue (wrong [[nodiscard]] placement)

diff --git a/Engine/Source/Runtime/Engine/Public/WorldPartition/StaticLightingData/StaticLightingDescriptors.h b/Engine/Source/Runtime/Engine/Public/WorldPartition/StaticLightingData/StaticLightingDescriptors.h
--- a/Engine/Source/Runtime/Engine/Public/WorldPartition/StaticLightingData/StaticLightingDescriptors.h
+++ b/Engine/Source/Runtime/Engine/Public/WorldPartition/StaticLightingData/StaticLightingDescriptors.h
@@ -60,7 +60,7 @@ struct FStaticLightingDescriptors
 	UNREALED_API UMapBuildDataRegistry* GetOrCreateRegistryForActor(AActor* Actor);
 	UNREALED_API UMapBuildDataRegistry* GetRegistryForActor(AActor* Actor, bool bCreateIfNotFound = false);
 
-	UNREALED_API [[nodiscard]] TArray<UMapBuildDataRegistry*> GetAllMapBuildData();
+	[[nodiscard]] UNREALED_API TArray<UMapBuildDataRegistry*> GetAllMapBuildData();
 
 	UNREALED_API bool CreateAndUpdateActors();
 
