Super minor cleanup to UMaterial::PostLoad to make it easier to read and reason about what is editor code and client code.

- Make BackwardsCompatibilityInputConversion, BackwardsCompatibilityVirtualTextureOutputConversion, BackwardsCompatibilityDecalConversion, ConvertMaterialToSubstrateMaterial all WITH_EDITOR
- Combine some WITH_EDITOR sections
- Delete dead code

#rnx

diff --git a/Engine/Source/Runtime/Engine/Private/Materials/Material.cpp b/Engine/Source/Runtime/Engine/Private/Materials/Material.cpp
--- a/Engine/Source/Runtime/Engine/Private/Materials/Material.cpp
+++ b/Engine/Source/Runtime/Engine/Private/Materials/Material.cpp
@@ -3013,9 +3013,10 @@ void UMaterial::PostDuplicate(bool bDuplicateForPIE)
 	}
 }
 
+#if WITH_EDITOR
+
 void UMaterial::BackwardsCompatibilityInputConversion()
 {
-#if WITH_EDITOR
 	if( ShadingModel != MSM_Unlit )
 	{
 		UMaterialEditorOnlyData* EditorOnly = GetEditorOnlyData();
@@ -3051,12 +3052,10 @@ void UMaterial::BackwardsCompatibilityInputConversion()
 			EditorOnly->Specular.Connect( 2, FunctionExpression );
 		}
 	}
-#endif // WITH_EDITOR
 }
 
 void UMaterial::BackwardsCompatibilityVirtualTextureOutputConversion()
 {
-#if WITH_EDITOR
 	// Remove MD_RuntimeVirtualTexture support and replace with an explicit UMaterialExpressionRuntimeVirtualTextureOutput.
 	if (MaterialDomain == MD_RuntimeVirtualTexture)
 	{
@@ -3131,12 +3130,10 @@ void UMaterial::BackwardsCompatibilityVirtualTextureOutputConversion()
 		// Recompile after changes with a guid representing the conversion applied here.
 		ReleaseResourcesAndMutateDDCKey(BackwardsCompatibilityVirtualTextureOutputConversionGuid);
 	}
-#endif // WITH_EDITOR
 }
 
 void UMaterial::BackwardsCompatibilityDecalConversion()
 {
-#if WITH_EDITOR
 	if (GMaterialsThatNeedDecalFix.Get(this))
 	{
 		// Change this guid if you change the conversion code below
@@ -3220,8 +3217,8 @@ void UMaterial::BackwardsCompatibilityDecalConversion()
 		// Recompile after changes with a guid representing the conversion applied here.
 		ReleaseResourcesAndMutateDDCKey(BackwardsCompatibilityDecalConversionGuid);
 	}
-#endif // WITH_EDITOR
 }
+#endif // WITH_EDITOR
 
 static void AddSurfaceSubstrateShadingModelFromMaterialShadingModels(FSubstrateMaterialInfo& OutInfo, const FMaterialShadingModelField& InShadingModels)
 {
@@ -3265,6 +3262,8 @@ EBlendMode ConvertLegacyBlendMode(EBlendMode InBlendMode, FMaterialShadingModelF
 #define SUBSTRATE_MOVE_CONNECTION 0
 #define SUBSTRATE_COPY_CONNECTION 1
 
+#if WITH_EDITOR
+
 void UMaterial::ConvertMaterialToSubstrateMaterial()
 {
 	/*
@@ -3283,7 +3282,6 @@ void UMaterial::ConvertMaterialToSubstrateMaterial()
 	* --- Material instance shading model override
 	*     Overridden from the HLSLTranslator when detected by comparing base and instanced materials.
 	*/
-#if WITH_EDITOR
 	if (!Substrate::IsSubstrateEnabled())
 	{
 		return;
@@ -3929,8 +3927,8 @@ void UMaterial::ConvertMaterialToSubstrateMaterial()
 		// We might have moved connections above so update the CachedExpressionData from the EditorOnly connection data (ground truth).
 		UpdateCachedExpressionData();
 	}
-#endif
 }
+#endif // WITH_EDITOR
 
 TMap<FGuid, UMaterialInterface*> LightingGuidFixupMap;
 
@@ -4072,11 +4070,12 @@ void UMaterial::PostLoad()
 		AssertDefaultMaterialsPostLoaded();
 	}	
 
+#if WITH_EDITOR
 	if ( GIsEditor && GetOuter() == GetTransientPackage() && FCString::Strstr(*GetName(), TEXT("MEStatsMaterial_")))
 	{
 		bIsMaterialEditorStatsMaterial = true;
 	}
-
+#endif // WITH_EDITOR
 
 	if( GetLinkerUEVersion() < VER_UE4_REMOVED_MATERIAL_USED_WITH_UI_FLAG && bUsedWithUI_DEPRECATED == true )
 	{
@@ -4169,11 +4168,12 @@ void UMaterial::PostLoad()
 		FPlatformMisc::CreateGuid(StateId);
 	}
 
+#if WITH_EDITOR
+
 	BackwardsCompatibilityInputConversion();
 	BackwardsCompatibilityVirtualTextureOutputConversion();
 	BackwardsCompatibilityDecalConversion();
 
-#if WITH_EDITOR
 	if ( GMaterialsThatNeedSamplerFixup.Get( this ) )
 	{
 		GMaterialsThatNeedSamplerFixup.Clear( this );
@@ -4241,9 +4241,11 @@ void UMaterial::PostLoad()
 	}
 #endif
 
+#if WITH_EDITOR
 	// Substrate materials conversion needs to be done after expressions are cached, otherwise material function won't have 
 	// valid inputs in certain cases
 	ConvertMaterialToSubstrateMaterial();
+#endif // WITH_EDITOR
 
 	checkf(CachedExpressionData, TEXT("Missing cached expression data for material, should have been either serialized or created during PostLoad"));
 
@@ -4286,9 +4288,7 @@ void UMaterial::PostLoad()
 			ReleaseResourcesAndMutateDDCKey(BackwardsCompatibilityFeatureLevelSM6ConversionGuid);
 		}
 	}
-#endif // #if WITH_EDITOR
 
-#if WITH_EDITOR
 	// Before, refraction was only enabled when the refraction pin was plugged in.
 	// Now it is enabled only when not OFF. Otherwise:
 	//    - if plugged the pin override the physically based material refraction
@@ -4322,20 +4322,7 @@ void UMaterial::PostLoad()
 	STAT(double MaterialLoadTime = 0);
 	{
 		SCOPE_SECONDS_COUNTER(MaterialLoadTime);
-// Daniel: Disable compiling shaders for cooked platforms as the cooker will manually call the BeginCacheForCookedPlatformData function and load balence
-#if 0 && WITH_EDITOR
-		// enable caching in postload for derived data cache commandlet and cook by the book
-		ITargetPlatformManagerModule* TPM = GetTargetPlatformManager();
-		if (TPM && (TPM->RestrictFormatsToRuntimeOnly() == false))
-		{
-			TArray<ITargetPlatform*> Platforms = TPM->GetActiveTargetPlatforms();
-			// Cache for all the shader formats that the cooking target requires
-			for (int32 FormatIndex = 0; FormatIndex < Platforms.Num(); FormatIndex++)
-			{
-				BeginCacheForCookedPlatformData(Platforms[FormatIndex]);
-			}
-		}
-#endif
+
 		//Don't compile shaders in post load for dev overhead materials.
 		if (FApp::CanEverRender() && !bIsMaterialEditorStatsMaterial && GAllowCompilationInPostLoad)
 		{
@@ -4362,19 +4349,19 @@ void UMaterial::PostLoad()
 	}
 	INC_FLOAT_STAT_BY(STAT_ShaderCompiling_MaterialLoading,(float)MaterialLoadTime);
 
+#if WITH_EDITOR
 	if( GIsEditor && !IsTemplate() )
 	{
 		// Ensure that the ReferencedTextureGuids array is up to date.
 		UpdateLightmassTextureTracking();
 	}
+#endif // WITH_EDITOR
 
 	if (IsDeferredDecal())
 	{
 		FPSOPrecacheParams PSOPrecacheParams;
 		UMaterialInterface::PrecachePSOs(&FLocalVertexFactory::StaticType, PSOPrecacheParams);
 	}
-
-	//DumpDebugInfo(*GLog);
 }
 
 #if WITH_EDITORONLY_DATA
