Extend FSkeletalMeshSectionIter to support ignore preview filter.
- avoids the need to modify MeshObject during GetDynamicRayTracingInstances(...) which prevents it from being called in parallel with GetDynamicMeshElements(...) etc.

diff --git a/Engine/Source/Runtime/Engine/Private/SkeletalMesh.cpp b/Engine/Source/Runtime/Engine/Private/SkeletalMesh.cpp
--- a/Engine/Source/Runtime/Engine/Private/SkeletalMesh.cpp
+++ b/Engine/Source/Runtime/Engine/Private/SkeletalMesh.cpp
@@ -6670,14 +6670,14 @@ FSkeletalMeshSceneProxy::FSkeletalMeshSceneProxy(const USkinnedMeshComponent* Co
 class FSkeletalMeshSectionIter
 {
 public:
-	FSkeletalMeshSectionIter(const int32 InLODIdx, const FSkeletalMeshObject& InMeshObject, const FSkeletalMeshLODRenderData& InLODData, const FSkeletalMeshSceneProxy::FLODSectionElements& InLODSectionElements)
+	FSkeletalMeshSectionIter(const int32 InLODIdx, const FSkeletalMeshObject& InMeshObject, const FSkeletalMeshLODRenderData& InLODData, const FSkeletalMeshSceneProxy::FLODSectionElements& InLODSectionElements, bool bIgnorePreviewFilter = false)
 		: SectionIndex(0)
 		, MeshObject(InMeshObject)
 		, LODSectionElements(InLODSectionElements)
 		, Sections(InLODData.RenderSections)
 #if WITH_EDITORONLY_DATA
-		, SectionIndexPreview(InMeshObject.SectionIndexPreview)
-		, MaterialIndexPreview(InMeshObject.MaterialIndexPreview)
+		, SectionIndexPreview(bIgnorePreviewFilter ? INDEX_NONE : InMeshObject.SectionIndexPreview)
+		, MaterialIndexPreview(bIgnorePreviewFilter ? INDEX_NONE : InMeshObject.MaterialIndexPreview)
 #endif
 	{
 		while (NotValidPreviewSection())
@@ -7263,17 +7263,12 @@ void FSkeletalMeshSceneProxy::GetDynamicRayTracingInstances(FRayTracingMaterialG
 			{
 				return;
 			}
-
-		#if WITH_EDITORONLY_DATA
-			int32 SectionIndexPreview = MeshObject->SectionIndexPreview;
-			int32 MaterialIndexPreview = MeshObject->MaterialIndexPreview;
-			MeshObject->SectionIndexPreview = INDEX_NONE;
-			MeshObject->MaterialIndexPreview = INDEX_NONE;
-		#endif
+			
+			const bool bIgnorePreviewFilter = true;
 
 			uint32 TotalNumVertices = 0;
 
-			for (FSkeletalMeshSectionIter Iter(LODIndex, *MeshObject, LODData, LODSection); Iter; ++Iter)
+			for (FSkeletalMeshSectionIter Iter(LODIndex, *MeshObject, LODData, LODSection, bIgnorePreviewFilter); Iter; ++Iter)
 			{
 				const FSkelMeshRenderSection& Section = Iter.GetSection();
 				const int32 SectionIndex = Iter.GetSectionElementIndex();
@@ -7300,7 +7295,7 @@ void FSkeletalMeshSceneProxy::GetDynamicRayTracingInstances(FRayTracingMaterialG
 				TArray<FRayTracingGeometrySegment> GeometrySections;
 				GeometrySections.Reserve(LODData.RenderSections.Num());
 
-				for (FSkeletalMeshSectionIter Iter(LODIndex, *MeshObject, LODData, LODSection); Iter; ++Iter)
+				for (FSkeletalMeshSectionIter Iter(LODIndex, *MeshObject, LODData, LODSection, bIgnorePreviewFilter); Iter; ++Iter)
 				{
 					const FSkelMeshRenderSection& Section = Iter.GetSection();
 					const FSectionElementInfo& SectionElementInfo = Iter.GetSectionElementInfo();
@@ -7332,11 +7327,6 @@ void FSkeletalMeshSceneProxy::GetDynamicRayTracingInstances(FRayTracingMaterialG
 			}
 
 			OutRayTracingInstances.Add(RayTracingInstance);
-			
-		#if WITH_EDITORONLY_DATA
-			MeshObject->SectionIndexPreview = SectionIndexPreview;
-			MeshObject->MaterialIndexPreview = MaterialIndexPreview;
-		#endif
 		}
 	}
 }
