Delete debug code in UMaterial::PostLoad

#rnx

diff --git a/Engine/Source/Runtime/Engine/Private/Materials/Material.cpp b/Engine/Source/Runtime/Engine/Private/Materials/Material.cpp
--- a/Engine/Source/Runtime/Engine/Private/Materials/Material.cpp
+++ b/Engine/Source/Runtime/Engine/Private/Materials/Material.cpp
@@ -4131,12 +4131,6 @@ void UMaterial::PostLoad()
 	}
 
 #if !WITH_EDITOR
-
-	if (GetFullName().Contains(TEXT("clearcoat_bn")))
-	{
-		UE_LOG(LogMaterial, Display, TEXT("DEBUG break"));
-	}
-
 	// Filter out ShadingModels field to a current platform settings
 	FilterOutPlatformShadingModels(GMaxRHIShaderPlatform, ShadingModels);
 	// Override material shader model if it was filtered out
