Update GetActor comments in GameplayStatics to match actual performance

diff --git a/Engine/Source/Runtime/Engine/Classes/Kismet/GameplayStatics.h b/Engine/Source/Runtime/Engine/Classes/Kismet/GameplayStatics.h
--- a/Engine/Source/Runtime/Engine/Classes/Kismet/GameplayStatics.h
+++ b/Engine/Source/Runtime/Engine/Classes/Kismet/GameplayStatics.h
@@ -82,7 +82,6 @@ class UGameplayStatics : public UBlueprintFunctionLibrary
 	
 	/** 
 	 *	Find the first Actor in the world of the specified class. 
-	 *	This is a slow operation, use with caution e.g. do not use every frame.
 	 *	@param	ActorClass	Class of Actor to find. Must be specified or result will be empty.
 	 *	@return				Actor of the specified class.
 	 */
@@ -91,7 +90,7 @@ class UGameplayStatics : public UBlueprintFunctionLibrary
 
 	/** 
 	 *	Find all Actors in the world of the specified class. 
-	 *	This is a slow operation, use with caution e.g. do not use every frame.
+	 *	This will be slow if there are many actors of the specified class.
 	 *	@param	ActorClass	Class of Actor to find. Must be specified or result array will be empty.
 	 *	@param	OutActors	Output array of Actors of the specified class.
 	 */
@@ -100,7 +99,7 @@ class UGameplayStatics : public UBlueprintFunctionLibrary
 
 	/** 
 	 *	Find all Actors in the world with the specified interface.
-	 *	This is a slow operation, use with caution e.g. do not use every frame.
+	 *	This is a very slow operation, as it will search over every actor in the world.
 	 *	@param	Interface	Interface to find. Must be specified or result array will be empty.
 	 *	@param	OutActors	Output array of Actors of the specified interface.
 	 */
@@ -109,7 +108,7 @@ class UGameplayStatics : public UBlueprintFunctionLibrary
 
 	/**
 	 *	Find all Actors in the world with the specified tag.
-	 *	This is a slow operation, use with caution e.g. do not use every frame.
+	 *	This is a very slow operation, as it will search over every actor in the world.
 	 *	@param	Tag			Tag to find. Must be specified or result array will be empty.
 	 *	@param	OutActors	Output array of Actors of the specified tag.
 	 */
@@ -118,7 +117,7 @@ class UGameplayStatics : public UBlueprintFunctionLibrary
 
 	/**
 	 *	Find all Actors in the world of the specified class with the specified tag.
-	 *	This is a slow operation, use with caution e.g. do not use every frame.
+	 *	This will be slow if there are many actors of the specified class.
 	 *	@param	Tag			Tag to find. Must be specified or result array will be empty.
 	 *	@param	ActorClass	Class of Actor to find. Must be specified or result array will be empty.
 	 *	@param	OutActors	Output array of Actors of the specified tag.
