SkeletalMesh: Add a BP decorator to access the GetHasVertexColors to check if a skeletal mesh _actually_ has vertex colors.

#rnx

diff --git a/Engine/Source/Runtime/Engine/Classes/Engine/SkeletalMesh.h b/Engine/Source/Runtime/Engine/Classes/Engine/SkeletalMesh.h
--- a/Engine/Source/Runtime/Engine/Classes/Engine/SkeletalMesh.h
+++ b/Engine/Source/Runtime/Engine/Classes/Engine/SkeletalMesh.h
@@ -1313,7 +1313,8 @@ public:
 		PRAGMA_ENABLE_DEPRECATION_WARNINGS
 	}
 
-	/** Return whether or not the mesh has vertex colors. USkinnedAsset interface. */
+	/** Return whether the mesh has vertex colors. USkinnedAsset interface. */
+	UFUNCTION(BlueprintPure, Category = Mesh, meta = (DisplayName = "Has Vertex Colors", ScriptName = "HasVertexColors", Keywords = "vertex color"))
 	virtual bool GetHasVertexColors() const override
 	{
 		WaitUntilAsyncPropertyReleased(ESkeletalMeshAsyncProperties::HasVertexColors, ESkinnedAssetAsyncPropertyLockType::ReadOnly);
