Add profile scopes to OnPostReachabilityAnalysis callbacks

diff --git a/Engine/Plugins/FX/Niagara/Source/Niagara/Private/NiagaraWorldManager.cpp b/Engine/Plugins/FX/Niagara/Source/Niagara/Private/NiagaraWorldManager.cpp
--- a/Engine/Plugins/FX/Niagara/Source/Niagara/Private/NiagaraWorldManager.cpp
+++ b/Engine/Plugins/FX/Niagara/Source/Niagara/Private/NiagaraWorldManager.cpp
@@ -1007,6 +1007,7 @@ void FNiagaraWorldManager::OnPreGarbageCollect()
 
 void FNiagaraWorldManager::OnPostReachabilityAnalysis()
 {
+	TRACE_CPUPROFILER_EVENT_SCOPE(FNiagaraWorldManager::OnPostReachabilityAnalysis);
 	for (TPair<UWorld*, FNiagaraWorldManager*>& Pair : WorldManagers)
 	{
 		Pair.Value->PostReachabilityAnalysis();
diff --git a/Engine/Source/Editor/UnrealEd/Private/PackageTools.cpp b/Engine/Source/Editor/UnrealEd/Private/PackageTools.cpp
--- a/Engine/Source/Editor/UnrealEd/Private/PackageTools.cpp
+++ b/Engine/Source/Editor/UnrealEd/Private/PackageTools.cpp
@@ -152,6 +152,7 @@ UPackageTools::UPackageTools(const FObjectInitializer& ObjectInitializer)
 	void UPackageTools::RestoreStandaloneOnReachableObjects()
 	{
 		check(GIsEditor);
+		TRACE_CPUPROFILER_EVENT_SCOPE(UPackageTools::RestoreStandaloneOnReachableObjects);
 
 		if (PackagesBeingUnloaded && ObjectsThatHadFlagsCleared.Num() > 0)
 		{
diff --git a/Engine/Source/Runtime/Engine/Private/WorldPartition/WorldPartition.cpp b/Engine/Source/Runtime/Engine/Private/WorldPartition/WorldPartition.cpp
--- a/Engine/Source/Runtime/Engine/Private/WorldPartition/WorldPartition.cpp
+++ b/Engine/Source/Runtime/Engine/Private/WorldPartition/WorldPartition.cpp
@@ -359,6 +359,7 @@ UWorldPartition::UWorldPartition(const FObjectInitializer& ObjectInitializer)
 #if WITH_EDITOR
 void UWorldPartition::OnGCPostReachabilityAnalysis()
 {
+	TRACE_CPUPROFILER_EVENT_SCOPE(UWorldPartition::OnGCPostReachabilityAnalysis);
 	const TIndirectArray<FWorldContext>& WorldContextList = GEngine->GetWorldContexts();
 
 	// Avoid running this process while a game world is live
