[HWRT] For skinned meshes, move the check for non-resident LODs to happen before adding them to the build queue.
During streaming, under very rare circumstances, we might queue an RTGeometry to build in the next frame while the mesh is also being streamed out leading to a crash during the build.

#rnx

diff --git a/Engine/Source/Runtime/Engine/Private/SkeletalMesh.cpp b/Engine/Source/Runtime/Engine/Private/SkeletalMesh.cpp
--- a/Engine/Source/Runtime/Engine/Private/SkeletalMesh.cpp
+++ b/Engine/Source/Runtime/Engine/Private/SkeletalMesh.cpp
@@ -7228,6 +7228,12 @@ void FSkeletalMeshSceneProxy::GetDynamicRayTracingInstances(FRayTracingMaterialG
 	{
 		return;
 	}
+	
+	// According to GetMeshElementsConditionallySelectable(), non-resident LODs should just be skipped
+	if (MeshObject->GetRayTracingLOD() < SkeletalMeshRenderData->CurrentFirstLODIdx)
+	{
+		return;
+	}
 
 	MeshObject->QueuePendingRayTracingGeometryUpdate(Context.RHICmdList);
 
@@ -7248,12 +7254,6 @@ void FSkeletalMeshSceneProxy::GetDynamicRayTracingInstances(FRayTracingMaterialG
 			check(LODIndex < SkeletalMeshRenderData->LODRenderData.Num());
 			const FSkeletalMeshLODRenderData& LODData = SkeletalMeshRenderData->LODRenderData[LODIndex];
 
-			if (LODIndex < SkeletalMeshRenderData->CurrentFirstLODIdx)
-			{
-				// According to GetMeshElementsConditionallySelectable(), non-resident LODs should just be skipped
-				return;
-			}
-
 			ensure(LODSections.Num() > 0);
 			const FLODSectionElements& LODSection = LODSections[LODIndex];
 			check(LODSection.SectionElements.Num() == LODData.RenderSections.Num());
