Exposed FText::AsMemory to BP

Note: Due to limitations with BP exposed functions, this takes an int64 rather than a uint64. Any negative values are clamped to zero rather than overflowing.

diff --git a/Engine/Source/Runtime/Core/Public/Internationalization/Text.h b/Engine/Source/Runtime/Core/Public/Internationalization/Text.h
--- a/Engine/Source/Runtime/Core/Public/Internationalization/Text.h
+++ b/Engine/Source/Runtime/Core/Public/Internationalization/Text.h
@@ -186,7 +186,8 @@ CORE_API bool LexTryParseString(ERoundingMode& OutValue, const TCHAR* Buffer);
 CORE_API void LexFromString(ERoundingMode& OutValue, const TCHAR* Buffer);
 CORE_API const TCHAR* LexToString(ERoundingMode InValue);
 
-enum EMemoryUnitStandard
+/** Redeclared in KismetTextLibrary for meta-data extraction purposes, be sure to update there as well */
+enum EMemoryUnitStandard : int
 {
 	/* International Electrotechnical Commission (MiB) 1024-based */
 	IEC,
diff --git a/Engine/Source/Runtime/Engine/Classes/Kismet/KismetTextLibrary.h b/Engine/Source/Runtime/Engine/Classes/Kismet/KismetTextLibrary.h
--- a/Engine/Source/Runtime/Engine/Classes/Kismet/KismetTextLibrary.h
+++ b/Engine/Source/Runtime/Engine/Classes/Kismet/KismetTextLibrary.h
@@ -32,6 +32,15 @@ enum ERoundingMode : int
 	ToPositiveInfinity,
 };
 
+UENUM(BlueprintType)
+enum EMemoryUnitStandard : int
+{
+	/* International Electrotechnical Commission (MiB) 1024-based */
+	IEC,
+	/* International System of Units 1000-based */
+	SI
+};
+
 UENUM(BlueprintType)
 enum class ETextGender : uint8
 {
@@ -288,6 +297,10 @@ class UKismetTextLibrary : public UBlueprintFunctionLibrary
 	UFUNCTION(BlueprintPure, meta=(DisplayName = "As Timespan", AdvancedDisplay = "1"), Category="Utilities|Text")
 	static ENGINE_API FText AsTimespan_Timespan(const FTimespan& InTimespan);
 
+	/** Generate an FText that represents the passed number as a memory size in the current culture */
+	UFUNCTION(BlueprintPure, meta=(AdvancedDisplay = "1"), Category="Utilities|Text")
+	static ENGINE_API FText AsMemory(int64 NumBytes, TEnumAsByte<EMemoryUnitStandard> UnitStandard = EMemoryUnitStandard::IEC, bool bUseGrouping = true);
+
 	/** Used for formatting text using the FText::Format function and utilized by the UK2Node_FormatText */
 	UFUNCTION(BlueprintPure, meta=(BlueprintInternalUseOnly = "true"))
 	static ENGINE_API FText Format(FText InPattern, TArray<FFormatArgumentData> InArgs);
diff --git a/Engine/Source/Runtime/Engine/Private/KismetTextLibrary.cpp b/Engine/Source/Runtime/Engine/Private/KismetTextLibrary.cpp
--- a/Engine/Source/Runtime/Engine/Private/KismetTextLibrary.cpp
+++ b/Engine/Source/Runtime/Engine/Private/KismetTextLibrary.cpp
@@ -307,6 +307,12 @@ FText UKismetTextLibrary::AsTimespan_Timespan(const FTimespan& InTimespan)
 	return FText::AsTimespan(InTimespan);
 }
 
+FText UKismetTextLibrary::AsMemory(int64 NumBytes, TEnumAsByte<EMemoryUnitStandard> UnitStandard, bool bUseGrouping)
+{
+	uint64 UnsignedNumBytes = static_cast<uint64>(FMath::Max(0, NumBytes));
+	return FText::AsMemory(UnsignedNumBytes, bUseGrouping ? &FNumberFormattingOptions::DefaultWithGrouping() : &FNumberFormattingOptions::DefaultNoGrouping(), nullptr, UnitStandard);
+}
+
 FText UKismetTextLibrary::Format(FText InPattern, TArray<FFormatArgumentData> InArgs)
 {
 	return FTextFormatter::Format(MoveTemp(InPattern), MoveTemp(InArgs), false, false);
