Added a log at additive retarget bone transform function, to detect bone size differences between the retarget transforms array and the skeleton required bones
For performance and spam reasons it will log only once per session, the first case that it encounters

diff --git a/Engine/Source/Runtime/Engine/Private/Animation/AnimationRuntime.cpp b/Engine/Source/Runtime/Engine/Private/Animation/AnimationRuntime.cpp
--- a/Engine/Source/Runtime/Engine/Private/Animation/AnimationRuntime.cpp
+++ b/Engine/Source/Runtime/Engine/Private/Animation/AnimationRuntime.cpp
@@ -2711,17 +2711,34 @@ void FAnimationRuntime::RetargetBoneTransform(const USkeleton* SourceSkeleton, c
 					const FTransform& RefPoseTransform = RequiredBones.GetRefPoseTransform(BoneIndex);
 
 					// Remap the base pose onto the target skeleton so that we are working entirely in target space
-					FTransform BaseTransform = AuthoredOnRefSkeleton[SourceSkeletonBoneIndex];
-					if (SkeletonRemapping.RequiresReferencePoseRetarget())
+					if(AuthoredOnRefSkeleton.IsValidIndex(SourceSkeletonBoneIndex))
 					{
-						BaseTransform = SkeletonRemapping.RetargetBoneTransformToTargetSkeleton(TargetSkeletonBoneIndex, BaseTransform);
+						FTransform BaseTransform = AuthoredOnRefSkeleton[SourceSkeletonBoneIndex];
+						if (SkeletonRemapping.RequiresReferencePoseRetarget())
+						{
+							BaseTransform = SkeletonRemapping.RetargetBoneTransformToTargetSkeleton(TargetSkeletonBoneIndex, BaseTransform);
+						}
+
+						// Apply the retargeting as if it were an additive difference between the current skeleton and the retarget skeleton. 
+						BoneTransform.SetRotation(BoneTransform.GetRotation() * BaseTransform.GetRotation().Inverse() * RefPoseTransform.GetRotation());
+						BoneTransform.SetTranslation(BoneTransform.GetTranslation() + (RefPoseTransform.GetTranslation() - BaseTransform.GetTranslation()));
+						BoneTransform.SetScale3D(BoneTransform.GetScale3D() * (RefPoseTransform.GetScale3D() * BaseTransform.GetSafeScaleReciprocal(BaseTransform.GetScale3D())));
+						BoneTransform.NormalizeRotation();
 					}
+					else
+					{
+						static bool bLogOnce = true;
+						if (bLogOnce)	// If there are several assets broken, we will skip all but first. This is done to avoid log spam.
+						{
+							UE_LOG(LogAnimation, Error, TEXT("RetargetBoneTransform: Trying to retarget an additive bone (%s) from SourceName (%s) with Source Skeleton (%s) and Target Skeleton (%s) that does not exist at RetargetTransforms Array. Please check the animation and its Base Pose Aimation and re-save them."),
+								*SourceSkeleton->GetReferenceSkeleton().GetBoneName(SourceSkeletonBoneIndex).ToString()
+								, *SourceName.ToString()
+								, *SourceSkeleton->GetFullName()
+								, *TargetSkeleton->GetFullName());
 
-					// Apply the retargeting as if it were an additive difference between the current skeleton and the retarget skeleton. 
-					BoneTransform.SetRotation(BoneTransform.GetRotation() * BaseTransform.GetRotation().Inverse() * RefPoseTransform.GetRotation());
-					BoneTransform.SetTranslation(BoneTransform.GetTranslation() + (RefPoseTransform.GetTranslation() - BaseTransform.GetTranslation()));
-					BoneTransform.SetScale3D(BoneTransform.GetScale3D() * (RefPoseTransform.GetScale3D() * BaseTransform.GetSafeScaleReciprocal(BaseTransform.GetScale3D())));
-					BoneTransform.NormalizeRotation();
+							bLogOnce = false;
+						}
+					}
 				}
 				break;
 			}
