[NetPhysics] Predictive Interpolation, put a max clamp on how long to run predicted kinematic replication, using the extrapolation limit.
Also allow clearing the target state if replication say it's okay to clear it and the object can't simulate (is kinematic).

diff --git a/Engine/Source/Runtime/Engine/Private/PhysicsEngine/PhysicsReplication.cpp b/Engine/Source/Runtime/Engine/Private/PhysicsEngine/PhysicsReplication.cpp
--- a/Engine/Source/Runtime/Engine/Private/PhysicsEngine/PhysicsReplication.cpp
+++ b/Engine/Source/Runtime/Engine/Private/PhysicsEngine/PhysicsReplication.cpp
@@ -1615,8 +1615,9 @@ bool FPhysicsReplicationAsync::PredictiveInterpolation(Chaos::FPBDRigidParticleH
 
 		// --- Should replication stop? ---
 		const bool bClearTarget =
-			((bOkToClear && bShouldSleep && Target.AccumulatedSleepSeconds >= PhysicsReplicationCVars::PredictiveInterpolationCVars::SleepSecondsClearTarget) // Don't clear the target due to sleeping until the object both should sleep and is sleeping for n seconds
-			|| (bOkToClear && !bReplicatingPhysics))
+			((bOkToClear && bShouldSleep && Target.AccumulatedSleepSeconds >= PhysicsReplicationCVars::PredictiveInterpolationCVars::SleepSecondsClearTarget) // Allow clearing the target due to sleeping after the object has been sleeping for n seconds
+			|| (bOkToClear && !bReplicatingPhysics) // If replication say it's okay to clear the target and the object shouldn't replicate physics anymore, clear the target
+			|| (bOkToClear && !bCanSimulate)) // If replication say it's okay to clear the target and the object can't simulate, clear the target
 			&& !PhysicsReplicationCVars::PredictiveInterpolationCVars::bDontClearTarget;
 
 		// --- Target Prediction ---
@@ -1625,10 +1626,16 @@ bool FPhysicsReplicationAsync::PredictiveInterpolation(Chaos::FPBDRigidParticleH
 			const int32 ExtrapolationTickLimit = FMath::Max(
 				FMath::CeilToInt(Target.AverageReceiveInterval * PhysicsReplicationCVars::PredictiveInterpolationCVars::ExtrapolationTimeMultiplier), // Extrapolate time based on receive interval * multiplier
 				FMath::CeilToInt(PhysicsReplicationCVars::PredictiveInterpolationCVars::ExtrapolationMinTime / DeltaSeconds)); // At least extrapolate for N seconds
+
 			if (Target.TickCount <= ExtrapolationTickLimit)
 			{
 				FPhysicsReplicationAsync::ExtrapolateTarget(Target, 1, DeltaSeconds);
 			}
+			else
+			{
+				// If we reach the extrapolation limit, disable target from being altered
+				Target.bAllowTargetAltering = false;
+			}
 		}
 
 		return bClearTarget;
@@ -1745,7 +1752,7 @@ bool FPhysicsReplicationAsync::PredictiveInterpolation(Chaos::FPBDRigidParticleH
 		const bool bKinematicPrediction = PhysicsReplicationCVars::PredictiveInterpolationCVars::bKinematicPrediction;
 		const float InterpolationTicks = FMath::CeilToInt(Target.AverageReceiveInterval) - (RigidsSolver->GetCurrentFrame() - Target.ReceiveFrame);
 
-		if (bKinematicPrediction || InterpolationTicks > 0)
+		if ((bKinematicPrediction && Target.bAllowTargetAltering) || InterpolationTicks > 0)
 		{
 			/* Calculate the Lerp value for a smooth interpolation
 			* ------------------------------------------------------------------------------
