[NetPhysics] Clean up logic for how many extra inputs and states to send along the most recent input/state in NetworkPhysicsComponent, to cover for lost packets.

diff --git a/Engine/Source/Runtime/Engine/Private/PhysicsEngine/NetworkPhysicsComponent.cpp b/Engine/Source/Runtime/Engine/Private/PhysicsEngine/NetworkPhysicsComponent.cpp
--- a/Engine/Source/Runtime/Engine/Private/PhysicsEngine/NetworkPhysicsComponent.cpp
+++ b/Engine/Source/Runtime/Engine/Private/PhysicsEngine/NetworkPhysicsComponent.cpp
@@ -22,10 +22,10 @@ namespace PhysicsReplicationCVars
 {
 	namespace ResimulationCVars
 	{
-		int32 InputRedundancy = 3;
-		static FAutoConsoleVariableRef CVarResimInputRedundancy(TEXT("np2.Resim.InputRedundancy"), InputRedundancy, TEXT("How many inputs to send with each unreliable network message to account for packetloss."));
-		int32 StateRedundancy = 1;
-		static FAutoConsoleVariableRef CVarResimStateRedundancy(TEXT("np2.Resim.StateRedundancy"), StateRedundancy, TEXT("How many states to send with each unreliable network message to account for packetloss."));
+		int32 RedundantInputs = 2;
+		static FAutoConsoleVariableRef CVarResimRedundantInputs(TEXT("np2.Resim.RedundantInputs"), RedundantInputs, TEXT("How many extra inputs to send with each unreliable network message, to account for packetloss."));
+		int32 RedundantStates = 0;
+		static FAutoConsoleVariableRef CVarResimRedundantStates(TEXT("np2.Resim.RedundantStates"), RedundantStates, TEXT("How many extra states to send with each unreliable network message, to account for packetloss."));
 		bool bAllowRewindToClosestState = true;
 		static FAutoConsoleVariableRef CVarResimAllowRewindToClosestState(TEXT("np2.Resim.AllowRewindToClosestState"), bAllowRewindToClosestState, TEXT("When rewinding to a specific frame, if the client doens't have state data for that frame, use closest data available. Only affects the first rewind frame, when FPBDRigidsEvolution is set to Reset."));
 		bool bCompareStateToTriggerRewind = false;
@@ -423,13 +423,13 @@ UNetworkPhysicsComponent::UNetworkPhysicsComponent() : Super()
 
 void UNetworkPhysicsComponent::InitPhysics()
 {
-	if (const IConsoleVariable* CVarInputRedundancy = IConsoleManager::Get().FindConsoleVariable(TEXT("np2.Resim.InputRedundancy")))
+	if (const IConsoleVariable* CVarRedundantInputs = IConsoleManager::Get().FindConsoleVariable(TEXT("np2.Resim.RedundantInputs")))
 	{
-		InputRedundancy = CVarInputRedundancy->GetInt();
+		SetNumberOfInputsToNetwork(CVarRedundantInputs->GetInt());
 	}
-	if (const IConsoleVariable* CVarStateRedundancy = IConsoleManager::Get().FindConsoleVariable(TEXT("np2.Resim.StateRedundancy")))
+	if (const IConsoleVariable* CVarRedundantStates = IConsoleManager::Get().FindConsoleVariable(TEXT("np2.Resim.RedundantStates")))
 	{
-		StateRedundancy = CVarStateRedundancy->GetInt();
+		SetNumberOfStatesToNetwork(CVarRedundantStates->GetInt());
 	}
 
 	if (AActor* Owner = GetOwner())
@@ -474,19 +474,19 @@ void UNetworkPhysicsComponent::InitializeComponent()
 		SettingsComponent = Owner->FindComponentByClass<UNetworkPhysicsSettingsComponent>();
 		if (SettingsComponent)
 		{
-			InputRedundancy = SettingsComponent->NetworkPhysicsComponentSettings.GetRedundantInputs();
-			StateRedundancy = SettingsComponent->NetworkPhysicsComponentSettings.GetRedundantStates();
+			SetNumberOfInputsToNetwork(SettingsComponent->NetworkPhysicsComponentSettings.GetRedundantInputs());
+			SetNumberOfStatesToNetwork(SettingsComponent->NetworkPhysicsComponentSettings.GetRedundantStates());
 			bEnableUnreliableFlow = SettingsComponent->NetworkPhysicsComponentSettings.GetEnableUnreliableFlow();
 			bEnableReliableFlow = SettingsComponent->NetworkPhysicsComponentSettings.GetEnableReliableFlow();
 			bValidateDataOnGameThread = SettingsComponent->NetworkPhysicsComponentSettings.GetValidateDataOnGameThread();
 
 			if (ReplicatedInputs.History)
 			{
-				ReplicatedInputs.History->ResizeDataHistory(InputRedundancy);
+				ReplicatedInputs.History->ResizeDataHistory(InputsToNetwork);
 			}
 			if (ReplicatedStates.History)
 			{
-				ReplicatedStates.History->ResizeDataHistory(StateRedundancy);
+				ReplicatedStates.History->ResizeDataHistory(StatesToNetwork);
 			}
 		}
 	}
@@ -501,6 +501,8 @@ void UNetworkPhysicsComponent::InitializeComponent()
 				NetworkPhysicsComponent_Internal = Solver->CreateAndRegisterSimCallbackObject_External<FAsyncNetworkPhysicsComponent>();
 				NetworkPhysicsComponent_Internal->SettingsComponent = SettingsComponent ? SettingsComponent->GetNetworkPhysicsSettings_Internal() : nullptr;
 				NetworkPhysicsComponent_Internal->RootPhysicsObject = RootPhysicsObject;
+				NetworkPhysicsComponent_Internal->InputsToNetwork = InputsToNetwork;
+				NetworkPhysicsComponent_Internal->StatesToNetwork = StatesToNetwork;
 				CreateAsyncDataHistory();
 			}
 		}
@@ -682,7 +684,7 @@ void UNetworkPhysicsComponent::OnRep_SetReplicatedStates()
 	{
 		if (!AsyncInput->StateData)
 		{
-			AsyncInput->StateData = StateHelper->CreateUniqueRewindHistory(StateRedundancy);
+			AsyncInput->StateData = StateHelper->CreateUniqueRewindHistory(StatesToNetwork);
 		}
 
 		AsyncInput->StateData->ResetFast();
@@ -719,7 +721,7 @@ void UNetworkPhysicsComponent::OnRep_SetReplicatedInputs()
 	{
 		if (!AsyncInput->InputData)
 		{
-			AsyncInput->InputData = InputHelper->CreateUniqueRewindHistory(InputRedundancy);
+			AsyncInput->InputData = InputHelper->CreateUniqueRewindHistory(InputsToNetwork);
 		}
 
 		AsyncInput->InputData->ResetFast();
@@ -744,7 +746,7 @@ void UNetworkPhysicsComponent::ServerReceiveInputData_Implementation(const FNetw
 	{
 		if (!AsyncInput->InputData)
 		{
-			AsyncInput->InputData = InputHelper->CreateUniqueRewindHistory(InputRedundancy);
+			AsyncInput->InputData = InputHelper->CreateUniqueRewindHistory(InputsToNetwork);
 		}
 
 		// Validate data in the received inputs
@@ -1253,20 +1255,19 @@ void FAsyncNetworkPhysicsComponent::ConsumeAsyncInput(const int32 PhysicsStep)
 FAsyncNetworkPhysicsComponentOutput& FAsyncNetworkPhysicsComponent::GetAsyncOutput_Internal()
 {
 	FAsyncNetworkPhysicsComponentOutput& AsyncOutput = GetProducerOutputData_Internal();
-	const FNetworkPhysicsSettingsNetworkPhysicsComponent& ComponentSettings = GetComponentSettings();
 
 	// InputData marshal from PT to GT is needed for: LocallyControlled and Server
 	if ((IsLocallyControlled() || IsServer()) && AsyncOutput.InputData == nullptr && InputHistory != nullptr)
 	{
 		AsyncOutput.InputData = InputHistory->CreateNew();
-		AsyncOutput.InputData->ResizeDataHistory(ComponentSettings.GetRedundantInputs());
+		AsyncOutput.InputData->ResizeDataHistory(InputsToNetwork);
 	}
 
 	// StateData marshal from PT to GT is needed for: Server
 	if (IsServer() && AsyncOutput.StateData == nullptr && StateHistory != nullptr)
 	{
 		AsyncOutput.StateData = StateHistory->CreateNew();
-		AsyncOutput.StateData->ResizeDataHistory(ComponentSettings.GetRedundantStates());
+		AsyncOutput.StateData->ResizeDataHistory(StatesToNetwork);
 	}
 
 	return AsyncOutput;
@@ -1524,11 +1525,10 @@ void FAsyncNetworkPhysicsComponent::SendInputData_Internal(FAsyncNetworkPhysicsC
 		// -- Default / Unreliable Flow --
 		if (ComponentSettings.GetEnableUnreliableFlow())
 		{
-			const uint8 InputsToSend = ComponentSettings.GetRedundantInputs();
-			const int32 FromFrame = FMath::Max(0, ToFrame - InputsToSend - 1); // Remove 1 since both ToFrame and FromFrame are inclusive
+			const int32 FromFrame = FMath::Max(0, ToFrame - InputsToNetwork - 1); // Remove 1 since both ToFrame and FromFrame are inclusive
 
 			// Resize marshaling history if needed
-			AsyncOutput.InputData->ResizeDataHistory(InputsToSend);
+			AsyncOutput.InputData->ResizeDataHistory(InputsToNetwork);
 			
 			if (InputHistory->CopyData(*AsyncOutput.InputData, FromFrame, ToFrame, /*bIncludeUnimportant*/ true, /*bIncludeImportant*/ ComponentSettings.GetEnableReliableFlow() == false))
 			{
@@ -1590,11 +1590,10 @@ void FAsyncNetworkPhysicsComponent::SendStateData_Internal(FAsyncNetworkPhysicsC
 		// -- Default / Unreliable Flow --
 		if (ComponentSettings.GetEnableUnreliableFlow())
 		{
-			const uint8 StatesToSend = ComponentSettings.GetRedundantStates();
-			const int32 FromFrame = FMath::Max(0, ToFrame - StatesToSend - 1); // Remove 1 since both ToFrame and FromFrame are inclusive
+			const int32 FromFrame = FMath::Max(0, ToFrame - StatesToNetwork - 1); // Remove 1 since both ToFrame and FromFrame are inclusive
 
 			// Resize marshaling history if needed
-			AsyncOutput.StateData->ResizeDataHistory(StatesToSend);
+			AsyncOutput.StateData->ResizeDataHistory(StatesToNetwork);
 
 			if (StateHistory->CopyData(*AsyncOutput.StateData, FromFrame, ToFrame, /*bIncludeUnimportant*/ true, /*bIncludeImportant*/ ComponentSettings.GetEnableReliableFlow() == false))
 			{
diff --git a/Engine/Source/Runtime/Engine/Public/Physics/NetworkPhysicsComponent.h b/Engine/Source/Runtime/Engine/Public/Physics/NetworkPhysicsComponent.h
--- a/Engine/Source/Runtime/Engine/Public/Physics/NetworkPhysicsComponent.h
+++ b/Engine/Source/Runtime/Engine/Public/Physics/NetworkPhysicsComponent.h
@@ -816,6 +816,18 @@ private:
 	/** Setup data type and history in async input and trigger the creation of data history in the async component */
 	ENGINE_API void CreateAsyncDataHistory();
 
+	/** Set number of inputs to send over the network with each message, clamped to 1 as minimum */
+	void SetNumberOfInputsToNetwork(uint16 RedundantInputs)
+	{
+		InputsToNetwork = FMath::Max((uint16)(1 + RedundantInputs), (uint16)1);
+	}
+
+	/** Set number of states to send over the network with each message, clamped to 1 as minimum */
+	void SetNumberOfStatesToNetwork(uint16 RedundantStates)
+	{
+		StatesToNetwork = FMath::Max((uint16)(1 + RedundantStates), (uint16)1);
+	}
+
 public:
 	/** Register and create both state and input to be both networked and cached in history */
 	template<typename PhysicsTraits>
@@ -851,10 +863,10 @@ private:
 	TUniquePtr<FNetworkPhysicsDataHelper> StateHelper;
 
 	// Send last N number of inputs each replication call to patch up holes due to packet loss
-	int8 InputRedundancy = 3;
+	uint16 InputsToNetwork = 3;
 
 	// Send last N number of states each replication call to patch up holes due to packet loss
-	int8 StateRedundancy = 1;
+	uint16 StatesToNetwork = 1;
 
 	// Actor component that will be used to fill the histories
 	TObjectPtr<UActorComponent> ActorComponent;
@@ -974,10 +986,10 @@ FORCEINLINE void UNetworkPhysicsComponent::CreateDataHistory(UActorComponent* Hi
 	InputHelper = MakeUnique<TNetworkPhysicsDataHelper<typename PhysicsTraits::InputsType>>();
 	StateHelper = MakeUnique<TNetworkPhysicsDataHelper<typename PhysicsTraits::StatesType>>();
 
-	ReplicatedInputs.History = MakeUnique<TNetRewindHistory<typename PhysicsTraits::InputsType>>(InputRedundancy);
+	ReplicatedInputs.History = MakeUnique<TNetRewindHistory<typename PhysicsTraits::InputsType>>(InputsToNetwork);
 	ReplicatedInputs.Owner = this;
 
-	ReplicatedStates.History = MakeUnique<TNetRewindHistory<typename PhysicsTraits::StatesType>>(StateRedundancy);
+	ReplicatedStates.History = MakeUnique<TNetRewindHistory<typename PhysicsTraits::StatesType>>(StatesToNetwork);
 	ReplicatedStates.Owner = this;
 
 	ReplicatedImportantInput.History = MakeUnique<TNetRewindHistory<typename PhysicsTraits::InputsType>>(1);
@@ -996,7 +1008,7 @@ FORCEINLINE void UNetworkPhysicsComponent::CreateInputHistory(UActorComponent* H
 {
 	InputHelper = MakeUnique<TNetworkPhysicsDataHelper<InputsType>>();
 
-	ReplicatedInputs.History = MakeUnique<TNetRewindHistory<InputsType>>(InputRedundancy);
+	ReplicatedInputs.History = MakeUnique<TNetRewindHistory<InputsType>>(InputsToNetwork);
 	ReplicatedInputs.Owner = this;
 
 	ReplicatedImportantInput.History = MakeUnique<TNetRewindHistory<InputsType>>(1);
@@ -1197,6 +1209,12 @@ private:
 	// Local temporary states data used by pre/post process inputs functions
 	TUniquePtr<FNetworkPhysicsData> StateData;
 
+	// Send last N number of inputs each replication call to patch up holes due to packet loss
+	uint16 InputsToNetwork = 3;
+
+	// Send last N number of states each replication call to patch up holes due to packet loss
+	uint16 StatesToNetwork = 1;
+
 	FDelegateHandle DelegateOnPreProcessInputs_Internal;
 	FDelegateHandle DelegateOnPostProcessInputs_Internal;
 };
diff --git a/Engine/Source/Runtime/Engine/Public/Physics/NetworkPhysicsSettingsComponent.h b/Engine/Source/Runtime/Engine/Public/Physics/NetworkPhysicsSettingsComponent.h
--- a/Engine/Source/Runtime/Engine/Public/Physics/NetworkPhysicsSettingsComponent.h
+++ b/Engine/Source/Runtime/Engine/Public/Physics/NetworkPhysicsSettingsComponent.h
@@ -67,8 +67,8 @@ namespace PhysicsReplicationCVars
 		extern bool bApplyDataInsteadOfMergeData;
 		extern bool bAllowInputExtrapolation;
 		extern bool bValidateDataOnGameThread;
-		extern int32 InputRedundancy;
-		extern int32 StateRedundancy;
+		extern int32 RedundantInputs;
+		extern int32 RedundantStates;
 		extern bool bCompareStateToTriggerRewind;
 		extern bool bCompareInputToTriggerRewind;
 	}
@@ -309,17 +309,17 @@ struct FNetworkPhysicsSettingsNetworkPhysicsComponent
 
 	UPROPERTY(config, EditDefaultsOnly, Category = "Overrides", Meta = (InlineEditConditionToggle))
 	uint32 bOverrideRedundantInputs : 1;
-	// Overrides CVar: np2.Resim.InputRedundancy -- How many inputs to send with each unreliable network message to account for packetloss.
+	// Overrides CVar: np2.Resim.RedundantInputs -- How many extra inputs to send with each unreliable network message, to account for packetloss.
 	UPROPERTY(config, EditDefaultsOnly, Category = "Overrides", Meta = (EditCondition = "bOverrideRedundantInputs"))
-	uint8 RedundantInputs = PhysicsReplicationCVars::ResimulationCVars::InputRedundancy;
-	const uint8 GetRedundantInputs() const { return bOverrideRedundantInputs ? RedundantInputs : PhysicsReplicationCVars::ResimulationCVars::InputRedundancy; }
+	uint16 RedundantInputs = PhysicsReplicationCVars::ResimulationCVars::RedundantInputs;
+	const uint16 GetRedundantInputs() const { return bOverrideRedundantInputs ? RedundantInputs : PhysicsReplicationCVars::ResimulationCVars::RedundantInputs; }
 
 	UPROPERTY(config, EditDefaultsOnly, Category = "Overrides", Meta = (InlineEditConditionToggle))
 	uint32 bOverrideRedundantStates : 1;
-	// Overrides CVar: np2.Resim.StateRedundancy -- How many states to send with each unreliable network message to account for packetloss.
+	// Overrides CVar: np2.Resim.RedundantStates -- How many extra states to send with each unreliable network message, to account for packetloss.
 	UPROPERTY(config, EditDefaultsOnly, Category = "Overrides", Meta = (EditCondition = "bOverrideRedundantStates"))
-	uint8 RedundantStates = PhysicsReplicationCVars::ResimulationCVars::StateRedundancy;
-	const uint8 GetRedundantStates() const { return bOverrideRedundantStates ? RedundantStates : PhysicsReplicationCVars::ResimulationCVars::StateRedundancy; }
+	uint16 RedundantStates = PhysicsReplicationCVars::ResimulationCVars::RedundantStates;
+	const uint16 GetRedundantStates() const { return bOverrideRedundantStates ? RedundantStates : PhysicsReplicationCVars::ResimulationCVars::RedundantStates; }
 
 	UPROPERTY(config, EditDefaultsOnly, Category = "Overrides", Meta = (InlineEditConditionToggle))
 	uint32 bOverrideCompareStateToTriggerRewind : 1;
