Fix hidden pose watches being drawn in PIE
* Debugging an ABP draws the pose watches of linked anim graphs.
* It can be confusing why a pose watch is shown and where it lives as the pose watches are not shown in the pose watch manager.
* I've chosen against showing linked pose watches in the pose watch manager as if layers are swapped out at runtime the problem of trying to track down where the pose watches live remains.
* The solution chosen is to only show PIE pose watches for debugged ABPs that're actively open in an editor, i.e. if a pose watches in PIE it should be very clear which ABP asset they originate from.

diff --git a/Engine/Source/Runtime/Engine/Private/Components/SkeletalMeshComponent.cpp b/Engine/Source/Runtime/Engine/Private/Components/SkeletalMeshComponent.cpp
--- a/Engine/Source/Runtime/Engine/Private/Components/SkeletalMeshComponent.cpp
+++ b/Engine/Source/Runtime/Engine/Private/Components/SkeletalMeshComponent.cpp
@@ -47,6 +47,8 @@
 #include "Engine/PoseWatch.h"
 #include "Settings/AnimBlueprintSettings.h"
 #include "SkeletalMeshComponentInstanceData.h"
+#include "Editor.h"
+#include "Subsystems/AssetEditorSubsystem.h"
 #endif
 
 #include UE_INLINE_GENERATED_CPP_BY_NAME(SkeletalMeshComponent)
@@ -4953,9 +4955,14 @@ void USkeletalMeshComponent::UpdatePoseWatches()
 		{
 			if (const UAnimBlueprintGeneratedClass* AnimBPGenClass = Cast<UAnimBlueprintGeneratedClass>(InAnimInstance->GetClass()))
 			{
-				if (const UAnimBlueprint* AnimBlueprint = Cast<UAnimBlueprint>(AnimBPGenClass->ClassGeneratedBy))
+				if (UAnimBlueprint* AnimBlueprint = Cast<UAnimBlueprint>(AnimBPGenClass->ClassGeneratedBy))
 				{
-					const UAnimBlueprint* RootAnimBlueprint = UAnimBlueprint::FindRootAnimBlueprint(AnimBlueprint);
+					if (!GEditor->GetEditorSubsystem<UAssetEditorSubsystem>()->FindEditorForAsset(AnimBlueprint, false))
+					{
+						return;
+					}
+
+					UAnimBlueprint* RootAnimBlueprint = UAnimBlueprint::FindRootAnimBlueprint(AnimBlueprint);
 					AnimBlueprint = RootAnimBlueprint ? RootAnimBlueprint : AnimBlueprint;
 
 					const FAnimBlueprintDebugData& DebugData = AnimBlueprint->GetAnimBlueprintGeneratedClass()->GetAnimBlueprintDebugData();
