Added extra number formatting options to the BP version of FText::AsMemory

#rnx

diff --git a/Engine/Source/Runtime/Engine/Classes/Kismet/KismetTextLibrary.h b/Engine/Source/Runtime/Engine/Classes/Kismet/KismetTextLibrary.h
--- a/Engine/Source/Runtime/Engine/Classes/Kismet/KismetTextLibrary.h
+++ b/Engine/Source/Runtime/Engine/Classes/Kismet/KismetTextLibrary.h
@@ -299,7 +299,7 @@ class UKismetTextLibrary : public UBlueprintFunctionLibrary
 
 	/** Generate an FText that represents the passed number as a memory size in the current culture */
 	UFUNCTION(BlueprintPure, meta=(AdvancedDisplay = "1"), Category="Utilities|Text")
-	static ENGINE_API FText AsMemory(int64 NumBytes, TEnumAsByte<EMemoryUnitStandard> UnitStandard = EMemoryUnitStandard::IEC, bool bUseGrouping = true);
+	static ENGINE_API FText AsMemory(int64 NumBytes, TEnumAsByte<EMemoryUnitStandard> UnitStandard = EMemoryUnitStandard::IEC, bool bUseGrouping = true, int32 MinimumIntegralDigits = 1, int32 MaximumIntegralDigits = 324, int32 MinimumFractionalDigits = 0, int32 MaximumFractionalDigits = 3);
 
 	/** Used for formatting text using the FText::Format function and utilized by the UK2Node_FormatText */
 	UFUNCTION(BlueprintPure, meta=(BlueprintInternalUseOnly = "true"))
diff --git a/Engine/Source/Runtime/Engine/Private/KismetTextLibrary.cpp b/Engine/Source/Runtime/Engine/Private/KismetTextLibrary.cpp
--- a/Engine/Source/Runtime/Engine/Private/KismetTextLibrary.cpp
+++ b/Engine/Source/Runtime/Engine/Private/KismetTextLibrary.cpp
@@ -307,10 +307,17 @@ FText UKismetTextLibrary::AsTimespan_Timespan(const FTimespan& InTimespan)
 	return FText::AsTimespan(InTimespan);
 }
 
-FText UKismetTextLibrary::AsMemory(int64 NumBytes, TEnumAsByte<EMemoryUnitStandard> UnitStandard, bool bUseGrouping)
+FText UKismetTextLibrary::AsMemory(int64 NumBytes, TEnumAsByte<EMemoryUnitStandard> UnitStandard, bool bUseGrouping, int32 MinimumIntegralDigits, int32 MaximumIntegralDigits, int32 MinimumFractionalDigits, int32 MaximumFractionalDigits)
 {
+	FNumberFormattingOptions NumberFormatOptions;
+	NumberFormatOptions.UseGrouping = bUseGrouping;
+	NumberFormatOptions.MinimumIntegralDigits = MinimumIntegralDigits;
+	NumberFormatOptions.MaximumIntegralDigits = MaximumIntegralDigits;
+	NumberFormatOptions.MinimumFractionalDigits = MinimumFractionalDigits;
+	NumberFormatOptions.MaximumFractionalDigits = MaximumFractionalDigits;
+
 	uint64 UnsignedNumBytes = static_cast<uint64>(FMath::Max(0, NumBytes));
-	return FText::AsMemory(UnsignedNumBytes, bUseGrouping ? &FNumberFormattingOptions::DefaultWithGrouping() : &FNumberFormattingOptions::DefaultNoGrouping(), nullptr, UnitStandard);
+	return FText::AsMemory(UnsignedNumBytes, &NumberFormatOptions, nullptr, UnitStandard);
 }
 
 FText UKismetTextLibrary::Format(FText InPattern, TArray<FFormatArgumentData> InArgs)
