Alpha: Reorganization/cleanup of alpha settings.

* Moved settings to be adjancent as requested (into the Default category)
* Renamed "Enable Alpha Output" to "Alpha Output"
* Linked "Support Primitive Alpha Holdout" setting's enabled state to "Alpha Output"

diff --git a/Engine/Plugins/Experimental/Avalanche/Source/AvalancheMediaEditor/Private/Rundown/Preview/SAvaRundownPagePreview.cpp b/Engine/Plugins/Experimental/Avalanche/Source/AvalancheMediaEditor/Private/Rundown/Preview/SAvaRundownPagePreview.cpp
--- a/Engine/Plugins/Experimental/Avalanche/Source/AvalancheMediaEditor/Private/Rundown/Preview/SAvaRundownPagePreview.cpp
+++ b/Engine/Plugins/Experimental/Avalanche/Source/AvalancheMediaEditor/Private/Rundown/Preview/SAvaRundownPagePreview.cpp
@@ -308,7 +308,7 @@ void SAvaRundownPagePreview::HandleCheckerboardActionExecute() const
 		{
 			const FText NotificationText = LOCTEXT("AlphaSupport",
 				"An output requested Alpha Support but the required project setting is not enabled!\n"
-				"Go to Project Settings > Rendering > PostProcessing > 'Enable Alpha Output' and enable it.");
+				"Go to Project Settings > Rendering > PostProcessing > 'Alpha Output' and enable it.");
 
 			FNotificationInfo Info(NotificationText);
 			Info.ExpireDuration = 5.0f;	// The message is long, need more time to read it.
diff --git a/Engine/Plugins/Experimental/Compositing/HoldoutComposite/Source/HoldoutComposite/Private/HoldoutCompositeSubsystem.cpp b/Engine/Plugins/Experimental/Compositing/HoldoutComposite/Source/HoldoutComposite/Private/HoldoutCompositeSubsystem.cpp
--- a/Engine/Plugins/Experimental/Compositing/HoldoutComposite/Source/HoldoutComposite/Private/HoldoutCompositeSubsystem.cpp
+++ b/Engine/Plugins/Experimental/Compositing/HoldoutComposite/Source/HoldoutComposite/Private/HoldoutCompositeSubsystem.cpp
@@ -82,7 +82,7 @@ bool UHoldoutCompositeSubsystem::ValidateProjectSettings()
 		// We inform the user and offer them the option to activate project settings.
 		UE_CALL_ONCE([&]{ PrimitiveHoldoutSettingsNotification(RendererSettings); });
 #else
-		UE_CALL_ONCE([&]{ UE_LOG(LogHoldoutComposite, Warning, TEXT("Both \"Enable Alpha Output\" and \"Support Primitive Alpha Holdout\" project settings must be enabled for holdout composite.")); });
+		UE_CALL_ONCE([&]{ UE_LOG(LogHoldoutComposite, Warning, TEXT("Both \"Alpha Output\" and \"Support Primitive Alpha Holdout\" project settings must be enabled for holdout composite.")); });
 #endif
 	}
 
@@ -92,7 +92,7 @@ bool UHoldoutCompositeSubsystem::ValidateProjectSettings()
 #if WITH_EDITOR
 void UHoldoutCompositeSubsystem::PrimitiveHoldoutSettingsNotification(class URendererSettings* RendererSettings)
 {
-	const FText HoldoutText = LOCTEXT("HoldoutSettingPrompt", "The following project settings must be enabled for holdout composite:\n- Enable Alpha Output\n- Support Primitive Alpha Holdout");
+	const FText HoldoutText = LOCTEXT("HoldoutSettingPrompt", "The following project settings must be enabled for holdout composite:\n- Alpha Output\n- Support Primitive Alpha Holdout");
 	const FText HoldoutConfirmText = LOCTEXT("HoldoutSettingConfirm", "Update");
 	const FText HoldoutCancelText = LOCTEXT("HoldoutSettingCancel", "Not Now");
 
diff --git a/Engine/Plugins/MovieScene/MovieRenderPipeline/Source/MovieRenderPipelineCore/Private/Graph/MovieGraphRenderLayerSubsystem.cpp b/Engine/Plugins/MovieScene/MovieRenderPipeline/Source/MovieRenderPipelineCore/Private/Graph/MovieGraphRenderLayerSubsystem.cpp
--- a/Engine/Plugins/MovieScene/MovieRenderPipeline/Source/MovieRenderPipelineCore/Private/Graph/MovieGraphRenderLayerSubsystem.cpp
+++ b/Engine/Plugins/MovieScene/MovieRenderPipeline/Source/MovieRenderPipelineCore/Private/Graph/MovieGraphRenderLayerSubsystem.cpp
@@ -197,7 +197,7 @@ void UMovieGraphRenderPropertyModifier::ApplyModifier(const UWorld* World)
 	if (bHoldout && (RendererSettings->bEnableAlphaChannelInPostProcessing == false))
 	{
 		// TODO: Ideally this is called in a general-purpose validation step instead, but that framework does not exist yet.
-		UE_LOG(LogMovieRenderPipeline, Warning, TEXT("A modifier with 'Holdout' is active, but 'Enable Alpha Output' in "
+		UE_LOG(LogMovieRenderPipeline, Warning, TEXT("A modifier with 'Holdout' is active, but 'Alpha Output' in "
 													 "the project's Rendering settings is disabled: holdout will not work properly."));
 	}
 	
diff --git a/Engine/Plugins/MovieScene/MovieRenderPipeline/Source/MovieRenderPipelineCore/Private/MoviePipeline.cpp b/Engine/Plugins/MovieScene/MovieRenderPipeline/Source/MovieRenderPipelineCore/Private/MoviePipeline.cpp
--- a/Engine/Plugins/MovieScene/MovieRenderPipeline/Source/MovieRenderPipelineCore/Private/MoviePipeline.cpp
+++ b/Engine/Plugins/MovieScene/MovieRenderPipeline/Source/MovieRenderPipelineCore/Private/MoviePipeline.cpp
@@ -122,7 +122,7 @@ void UMoviePipeline::ValidateSequenceAndSettings() const
 
 		if (bAnyOutputWantsAlpha && TonemapAlphaCVar->GetBool() == false)
 		{
-			UE_LOG(LogMovieRenderPipeline, Warning, TEXT("An output requested Alpha Support but the required project setting is not enabled! Go to Project Settings > Rendering > PostProcessing > 'Enable Alpha Output' and set it to 'True'."));
+			UE_LOG(LogMovieRenderPipeline, Warning, TEXT("An output requested Alpha Support but the required project setting is not enabled! Go to Project Settings > Rendering > PostProcessing > 'Alpha Output' and set it to 'True'."));
 		}
 	}
 }
diff --git a/Engine/Plugins/MovieScene/MovieRenderPipeline/Source/MovieRenderPipelineCore/Public/Graph/MovieGraphRenderLayerSubsystem.h b/Engine/Plugins/MovieScene/MovieRenderPipeline/Source/MovieRenderPipelineCore/Public/Graph/MovieGraphRenderLayerSubsystem.h
--- a/Engine/Plugins/MovieScene/MovieRenderPipeline/Source/MovieRenderPipelineCore/Public/Graph/MovieGraphRenderLayerSubsystem.h
+++ b/Engine/Plugins/MovieScene/MovieRenderPipeline/Source/MovieRenderPipelineCore/Public/Graph/MovieGraphRenderLayerSubsystem.h
@@ -877,7 +877,7 @@ public:
 
 	/**
 	 * If true, the primitive will render black with an alpha of 0, but all secondary effects (shadows, reflections,
-	 * indirect lighting) remain. This feature requires activating the project setting(s) "Enable Alpha Output", and "Support Primitive Alpha Holdout" if using the deferred renderer.
+	 * indirect lighting) remain. This feature requires activating the project setting(s) "Alpha Output", and "Support Primitive Alpha Holdout" if using the deferred renderer.
 	 */
 	UPROPERTY(BlueprintReadWrite, EditAnywhere, Category = "Settings", meta = (EditCondition = "bOverride_bHoldout"))
 	uint8 bHoldout : 1;
diff --git a/Engine/Plugins/MovieScene/MovieRenderPipeline/Source/MovieRenderPipelineRenderPasses/Public/MoviePipelineDeferredPasses.h b/Engine/Plugins/MovieScene/MovieRenderPipeline/Source/MovieRenderPipelineRenderPasses/Public/MoviePipelineDeferredPasses.h
--- a/Engine/Plugins/MovieScene/MovieRenderPipeline/Source/MovieRenderPipelineRenderPasses/Public/MoviePipelineDeferredPasses.h
+++ b/Engine/Plugins/MovieScene/MovieRenderPipeline/Source/MovieRenderPipelineRenderPasses/Public/MoviePipelineDeferredPasses.h
@@ -101,7 +101,7 @@ protected:
 public:
 	/**
 	* Should multiple temporal/spatial samples accumulate the alpha channel? This requires r.PostProcessing.PropagateAlpha
-	* to be enabled (see "Enable Alpha Output" under Project Settings > Rendering). This adds
+	* to be enabled (see "Alpha Output" under Project Settings > Rendering). This adds
 	* ~30% cost to the accumulation so you should not enable it unless necessary. You must delete both the sky and fog to ensure
 	* that they do not make all pixels opaque.
 	*/
diff --git a/Engine/Source/Runtime/Engine/Classes/Components/ExponentialHeightFogComponent.h b/Engine/Source/Runtime/Engine/Classes/Components/ExponentialHeightFogComponent.h
--- a/Engine/Source/Runtime/Engine/Classes/Components/ExponentialHeightFogComponent.h
+++ b/Engine/Source/Runtime/Engine/Classes/Components/ExponentialHeightFogComponent.h
@@ -213,7 +213,7 @@ class UExponentialHeightFogComponent : public USceneComponent
 	UPROPERTY(EditAnywhere, BlueprintReadOnly, Category = VolumetricFog, AdvancedDisplay)
 	bool bOverrideLightColorsWithFogInscatteringColors;
 
-	/** If this is True, this primitive will render black with an alpha of 0, but all secondary effects (shadows, reflections, indirect lighting) remain. This feature requires activating the project setting(s) "Enable Alpha Output", and "Support Primitive Alpha Holdout" if using the deferred renderer. */
+	/** If this is True, this primitive will render black with an alpha of 0, but all secondary effects (shadows, reflections, indirect lighting) remain. This feature requires activating the project setting(s) "Alpha Output", and "Support Primitive Alpha Holdout" if using the deferred renderer. */
 	UPROPERTY(EditAnywhere, AdvancedDisplay, BlueprintReadOnly, Category = Rendering, Interp)
 	uint8 bHoldout : 1;
 
diff --git a/Engine/Source/Runtime/Engine/Classes/Components/PrimitiveComponent.h b/Engine/Source/Runtime/Engine/Classes/Components/PrimitiveComponent.h
--- a/Engine/Source/Runtime/Engine/Classes/Components/PrimitiveComponent.h
+++ b/Engine/Source/Runtime/Engine/Classes/Components/PrimitiveComponent.h
@@ -428,7 +428,7 @@ public:
 	UPROPERTY(EditAnywhere, AdvancedDisplay, BlueprintReadOnly, Category=Rendering)
 	uint8 bReceivesDecals:1;
 
-	/** If this is True, this primitive will render black with an alpha of 0, but all secondary effects (shadows, reflections, indirect lighting) remain. This feature requires activating the project setting(s) "Enable Alpha Output", and "Support Primitive Alpha Holdout" if using the deferred renderer. */
+	/** If this is True, this primitive will render black with an alpha of 0, but all secondary effects (shadows, reflections, indirect lighting) remain. This feature requires activating the project setting(s) "Alpha Output", and "Support Primitive Alpha Holdout" if using the deferred renderer. */
 	UPROPERTY(EditAnywhere, AdvancedDisplay, BlueprintReadOnly, Category = Rendering, Interp)
 	uint8 bHoldout : 1;
 
diff --git a/Engine/Source/Runtime/Engine/Classes/Components/SkyAtmosphereComponent.h b/Engine/Source/Runtime/Engine/Classes/Components/SkyAtmosphereComponent.h
--- a/Engine/Source/Runtime/Engine/Classes/Components/SkyAtmosphereComponent.h
+++ b/Engine/Source/Runtime/Engine/Classes/Components/SkyAtmosphereComponent.h
@@ -160,7 +160,7 @@ class USkyAtmosphereComponent : public USceneComponent
 	UPROPERTY(EditAnywhere, BlueprintReadOnly, interp, Category = "Art Direction", meta = (UIMin = 0.001f, UIMax = 10.0f, ClampMin = 0.001f))
 	float AerialPerspectiveStartDepth;
 
-	/** If this is True, this primitive will render black with an alpha of 0, but all secondary effects (shadows, reflections, indirect lighting) remain. This feature requires activating the project setting(s) "Enable Alpha Output", and "Support Primitive Alpha Holdout" if using the deferred renderer. */
+	/** If this is True, this primitive will render black with an alpha of 0, but all secondary effects (shadows, reflections, indirect lighting) remain. This feature requires activating the project setting(s) "Alpha Output", and "Support Primitive Alpha Holdout" if using the deferred renderer. */
 	UPROPERTY(EditAnywhere, AdvancedDisplay, BlueprintReadOnly, Category = Rendering, Interp)
 	uint8 bHoldout : 1;
 
diff --git a/Engine/Source/Runtime/Engine/Classes/Components/VolumetricCloudComponent.h b/Engine/Source/Runtime/Engine/Classes/Components/VolumetricCloudComponent.h
--- a/Engine/Source/Runtime/Engine/Classes/Components/VolumetricCloudComponent.h
+++ b/Engine/Source/Runtime/Engine/Classes/Components/VolumetricCloudComponent.h
@@ -139,7 +139,7 @@ class UVolumetricCloudComponent : public USceneComponent
 	UPROPERTY(EditAnywhere, BlueprintReadOnly, interp, Category = "Art Direction", meta = (UIMin = 0.0, UIMax = 100.0, ClampMin = 0.0, SliderExponent = 2.0))
 	float AerialPespectiveMieScatteringFadeDistance;
 
-	/** If this is True, this primitive will render black with an alpha of 0, but all secondary effects (shadows, reflections, indirect lighting) remain. This feature requires activating the project setting(s) "Enable Alpha Output", and "Support Primitive Alpha Holdout" if using the deferred renderer. */
+	/** If this is True, this primitive will render black with an alpha of 0, but all secondary effects (shadows, reflections, indirect lighting) remain. This feature requires activating the project setting(s) "Alpha Output", and "Support Primitive Alpha Holdout" if using the deferred renderer. */
 	UPROPERTY(EditAnywhere, AdvancedDisplay, BlueprintReadOnly, Category = Rendering, Interp)
 	uint8 bHoldout : 1;
 
diff --git a/Engine/Source/Runtime/Engine/Classes/Engine/RendererSettings.h b/Engine/Source/Runtime/Engine/Classes/Engine/RendererSettings.h
--- a/Engine/Source/Runtime/Engine/Classes/Engine/RendererSettings.h
+++ b/Engine/Source/Runtime/Engine/Classes/Engine/RendererSettings.h
@@ -740,11 +740,18 @@ class URendererSettings : public UDeveloperSettings
 		ToolTip = "Whether the custom depth pass has the TemporalAA jitter enabled. Disabling this can be useful when the result of the CustomDepth Pass is used after TAA (e.g. after Tonemapping)"))
 	uint32 bCustomDepthTaaJitter : 1;
 
-	UPROPERTY(config, EditAnywhere, Category = Postprocessing, meta = (
-		ConsoleVariable = "r.PostProcessing.PropagateAlpha", DisplayName = "Enable Alpha Output",
+	UPROPERTY(config, EditAnywhere, Category = DefaultSettings, meta = (
+		ConsoleVariable = "r.PostProcessing.PropagateAlpha", DisplayName = "Alpha Output",
 		ToolTip = "Enable r.PostProcessing.PropagateAlpha to enforce alpha in scene color (overriding r.SceneColorFormat if necessary) and propagate it through the renderer's post-processing chain. The legacy \"Linear color space only\" and \"Allow through tonemapper\" options now map to True and the engine keeps full alpha precision throughout post-processing. This feature can now be toggled without an engine restart."))
 	uint32 bEnableAlphaChannelInPostProcessing : 1;
 
+	UPROPERTY(config, EditAnywhere, Category = DefaultSettings, meta = (
+		EditCondition = "bEnableAlphaChannelInPostProcessing",
+		ConsoleVariable = "r.Deferred.SupportPrimitiveAlphaHoldout", DisplayName = "Support Primitive Alpha Holdout (Deferred)",
+		ToolTip = "Enable primitive alpha holdout support in multiple deferred renderer passes. If primitive holdout masks are not to be used, keep this setting disabled for increased performance. Requires \"Alpha Output\" to be enabled.",
+		ConfigRestartRequired = true))
+	uint32 bDeferredSupportPrimitiveAlphaHoldout : 1;
+
 	UPROPERTY(config, EditAnywhere, Category = DefaultSettings, meta = (
 		ConsoleVariable = "r.DefaultFeature.Bloom", DisplayName = "Bloom",
 		ToolTip = "Whether the default for Bloom is enabled or not (postprocess volume/camera/game setting can still override and enable or disable it independently)"))
@@ -996,12 +1003,6 @@ class URendererSettings : public UDeveloperSettings
 		ConfigRestartRequired = true))
 		uint32 bSupportCloudShadowOnForwardLitTranslucent : 1;
 
-	UPROPERTY(config, EditAnywhere, Category = Optimizations, meta = (
-		ConsoleVariable = "r.Deferred.SupportPrimitiveAlphaHoldout", DisplayName = "Support Primitive Alpha Holdout (Deferred)",
-		ToolTip = "Enable primitive alpha holdout support in multiple deferred renderer passes. If primitive holdout masks are not to be used, keep this setting disabled for increased performance.",
-		ConfigRestartRequired = true))
-	uint32 bDeferredSupportPrimitiveAlphaHoldout : 1;
-
 	/**
 	"Select the format of the light function atlas texture."
 	*/
diff --git a/Engine/Source/Runtime/Engine/Public/PrimitiveSceneProxy.h b/Engine/Source/Runtime/Engine/Public/PrimitiveSceneProxy.h
--- a/Engine/Source/Runtime/Engine/Public/PrimitiveSceneProxy.h
+++ b/Engine/Source/Runtime/Engine/Public/PrimitiveSceneProxy.h
@@ -1457,7 +1457,7 @@ protected:
 
 	/**
 	 * If this is True, this primitive should render black with an alpha of 0, but all secondary effects (shadows, refletions, indirect lighting)
-	 * should behave as usual. This feature requires activating the project setting(s) "Enable Alpha Output", and "Support Primitive Alpha Holdout" if using the deferred renderer.
+	 * should behave as usual. This feature requires activating the project setting(s) "Alpha Output", and "Support Primitive Alpha Holdout" if using the deferred renderer.
 	 */
 	uint8 bHoldout : 1;
 
diff --git a/Engine/Source/Runtime/Landscape/Classes/LandscapeProxy.h b/Engine/Source/Runtime/Landscape/Classes/LandscapeProxy.h
--- a/Engine/Source/Runtime/Landscape/Classes/LandscapeProxy.h
+++ b/Engine/Source/Runtime/Landscape/Classes/LandscapeProxy.h
@@ -795,7 +795,7 @@ public:
 	UPROPERTY(EditAnywhere, AdvancedDisplay, Category = Lighting, meta = (LandscapeOverridable, ClampMin = "0.0", ClampMax = "10.0", UIMin = "0.0", UIMax = "10.0"))
 	float NonNaniteVirtualShadowMapInvalidationScreenSizeLimit = 0.2f;
 
-	/** If this is True, this primitive will render black with an alpha of 0, but all secondary effects (shadows, reflections, indirect lighting) remain. This feature requires the project settings "Enable Alpha Output" and "Support Primitive Alpha Holdout". */
+	/** If this is True, this primitive will render black with an alpha of 0, but all secondary effects (shadows, reflections, indirect lighting) remain. This feature requires the project settings "Alpha Output" and "Support Primitive Alpha Holdout". */
 	UPROPERTY(EditAnywhere, AdvancedDisplay, Category = Rendering, meta=(LandscapeOverridable), Interp)
 	uint8 bHoldout : 1;
 
