World Partition - HLOD - Fixed warning triggered by the inclusion of non-spatially loaded actors in the HLOD setup (when using HashSet)
* Harmless warning, had no impact on the actual result of the HLOD generation

diff --git a/Engine/Source/Runtime/Engine/Private/WorldPartition/RuntimeHashSet/WorldPartitionRuntimeHashSetHLODGeneration.cpp b/Engine/Source/Runtime/Engine/Private/WorldPartition/RuntimeHashSet/WorldPartitionRuntimeHashSetHLODGeneration.cpp
--- a/Engine/Source/Runtime/Engine/Private/WorldPartition/RuntimeHashSet/WorldPartitionRuntimeHashSetHLODGeneration.cpp
+++ b/Engine/Source/Runtime/Engine/Private/WorldPartition/RuntimeHashSet/WorldPartitionRuntimeHashSetHLODGeneration.cpp
@@ -200,6 +200,12 @@ bool UWorldPartitionRuntimeHashSet::SetupHLODActors(const IStreamingGenerationCo
 			int32 CellDescInstanceIndex = 0;
 			for (URuntimePartition::FCellDescInstance& CellDescInstance : CellDescInstances)
 			{
+				// Skip non-spatially loaded cells as they require no HLOD representation
+				if (!CellDescInstance.bIsSpatiallyLoaded)
+				{
+					continue;
+				}
+
 				const FCellUniqueId CellUniqueId = GetCellUniqueId(CellDescInstance);
 
 				UE_LOG(LogWorldPartition, Display, TEXT("[%d / %d] Processing cell %s..."), ++CellDescInstanceIndex, CellDescInstances.Num(), *CellUniqueId.Name);
