Tests: Streaming: Fixed async loading tests failing because of some of issues with assets (System.Engine.Loading.ThreadSafeAsyncLoadingTest, System.Engine.Loading.UnattendedAsyncLoadingTest)

Tese streaming tests take the first 5000 assets from AssetRegistry and try to stream them. However, some of the assets are either problematic or not expected to be loaded outside of the editor, which leads to loading errors failing the tests.

Some of the problems I've found include:
   - ARUtilities, OpenXR: Trying to load assets from plugins which are not enabled.
   - WorldPartition: HLODLayer objects from the engine referencing HLODBuilderMeshMergeSettings objects implemented on the editor side (these assets are probably not meant to be loaded outside of the editor)
   - Harmonix, MovieRenderPipeline: issues that I haven't been able to identify, possibly broken assets.

Resolved by limiting which assets we're trying to use in the tests. It doesn't reduce the generality of the tests, as they are meant to test our streaming logic, not validate the correctness of assets. This change is more of a workaround, though.

diff --git a/Engine/Source/Runtime/Engine/Private/Tests/Loading/AsyncLoadingTests.cpp b/Engine/Source/Runtime/Engine/Private/Tests/Loading/AsyncLoadingTests.cpp
--- a/Engine/Source/Runtime/Engine/Private/Tests/Loading/AsyncLoadingTests.cpp
+++ b/Engine/Source/Runtime/Engine/Private/Tests/Loading/AsyncLoadingTests.cpp
@@ -34,7 +34,11 @@ bool FThreadSafeAsyncLoadingTest::RunTest(const FString& Parameters)
 		{
 			if (UniquePackages.Num() < MaxPackageCount)
 			{
-				UniquePackages.FindOrAdd(AssetData.PackageName);
+				if (LoadingTestsUtils::IsAssetSuitableForTests(AssetData))
+				{
+					UniquePackages.FindOrAdd(AssetData.PackageName);
+				}
+
 				return true;
 			}
 			
diff --git a/Engine/Source/Runtime/Engine/Private/Tests/Loading/AsyncLoadingTests_Shared.cpp b/Engine/Source/Runtime/Engine/Private/Tests/Loading/AsyncLoadingTests_Shared.cpp
--- a/Engine/Source/Runtime/Engine/Private/Tests/Loading/AsyncLoadingTests_Shared.cpp
+++ b/Engine/Source/Runtime/Engine/Private/Tests/Loading/AsyncLoadingTests_Shared.cpp
@@ -5,6 +5,7 @@
 #include "Misc/PackageName.h"
 #include "UObject/MetaData.h"
 #include "UObject/SavePackage.h"
+#include "AssetRegistry/AssetData.h"
 
 UAsyncLoadingTests_Shared::FOnPostLoadDelegate UAsyncLoadingTests_Shared::OnPostLoad;
 UAsyncLoadingTests_Shared::FOnSerializeDelegate UAsyncLoadingTests_Shared::OnSerialize;
@@ -121,3 +122,27 @@ void FLoadingTestsScope::CleanupObjects()
 }
 
 #endif // WITH_DEV_AUTOMATION_TESTS
+
+namespace LoadingTestsUtils
+{
+	extern bool IsAssetSuitableForTests(const FAssetData& AssetData)
+	{
+		FString PackageName = AssetData.PackageName.ToString();
+
+		// Assets from plugins can be problematic because they are either not accessible, or sometimes have issues,
+		// so let's limit ourselves to ones from the engine or the game.
+		if (!PackageName.StartsWith(TEXT("/Engine/")) && !PackageName.StartsWith(TEXT("/Game/")))
+		{
+			return false;
+		}
+
+		// We're skipping WorldPartition assets because some HLOD Layers (engine objects) reference settings objects
+		// defined in editor-only plugins, which obviously fails to load on non-editor targets.
+		if (PackageName.StartsWith(TEXT("/Game/Tests/WorldPartition")))
+		{
+			return false;
+		}
+
+		return true;
+	}
+}
diff --git a/Engine/Source/Runtime/Engine/Private/Tests/Loading/AsyncLoadingTests_Shared.h b/Engine/Source/Runtime/Engine/Private/Tests/Loading/AsyncLoadingTests_Shared.h
--- a/Engine/Source/Runtime/Engine/Private/Tests/Loading/AsyncLoadingTests_Shared.h
+++ b/Engine/Source/Runtime/Engine/Private/Tests/Loading/AsyncLoadingTests_Shared.h
@@ -160,3 +160,8 @@ protected:
 };
 
 #endif // WITH_DEV_AUTOMATION_TESTS
+
+namespace LoadingTestsUtils
+{
+	extern bool IsAssetSuitableForTests(const FAssetData& AssetData);
+}
diff --git a/Engine/Source/Runtime/Engine/Private/Tests/Loading/UnattendedAsyncLoadingTests.cpp b/Engine/Source/Runtime/Engine/Private/Tests/Loading/UnattendedAsyncLoadingTests.cpp
--- a/Engine/Source/Runtime/Engine/Private/Tests/Loading/UnattendedAsyncLoadingTests.cpp
+++ b/Engine/Source/Runtime/Engine/Private/Tests/Loading/UnattendedAsyncLoadingTests.cpp
@@ -6,6 +6,7 @@
 #include "AssetRegistry/AssetRegistryState.h"
 #include "AssetRegistry/IAssetRegistry.h"
 #include "Tasks/Task.h"
+#include "AsyncLoadingTests_Shared.h"
 
 #if WITH_DEV_AUTOMATION_TESTS
 
@@ -200,7 +201,11 @@ bool FUnattendedLoadingTest::RunTest(const FString& Parameters)
 		{
 			if (UniquePackages.Num() < MaxPackageCount)
 			{
-				UniquePackages.FindOrAdd(AssetData.PackageName);
+				if (LoadingTestsUtils::IsAssetSuitableForTests(AssetData))
+				{
+					UniquePackages.FindOrAdd(AssetData.PackageName);
+				}
+
 				return true;
 			}
 			
