Chaos SKM Cloth: Make ClothCollisionSource methods BlueprintCallable.

Also make ClothGeometryScale EditAnywhere.

diff --git a/Engine/Source/Runtime/Engine/Classes/Components/SkeletalMeshComponent.h b/Engine/Source/Runtime/Engine/Classes/Components/SkeletalMeshComponent.h
--- a/Engine/Source/Runtime/Engine/Classes/Components/SkeletalMeshComponent.h
+++ b/Engine/Source/Runtime/Engine/Classes/Components/SkeletalMeshComponent.h
@@ -892,7 +892,7 @@ public:
 	/** This scale is applied to all cloth geometry (e.g., cloth meshes and collisions) in order to simulate in a different scale space than world.This scale is not applied to distance-based simulation parameters such as MaxDistance. 
 	* This property is currently only read by the cloth solver when creating cloth actors, but may become animatable in the future.
 	*/
-	UPROPERTY(BlueprintReadWrite, Category = Clothing, meta = (UIMin = 0.0, UIMax = 10.0, ClampMin = 0.0, ClampMax = 10000.0))
+	UPROPERTY(EditAnywhere, BlueprintReadWrite, Category = Clothing, meta = (UIMin = 0.0, UIMax = 10.0, ClampMin = 0.0, ClampMax = 10000.0))
 	float ClothGeometryScale = 1.f;
 
 	/** Notification when constraint is broken. */
@@ -2338,8 +2338,6 @@ public:
 
 	const TArray<FClothCollisionSource>& GetClothCollisionSources() const { return ClothCollisionSources; }
 
-#if WITH_CLOTH_COLLISION_DETECTION
-
 	/**
 	 * Add a collision source for the cloth on this component.
 	 * Each cloth tick, the collision defined by the physics asset, transformed by the bones in the source
@@ -2347,18 +2345,31 @@ public:
 	 * @param	InSourceComponent		The component to extract collision transforms from
 	 * @param	InSourcePhysicsAsset	The physics asset that defines the collision primitives (that will be transformed by InSourceComponent's bones)
 	 */
+	UFUNCTION(BlueprintCallable, Category = "Clothing")
 	ENGINE_API void AddClothCollisionSource(USkeletalMeshComponent* InSourceComponent, UPhysicsAsset* InSourcePhysicsAsset);
 
 	/** Remove a cloth collision source defined by a component */
-	ENGINE_API void RemoveClothCollisionSource(USkeletalMeshComponent* InSourceComponent);
+	UFUNCTION(BlueprintCallable, Category = "Clothing")
+	ENGINE_API void RemoveClothCollisionSources(USkeletalMeshComponent* InSourceComponent);
+
+	/** Remove a cloth collision source defined by a component */
+	UE_DEPRECATED(5.5, "This function has been renamed RemoveClothCollisionSources")
+	void RemoveClothCollisionSource(USkeletalMeshComponent* InSourceComponent)
+	{
+		RemoveClothCollisionSources(InSourceComponent);
+	}
 
 	/** Remove a cloth collision source defined by both a component and a physics asset */
+	UFUNCTION(BlueprintCallable, Category = "Clothing")
 	ENGINE_API void RemoveClothCollisionSource(USkeletalMeshComponent* InSourceComponent, UPhysicsAsset* InSourcePhysicsAsset);
 
 	/** Remove all cloth collision sources */
+	UFUNCTION(BlueprintCallable, Category = "Clothing")
 	ENGINE_API void ResetClothCollisionSources();
 
 protected:
+
+#if WITH_CLOTH_COLLISION_DETECTION
 	/** copy cloth collision sources to this, where parent means components above it in the hierarchy */
 	ENGINE_API void CopyClothCollisionSources();
 
@@ -2373,15 +2384,6 @@ protected:
 	/** find if this component has collisions for clothing and return the results calculated by bone transforms */
 	ENGINE_API void FindClothCollisions(FClothCollisionData& OutCollisions);
 
-#else
-
-public:
-	/** Stub out these public functions if cloth collision is disabled */
-	void AddClothCollisionSource(USkeletalMeshComponent* InSourceComponent, UPhysicsAsset* InSourcePhysicsAsset) {}
-	void RemoveClothCollisionSource(USkeletalMeshComponent* InSourceComponent) {}
-	void RemoveClothCollisionSource(USkeletalMeshComponent* InSourceComponent, UPhysicsAsset* InSourcePhysicsAsset) {}
-	void ResetClothCollisionSources() {}
-
 #endif
 
 public:
diff --git a/Engine/Source/Runtime/Engine/Private/SkeletalMeshComponentPhysics.cpp b/Engine/Source/Runtime/Engine/Private/SkeletalMeshComponentPhysics.cpp
--- a/Engine/Source/Runtime/Engine/Private/SkeletalMeshComponentPhysics.cpp
+++ b/Engine/Source/Runtime/Engine/Private/SkeletalMeshComponentPhysics.cpp
@@ -3431,8 +3431,11 @@ void USkeletalMeshComponent::ProcessClothCollisionWithEnvironment()
 	ClothingSimulation->AddExternalCollisions(NewCollisionData);
 }
 
+#endif// #if WITH_CLOTH_COLLISION_DETECTION
+
 void USkeletalMeshComponent::AddClothCollisionSource(USkeletalMeshComponent* InSourceComponent, UPhysicsAsset* InSourcePhysicsAsset)
 {
+#if WITH_CLOTH_COLLISION_DETECTION
 	if(InSourceComponent && InSourcePhysicsAsset)
 	{
 		FClothCollisionSource* FoundCollisionSource = ClothCollisionSources.FindByPredicate(
@@ -3455,10 +3458,12 @@ void USkeletalMeshComponent::AddClothCollisionSource(USkeletalMeshComponent* InS
 			ClothTickFunction.AddPrerequisite(InSourceComponent, InSourceComponent->PrimaryComponentTick);
 		}
 	}
+#endif
 }
 
-void USkeletalMeshComponent::RemoveClothCollisionSource(USkeletalMeshComponent* InSourceComponent)
+void USkeletalMeshComponent::RemoveClothCollisionSources(USkeletalMeshComponent* InSourceComponent)
 {
+#if WITH_CLOTH_COLLISION_DETECTION
 	if(InSourceComponent)
 	{
 		ClothCollisionSources.RemoveAll([InSourceComponent](const FClothCollisionSource& InCollisionSource)
@@ -3466,10 +3471,12 @@ void USkeletalMeshComponent::RemoveClothCollisionSource(USkeletalMeshComponent*
 			return !InCollisionSource.SourceComponent.IsValid() || InCollisionSource.SourceComponent.Get() == InSourceComponent; 
 		});
 	}
+#endif
 }
 
 void USkeletalMeshComponent::RemoveClothCollisionSource(USkeletalMeshComponent* InSourceComponent, UPhysicsAsset* InSourcePhysicsAsset)
 {
+#if WITH_CLOTH_COLLISION_DETECTION
 	if(InSourceComponent && InSourcePhysicsAsset)
 	{
 		ClothCollisionSources.RemoveAll([InSourceComponent, InSourcePhysicsAsset](const FClothCollisionSource& InCollisionSource)
@@ -3477,15 +3484,16 @@ void USkeletalMeshComponent::RemoveClothCollisionSource(USkeletalMeshComponent*
 			return !InCollisionSource.SourceComponent.IsValid() || (InCollisionSource.SourceComponent.Get() == InSourceComponent && InCollisionSource.SourcePhysicsAsset.Get() == InSourcePhysicsAsset); 
 		});
 	}
+#endif
 }
 
 void USkeletalMeshComponent::ResetClothCollisionSources()
 {
+#if WITH_CLOTH_COLLISION_DETECTION
 	ClothCollisionSources.Reset();
+#endif
 }
 
-#endif// #if WITH_CLOTH_COLLISION_DETECTION
-
 void USkeletalMeshComponent::EndPhysicsTickComponent(FSkeletalMeshComponentEndPhysicsTickFunction& ThisTickFunction)
 {
 	//IMPORTANT!
