Fix potential crash in mesh painting when the selected UV channel is out of bounds for the given paint mesh.
#rnx

diff --git a/Engine/Source/Runtime/Engine/Private/Components/StaticMeshComponent.cpp b/Engine/Source/Runtime/Engine/Private/Components/StaticMeshComponent.cpp
--- a/Engine/Source/Runtime/Engine/Private/Components/StaticMeshComponent.cpp
+++ b/Engine/Source/Runtime/Engine/Private/Components/StaticMeshComponent.cpp
@@ -3193,13 +3193,14 @@ void UStaticMeshComponent::SetMeshPaintTextureOverride(UTexture* OverrideTexture
 
 int32 UStaticMeshComponent::GetMeshPaintTextureCoordinateIndex() const
 {
+	const int32 NumTexCoords = StaticMesh->GetNumTexCoords(0);
 	if (bOverrideMeshPaintTextureCoordinateIndex)
 	{
-		return OverriddenMeshPaintTextureCoordinateIndex;
+		return FMath::Clamp(OverriddenMeshPaintTextureCoordinateIndex, 0, NumTexCoords - 1);
 	}
 	if (StaticMesh != nullptr)
 	{
-		return StaticMesh->MeshPaintTextureCoordinateIndex;
+		return FMath::Clamp(StaticMesh->MeshPaintTextureCoordinateIndex, 0, NumTexCoords - 1);
 	}
 	return 0;
 }
