[HWRT] Comment out ensure about duplicated referenced geometries for now.

#rnx

diff --git a/Engine/Source/Runtime/Engine/Private/Rendering/RayTracingGeometryManager.cpp b/Engine/Source/Runtime/Engine/Private/Rendering/RayTracingGeometryManager.cpp
--- a/Engine/Source/Runtime/Engine/Private/Rendering/RayTracingGeometryManager.cpp
+++ b/Engine/Source/Runtime/Engine/Private/Rendering/RayTracingGeometryManager.cpp
@@ -701,7 +701,7 @@ void FRayTracingGeometryManager::Tick(FRHICommandList& RHICmdList)
 		}
 
 #if DO_CHECK
-		ensure(ReferencedGeometries.Num() == TSet(ReferencedGeometries).Num());
+		// ensure(ReferencedGeometries.Num() == TSet(ReferencedGeometries).Num());
 #endif
 
 		SET_MEMORY_STAT(STAT_RayTracingGeometryRequestedMemory, RequestedSize);
