[NetPhysics] Kinematic character movement physics interactions with physics prediction enabled has a client authoritative flow that schedules forces on the server, move this behind a CVar setting rather than being the default mode.

diff --git a/Engine/Source/Runtime/Engine/Private/Components/CharacterMovementComponent.cpp b/Engine/Source/Runtime/Engine/Private/Components/CharacterMovementComponent.cpp
--- a/Engine/Source/Runtime/Engine/Private/Components/CharacterMovementComponent.cpp
+++ b/Engine/Source/Runtime/Engine/Private/Components/CharacterMovementComponent.cpp
@@ -362,6 +362,13 @@ namespace CharacterMovementCVars
 		TEXT("Apply the ledge movement vector directly, rather than the old method that reapplied acceleration."),
 		ECVF_Default);
 
+	static bool bEnableClientAuthScheduledPushForces = false;
+	FAutoConsoleVariableRef CVarEnableClientAuthScheduledPushForces(
+		TEXT("np2.CMC.EnableClientAuthScheduledPushForces"),
+		bEnableClientAuthScheduledPushForces,
+		TEXT("Enable a client-authoritative network flow for adding forces to pushed physics objects, requires physics prediction to be enabled in the project settings. NOTE: This is not recommended for networked physics and can produce inconsistency bugs where the client adds multiple forces for a single physics tick on the server, especially if async physics is not enabled."),
+		ECVF_Default);
+
 #if CSV_PROFILER_STATS
 	bool bClientRecordMovePackedRpcStatsToCsv = false;
 	FAutoConsoleVariableRef CVarClientRecordMovePackedRpcStatsToCsv(
@@ -7642,14 +7649,14 @@ void UCharacterMovementComponent::ApplyImpactPhysicsForces(const FHitResult& Imp
 				{
 					PushForceModificator *= BodyMass;
 				}
-				const bool bEnableNetworkPhysics = Chaos::FPhysicsSolverBase::IsNetworkPhysicsPredictionEnabled(); 
+				const bool bClientAuthScheduledPushForces = CharacterMovementCVars::bEnableClientAuthScheduledPushForces && Chaos::FPhysicsSolverBase::IsNetworkPhysicsPredictionEnabled();
 
 				Force *= PushForceModificator;
 				const float ZeroVelocityTolerance = 1.0f;
 				if (ComponentVelocity.IsNearlyZero(ZeroVelocityTolerance))
 				{
 					Force *= InitialPushForceFactor;
-					if (!bEnableNetworkPhysics)
+					if (!bClientAuthScheduledPushForces)
 					{
 						ImpactComponent->AddImpulseAtLocation(Force, ForcePoint, Impact.BoneName);
 					}
@@ -7661,7 +7668,7 @@ void UCharacterMovementComponent::ApplyImpactPhysicsForces(const FHitResult& Imp
 				else
 				{
 					Force *= PushForceFactor;
-					if (!bEnableNetworkPhysics)
+					if (!bClientAuthScheduledPushForces)
 					{
 						ImpactComponent->AddForceAtLocation(Force, ForcePoint, Impact.BoneName);
 					}
