Light Function - Rect light, light function material angle control

#tests Could add a test for rect light function with the new projection (Light Function Cone Angle > 0 for rect light)

diff --git a/Engine/Shaders/Private/LightFunctionAtlas/LightFunctionAtlasCommon.usf b/Engine/Shaders/Private/LightFunctionAtlas/LightFunctionAtlasCommon.usf
--- a/Engine/Shaders/Private/LightFunctionAtlas/LightFunctionAtlasCommon.usf
+++ b/Engine/Shaders/Private/LightFunctionAtlas/LightFunctionAtlasCommon.usf
@@ -78,6 +78,15 @@ FLightFunctionColor GetLocalLightFunctionCommon(float3 TranslatedWorldPosition,
 		break;
 	}
 
+	case LIGHT_TYPE_RECT:
+	{
+		if (TanOuterAngle > 0.0f)
+		{
+			LightFunctionUV = LightVector.xy / (LightVector.z * TanOuterAngle) * .5f + .5f;
+		}
+		break;
+	}
+
 	//case LIGHT_TYPE_RECT:
 	//case LIGHT_TYPE_DIRECTIONAL:
 	default:
diff --git a/Engine/Source/Runtime/Engine/Classes/Components/RectLightComponent.h b/Engine/Source/Runtime/Engine/Classes/Components/RectLightComponent.h
--- a/Engine/Source/Runtime/Engine/Classes/Components/RectLightComponent.h
+++ b/Engine/Source/Runtime/Engine/Classes/Components/RectLightComponent.h
@@ -47,6 +47,13 @@ class URectLightComponent : public ULocalLightComponent
 	UPROPERTY(EditAnywhere, BlueprintReadOnly, interp, Category = Light, meta = (UIMin = "0.0"))
 	float BarnDoorLength;
 
+	/**
+	 * Aperture of cone angle for the perspective projection of the light function material.
+	 * If 0, an orthographic projection is used instead.
+	 */
+	UPROPERTY(EditAnywhere, BlueprintReadOnly, interp, Category = LightFunction, meta = (UIMin = "0.0", UIMax = "89.0", ClampMin = "0.0", ClampMax = "89.0"))
+	float LightFunctionConeAngle;
+
 	/** Texture mapped to the light source rectangle */
 	UPROPERTY(EditAnywhere, BlueprintReadOnly, Category = Light)
 	TObjectPtr<class UTexture> SourceTexture;
@@ -71,6 +78,7 @@ public:
 	ENGINE_API virtual float ComputeLightBrightness() const override;
 #if WITH_EDITOR
 	ENGINE_API virtual void SetLightBrightness(float InBrightness) override;
+	ENGINE_API virtual bool CanEditChange(const FProperty* InProperty) const override;
 #endif
 
 	//~ Begin ULightComponent Interface.
diff --git a/Engine/Source/Runtime/Engine/Private/Components/RectLightComponent.cpp b/Engine/Source/Runtime/Engine/Private/Components/RectLightComponent.cpp
--- a/Engine/Source/Runtime/Engine/Private/Components/RectLightComponent.cpp
+++ b/Engine/Source/Runtime/Engine/Private/Components/RectLightComponent.cpp
@@ -44,6 +44,7 @@ URectLightComponent::URectLightComponent(const FObjectInitializer& ObjectInitial
 	SourceTexture = nullptr;
 	BarnDoorAngle = GetRectLightBarnDoorMaxAngle();
 	BarnDoorLength = 20.0f;
+	LightFunctionConeAngle = 0.0f;
 	// RayTracingData will be initialised on the render thread.
 }
 
@@ -149,6 +150,25 @@ void URectLightComponent::SetLightBrightness(float InBrightness)
 		Super::SetLightBrightness(InBrightness / 16); // Legacy scale of 16
 	}
 }
+
+bool URectLightComponent::CanEditChange(const FProperty* InProperty) const
+{
+	if (InProperty)
+	{
+		FString PropertyName = InProperty->GetName();
+
+		if (PropertyName == GET_MEMBER_NAME_STRING_CHECKED(URectLightComponent, LightFunctionConeAngle))
+		{
+			if (Mobility == EComponentMobility::Static)
+			{
+				return false;
+			}
+			return LightFunctionMaterial != NULL;
+		}
+	}
+
+	return Super::CanEditChange(InProperty);
+}
 #endif // WITH_EDITOR
 
 /**
@@ -209,6 +229,7 @@ FRectLightSceneProxy::FRectLightSceneProxy(const URectLightComponent* Component)
 	, BarnDoorLength(FMath::Max(0.1f, Component->BarnDoorLength))
 	, RayTracingData(Component->RayTracingData)
 	, SourceTexture(Component->SourceTexture)
+	, LightFunctionConeAngleTangent(Component->LightFunctionConeAngle > 0 ? FMath::Tan(FMath::Clamp(Component->LightFunctionConeAngle, 0.0f, 89.0f) * (float)UE_PI / 180.0f) : 0.0f)
 {
 	RectAtlasId = ~0u;
 }
@@ -306,4 +327,4 @@ bool FRectLightSceneProxy::GetWholeSceneProjectedShadowInitializer(const FSceneV
 	}
 
 	return false;
-}
\ No newline at end of file
+}
diff --git a/Engine/Source/Runtime/Engine/Public/RectLightSceneProxy.h b/Engine/Source/Runtime/Engine/Public/RectLightSceneProxy.h
--- a/Engine/Source/Runtime/Engine/Public/RectLightSceneProxy.h
+++ b/Engine/Source/Runtime/Engine/Public/RectLightSceneProxy.h
@@ -31,6 +31,7 @@ public:
 	FRectLightRayTracingData* RayTracingData; // Render thread data only
 	UTexture*	SourceTexture;
 	uint32		RectAtlasId;
+	float		LightFunctionConeAngleTangent;	// Use Ortho projection if 0
 
 	FRectLightSceneProxy(const URectLightComponent* Component);
 	virtual ~FRectLightSceneProxy();
diff --git a/Engine/Source/Runtime/Renderer/Private/LightFunctionAtlas.cpp b/Engine/Source/Runtime/Renderer/Private/LightFunctionAtlas.cpp
--- a/Engine/Source/Runtime/Renderer/Private/LightFunctionAtlas.cpp
+++ b/Engine/Source/Runtime/Renderer/Private/LightFunctionAtlas.cpp
@@ -20,6 +20,7 @@
 #include "Containers/HashTable.h"
 #include "RenderUtils.h"
 #include "VolumetricFog.h"
+#include "RectLightSceneProxy.h"
 
 DECLARE_GPU_STAT(LightFunctionAtlasGeneration);
 
@@ -649,7 +650,15 @@ void FLightFunctionAtlas::RenderLightFunctionAtlas(FRDGBuilder& GraphBuilder, TA
 			ensure(FMath::IsNearlyEqual((PackedAtlasSlotMinUV & 0xFFFF) / 65536.0f, AtlasSlot.MinU));
 			ensure(FMath::IsNearlyEqual(((PackedAtlasSlotMinUV >> 16) & 0xFFFF) / 65536.0f, AtlasSlot.MinV));
 
-			const float TanOuterAngle = LightType == LightType_Spot ? FMath::Tan(LightSceneInfo->Proxy->GetOuterConeAngle()) : 1.0f;
+			float TanOuterAngle = LightType == LightType_Spot ? FMath::Tan(LightSceneInfo->Proxy->GetOuterConeAngle()) : -1.0f;
+			if (LightType == LightType_Rect)
+			{
+				FRectLightSceneProxy* RectLightProxy = (FRectLightSceneProxy*)LightSceneInfo->Proxy;
+				if (RectLightProxy->LightFunctionConeAngleTangent > 0.0f)
+				{
+					TanOuterAngle = RectLightProxy->LightFunctionConeAngleTangent;
+				}
+			}
 
 			// ShadowFadeFraction is unused.
 			LightInfoDataBufferPtr[OutputBufferLightIndex].Parameters = FVector4f(Proxy->GetLightFunctionFadeDistance(), FMath::AsFloat(PackedLightInfoDataParams), FMath::AsFloat(PackedAtlasSlotMinUV), TanOuterAngle);
diff --git a/Engine/Source/Runtime/Renderer/Private/LightRendering.cpp b/Engine/Source/Runtime/Renderer/Private/LightRendering.cpp
--- a/Engine/Source/Runtime/Renderer/Private/LightRendering.cpp
+++ b/Engine/Source/Runtime/Renderer/Private/LightRendering.cpp
@@ -35,6 +35,7 @@
 #include "HeterogeneousVolumes/HeterogeneousVolumes.h"
 #include "Materials/MaterialRenderProxy.h"
 #include "RHIResourceUtils.h"
+#include "RectLightSceneProxy.h"
 
 using namespace LightFunctionAtlas;
 
@@ -274,9 +275,26 @@ void FLightFunctionSharedParameters::Bind(const FShaderParameterMap& ParameterMa
 
 FVector4f FLightFunctionSharedParameters::GetLightFunctionSharedParameters(const FLightSceneInfo* LightSceneInfo, float ShadowFadeFraction)
 {
-	const bool bIsSpotLight = LightSceneInfo->Proxy->GetLightType() == LightType_Spot;
+	bool bIsSpotLight = LightSceneInfo->Proxy->GetLightType() == LightType_Spot;
 	const bool bIsPointLight = LightSceneInfo->Proxy->GetLightType() == LightType_Point;
-	const float TanOuterAngle = bIsSpotLight ? FMath::Tan(LightSceneInfo->Proxy->GetOuterConeAngle()) : 1.0f;
+	float TanOuterAngle = bIsSpotLight ? FMath::Tan(LightSceneInfo->Proxy->GetOuterConeAngle()) : 1.0f;
+
+	if (LightSceneInfo->Proxy->GetLightType() == LightType_Rect)
+	{
+		// Rect light can have a spot like perspective projection
+		FRectLightSceneProxy* RectLightProxy = (FRectLightSceneProxy*)LightSceneInfo->Proxy;
+		if (RectLightProxy->LightFunctionConeAngleTangent > 0.0f)
+		{
+			bIsSpotLight = true;
+			TanOuterAngle = RectLightProxy->LightFunctionConeAngleTangent;
+		}
+		else
+		{
+			bIsSpotLight = false;
+			TanOuterAngle = 0.0f;
+		}
+	}
+
 	return FVector4f(TanOuterAngle, ShadowFadeFraction, bIsSpotLight ? 1.0f : 0.0f, bIsPointLight ? 1.0f : 0.0f);
 }
 
