Substrate - added multiple subsurface profile support from a material for swwitch nodes to work correctly.
Kept SetSubsurfaceProfileOverrideRT on proxy since otherwise there was some weird interaction with material/materialinstance when a profile was already there.

#rnx

diff --git a/Engine/Shaders/Public/ShaderVersion.ush b/Engine/Shaders/Public/ShaderVersion.ush
--- a/Engine/Shaders/Public/ShaderVersion.ush
+++ b/Engine/Shaders/Public/ShaderVersion.ush
@@ -3,4 +3,4 @@
 // Each time the console command is executed it generates a new GUID. As this file is included
 // in Platform.ush (which should be included in any shader) it allows to invalidate the shader DDC.
 // 
-#pragma message("UESHADERMETADATA_VERSION ECC0ACE5-64A8-47CA-A772-C34CD9E09C21")
\ No newline at end of file
+#pragma message("UESHADERMETADATA_VERSION EDDDB80A-CEF7-43E6-9888-56DFAB6BCA2E")
diff --git a/Engine/Source/Editor/UnrealEd/Private/Factories/EditorFactories.cpp b/Engine/Source/Editor/UnrealEd/Private/Factories/EditorFactories.cpp
--- a/Engine/Source/Editor/UnrealEd/Private/Factories/EditorFactories.cpp
+++ b/Engine/Source/Editor/UnrealEd/Private/Factories/EditorFactories.cpp
@@ -7686,6 +7686,9 @@ UObject* USubsurfaceProfileFactory::FactoryCreateNew(UClass* InClass, UObject* I
 	// loaded from files to be automatically converted to MFP.
 	Object->Settings.bEnableMeanFreePath = true;
 
+	// Assign a GUID to the asset
+	Object->Guid = FGuid::NewGuid();
+
 	return Object;
 }
 
diff --git a/Engine/Source/Runtime/Core/Private/UObject/FortniteMainBranchObjectVersion.cpp b/Engine/Source/Runtime/Core/Private/UObject/FortniteMainBranchObjectVersion.cpp
--- a/Engine/Source/Runtime/Core/Private/UObject/FortniteMainBranchObjectVersion.cpp
+++ b/Engine/Source/Runtime/Core/Private/UObject/FortniteMainBranchObjectVersion.cpp
@@ -14,6 +14,6 @@ TMap<FGuid, FGuid> FFortniteMainBranchObjectVersion::GetSystemGuids()
 	SystemGuids.Add(DevGuids.Niagara_LatestScriptCompileVersion, FGuid("F8D995224324470E97ED718D33939A0E"));
 	SystemGuids.Add(DevGuids.SkeletalMeshDerivedDataVersion, FGuid("B9AEAA3EE7AC4FCFB3860E8F3E06DD75"));
 	SystemGuids.Add(DevGuids.STATICMESH_DERIVEDDATA_VER, FGuid("1C2EBE3310464FD7995ADE1D6AB92BDC"));
-	SystemGuids.Add(DevGuids.MaterialTranslationDDCVersion, FGuid("3A4DC1DF808D4096BEEA6BA2E3ACC38D"));
+	SystemGuids.Add(DevGuids.MaterialTranslationDDCVersion, FGuid("79D0C4D88D7C4D1491C86B9922001F20"));
 	return SystemGuids;
 }
diff --git a/Engine/Source/Runtime/Core/Public/UObject/FortniteMainBranchObjectVersion.h b/Engine/Source/Runtime/Core/Public/UObject/FortniteMainBranchObjectVersion.h
--- a/Engine/Source/Runtime/Core/Public/UObject/FortniteMainBranchObjectVersion.h
+++ b/Engine/Source/Runtime/Core/Public/UObject/FortniteMainBranchObjectVersion.h
@@ -521,6 +521,9 @@ struct FFortniteMainBranchObjectVersion
 		// Multiple changes to AnimNext modules, variables etc.
 		AnimNextModuleRefactor,
 
+		// Subsurface profile now has a guid to be able to select one of many in a Substrate material.
+		SubsurfaceProfileGuid,
+
 		// -----<new versions can be added above this line>-------------------------------------------------
 		VersionPlusOne,
 		LatestVersion = VersionPlusOne - 1
diff --git a/Engine/Source/Runtime/Engine/Classes/Engine/EngineTypes.h b/Engine/Source/Runtime/Engine/Classes/Engine/EngineTypes.h
--- a/Engine/Source/Runtime/Engine/Classes/Engine/EngineTypes.h
+++ b/Engine/Source/Runtime/Engine/Classes/Engine/EngineTypes.h
@@ -753,12 +753,12 @@ public:
 	// Subsurface profiles
 	void AddSubsurfaceProfile(USubsurfaceProfile* InProfile) { if (InProfile) SubsurfaceProfiles.Add(InProfile); }
 	int32 CountSubsurfaceProfiles() const { return SubsurfaceProfiles.Num(); }
-	USubsurfaceProfile* GetSubsurfaceProfile() const { return SubsurfaceProfiles.Num() > 0 ? SubsurfaceProfiles[0] : nullptr; }
+	USubsurfaceProfile* GetSubsurfaceProfile(int32 Index) const { return SubsurfaceProfiles.IsValidIndex(Index) ? SubsurfaceProfiles[Index] : nullptr; }
 
 	// Specular profiles
 	void AddSpecularProfile(USpecularProfile* InProfile) { if (InProfile) SpecularProfiles.Add(InProfile); }
 	int32 CountSpecularProfiles() const { return SpecularProfiles.Num(); }
-	USpecularProfile* GetSpecularProfile(int32 Index) const { return Index < SpecularProfiles.Num() ? SpecularProfiles[Index] : nullptr; }
+	USpecularProfile* GetSpecularProfile(int32 Index) const { return SpecularProfiles.IsValidIndex(Index) ? SpecularProfiles[Index] : nullptr; }
 
 	// Shading model from expression
 	void SetShadingModelFromExpression(bool bIn) { bHasShadingModelFromExpression = bIn ? 1u : 0u; }
diff --git a/Engine/Source/Runtime/Engine/Classes/Engine/SubsurfaceProfile.h b/Engine/Source/Runtime/Engine/Classes/Engine/SubsurfaceProfile.h
--- a/Engine/Source/Runtime/Engine/Classes/Engine/SubsurfaceProfile.h
+++ b/Engine/Source/Runtime/Engine/Classes/Engine/SubsurfaceProfile.h
@@ -154,9 +154,14 @@ class USubsurfaceProfile : public UObject
 	UPROPERTY(Category = USubsurfaceProfile, EditAnywhere, meta = (ShowOnlyInnerProperties))
 	struct FSubsurfaceProfileStruct Settings;
 
+	UPROPERTY()
+	FGuid Guid;
+
 	//~ Begin UObject Interface
 	virtual void BeginDestroy();
 	virtual void PostEditChangeProperty(struct FPropertyChangedEvent& PropertyChangedEvent);
+	virtual void PostDuplicate(EDuplicateMode::Type DuplicateMode);
+	virtual void Serialize(FArchive& Ar) override;
 
 	// Upgrade parameters from Separable to Burley.
 	virtual void PostLoad();
@@ -281,4 +286,7 @@ ENGINE_API FRHITexture* GetSSProfilesPreIntegratedTextureWithFallback();
 ENGINE_API FName GetSubsurfaceProfileParameterName();
 
 // Returns the subsurface profile ID for a given Sub-surface Profile object
-ENGINE_API float GetSubsurfaceProfileId(const USubsurfaceProfile* In);
\ No newline at end of file
+ENGINE_API float GetSubsurfaceProfileId(const USubsurfaceProfile* In);
+
+// Returns the shader parameter name for a Subsurface profile.
+ENGINE_API FName CreateSubsurfaceProfileParameterName(const USubsurfaceProfile* InProfile);
\ No newline at end of file
diff --git a/Engine/Source/Runtime/Engine/Private/Materials/Material.cpp b/Engine/Source/Runtime/Engine/Private/Materials/Material.cpp
--- a/Engine/Source/Runtime/Engine/Private/Materials/Material.cpp
+++ b/Engine/Source/Runtime/Engine/Private/Materials/Material.cpp
@@ -448,9 +448,25 @@ public:
 		{
 			if (Type == EMaterialParameterType::Scalar && ParameterInfo.Name == GetSubsurfaceProfileParameterName())
 			{
+				// Legacy single SubsurfaceProfile (Substrate do not add this one in the material shader)
 				OutValue = GetSubsurfaceProfileId(GetSubsurfaceProfileRT());
 				return true;
 			}
+			else if (Type == EMaterialParameterType::Scalar && NumSubsurfaceProfileRT() > 0)
+			{
+				const USubsurfaceProfile* SSProfileOverrideRT = GetSubsurfaceProfileRT();
+				// Substrate general SubsurfaceProfile
+				for (uint32 It = 0, Count = NumSubsurfaceProfileRT(); It < Count; ++It)
+				{
+					const USubsurfaceProfile* SSProfileRT = GetSubsurfaceProfileRT(It);
+					if (ParameterInfo.Name == CreateSubsurfaceProfileParameterName(SSProfileRT))
+					{
+						// Set the root material Profile, or the profile overriden by any instances.
+						OutValue = GetSubsurfaceProfileId(SSProfileOverrideRT ? SSProfileOverrideRT : SSProfileRT);
+						return true;
+					}
+				}
+			}
 			else if (Type == EMaterialParameterType::Scalar && NumSpecularProfileRT() > 0)
 			{
 				for (uint32 It=0,Count=NumSpecularProfileRT();It<Count;++It)
@@ -3000,6 +3016,15 @@ void UMaterial::Serialize(FArchive& Ar)
 		// All previous material must use the legacy pixel depth offset mode. New material will use the new mode.
 		PixelDepthOffsetMode = PDOM_Legacy;
 	}
+
+	if (Ar.IsLoading() && Ar.CustomVer(FFortniteMainBranchObjectVersion::GUID) < FFortniteMainBranchObjectVersion::SubsurfaceProfileGuid)
+	{
+		// Add the unique legacy profile to the Substrate profile list the first time we load an older asset.
+		if (SubsurfaceProfiles.IsEmpty() && SubsurfaceProfile)
+		{
+			SubsurfaceProfiles.Add(SubsurfaceProfile);
+		}
+	}
 }
 
 void UMaterial::PostDuplicate(bool bDuplicateForPIE)
@@ -5408,7 +5433,14 @@ void UMaterial::RebuildShadingModelField()
 		SubsurfaceProfile = nullptr;
 		if ((SubstrateMaterialInfo.HasShadingModel(SSM_Eye) || SubstrateMaterialInfo.HasShadingModel(SSM_SubsurfaceProfile)) && SubstrateMaterialInfo.CountSubsurfaceProfiles() > 0)
 		{
-			SubsurfaceProfile = SubstrateMaterialInfo.GetSubsurfaceProfile();
+			SubsurfaceProfile = SubstrateMaterialInfo.GetSubsurfaceProfile(0);
+		}
+
+		// Set subsurface profiles if any
+		SubsurfaceProfiles.SetNum(SubstrateMaterialInfo.CountSubsurfaceProfiles());
+		for (int32 It = 0, Count = SubstrateMaterialInfo.CountSubsurfaceProfiles(); It < Count; ++It)
+		{
+			SubsurfaceProfiles[It] = SubstrateMaterialInfo.GetSubsurfaceProfile(It);
 		}
 
 		// Set specular profile if any
@@ -6931,6 +6963,24 @@ USubsurfaceProfile* UMaterial::GetSubsurfaceProfile_Internal() const
 	return SubsurfaceProfile; 
 }
 
+uint32 UMaterial::NumSubsurfaceProfileRoot_Internal() const
+{
+	return SubsurfaceProfiles.Num();
+}
+
+USubsurfaceProfile* UMaterial::GetSubsurfaceProfileRoot_Internal(uint32 Index) const
+{
+	checkSlow(IsInGameThread());
+	check(Index<uint32(SubsurfaceProfiles.Num()));
+	return SubsurfaceProfiles[Index];
+}
+
+USubsurfaceProfile* UMaterial::GetSubsurfaceProfileOverride_Internal() const
+{
+	checkSlow(IsInGameThread());
+	return nullptr; // No override for root material
+}
+
 uint32 UMaterial::NumSpecularProfile_Internal() const
 {
 	return SpecularProfiles.Num();
diff --git a/Engine/Source/Runtime/Engine/Private/Materials/MaterialExpressionSubstrate.cpp b/Engine/Source/Runtime/Engine/Private/Materials/MaterialExpressionSubstrate.cpp
--- a/Engine/Source/Runtime/Engine/Private/Materials/MaterialExpressionSubstrate.cpp
+++ b/Engine/Source/Runtime/Engine/Private/Materials/MaterialExpressionSubstrate.cpp
@@ -290,7 +290,8 @@ int32 UMaterialExpressionSubstrateShadingModels::Compile(class FMaterialCompiler
 	const bool bHasSSS = HasSSS();
 	if (bHasSSS)
 	{
-		SSSProfileCodeChunk = Compiler->ForceCast(Compiler->ScalarParameter(GetSubsurfaceProfileParameterName(), 1.0f), MCT_Float1);
+		const FName SubsurfaceProfileParameterName = CreateSubsurfaceProfileParameterName(SubsurfaceProfile);
+		SSSProfileCodeChunk = Compiler->ForceCast(Compiler->ScalarParameter(SubsurfaceProfileParameterName, 1.0f), MCT_Float1);
 	}
 
 	FSubstrateOperator& SubstrateOperator = Compiler->SubstrateCompilationGetOperator(Compiler->SubstrateTreeStackGetPathUniqueId());
@@ -713,6 +714,7 @@ UMaterialExpressionSubstrateSlabBSDF::UMaterialExpressionSubstrateSlabBSDF(const
 
 #if WITH_EDITOR
 FName CreateSpecularProfileParameterName(USpecularProfile* InProfile);
+FName CreateSubsurfaceProfileParameterName(const USubsurfaceProfile* InProfile);
 int32 UMaterialExpressionSubstrateSlabBSDF::Compile(class FMaterialCompiler* Compiler, int32 OutputIndex)
 {
 	FGuid PathUniqueId = Compiler->SubstrateTreeStackGetPathUniqueId();
@@ -725,7 +727,8 @@ int32 UMaterialExpressionSubstrateSlabBSDF::Compile(class FMaterialCompiler* Com
 	int32 SSSProfileCodeChunk = INDEX_NONE;
 	if (SubstrateOperator.Has(ESubstrateBsdfFeature_SSS))
 	{
-		SSSProfileCodeChunk = Compiler->ForceCast(Compiler->ScalarParameter(GetSubsurfaceProfileParameterName(), 1.0f), MCT_Float1);
+		const FName SubsurfaceProfileParameterName = CreateSubsurfaceProfileParameterName(SubsurfaceProfile);
+		SSSProfileCodeChunk = Compiler->ForceCast(Compiler->ScalarParameter(SubsurfaceProfileParameterName, 1.0f), MCT_Float1);
 	}
 
 	int32 SpecularProfileCodeChunk = INDEX_NONE;
@@ -2058,7 +2061,8 @@ int32 UMaterialExpressionSubstrateEyeBSDF::Compile(class FMaterialCompiler* Comp
 	int32 SSSProfileCodeChunk = INDEX_NONE;
 	if (SubsurfaceProfile != nullptr)
 	{
-		SSSProfileCodeChunk = Compiler->ForceCast(Compiler->ScalarParameter(GetSubsurfaceProfileParameterName(), 1.0f), MCT_Float1);
+		const FName SubsurfaceProfileParameterName = CreateSubsurfaceProfileParameterName(SubsurfaceProfile);
+		SSSProfileCodeChunk = Compiler->ForceCast(Compiler->ScalarParameter(SubsurfaceProfileParameterName, 1.0f), MCT_Float1);
 	}
 
 	FSubstrateOperator& SubstrateOperator = Compiler->SubstrateCompilationGetOperator(Compiler->SubstrateTreeStackGetPathUniqueId());
@@ -3580,7 +3584,8 @@ int32 UMaterialExpressionSubstrateConvertMaterialAttributes::Compile(class FMate
 	const bool bHasSSS = HasSSS();
 	if (bHasSSS)
 	{
-		SSSProfileCodeChunk = Compiler->ForceCast(Compiler->ScalarParameter(GetSubsurfaceProfileParameterName(), 1.0f), MCT_Float1);
+		const FName SubsurfaceProfileParameterName = CreateSubsurfaceProfileParameterName(SubsurfaceProfile);
+		SSSProfileCodeChunk = Compiler->ForceCast(Compiler->ScalarParameter(SubsurfaceProfileParameterName, 1.0f), MCT_Float1);
 	}
 	SSSProfileCodeChunk = SSSProfileCodeChunk != INDEX_NONE ? SSSProfileCodeChunk : Compiler->Constant(0.0f);
 
diff --git a/Engine/Source/Runtime/Engine/Private/Materials/MaterialInstance.cpp b/Engine/Source/Runtime/Engine/Private/Materials/MaterialInstance.cpp
--- a/Engine/Source/Runtime/Engine/Private/Materials/MaterialInstance.cpp
+++ b/Engine/Source/Runtime/Engine/Private/Materials/MaterialInstance.cpp
@@ -300,6 +300,25 @@ bool FMaterialInstanceResource::GetParameterValue(EMaterialParameterType Type, c
 		OutValue = GetSubsurfaceProfileId(MySubsurfaceProfileRT);
 		bResult = true;
 	}
+	else if (Type == EMaterialParameterType::Scalar && NumSubsurfaceProfileRT() > 0)
+	{
+#if !ENABLE_MATERIAL_LAYER_PROTOTYPE
+		check(ParameterInfo.Association == EMaterialParameterAssociation::GlobalParameter);
+#endif
+		const USubsurfaceProfile* SSProfileOverrideRT = GetSubsurfaceProfileRT();
+		// Substrate general SubsurfaceProfile
+		for (uint32 It = 0, Count = NumSubsurfaceProfileRT(); It < Count; ++It)
+		{
+			const USubsurfaceProfile* SSProfileRT = GetSubsurfaceProfileRT(It);
+			if (ParameterInfo.Name == CreateSubsurfaceProfileParameterName(SSProfileRT))
+			{
+				// Set the root material Profile, or the profile overriden by any instances.
+				OutValue = GetSubsurfaceProfileId(SSProfileOverrideRT ? SSProfileOverrideRT : SSProfileRT);
+				bResult = true;
+				break;
+			}
+		}
+	}
 	else if (Type == EMaterialParameterType::Scalar && NumSpecularProfileRT() > 0)
 	{
 		for (uint32 It=0,Count=NumSpecularProfileRT();It<Count;++It)
@@ -4873,6 +4892,32 @@ USubsurfaceProfile* UMaterialInstance::GetSubsurfaceProfile_Internal() const
 	return Parent ? Parent->GetSubsurfaceProfile_Internal() : 0;
 }
 
+uint32 UMaterialInstance::NumSubsurfaceProfileRoot_Internal() const
+{
+	// Return the subsurface profile count form the root material.
+	checkSlow(IsInGameThread());
+	return Parent ? Parent->NumSubsurfaceProfileRoot_Internal() : 0;
+}
+
+USubsurfaceProfile* UMaterialInstance::GetSubsurfaceProfileRoot_Internal(uint32 Index) const
+{
+	// Return the Subsurface profile from the root material.
+	checkSlow(IsInGameThread());
+	return Parent ? Parent->GetSubsurfaceProfileRoot_Internal(Index) : 0;
+}
+
+USubsurfaceProfile* UMaterialInstance::GetSubsurfaceProfileOverride_Internal() const
+{
+	// Return the possible override for all the instance, but root material always return null as no override since in this case the material Profile itself will be used.
+	// The single overriden SSSProbile will overide all the Probile from the root material.
+	checkSlow(IsInGameThread());
+	if (bOverrideSubsurfaceProfile)
+	{
+		return SubsurfaceProfile;
+	}
+	return Parent ? Parent->GetSubsurfaceProfileOverride_Internal() : 0;
+}
+
 bool UMaterialInstance::CastsRayTracedShadows() const
 {
 	//#dxr_todo: do per material instance override?
diff --git a/Engine/Source/Runtime/Engine/Private/Materials/MaterialInterface.cpp b/Engine/Source/Runtime/Engine/Private/Materials/MaterialInterface.cpp
--- a/Engine/Source/Runtime/Engine/Private/Materials/MaterialInterface.cpp
+++ b/Engine/Source/Runtime/Engine/Private/Materials/MaterialInterface.cpp
@@ -1451,6 +1451,16 @@ USubsurfaceProfile* UMaterialInterface::GetSubsurfaceProfile_Internal() const
 	return nullptr;
 }
 
+USubsurfaceProfile* UMaterialInterface::GetSubsurfaceProfileRoot_Internal(uint32 Index) const
+{
+	return nullptr;
+}
+
+USubsurfaceProfile* UMaterialInterface::GetSubsurfaceProfileOverride_Internal() const
+{
+	return nullptr;
+}
+
 USpecularProfile* UMaterialInterface::GetSpecularProfile_Internal(uint32 Index) const
 {
 	return nullptr;
@@ -1461,6 +1471,10 @@ UNeuralProfile* UMaterialInterface::GetNeuralProfile_Internal() const
 	return nullptr;
 }
 
+uint32 UMaterialInterface::NumSubsurfaceProfileRoot_Internal() const
+{
+	return 0u;
+}
 
 uint32 UMaterialInterface::NumSpecularProfile_Internal() const
 {
@@ -1519,8 +1533,9 @@ void UMaterialInterface::UpdateMaterialRenderProxy(FMaterialRenderProxy& Proxy)
 
 	FMaterialShadingModelField MaterialShadingModels = GetShadingModels();
 
-	// for better performance we only update SubsurfaceProfileRT if the feature is used
-	if (UseSubsurfaceProfile(MaterialShadingModels))
+	// For better performance we only update SubsurfaceProfileRT if the feature is used and substrate is not enabled
+	// When Substrate is enabled, this is ONLY used as an override for Subsurface Profile on material instance (override all Subsurface Profiles at once for now)
+	if (UseSubsurfaceProfile(MaterialShadingModels) && !Substrate::IsSubstrateEnabled())
 	{
 		USubsurfaceProfile* LocalSubsurfaceProfile = GetSubsurfaceProfile_Internal();
 		
@@ -1579,43 +1594,93 @@ void UMaterialInterface::UpdateMaterialRenderProxy(FMaterialRenderProxy& Proxy)
 
 	if (Substrate::IsSubstrateEnabled())
 	{
-		struct FEntry
-		{
-			USpecularProfile* Profile = nullptr;
-			FSpecularProfileStruct Settings;
-			const FTextureReference* Texture = nullptr;
-			FGuid Guid;
-		};
-		TArray<FEntry> Entries;
-		for (int32 It = 0, Count = NumSpecularProfile_Internal(); It<Count; ++It)
+		FMaterialRenderProxy* InProxy = &Proxy;
+
 		{
-			FEntry& Entry = Entries.AddDefaulted_GetRef();
-			Entry.Profile = GetSpecularProfile_Internal(It);
-			if (Entry.Profile)
+			struct FEntry
 			{
-				Entry.Settings 	= Entry.Profile->Settings;
-				Entry.Guid 		= Entry.Profile->Guid;	
-				if (!Entry.Settings.IsProcedural())
+				USpecularProfile* Profile = nullptr;
+				FSpecularProfileStruct Settings;
+				const FTextureReference* Texture = nullptr;
+				FGuid Guid;
+			};
+			TArray<FEntry> Entries;
+			for (int32 It = 0, Count = NumSpecularProfile_Internal(); It<Count; ++It)
+			{
+				FEntry& Entry = Entries.AddDefaulted_GetRef();
+				Entry.Profile = GetSpecularProfile_Internal(It);
+				if (Entry.Profile)
 				{
-					Entry.Texture = &Entry.Profile->Settings.Texture->TextureReference;
+					Entry.Settings 	= Entry.Profile->Settings;
+					Entry.Guid 		= Entry.Profile->Guid;	
+					if (!Entry.Settings.IsProcedural())
+					{
+						Entry.Texture = &Entry.Profile->Settings.Texture->TextureReference;
+					}
 				}
 			}
+
+			ENQUEUE_RENDER_COMMAND(UpdateMaterialRenderProxySpecular)(
+			[InProxy, Entries](FRHICommandListImmediate& RHICmdList)
+			{
+				for (const FEntry& Entry : Entries)
+				{
+					if (Entry.Profile)
+					{
+						const uint32 AllocationId = SpecularProfileAtlas::AddOrUpdateProfile(Entry.Profile, Entry.Guid, Entry.Settings, Entry.Texture);
+						check(AllocationId >= 0 && AllocationId < MAX_SPECULAR_PROFILE_COUNT);
+					}
+					InProxy->AddSpecularProfileRT(Entry.Profile);
+				}
+			});
 		}
 
-		FMaterialRenderProxy* InProxy = &Proxy;
-		ENQUEUE_RENDER_COMMAND(UpdateMaterialRenderProxySpecular)(
-		[InProxy, Entries](FRHICommandListImmediate& RHICmdList)
 		{
-			for (const FEntry& Entry : Entries)
+			struct FEntry
 			{
+				USubsurfaceProfile* Profile = nullptr;
+				FSubsurfaceProfileStruct Settings;
+			};
+			TArray<FEntry> Entries;
+			for (int32 It = 0, Count = NumSubsurfaceProfileRoot_Internal(); It < Count; ++It)
+			{
+				FEntry& Entry = Entries.AddDefaulted_GetRef();
+				Entry.Profile = GetSubsurfaceProfileRoot_Internal(It);
 				if (Entry.Profile)
 				{
-					const uint32 AllocationId = SpecularProfileAtlas::AddOrUpdateProfile(Entry.Profile, Entry.Guid, Entry.Settings, Entry.Texture);
-					check(AllocationId >= 0 && AllocationId < MAX_SPECULAR_PROFILE_COUNT);
+					Entry.Settings = Entry.Profile->Settings;
 				}
-				InProxy->AddSpecularProfileRT(Entry.Profile);
 			}
-		});
+
+			USubsurfaceProfile* SubsurfaceProfileOverride = GetSubsurfaceProfileOverride_Internal();
+			FSubsurfaceProfileStruct SubsurfaceProfileOverrideSettings;
+			if (SubsurfaceProfileOverride)
+			{
+				SubsurfaceProfileOverrideSettings = SubsurfaceProfileOverride->Settings;
+			}
+
+			ENQUEUE_RENDER_COMMAND(UpdateMaterialRenderProxySubsurfaceProfiles)(
+			[InProxy, Entries, SubsurfaceProfileOverride, SubsurfaceProfileOverrideSettings](FRHICommandListImmediate& RHICmdList)
+			{
+				for (const FEntry& Entry : Entries)
+				{
+					if (Entry.Profile)
+					{
+						const uint32 AllocationId = GSubsurfaceProfileTextureObject.AddOrUpdateProfile(Entry.Settings, Entry.Profile);
+						check(AllocationId >= 0 && AllocationId < MAX_SUBSURFACE_PROFILE_COUNT);
+					}
+					InProxy->AddSubsurfaceProfileRT(Entry.Profile);
+				}
+
+				if (SubsurfaceProfileOverride)
+				{
+					GSubsurfaceProfileTextureObject.AddOrUpdateProfile(SubsurfaceProfileOverrideSettings, SubsurfaceProfileOverride);
+				}
+
+				// Set the subsurface profile override using the default SubsurfaceProfile pointer.
+				InProxy->SetSubsurfaceProfileRT(SubsurfaceProfileOverride);
+			});
+		}
 	}
 }
 
diff --git a/Engine/Source/Runtime/Engine/Private/Rendering/SubsurfaceProfile.cpp b/Engine/Source/Runtime/Engine/Private/Rendering/SubsurfaceProfile.cpp
--- a/Engine/Source/Runtime/Engine/Private/Rendering/SubsurfaceProfile.cpp
+++ b/Engine/Source/Runtime/Engine/Private/Rendering/SubsurfaceProfile.cpp
@@ -8,6 +8,7 @@
 #include "RenderTargetPool.h"
 #include "PixelShaderUtils.h"
 #include "RenderingThread.h"
+#include "UObject/FortniteMainBranchObjectVersion.h"
 
 #include UE_INLINE_GENERATED_CPP_BY_NAME(SubsurfaceProfile)
 
@@ -716,6 +717,16 @@ FName GetSubsurfaceProfileParameterName()
 	return NameSubsurfaceProfile;
 }
 
+static FName CreateSubsurfaceProfileParameterName(const FGuid& InGuid)
+{
+	return FName(TEXT("__SubsurfaceProfile") + InGuid.ToString());
+}
+
+FName CreateSubsurfaceProfileParameterName(const USubsurfaceProfile* InProfile)
+{
+	return InProfile ? CreateSubsurfaceProfileParameterName(InProfile->Guid) : FName();
+}
+
 float GetSubsurfaceProfileId(const USubsurfaceProfile* In)
 {
 	int32 AllocationId = 0;
@@ -785,6 +796,14 @@ void USubsurfaceProfile::PostEditChangeProperty(struct FPropertyChangedEvent& Pr
 		});
 }
 
+void USubsurfaceProfile::PostDuplicate(EDuplicateMode::Type DuplicateMode)
+{
+	Super::PostDuplicate(DuplicateMode);
+
+	// When a Subsurface Profile asset is duplicated/copied pasted (e.g. from the asset browser), we want the guid to be regenerated.
+	Guid = FGuid::NewGuid();
+}
+
 void USubsurfaceProfile::PostLoad()
 {
 	Super::PostLoad();
@@ -792,3 +811,17 @@ void USubsurfaceProfile::PostLoad()
 	UpgradeSubsurfaceProfileParameters(this->Settings);
 }
 
+void USubsurfaceProfile::Serialize(FArchive& Ar)
+{
+	Super::Serialize(Ar);
+
+	Ar.UsingCustomVersion(FFortniteMainBranchObjectVersion::GUID);
+
+	if (Ar.CustomVer(FFortniteMainBranchObjectVersion::GUID) < FFortniteMainBranchObjectVersion::SubsurfaceProfileGuid && Ar.IsLoading())
+	{
+		// Generate a unique GUID from the unique project asset path the first time we load an older asset.
+		const FString PathName = this->GetPathName();
+		Guid = FGuid::NewDeterministicGuid(PathName, 123u);
+	}
+}
+
diff --git a/Engine/Source/Runtime/Engine/Public/Materials/Material.h b/Engine/Source/Runtime/Engine/Public/Materials/Material.h
--- a/Engine/Source/Runtime/Engine/Public/Materials/Material.h
+++ b/Engine/Source/Runtime/Engine/Public/Materials/Material.h
@@ -1201,6 +1201,9 @@ public:
 	ENGINE_API virtual bool IsPostProcessMaterial() const;
 	ENGINE_API bool IsPostProcessMaterialOutputingAlpha() const;
 	ENGINE_API virtual USubsurfaceProfile* GetSubsurfaceProfile_Internal() const override;
+	ENGINE_API virtual uint32 NumSubsurfaceProfileRoot_Internal() const override;
+	ENGINE_API virtual USubsurfaceProfile* GetSubsurfaceProfileRoot_Internal(uint32 Index) const override;
+	ENGINE_API virtual USubsurfaceProfile* GetSubsurfaceProfileOverride_Internal() const override;
 	ENGINE_API virtual uint32 NumSpecularProfile_Internal() const override;
 	ENGINE_API virtual USpecularProfile* GetSpecularProfile_Internal(uint32 Index) const override;
 	ENGINE_API virtual UNeuralProfile* GetNeuralProfile_Internal() const override;
diff --git a/Engine/Source/Runtime/Engine/Public/Materials/MaterialInstance.h b/Engine/Source/Runtime/Engine/Public/Materials/MaterialInstance.h
--- a/Engine/Source/Runtime/Engine/Public/Materials/MaterialInstance.h
+++ b/Engine/Source/Runtime/Engine/Public/Materials/MaterialInstance.h
@@ -901,6 +901,9 @@ public:
 	ENGINE_API virtual bool HasCustomizedUVs() const override;
 	ENGINE_API virtual bool HasPixelAnimation() const override;
 	ENGINE_API virtual USubsurfaceProfile* GetSubsurfaceProfile_Internal() const override;
+	ENGINE_API uint32 NumSubsurfaceProfileRoot_Internal() const override;
+	ENGINE_API USubsurfaceProfile* GetSubsurfaceProfileRoot_Internal(uint32 Index) const override;
+	ENGINE_API virtual USubsurfaceProfile* GetSubsurfaceProfileOverride_Internal() const override;
 	ENGINE_API virtual bool CastsRayTracedShadows() const override;
 	ENGINE_API virtual bool IsTessellationEnabled() const override;
 
diff --git a/Engine/Source/Runtime/Engine/Public/Materials/MaterialInterface.h b/Engine/Source/Runtime/Engine/Public/Materials/MaterialInterface.h
--- a/Engine/Source/Runtime/Engine/Public/Materials/MaterialInterface.h
+++ b/Engine/Source/Runtime/Engine/Public/Materials/MaterialInterface.h
@@ -275,6 +275,12 @@ public:
 	UPROPERTY(EditAnywhere, BlueprintReadOnly, Category = Material, meta = (DisplayName = "Subsurface Profile"))
 	TObjectPtr<class USubsurfaceProfile> SubsurfaceProfile;
 
+	/** Subsurface Profiles. For internal usage, not editable/visible. 
+	 * For Substrate, there can be many in a material similarly to SpecularProfile (even though only one can be specified per pixel due to the post processing) 
+	 */
+	UPROPERTY()
+	TArray<TObjectPtr<class USubsurfaceProfile>> SubsurfaceProfiles;
+
 	/** Specular Profile. For internal usage, not editable/visible */
 	UPROPERTY()
 	TArray<TObjectPtr<class USpecularProfile>> SpecularProfiles;
@@ -945,6 +951,9 @@ public:
 	ENGINE_API virtual bool HasCustomizedUVs() const;
 	ENGINE_API virtual bool HasPixelAnimation() const;
 	ENGINE_API virtual USubsurfaceProfile* GetSubsurfaceProfile_Internal() const;
+	ENGINE_API virtual uint32 NumSubsurfaceProfileRoot_Internal() const;
+	ENGINE_API virtual USubsurfaceProfile* GetSubsurfaceProfileRoot_Internal(uint32 Index) const;
+	ENGINE_API virtual USubsurfaceProfile* GetSubsurfaceProfileOverride_Internal() const;
 	ENGINE_API virtual uint32 NumSpecularProfile_Internal() const;
 	ENGINE_API virtual USpecularProfile* GetSpecularProfile_Internal(uint32 Index) const;
 	ENGINE_API virtual UNeuralProfile* GetNeuralProfile_Internal() const;
diff --git a/Engine/Source/Runtime/Engine/Public/Materials/MaterialRenderProxy.h b/Engine/Source/Runtime/Engine/Public/Materials/MaterialRenderProxy.h
--- a/Engine/Source/Runtime/Engine/Public/Materials/MaterialRenderProxy.h
+++ b/Engine/Source/Runtime/Engine/Public/Materials/MaterialRenderProxy.h
@@ -213,10 +213,16 @@ public:
 	}
 #endif
 
-	// Subsurface profiles
+	// Subsurface profile.
+	// When Substrate is enabled, this is ONLY used as an override for Subsurface Profile on material instance (override all Subsurface Profiles at once for now)
 	void SetSubsurfaceProfileRT(const USubsurfaceProfile* Ptr) { SubsurfaceProfileRT = Ptr; }
 	const USubsurfaceProfile* GetSubsurfaceProfileRT() const { return SubsurfaceProfileRT; }
 
+	// Subsurface profiles
+	void AddSubsurfaceProfileRT(const USubsurfaceProfile* Ptr) { SubsurfaceProfilesRT.Add(Ptr); }
+	const USubsurfaceProfile* GetSubsurfaceProfileRT(uint32 Index) const { check(Index<uint32(SubsurfaceProfilesRT.Num())); return SubsurfaceProfilesRT[Index]; }
+	const uint32 NumSubsurfaceProfileRT() const { return SubsurfaceProfilesRT.Num(); }
+
 	// Specular profiles
 	void AddSpecularProfileRT(const USpecularProfile* Ptr) { SpecularProfilesRT.Add(Ptr); }
 	const USpecularProfile* GetSpecularProfileRT(uint32 Index) const { check(Index<uint32(SpecularProfilesRT.Num())); return SpecularProfilesRT[Index]; }
@@ -243,7 +249,8 @@ private:
 	virtual void FinishCacheUniformExpressions() const {}
 
 	/** 0 if not set, game thread pointer, do not dereference, only for comparison */
-	const USubsurfaceProfile* SubsurfaceProfileRT;
+	const USubsurfaceProfile* SubsurfaceProfileRT;	// Overrides all SubsurfaceProfilesRT when set, used when set from material instance to respect the legacy workflow
+	TArray<const USubsurfaceProfile*> SubsurfaceProfilesRT;
 	TArray<const USpecularProfile*> SpecularProfilesRT;
 	const UNeuralProfile* NeuralProfileRT;
 	FString MaterialName;
