Substrate - Fixed override of SubsurfaceProfile on instance when the parent material is not SubsurfaceProfile and do not specify any profile asset.

#rnx
#tests Test with substrate the overide of a SubsutrateShadingModel node without subsurfaceprofile from a mateial instance to a SubsurfaceShading model with a Subsurface profile.

diff --git a/Engine/Source/Runtime/Core/Private/UObject/FortniteMainBranchObjectVersion.cpp b/Engine/Source/Runtime/Core/Private/UObject/FortniteMainBranchObjectVersion.cpp
--- a/Engine/Source/Runtime/Core/Private/UObject/FortniteMainBranchObjectVersion.cpp
+++ b/Engine/Source/Runtime/Core/Private/UObject/FortniteMainBranchObjectVersion.cpp
@@ -14,6 +14,6 @@ TMap<FGuid, FGuid> FFortniteMainBranchObjectVersion::GetSystemGuids()
 	SystemGuids.Add(DevGuids.Niagara_LatestScriptCompileVersion, FGuid("F8D995224324470E97ED718D33939A0E"));
 	SystemGuids.Add(DevGuids.SkeletalMeshDerivedDataVersion, FGuid("B9AEAA3EE7AC4FCFB3860E8F3E06DD75"));
 	SystemGuids.Add(DevGuids.STATICMESH_DERIVEDDATA_VER, FGuid("1C2EBE3310464FD7995ADE1D6AB92BDC"));
-	SystemGuids.Add(DevGuids.MaterialTranslationDDCVersion, FGuid("79D0C4D88D7C4D1491C86B9922001F20"));
+	SystemGuids.Add(DevGuids.MaterialTranslationDDCVersion, FGuid("0F9C4453CDDC4B8BA1527D4942381ECA"));
 	return SystemGuids;
 }
diff --git a/Engine/Source/Runtime/Engine/Private/Materials/MaterialExpressionSubstrate.cpp b/Engine/Source/Runtime/Engine/Private/Materials/MaterialExpressionSubstrate.cpp
--- a/Engine/Source/Runtime/Engine/Private/Materials/MaterialExpressionSubstrate.cpp
+++ b/Engine/Source/Runtime/Engine/Private/Materials/MaterialExpressionSubstrate.cpp
@@ -175,6 +175,24 @@ static int32 CompileWithDefaultTangentWS(class FMaterialCompiler* Compiler, FExp
 	return Compiler->VertexTangent();
 }
 
+FName CreateSpecularProfileParameterName(USpecularProfile* InProfile);
+FName CreateSubsurfaceProfileParameterName(const USubsurfaceProfile* InProfile);
+
+static int32 CreateSubsurfaceProfileParameter(class FMaterialCompiler* Compiler, USubsurfaceProfile* InProfile)
+{
+	check(InProfile);
+	const FName SubsurfaceProfileParameterName = CreateSubsurfaceProfileParameterName(InProfile);
+	const int32 SSSProfileCodeChunk = Compiler->ForceCast(Compiler->ScalarParameter(SubsurfaceProfileParameterName, 1.0f), MCT_Float1);
+	return SSSProfileCodeChunk;
+}
+
+static int32 CreateDefaultSubsurfaceProfileParameter(class FMaterialCompiler* Compiler)
+{
+	const int32 SSSProfileCodeChunk = Compiler->ForceCast(Compiler->ScalarParameter(GetSubsurfaceProfileParameterName(), 1.0f), MCT_Float1);
+	return SSSProfileCodeChunk;
+}
+
+
 #endif // WITH_EDITOR
 
 UMaterialExpressionSubstrateShadingModels::UMaterialExpressionSubstrateShadingModels(const FObjectInitializer& ObjectInitializer)
@@ -290,8 +308,11 @@ int32 UMaterialExpressionSubstrateShadingModels::Compile(class FMaterialCompiler
 	const bool bHasSSS = HasSSS();
 	if (bHasSSS)
 	{
-		const FName SubsurfaceProfileParameterName = CreateSubsurfaceProfileParameterName(SubsurfaceProfile);
-		SSSProfileCodeChunk = Compiler->ForceCast(Compiler->ScalarParameter(SubsurfaceProfileParameterName, 1.0f), MCT_Float1);
+		SSSProfileCodeChunk = CreateSubsurfaceProfileParameter(Compiler, SubsurfaceProfile);
+	}
+	else
+	{
+		SSSProfileCodeChunk = CreateDefaultSubsurfaceProfileParameter(Compiler);
 	}
 
 	FSubstrateOperator& SubstrateOperator = Compiler->SubstrateCompilationGetOperator(Compiler->SubstrateTreeStackGetPathUniqueId());
@@ -713,8 +734,6 @@ UMaterialExpressionSubstrateSlabBSDF::UMaterialExpressionSubstrateSlabBSDF(const
 }
 
 #if WITH_EDITOR
-FName CreateSpecularProfileParameterName(USpecularProfile* InProfile);
-FName CreateSubsurfaceProfileParameterName(const USubsurfaceProfile* InProfile);
 int32 UMaterialExpressionSubstrateSlabBSDF::Compile(class FMaterialCompiler* Compiler, int32 OutputIndex)
 {
 	FGuid PathUniqueId = Compiler->SubstrateTreeStackGetPathUniqueId();
@@ -727,8 +746,11 @@ int32 UMaterialExpressionSubstrateSlabBSDF::Compile(class FMaterialCompiler* Com
 	int32 SSSProfileCodeChunk = INDEX_NONE;
 	if (SubstrateOperator.Has(ESubstrateBsdfFeature_SSS))
 	{
-		const FName SubsurfaceProfileParameterName = CreateSubsurfaceProfileParameterName(SubsurfaceProfile);
-		SSSProfileCodeChunk = Compiler->ForceCast(Compiler->ScalarParameter(SubsurfaceProfileParameterName, 1.0f), MCT_Float1);
+		SSSProfileCodeChunk = CreateSubsurfaceProfileParameter(Compiler, SubsurfaceProfile);
+	}
+	else
+	{
+		SSSProfileCodeChunk = CreateDefaultSubsurfaceProfileParameter(Compiler);
 	}
 
 	int32 SpecularProfileCodeChunk = INDEX_NONE;
@@ -2061,8 +2083,11 @@ int32 UMaterialExpressionSubstrateEyeBSDF::Compile(class FMaterialCompiler* Comp
 	int32 SSSProfileCodeChunk = INDEX_NONE;
 	if (SubsurfaceProfile != nullptr)
 	{
-		const FName SubsurfaceProfileParameterName = CreateSubsurfaceProfileParameterName(SubsurfaceProfile);
-		SSSProfileCodeChunk = Compiler->ForceCast(Compiler->ScalarParameter(SubsurfaceProfileParameterName, 1.0f), MCT_Float1);
+		SSSProfileCodeChunk = CreateSubsurfaceProfileParameter(Compiler, SubsurfaceProfile);
+	}
+	else
+	{
+		SSSProfileCodeChunk = CreateDefaultSubsurfaceProfileParameter(Compiler);
 	}
 
 	FSubstrateOperator& SubstrateOperator = Compiler->SubstrateCompilationGetOperator(Compiler->SubstrateTreeStackGetPathUniqueId());
@@ -3584,8 +3609,11 @@ int32 UMaterialExpressionSubstrateConvertMaterialAttributes::Compile(class FMate
 	const bool bHasSSS = HasSSS();
 	if (bHasSSS)
 	{
-		const FName SubsurfaceProfileParameterName = CreateSubsurfaceProfileParameterName(SubsurfaceProfile);
-		SSSProfileCodeChunk = Compiler->ForceCast(Compiler->ScalarParameter(SubsurfaceProfileParameterName, 1.0f), MCT_Float1);
+		SSSProfileCodeChunk = CreateSubsurfaceProfileParameter(Compiler, SubsurfaceProfile);
+	}
+	else
+	{
+		SSSProfileCodeChunk = CreateDefaultSubsurfaceProfileParameter(Compiler);
 	}
 	SSSProfileCodeChunk = SSSProfileCodeChunk != INDEX_NONE ? SSSProfileCodeChunk : Compiler->Constant(0.0f);
 
