Substrate - fixed assert with null Subsurface Profiles on slab.

#rnx

diff --git a/Engine/Source/Runtime/Engine/Private/Materials/MaterialExpressionSubstrate.cpp b/Engine/Source/Runtime/Engine/Private/Materials/MaterialExpressionSubstrate.cpp
--- a/Engine/Source/Runtime/Engine/Private/Materials/MaterialExpressionSubstrate.cpp
+++ b/Engine/Source/Runtime/Engine/Private/Materials/MaterialExpressionSubstrate.cpp
@@ -306,7 +306,7 @@ int32 UMaterialExpressionSubstrateShadingModels::Compile(class FMaterialCompiler
 
 	int32 SSSProfileCodeChunk = INDEX_NONE;
 	const bool bHasSSS = HasSSS();
-	if (bHasSSS)
+	if (bHasSSS && SubsurfaceProfile)
 	{
 		SSSProfileCodeChunk = CreateSubsurfaceProfileParameter(Compiler, SubsurfaceProfile);
 	}
@@ -744,7 +744,7 @@ int32 UMaterialExpressionSubstrateSlabBSDF::Compile(class FMaterialCompiler* Com
 	const bool bHasAnisotropy = SubstrateOperator.Has(ESubstrateBsdfFeature_Anisotropy);
 
 	int32 SSSProfileCodeChunk = INDEX_NONE;
-	if (SubstrateOperator.Has(ESubstrateBsdfFeature_SSS))
+	if (SubstrateOperator.Has(ESubstrateBsdfFeature_SSS) && SubsurfaceProfile)
 	{
 		SSSProfileCodeChunk = CreateSubsurfaceProfileParameter(Compiler, SubsurfaceProfile);
 	}
@@ -3607,7 +3607,7 @@ int32 UMaterialExpressionSubstrateConvertMaterialAttributes::Compile(class FMate
 	// Need to handle this by looking at the material instead of the node?
 	int32 SSSProfileCodeChunk = INDEX_NONE;
 	const bool bHasSSS = HasSSS();
-	if (bHasSSS)
+	if (bHasSSS && SubsurfaceProfile)
 	{
 		SSSProfileCodeChunk = CreateSubsurfaceProfileParameter(Compiler, SubsurfaceProfile);
 	}
