Add namespace for SubsurfaceProfile functions
Rename SpecularProfileAtlas namespace into Specular profile.

This is more consitent with other profile type (e.g. NeuralProfile).

diff --git a/Engine/Source/Runtime/Engine/Classes/Engine/SpecularProfile.h b/Engine/Source/Runtime/Engine/Classes/Engine/SpecularProfile.h
--- a/Engine/Source/Runtime/Engine/Classes/Engine/SpecularProfile.h
+++ b/Engine/Source/Runtime/Engine/Classes/Engine/SpecularProfile.h
@@ -85,7 +85,7 @@ class USpecularProfile : public UObject
 	//~ End UObject Interface
 };
 
-namespace SpecularProfileAtlas
+namespace SpecularProfile
 {
 // Atlas - Initializes or updates the contents of the specular profile texture.
 ENGINE_API void UpdateSpecularProfileTextureAtlas(FRDGBuilder& GraphBuilder, EShaderPlatform ShaderPlatform);
@@ -104,4 +104,7 @@ ENGINE_API FName GetSpecularProfileParameterName(const USpecularProfile* InProfi
 
 // Profile - Returns the specular profile ID for a given Specular Profile object
 ENGINE_API float GetSpecularProfileId(const USpecularProfile* In);
+
+// Profile - Returns the shader parameter name for a Specular profile.
+ENGINE_API FName CreateSpecularProfileParameterName(USpecularProfile* InProfile);
 }
\ No newline at end of file
diff --git a/Engine/Source/Runtime/Engine/Classes/Engine/SubsurfaceProfile.h b/Engine/Source/Runtime/Engine/Classes/Engine/SubsurfaceProfile.h
--- a/Engine/Source/Runtime/Engine/Classes/Engine/SubsurfaceProfile.h
+++ b/Engine/Source/Runtime/Engine/Classes/Engine/SubsurfaceProfile.h
@@ -271,22 +271,53 @@ static const int32 SUBSURFACE_KERNEL_SIZE = 3;
 extern ENGINE_API TGlobalResource<FSubsurfaceProfileTexture> GSubsurfaceProfileTextureObject;
 
 // Initializes or updates the contents of the subsurface profile texture.
+UE_DEPRECATED(5.5, "Please, use the name space SubsurfaceProfile::")
 ENGINE_API void UpdateSubsurfaceProfileTexture(class FRDGBuilder& GraphBuilder, EShaderPlatform ShaderPlatform);
 
 // Returns the subsurface profile texture if it exists, or null.
+UE_DEPRECATED(5.5, "Please, use the name space SubsurfaceProfile::")
 ENGINE_API FRHITexture* GetSubsurfaceProfileTexture();
 
 // Returns the subsurface profile texture if it exists, or black.
+UE_DEPRECATED(5.5, "Please, use the name space SubsurfaceProfile::")
 ENGINE_API FRHITexture* GetSubsurfaceProfileTextureWithFallback();
 
 // Returns the Preintegrated texture if it exists, or black.
+UE_DEPRECATED(5.5, "Please, use the name space SubsurfaceProfile::")
 ENGINE_API FRHITexture* GetSSProfilesPreIntegratedTextureWithFallback();
 
 // Returns the subsurface profile ID shader parameter name
+UE_DEPRECATED(5.5, "Please, use the name space SubsurfaceProfile::")
 ENGINE_API FName GetSubsurfaceProfileParameterName();
 
 // Returns the subsurface profile ID for a given Sub-surface Profile object
+UE_DEPRECATED(5.5, "Please, use the name space SubsurfaceProfile::")
 ENGINE_API float GetSubsurfaceProfileId(const USubsurfaceProfile* In);
 
 // Returns the shader parameter name for a Subsurface profile.
-ENGINE_API FName CreateSubsurfaceProfileParameterName(const USubsurfaceProfile* InProfile);
\ No newline at end of file
+UE_DEPRECATED(5.5, "Please, use the name space SubsurfaceProfile::")
+ENGINE_API FName CreateSubsurfaceProfileParameterName(const USubsurfaceProfile* InProfile);
+
+namespace SubsurfaceProfile
+{
+// Initializes or updates the contents of the subsurface profile texture.
+ENGINE_API void UpdateSubsurfaceProfileTexture(class FRDGBuilder& GraphBuilder, EShaderPlatform ShaderPlatform);
+
+// Returns the subsurface profile texture if it exists, or null.
+ENGINE_API FRHITexture* GetSubsurfaceProfileTexture();
+
+// Returns the subsurface profile texture if it exists, or black.
+ENGINE_API FRHITexture* GetSubsurfaceProfileTextureWithFallback();
+
+// Returns the Preintegrated texture if it exists, or black.
+ENGINE_API FRHITexture* GetSSProfilesPreIntegratedTextureWithFallback();
+
+// Returns the subsurface profile ID shader parameter name
+ENGINE_API FName GetSubsurfaceProfileParameterName();
+
+// Returns the subsurface profile ID for a given Sub-surface Profile object
+ENGINE_API float GetSubsurfaceProfileId(const USubsurfaceProfile* In);
+
+// Returns the shader parameter name for a Subsurface profile.
+ENGINE_API FName CreateSubsurfaceProfileParameterName(const USubsurfaceProfile* InProfile);
+};
\ No newline at end of file
diff --git a/Engine/Source/Runtime/Engine/Private/Materials/HLSLMaterialTranslator.cpp b/Engine/Source/Runtime/Engine/Private/Materials/HLSLMaterialTranslator.cpp
--- a/Engine/Source/Runtime/Engine/Private/Materials/HLSLMaterialTranslator.cpp
+++ b/Engine/Source/Runtime/Engine/Private/Materials/HLSLMaterialTranslator.cpp
@@ -61,6 +61,7 @@
 #include "DerivedDataCache.h"
 #include "MaterialCachedData.h"
 #include "PostProcess/PostProcessMaterialInputs.h"
+#include "Engine/SubsurfaceProfile.h"
 
 #if WITH_EDITORONLY_DATA
 #include "Materials/MaterialExpressionSubstrate.h"
@@ -1623,7 +1624,7 @@ void FHLSLMaterialTranslator::TranslateMaterial()
 			int32 SubsurfaceColor = Material->CompilePropertyAndSetMaterialProperty(MP_SubsurfaceColor, this);
 			SubsurfaceColor = ForceCast(SubsurfaceColor, FMaterialAttributeDefinitionMap::GetValueType(MP_SubsurfaceColor), MFCF_ExactMatch | MFCF_ReplicateValue);
 
-			static FName NameSubsurfaceProfile(TEXT("__SubsurfaceProfile"));
+			static FName NameSubsurfaceProfile = SubsurfaceProfile::GetSubsurfaceProfileParameterName();
 
 			// 1.0f is is a not used profile - later this gets replaced with the actual profile
 			int32 CodeSubsurfaceProfile = ForceCast(ScalarParameter(NameSubsurfaceProfile, 1.0f), MCT_Float1);
diff --git a/Engine/Source/Runtime/Engine/Private/Materials/Material.cpp b/Engine/Source/Runtime/Engine/Private/Materials/Material.cpp
--- a/Engine/Source/Runtime/Engine/Private/Materials/Material.cpp
+++ b/Engine/Source/Runtime/Engine/Private/Materials/Material.cpp
@@ -446,10 +446,10 @@ public:
 		const FMaterialResource* MaterialResource = Material->GetMaterialResource(Context.Material.GetFeatureLevel());
 		if (MaterialResource && MaterialResource->GetRenderingThreadShaderMap())
 		{
-			if (Type == EMaterialParameterType::Scalar && ParameterInfo.Name == GetSubsurfaceProfileParameterName())
+			if (Type == EMaterialParameterType::Scalar && ParameterInfo.Name == SubsurfaceProfile::GetSubsurfaceProfileParameterName())
 			{
 				// Legacy single SubsurfaceProfile (Substrate do not add this one in the material shader)
-				OutValue = GetSubsurfaceProfileId(GetSubsurfaceProfileRT());
+				OutValue = SubsurfaceProfile::GetSubsurfaceProfileId(GetSubsurfaceProfileRT());
 				return true;
 			}
 			else if (Type == EMaterialParameterType::Scalar && NumSubsurfaceProfileRT() > 0)
@@ -459,10 +459,10 @@ public:
 				for (uint32 It = 0, Count = NumSubsurfaceProfileRT(); It < Count; ++It)
 				{
 					const USubsurfaceProfile* SSProfileRT = GetSubsurfaceProfileRT(It);
-					if (ParameterInfo.Name == CreateSubsurfaceProfileParameterName(SSProfileRT))
+					if (ParameterInfo.Name == SubsurfaceProfile::CreateSubsurfaceProfileParameterName(SSProfileRT))
 					{
 						// Set the root material Profile, or the profile overriden by any instances.
-						OutValue = GetSubsurfaceProfileId(SSProfileOverrideRT ? SSProfileOverrideRT : SSProfileRT);
+						OutValue = SubsurfaceProfile::GetSubsurfaceProfileId(SSProfileOverrideRT ? SSProfileOverrideRT : SSProfileRT);
 						return true;
 					}
 				}
@@ -471,9 +471,9 @@ public:
 			{
 				for (uint32 It=0,Count=NumSpecularProfileRT();It<Count;++It)
 				{
-					if (ParameterInfo.Name == SpecularProfileAtlas::GetSpecularProfileParameterName(GetSpecularProfileRT(It)))
+					if (ParameterInfo.Name == SpecularProfile::GetSpecularProfileParameterName(GetSpecularProfileRT(It)))
 					{
-						OutValue = SpecularProfileAtlas::GetSpecularProfileId(GetSpecularProfileRT(It));
+						OutValue = SpecularProfile::GetSpecularProfileId(GetSpecularProfileRT(It));
 						return true;
 					}
 				}
diff --git a/Engine/Source/Runtime/Engine/Private/Materials/MaterialExpressionSubstrate.cpp b/Engine/Source/Runtime/Engine/Private/Materials/MaterialExpressionSubstrate.cpp
--- a/Engine/Source/Runtime/Engine/Private/Materials/MaterialExpressionSubstrate.cpp
+++ b/Engine/Source/Runtime/Engine/Private/Materials/MaterialExpressionSubstrate.cpp
@@ -175,20 +175,17 @@ static int32 CompileWithDefaultTangentWS(class FMaterialCompiler* Compiler, FExp
 	return Compiler->VertexTangent();
 }
 
-FName CreateSpecularProfileParameterName(USpecularProfile* InProfile);
-FName CreateSubsurfaceProfileParameterName(const USubsurfaceProfile* InProfile);
-
 static int32 CreateSubsurfaceProfileParameter(class FMaterialCompiler* Compiler, USubsurfaceProfile* InProfile)
 {
 	check(InProfile);
-	const FName SubsurfaceProfileParameterName = CreateSubsurfaceProfileParameterName(InProfile);
+	const FName SubsurfaceProfileParameterName = SubsurfaceProfile::CreateSubsurfaceProfileParameterName(InProfile);
 	const int32 SSSProfileCodeChunk = Compiler->ForceCast(Compiler->ScalarParameter(SubsurfaceProfileParameterName, 1.0f), MCT_Float1);
 	return SSSProfileCodeChunk;
 }
 
 static int32 CreateDefaultSubsurfaceProfileParameter(class FMaterialCompiler* Compiler)
 {
-	const int32 SSSProfileCodeChunk = Compiler->ForceCast(Compiler->ScalarParameter(GetSubsurfaceProfileParameterName(), 1.0f), MCT_Float1);
+	const int32 SSSProfileCodeChunk = Compiler->ForceCast(Compiler->ScalarParameter(SubsurfaceProfile::GetSubsurfaceProfileParameterName(), 1.0f), MCT_Float1);
 	return SSSProfileCodeChunk;
 }
 
@@ -756,7 +753,7 @@ int32 UMaterialExpressionSubstrateSlabBSDF::Compile(class FMaterialCompiler* Com
 	int32 SpecularProfileCodeChunk = INDEX_NONE;
 	if (SubstrateOperator.Has(ESubstrateBsdfFeature_SpecularProfile))
 	{
-		const FName SpecularProfileParameterName = CreateSpecularProfileParameterName(SpecularProfile);
+		const FName SpecularProfileParameterName = SpecularProfile::CreateSpecularProfileParameterName(SpecularProfile);
 		SpecularProfileCodeChunk = Compiler->ForceCast(Compiler->ScalarParameter(SpecularProfileParameterName, 1.0f), MCT_Float1);
 	}
 
diff --git a/Engine/Source/Runtime/Engine/Private/Materials/MaterialHLSLEmitter.cpp b/Engine/Source/Runtime/Engine/Private/Materials/MaterialHLSLEmitter.cpp
--- a/Engine/Source/Runtime/Engine/Private/Materials/MaterialHLSLEmitter.cpp
+++ b/Engine/Source/Runtime/Engine/Private/Materials/MaterialHLSLEmitter.cpp
@@ -877,7 +877,7 @@ bool MaterialEmitHLSL(const FMaterialCompileTargetParameters& InCompilerTarget,
 			const EMaterialDomain MaterialDomain = EmitContext.Material->GetMaterialDomain();
 			if (MaterialDomain == MD_Volume || (MaterialDomain == MD_Surface && IsSubsurfaceShadingModel(EmitMaterialData.ShadingModelsFromCompilation)))
 			{
-				const FMaterialParameterInfo SubsurfaceProfileParameterInfo(GetSubsurfaceProfileParameterName());
+				const FMaterialParameterInfo SubsurfaceProfileParameterInfo(SubsurfaceProfile::GetSubsurfaceProfileParameterName());
 				const FMaterialParameterMetadata SubsurfaceProfileParameterMetadata(1.f);
 				const Material::FExpressionParameter SubsurfaceProfileExpression(SubsurfaceProfileParameterInfo, SubsurfaceProfileParameterMetadata);
 
diff --git a/Engine/Source/Runtime/Engine/Private/Materials/MaterialInstance.cpp b/Engine/Source/Runtime/Engine/Private/Materials/MaterialInstance.cpp
--- a/Engine/Source/Runtime/Engine/Private/Materials/MaterialInstance.cpp
+++ b/Engine/Source/Runtime/Engine/Private/Materials/MaterialInstance.cpp
@@ -291,13 +291,13 @@ bool FMaterialInstanceResource::GetParameterValue(EMaterialParameterType Type, c
 	bool bResult = false;
 
 	// Check for hard-coded parameters
-	if (Type == EMaterialParameterType::Scalar && ParameterInfo.Name == GetSubsurfaceProfileParameterName())
+	if (Type == EMaterialParameterType::Scalar && ParameterInfo.Name == SubsurfaceProfile::GetSubsurfaceProfileParameterName())
 	{
 #if !ENABLE_MATERIAL_LAYER_PROTOTYPE
 		check(ParameterInfo.Association == EMaterialParameterAssociation::GlobalParameter);
 #endif
 		const USubsurfaceProfile* MySubsurfaceProfileRT = GetSubsurfaceProfileRT();
-		OutValue = GetSubsurfaceProfileId(MySubsurfaceProfileRT);
+		OutValue = SubsurfaceProfile::GetSubsurfaceProfileId(MySubsurfaceProfileRT);
 		bResult = true;
 	}
 	else if (Type == EMaterialParameterType::Scalar && NumSubsurfaceProfileRT() > 0)
@@ -310,10 +310,10 @@ bool FMaterialInstanceResource::GetParameterValue(EMaterialParameterType Type, c
 		for (uint32 It = 0, Count = NumSubsurfaceProfileRT(); It < Count; ++It)
 		{
 			const USubsurfaceProfile* SSProfileRT = GetSubsurfaceProfileRT(It);
-			if (ParameterInfo.Name == CreateSubsurfaceProfileParameterName(SSProfileRT))
+			if (ParameterInfo.Name == SubsurfaceProfile::CreateSubsurfaceProfileParameterName(SSProfileRT))
 			{
 				// Set the root material Profile, or the profile overriden by any instances.
-				OutValue = GetSubsurfaceProfileId(SSProfileOverrideRT ? SSProfileOverrideRT : SSProfileRT);
+				OutValue = SubsurfaceProfile::GetSubsurfaceProfileId(SSProfileOverrideRT ? SSProfileOverrideRT : SSProfileRT);
 				bResult = true;
 				break;
 			}
@@ -323,12 +323,12 @@ bool FMaterialInstanceResource::GetParameterValue(EMaterialParameterType Type, c
 	{
 		for (uint32 It=0,Count=NumSpecularProfileRT();It<Count;++It)
 		{
-			if (ParameterInfo.Name == SpecularProfileAtlas::GetSpecularProfileParameterName(GetSpecularProfileRT(It)))
+			if (ParameterInfo.Name == SpecularProfile::GetSpecularProfileParameterName(GetSpecularProfileRT(It)))
 			{
 #if !ENABLE_MATERIAL_LAYER_PROTOTYPE
 				check(ParameterInfo.Association == EMaterialParameterAssociation::GlobalParameter);
 #endif
-				OutValue = SpecularProfileAtlas::GetSpecularProfileId(GetSpecularProfileRT(It));
+				OutValue = SpecularProfile::GetSpecularProfileId(GetSpecularProfileRT(It));
 				bResult = true;
 				break;
 			}
diff --git a/Engine/Source/Runtime/Engine/Private/Materials/MaterialInterface.cpp b/Engine/Source/Runtime/Engine/Private/Materials/MaterialInterface.cpp
--- a/Engine/Source/Runtime/Engine/Private/Materials/MaterialInterface.cpp
+++ b/Engine/Source/Runtime/Engine/Private/Materials/MaterialInterface.cpp
@@ -1627,7 +1627,7 @@ void UMaterialInterface::UpdateMaterialRenderProxy(FMaterialRenderProxy& Proxy)
 				{
 					if (Entry.Profile)
 					{
-						const uint32 AllocationId = SpecularProfileAtlas::AddOrUpdateProfile(Entry.Profile, Entry.Guid, Entry.Settings, Entry.Texture);
+						const uint32 AllocationId = SpecularProfile::AddOrUpdateProfile(Entry.Profile, Entry.Guid, Entry.Settings, Entry.Texture);
 						check(AllocationId >= 0 && AllocationId < MAX_SPECULAR_PROFILE_COUNT);
 					}
 					InProxy->AddSpecularProfileRT(Entry.Profile);
diff --git a/Engine/Source/Runtime/Engine/Private/Rendering/SpecularProfile.cpp b/Engine/Source/Runtime/Engine/Private/Rendering/SpecularProfile.cpp
--- a/Engine/Source/Runtime/Engine/Private/Rendering/SpecularProfile.cpp
+++ b/Engine/Source/Runtime/Engine/Private/Rendering/SpecularProfile.cpp
@@ -185,16 +185,6 @@ TGlobalResource<FSpecularProfileTextureManager> GSpecularProfileTextureManager;
 // SpecularProfile atlas storing several texture profiles or 0 if there is no user
 static TRefCountPtr<IPooledRenderTarget> GSpecularProfileTextureAtlas;
 
-static FName CreateSpecularProfileParameterName(const FGuid& InGuid)
-{
-	return FName(TEXT("__SpecularProfile") + InGuid.ToString());
-}
-
-FName CreateSpecularProfileParameterName(USpecularProfile* InProfile)
-{
-	return InProfile ? CreateSpecularProfileParameterName(InProfile->Guid) : FName();
-}
-
 FSpecularProfileTextureManager::FSpecularProfileTextureManager()
 {
 	check(IsInGameThread());
@@ -555,7 +545,7 @@ void USpecularProfile::PostDuplicate(EDuplicateMode::Type DuplicateMode)
 ///////////////////////////////////////////////////////////////////////////////////////////////////
 // Public API
 
-namespace SpecularProfileAtlas
+namespace SpecularProfile
 {
 	FName GetSpecularProfileParameterName(const USpecularProfile* In)
 	{
@@ -594,4 +584,14 @@ namespace SpecularProfileAtlas
 		GSpecularProfileTextureManager.GetAtlasTexture(GraphBuilder, ShaderPlatform);
 	}
 
-} // namespace SpecularProfileAtlas
+	
+	static FName CreateSpecularProfileParameterName(const FGuid& InGuid)
+	{
+		return FName(TEXT("__SpecularProfile") + InGuid.ToString());
+	}
+
+	FName CreateSpecularProfileParameterName(USpecularProfile* InProfile)
+	{
+		return InProfile ? CreateSpecularProfileParameterName(InProfile->Guid) : FName();
+	}
+} // namespace SpecularProfile
diff --git a/Engine/Source/Runtime/Engine/Private/Rendering/SubsurfaceProfile.cpp b/Engine/Source/Runtime/Engine/Private/Rendering/SubsurfaceProfile.cpp
--- a/Engine/Source/Runtime/Engine/Private/Rendering/SubsurfaceProfile.cpp
+++ b/Engine/Source/Runtime/Engine/Private/Rendering/SubsurfaceProfile.cpp
@@ -281,7 +281,7 @@ IPooledRenderTarget* FSubsurfaceProfileTexture::GetSSProfilesPreIntegratedTextur
 			FSSProfilePreIntegratedPS::FParameters* PassParameters = GraphBuilder.AllocParameters<FSSProfilePreIntegratedPS::FParameters>();
 			PassParameters->RenderTargets[0] = FRenderTargetBinding(ProfileTexture, ERenderTargetLoadAction::EClear, 0, i);
 
-			PassParameters->SourceSSProfilesTexture = GetSubsurfaceProfileTextureWithFallback();
+			PassParameters->SourceSSProfilesTexture = SubsurfaceProfile::GetSubsurfaceProfileTextureWithFallback();
 			PassParameters->SourceSSProfilesSampler = TStaticSamplerState<SF_Point, AM_Clamp, AM_Clamp, AM_Clamp>::GetRHI();
 
 			FIntVector SourceSSProfilesTextureSize = PassParameters->SourceSSProfilesTexture->GetSizeXYZ();
@@ -711,6 +711,44 @@ void FSubsurfaceProfileTexture::Dump()
 #endif
 }
 
+FName GetSubsurfaceProfileParameterName()
+{
+	return SubsurfaceProfile::GetSubsurfaceProfileParameterName();
+}
+
+FName CreateSubsurfaceProfileParameterName(const USubsurfaceProfile* InProfile)
+{
+	return SubsurfaceProfile::CreateSubsurfaceProfileParameterName(InProfile);
+}
+
+float GetSubsurfaceProfileId(const USubsurfaceProfile* In)
+{
+	return SubsurfaceProfile::GetSubsurfaceProfileId(In);
+}
+
+FRHITexture* GetSubsurfaceProfileTexture()
+{
+	return SubsurfaceProfile::GetSubsurfaceProfileTexture();
+}
+
+FRHITexture* GetSubsurfaceProfileTextureWithFallback()
+{
+	return SubsurfaceProfile::GetSubsurfaceProfileTextureWithFallback();
+}
+
+FRHITexture* GetSSProfilesPreIntegratedTextureWithFallback()
+{
+	return SubsurfaceProfile::GetSSProfilesPreIntegratedTextureWithFallback();
+}
+
+void UpdateSubsurfaceProfileTexture(FRDGBuilder& GraphBuilder, EShaderPlatform ShaderPlatform)
+{
+	SubsurfaceProfile::UpdateSubsurfaceProfileTexture(GraphBuilder, ShaderPlatform);
+}
+
+namespace SubsurfaceProfile
+{
+
 FName GetSubsurfaceProfileParameterName()
 {
 	static FName NameSubsurfaceProfile(TEXT("__SubsurfaceProfile"));
@@ -764,6 +802,8 @@ void UpdateSubsurfaceProfileTexture(FRDGBuilder& GraphBuilder, EShaderPlatform S
 	GSubsurfaceProfileTextureObject.GetSSProfilesPreIntegratedTexture(GraphBuilder, ShaderPlatform);
 }
 
+}// namespace SubsurfaceProfile
+
 // ------------------------------------------------------
 
 USubsurfaceProfile::USubsurfaceProfile(const FObjectInitializer& ObjectInitializer)
diff --git a/Engine/Source/Runtime/Renderer/Private/DeferredShadingRenderer.cpp b/Engine/Source/Runtime/Renderer/Private/DeferredShadingRenderer.cpp
--- a/Engine/Source/Runtime/Renderer/Private/DeferredShadingRenderer.cpp
+++ b/Engine/Source/Runtime/Renderer/Private/DeferredShadingRenderer.cpp
@@ -1555,8 +1555,8 @@ void FDeferredShadingSceneRenderer::Render(FRDGBuilder& GraphBuilder)
 		RDG_RHI_GPU_STAT_SCOPE(GraphBuilder, AllocateRendertargets);
 
 		// Force the subsurface profile texture to be updated.
-		UpdateSubsurfaceProfileTexture(GraphBuilder, ShaderPlatform);
-		SpecularProfileAtlas::UpdateSpecularProfileTextureAtlas(GraphBuilder, ShaderPlatform);
+		SubsurfaceProfile::UpdateSubsurfaceProfileTexture(GraphBuilder, ShaderPlatform);
+		SpecularProfile::UpdateSpecularProfileTextureAtlas(GraphBuilder, ShaderPlatform);
 
 		// Force the rect light texture & IES texture to be updated.
 		RectLightAtlas::UpdateAtlasTexture(GraphBuilder, FeatureLevel);
diff --git a/Engine/Source/Runtime/Renderer/Private/MobileShadingRenderer.cpp b/Engine/Source/Runtime/Renderer/Private/MobileShadingRenderer.cpp
--- a/Engine/Source/Runtime/Renderer/Private/MobileShadingRenderer.cpp
+++ b/Engine/Source/Runtime/Renderer/Private/MobileShadingRenderer.cpp
@@ -1037,8 +1037,8 @@ void FMobileSceneRenderer::Render(FRDGBuilder& GraphBuilder)
 	if (bRendererOutputFinalSceneColor)
 	{
 		// Force the subsurface profile texture to be updated.
-		UpdateSubsurfaceProfileTexture(GraphBuilder, ShaderPlatform);
-		SpecularProfileAtlas::UpdateSpecularProfileTextureAtlas(GraphBuilder, ShaderPlatform);
+		SubsurfaceProfile::UpdateSubsurfaceProfileTexture(GraphBuilder, ShaderPlatform);
+		SpecularProfile::UpdateSpecularProfileTextureAtlas(GraphBuilder, ShaderPlatform);
 
 		if (bDeferredShading)
 		{
diff --git a/Engine/Source/Runtime/Renderer/Private/SceneRendering.cpp b/Engine/Source/Runtime/Renderer/Private/SceneRendering.cpp
--- a/Engine/Source/Runtime/Renderer/Private/SceneRendering.cpp
+++ b/Engine/Source/Runtime/Renderer/Private/SceneRendering.cpp
@@ -1794,7 +1794,7 @@ void FViewInfo::SetupUniformBufferParameters(
 
 		// Profiles
 		{
-			FRHITexture* Texture = GetSubsurfaceProfileTextureWithFallback();
+			FRHITexture* Texture = SubsurfaceProfile::GetSubsurfaceProfileTextureWithFallback();
 			FIntVector TextureSize = Texture->GetSizeXYZ();
 			ViewUniformShaderParameters.SSProfilesTextureSizeAndInvSize = FVector4f(TextureSize.X, TextureSize.Y, 1.0f / TextureSize.X, 1.0f / TextureSize.Y);
 			ViewUniformShaderParameters.SSProfilesTexture = Texture;
@@ -1804,7 +1804,7 @@ void FViewInfo::SetupUniformBufferParameters(
 
 		// Pre-integrated profiles
 		{
-			FRHITexture* Texture = GetSSProfilesPreIntegratedTextureWithFallback();
+			FRHITexture* Texture = SubsurfaceProfile::GetSSProfilesPreIntegratedTextureWithFallback();
 			FIntVector TextureSize = Texture->GetSizeXYZ();
 			ViewUniformShaderParameters.SSProfilesPreIntegratedTextureSizeAndInvSize = FVector4f(TextureSize.X, TextureSize.Y, 1.0f / TextureSize.X, 1.0f / TextureSize.Y);
 			ViewUniformShaderParameters.SSProfilesPreIntegratedTexture = Texture;
@@ -1814,7 +1814,7 @@ void FViewInfo::SetupUniformBufferParameters(
 
 	// Specular Profiles
 	{
-		FRHITexture* Texture = SpecularProfileAtlas::GetSpecularProfileTextureAtlasWithFallback();
+		FRHITexture* Texture = SpecularProfile::GetSpecularProfileTextureAtlasWithFallback();
 		FIntVector TextureSize = Texture->GetSizeXYZ();
 		ViewUniformShaderParameters.SpecularProfileTextureSizeAndInvSize = FVector4f(TextureSize.X, TextureSize.Y, 1.0f / TextureSize.X, 1.0f / TextureSize.Y);
 		ViewUniformShaderParameters.SpecularProfileTexture = Texture;
