Fix bug where instance scene data buffer header is not returned.

Not 100% sure what will be affected/fixed by this. It seems that GeometryCollections+Nanite may be using this path but I did not verify.

#rnx

diff --git a/Engine/Source/Runtime/Engine/Private/PrimitiveSceneProxy.cpp b/Engine/Source/Runtime/Engine/Private/PrimitiveSceneProxy.cpp
--- a/Engine/Source/Runtime/Engine/Private/PrimitiveSceneProxy.cpp
+++ b/Engine/Source/Runtime/Engine/Private/PrimitiveSceneProxy.cpp
@@ -1286,7 +1286,7 @@ FInstanceDataBufferHeader FPrimitiveSceneProxy::GetInstanceDataHeader() const
 
 	if (InstanceSceneDataBuffersInternal)
 	{
-		InstanceSceneDataBuffersInternal->GetHeader();
+		return InstanceSceneDataBuffersInternal->GetHeader();
 	}
 	return FInstanceDataBufferHeader::SinglePrimitiveHeader;
 }
