[Chaos] Added an experimental option (disabled by default) that allows the initialization of FBodyInstances outside of the GameThread (using UE_CHAOS_ASYNC_INITBODY_ENABLED 1).

diff --git a/Engine/Source/Programs/HeadlessChaos/Private/HeadlessChaosTestSolverProxies.cpp b/Engine/Source/Programs/HeadlessChaos/Private/HeadlessChaosTestSolverProxies.cpp
--- a/Engine/Source/Programs/HeadlessChaos/Private/HeadlessChaosTestSolverProxies.cpp
+++ b/Engine/Source/Programs/HeadlessChaos/Private/HeadlessChaosTestSolverProxies.cpp
@@ -195,8 +195,8 @@ namespace ChaosTest {
 
 		// joint Proxy has been created, should be in dirty proxy list to be sent to physics thread.
 		// Note that we expect three dirty proxies, 2 for particles, 1 for joint.
-		EXPECT_EQ(Solver->GetDirtyProxyBucketInfo_External().Num[(uint32)EPhysicsProxyType::SingleParticleProxy], 2);
-		EXPECT_EQ(Solver->GetDirtyProxyBucketInfo_External().Num[(uint32)EPhysicsProxyType::JointConstraintType], 1);
+		EXPECT_EQ(Solver->GetDirtyProxyBucketInfoNum_External(EPhysicsProxyType::SingleParticleProxy), 2);
+		EXPECT_EQ(Solver->GetDirtyProxyBucketInfoNum_External(EPhysicsProxyType::JointConstraintType), 1);
 
 		// Although proxy exists, it should not have been added to physics thread only proxy array yet, as doing so from game thread code is wrong.
 		// Confirm game thread did not add it.
diff --git a/Engine/Source/Runtime/Engine/Classes/PhysicsEngine/BodyInstance.h b/Engine/Source/Runtime/Engine/Classes/PhysicsEngine/BodyInstance.h
--- a/Engine/Source/Runtime/Engine/Classes/PhysicsEngine/BodyInstance.h
+++ b/Engine/Source/Runtime/Engine/Classes/PhysicsEngine/BodyInstance.h
@@ -78,26 +78,40 @@ struct FInitBodySpawnParams
 	FPhysicsAggregateHandle Aggregate;
 };
 
-struct FInitBodiesHelperBase
+namespace Chaos::Private
 {
-	ENGINE_API FInitBodiesHelperBase(TArray<FBodyInstance*>& InBodies, TArray<FTransform>& InTransforms, 
+	// Private base class necessary to keep friend access to FBodyInstance
+	struct FInitBodiesHelperBaseInternal
+	{
+	protected:
+		void SetBodyInstanceExternalCollisionProfileBodySetup(FBodyInstance* BodyInstance, UBodySetup* ExternalCollisionProfileBodySetup);
+		void UpdateBodyInstanceSolverAsyncDeltaTime(FBodyInstance* BodyInstance);
+	};
+}
+
+template <typename InAllocatorType = FDefaultAllocator>
+struct TInitBodiesHelperBase : public Chaos::Private::FInitBodiesHelperBaseInternal
+{
+public:
+	ENGINE_API TInitBodiesHelperBase(TArray<FBodyInstance*, InAllocatorType>& InBodies, TArray<FTransform, InAllocatorType>& InTransforms,
 		class UBodySetup* InBodySetup, class UPrimitiveComponent* InPrimitiveComp, FPhysScene* InRBScene, 
 		const FInitBodySpawnParams& InSpawnParams, FPhysicsAggregateHandle InAggregate);
 	
-	ENGINE_API FInitBodiesHelperBase(TArray<FBodyInstance*>& InBodies, TArray<FTransform>& InTransforms, 
+	ENGINE_API TInitBodiesHelperBase(TArray<FBodyInstance*, InAllocatorType>& InBodies, TArray<FTransform, InAllocatorType>& InTransforms,
 		class UBodySetup* InBodySetup, class UPrimitiveComponent* InComponent, UObject* InSourceObject, FPhysScene* InRBScene,
 		const FInitBodySpawnParams& InSpawnParams, FPhysicsAggregateHandle InAggregate);
 
-	FInitBodiesHelperBase(const FInitBodiesHelperBase& InHelper) = delete;
-	FInitBodiesHelperBase(FInitBodiesHelperBase&& InHelper) = delete;
-	FInitBodiesHelperBase& operator=(const FInitBodiesHelperBase& InHelper) = delete;
-	FInitBodiesHelperBase& operator=(FInitBodiesHelperBase&& InHelper) = delete;
+	TInitBodiesHelperBase(const TInitBodiesHelperBase& InHelper) = delete;
+	TInitBodiesHelperBase(TInitBodiesHelperBase&& InHelper) = delete;
+	TInitBodiesHelperBase& operator=(const TInitBodiesHelperBase& InHelper) = delete;
+	TInitBodiesHelperBase& operator=(TInitBodiesHelperBase&& InHelper) = delete;
 
 	FORCEINLINE bool IsStatic() const { return bStatic; }
 
 	//The arguments passed into InitBodies
-	TArray<FBodyInstance*>& Bodies;   
-	TArray<FTransform>& Transforms;
+	TArray<FBodyInstance*, InAllocatorType>& Bodies;
+	TArray<FTransform, InAllocatorType>& Transforms;
+
 	class UBodySetup* BodySetup;
 	class UPrimitiveComponent* PrimitiveComp;
 	UObject* SourceObject;
@@ -134,22 +148,24 @@ protected:
 
 };
 
-template <bool bCompileStatic>
-struct FInitBodiesHelper : public FInitBodiesHelperBase
+using FInitBodiesHelperBase = TInitBodiesHelperBase<FDefaultAllocator>;
+
+template <bool bCompileStatic, typename InAllocatorType = FDefaultAllocator>
+struct FInitBodiesHelper : public TInitBodiesHelperBase<InAllocatorType>
 {
-	FInitBodiesHelper(TArray<FBodyInstance*>& InBodies, TArray<FTransform>& InTransforms, class UBodySetup* InBodySetup,
+	FInitBodiesHelper(TArray<FBodyInstance*, InAllocatorType>& InBodies, TArray<FTransform, InAllocatorType>& InTransforms, class UBodySetup* InBodySetup,
 		class UPrimitiveComponent* InPrimitiveComp, FPhysScene* InRBScene,
 		const FInitBodySpawnParams& InSpawnParams, FPhysicsAggregateHandle InAggregate);
 
-	FInitBodiesHelper(TArray<FBodyInstance*>& InBodies, TArray<FTransform>& InTransforms, class UBodySetup* InBodySetup, 
+	FInitBodiesHelper(TArray<FBodyInstance*, InAllocatorType>& InBodies, TArray<FTransform, InAllocatorType>& InTransforms, class UBodySetup* InBodySetup,
 		class UPrimitiveComponent* InComponent, class UObject* InOwnerObject, FPhysScene* InRBScene,
 		const FInitBodySpawnParams& InSpawnParams, FPhysicsAggregateHandle InAggregate)
-		: FInitBodiesHelperBase(InBodies, InTransforms, InBodySetup, InComponent, InOwnerObject, InRBScene, InSpawnParams, InAggregate)
+		: TInitBodiesHelperBase<InAllocatorType>(InBodies, InTransforms, InBodySetup, InComponent, InOwnerObject, InRBScene, InSpawnParams, InAggregate)
 	{
 		//Compute all the needed constants
-		bStatic = bCompileStatic || SpawnParams.bStaticPhysics;
-		SkelMeshComp = bCompileStatic ? nullptr : Cast<USkeletalMeshComponent>(PrimitiveComp);
-		if (SpawnParams.bPhysicsTypeDeterminesSimulation)
+		this->bStatic = bCompileStatic || this->SpawnParams.bStaticPhysics;
+		this->SkelMeshComp = bCompileStatic ? nullptr : Cast<USkeletalMeshComponent>(this->PrimitiveComp);
+		if (this->SpawnParams.bPhysicsTypeDeterminesSimulation)
 		{
 			this->UpdateSimulatingAndBlendWeight();
 		}
@@ -179,16 +195,8 @@ struct FInitBodiesHelperWithData : public FInitBodiesHelperBase
 	}
 
 	FInitBodiesHelperWithData(FInitBodiesHelperWithData&& InHelper)
-	: FInitBodiesHelperBase(OwnedBodies, OwnedTransforms, InHelper.BodySetup, InHelper.PrimitiveComp, InHelper.PhysScene, InHelper.SpawnParams, InHelper.Aggregate), OwnedBodies(MoveTemp(InHelper.OwnedBodies)), OwnedTransforms(MoveTemp(InHelper.OwnedTransforms)) //-V1050
-	{
-		//Compute all the needed constants
-		bStatic = bCompileStatic || SpawnParams.bStaticPhysics;
-		SkelMeshComp = bCompileStatic ? nullptr : Cast<USkeletalMeshComponent>(PrimitiveComp);
-		if(SpawnParams.bPhysicsTypeDeterminesSimulation)
-		{
-			this->UpdateSimulatingAndBlendWeight();
-		}
-	}
+	: FInitBodiesHelperWithData(OwnedBodies, OwnedTransforms, InHelper.BodySetup, InHelper.PrimitiveComp, InHelper.PhysScene, InHelper.SpawnParams, InHelper.Aggregate), OwnedBodies(MoveTemp(InHelper.OwnedBodies)), OwnedTransforms(MoveTemp(InHelper.OwnedTransforms)) //-V1050
+	{}
 
 	FInitBodiesHelperWithData& operator=(const FInitBodiesHelperWithData& InHelper) = delete;
 	FInitBodiesHelperWithData& operator=(FInitBodiesHelperWithData&& InHelper) = delete;
@@ -548,6 +556,8 @@ public:
 	/** Returns the mass override. See MassInKgOverride for documentation */
 	float GetMassOverride() const { return MassInKgOverride; }
 
+	bool IsUsingMACD() const { return bUseMACD; }
+
 	/** Sets the mass override */
 	ENGINE_API void SetMassOverride(float MassInKG, bool bNewOverrideMass = true);
 
@@ -716,8 +726,18 @@ public:
 	 *	@param InRBScene
 	 */
 	static ENGINE_API void InitStaticBodies(const TArray<FBodyInstance*>& Bodies, const TArray<FTransform>& Transforms, UBodySetup* BodySetup, class UPrimitiveComponent* PrimitiveComp, FPhysScene* InRBScene);
-
-
+	
+	/**
+	 * Standalone path to batch initialize large amounts of static bodies.
+	 * Note that this function will steal Bodies and Transforms arrays.
+	 *	@param Bodies The list of BodyInstances to initialize.
+	 *	@param Transforms The list of BodyInstance transform
+	 *	@param BodySetup The BodySetup to use
+	 *	@param PrimitiveComp The associated primitive component
+	 *	@param InRBScene The associated physics scene
+	 */
+	static ENGINE_API void InitStaticBodies(TArray<FBodyInstance*>&& Bodies, TArray<FTransform>&& Transforms, UBodySetup* BodySetup, class UPrimitiveComponent* PrimitiveComp, FPhysScene* InRBScene);
+	
 	/** Get the scene that owns this body. */
 	ENGINE_API FPhysScene* GetPhysicsScene();
 	ENGINE_API const FPhysScene* GetPhysicsScene() const;
@@ -1373,7 +1393,7 @@ private:
 	friend struct FUpdateCollisionResponseHelper;
 	friend class FBodySetupDetails;
 	
-	friend struct FInitBodiesHelperBase;
+	friend struct Chaos::Private::FInitBodiesHelperBaseInternal;
 	friend class FBodyInstanceCustomizationHelper;
 	friend class FFoliageTypeCustomizationHelpers;
 
diff --git a/Engine/Source/Runtime/Engine/Private/PhysicsEngine/BodyInstance.cpp b/Engine/Source/Runtime/Engine/Private/PhysicsEngine/BodyInstance.cpp
--- a/Engine/Source/Runtime/Engine/Private/PhysicsEngine/BodyInstance.cpp
+++ b/Engine/Source/Runtime/Engine/Private/PhysicsEngine/BodyInstance.cpp
@@ -1167,27 +1167,41 @@ static void GetSimulatingAndBlendWeight(const USkeletalMeshComponent* SkelMeshCo
 	}
 }
 
-void FInitBodiesHelperBase::UpdateSimulatingAndBlendWeight()
+namespace Chaos::Private
 {
-	GetSimulatingAndBlendWeight(SkelMeshComp, BodySetup, InstanceBlendWeight, bInstanceSimulatePhysics);
-}
+	void FInitBodiesHelperBaseInternal::SetBodyInstanceExternalCollisionProfileBodySetup(FBodyInstance* BodyInstance, UBodySetup* ExternalCollisionProfileBodySetup)
+	{
+		BodyInstance->ExternalCollisionProfileBodySetup = ExternalCollisionProfileBodySetup;
+	}
 
-template<bool bCompileStatic>
-FInitBodiesHelper<bCompileStatic>::FInitBodiesHelper(
-	TArray<FBodyInstance*>& InBodies, TArray<FTransform>& InTransforms, UBodySetup* InBodySetup, 
+	void FInitBodiesHelperBaseInternal::UpdateBodyInstanceSolverAsyncDeltaTime(FBodyInstance* BodyInstance)
+	{
+		BodyInstance->UpdateSolverAsyncDeltaTime();
+	}
+};
+
+template<bool bCompileStatic, typename InAllocatorType>
+FInitBodiesHelper<bCompileStatic, InAllocatorType>::FInitBodiesHelper(
+	TArray<FBodyInstance*, InAllocatorType>& InBodies, TArray<FTransform, InAllocatorType>& InTransforms, UBodySetup* InBodySetup,
 	UPrimitiveComponent* InPrimitiveComp, FPhysScene* InRBScene, 
 	const FInitBodySpawnParams& InSpawnParams, FPhysicsAggregateHandle InAggregate)
-	: FInitBodiesHelper<bCompileStatic>(InBodies, InTransforms, InBodySetup, InPrimitiveComp,
+	: FInitBodiesHelper<bCompileStatic, InAllocatorType>(InBodies, InTransforms, InBodySetup, InPrimitiveComp,
 		InPrimitiveComp ? InPrimitiveComp->GetOwner() : nullptr, 
 		InRBScene, InSpawnParams, InAggregate)
 {
 }
+
 template struct FInitBodiesHelper<true>;
 template struct FInitBodiesHelper<false>;
-
-
-FInitBodiesHelperBase::FInitBodiesHelperBase(TArray<FBodyInstance*>& InBodies, TArray<FTransform>& InTransforms, 
-	class UBodySetup* InBodySetup, UPrimitiveComponent* InComponent, UObject* InSourceObject, FPhysScene* InRBScene, 
+using BodyInstanceSingleElementArrayInlineAllocator = TInlineAllocator<1>;
+template struct FInitBodiesHelper<true, BodyInstanceSingleElementArrayInlineAllocator>;
+template struct FInitBodiesHelper<false, BodyInstanceSingleElementArrayInlineAllocator>;
+using FInitSingleBodyHelperStatic = FInitBodiesHelper<true, BodyInstanceSingleElementArrayInlineAllocator>;
+using FInitSingleBodyHelperDynamic = FInitBodiesHelper<false, BodyInstanceSingleElementArrayInlineAllocator>;
+
+template<typename InAllocatorType>
+TInitBodiesHelperBase<InAllocatorType>::TInitBodiesHelperBase(TArray<FBodyInstance*, InAllocatorType>& InBodies, TArray<FTransform, InAllocatorType>& InTransforms,
+	class UBodySetup* InBodySetup, UPrimitiveComponent* InComponent, UObject* InSourceObject, FPhysScene* InRBScene,
 	const FInitBodySpawnParams& InSpawnParams, FPhysicsAggregateHandle InAggregate)
 	: Bodies(InBodies)
 	, Transforms(InTransforms)
@@ -1203,26 +1217,34 @@ FInitBodiesHelperBase::FInitBodiesHelperBase(TArray<FBodyInstance*>& InBodies, T
 	, InstanceBlendWeight(-1.f)
 	, SkelMeshComp(nullptr)
 	, SpawnParams(InSpawnParams)
-	, DisableQueryOnlyActors(!!CDisableQueryOnlyActors.GetValueOnGameThread())
+	, DisableQueryOnlyActors(!!CDisableQueryOnlyActors.GetValueOnAnyThread())
 {
 #if USE_BODYINSTANCE_DEBUG_NAMES
 	PhysXName = GetDebugDebugName(PrimitiveComp, BodySetup, *DebugName);
 #endif
 }
 
-FInitBodiesHelperBase::FInitBodiesHelperBase(TArray<FBodyInstance*>& InBodies, TArray<FTransform>& InTransforms,
+template<typename InAllocatorType>
+TInitBodiesHelperBase<InAllocatorType>::TInitBodiesHelperBase(TArray<FBodyInstance*, InAllocatorType>& InBodies, TArray<FTransform, InAllocatorType>& InTransforms,
 	class UBodySetup* InBodySetup, class UPrimitiveComponent* InPrimitiveComp, FPhysScene* InRBScene,
 	const FInitBodySpawnParams& InSpawnParams, FPhysicsAggregateHandle InAggregate)
 	// Forward to the other constructor
-	: FInitBodiesHelperBase(InBodies, InTransforms, InBodySetup,
+	: TInitBodiesHelperBase<InAllocatorType>(InBodies, InTransforms, InBodySetup,
 		InPrimitiveComp,
 		InPrimitiveComp ? InPrimitiveComp->GetOwner() : nullptr,
 		InRBScene, InSpawnParams, InAggregate)
 {
 }
 
+template<typename InAllocatorType>
+void TInitBodiesHelperBase<InAllocatorType>::UpdateSimulatingAndBlendWeight()
+{
+	GetSimulatingAndBlendWeight(SkelMeshComp, BodySetup, InstanceBlendWeight, bInstanceSimulatePhysics);
+}
+
 // Return to actor ref
-void FInitBodiesHelperBase::CreateActor_AssumesLocked(FBodyInstance* Instance, const FTransform& Transform) const
+template<typename InAllocatorType>
+void TInitBodiesHelperBase<InAllocatorType>::CreateActor_AssumesLocked(FBodyInstance* Instance, const FTransform& Transform) const
 {
 	SCOPE_CYCLE_COUNTER(STAT_CreatePhysicsActor);
 	checkSlow(!FPhysicsInterface::IsValid(Instance->ActorHandle));
@@ -1250,19 +1272,20 @@ void FInitBodiesHelperBase::CreateActor_AssumesLocked(FBodyInstance* Instance, c
 	{
 		FPhysicsInterface::CreateActor(ActorParams, Instance->ActorHandle);
 		FPhysicsInterface::SetCcdEnabled_AssumesLocked(Instance->ActorHandle, Instance->bUseCCD);
-		FPhysicsInterface::SetMACDEnabled_AssumesLocked(Instance->ActorHandle, Instance->bUseMACD);
+		FPhysicsInterface::SetMACDEnabled_AssumesLocked(Instance->ActorHandle, Instance->IsUsingMACD());
 		FPhysicsInterface::SetIsKinematic_AssumesLocked(Instance->ActorHandle, !Instance->ShouldInstanceSimulatingPhysics());
 
 		FPhysicsInterface::SetMaxLinearVelocity_AssumesLocked(Instance->ActorHandle, TNumericLimits<float>::Max());
 		FPhysicsInterface::SetSmoothEdgeCollisionsEnabled_AssumesLocked(Instance->ActorHandle, Instance->bSmoothEdgeCollisions);
-		FPhysicsInterface::SetInertiaConditioningEnabled_AssumesLocked(Instance->ActorHandle, Instance->bInertiaConditioning);
+		FPhysicsInterface::SetInertiaConditioningEnabled_AssumesLocked(Instance->ActorHandle, Instance->IsInertiaConditioningEnabled());
 
 		// Set sleep event notification
 		FPhysicsInterface::SetSendsSleepNotifies_AssumesLocked(Instance->ActorHandle, Instance->bGenerateWakeEvents);
 	}
 }
 
-bool FInitBodiesHelperBase::CreateShapes_AssumesLocked(FBodyInstance* Instance) const
+template<typename InAllocatorType>
+bool TInitBodiesHelperBase<InAllocatorType>::CreateShapes_AssumesLocked(FBodyInstance* Instance) const
 {
 	SCOPE_CYCLE_COUNTER(STAT_CreatePhysicsShapes);
 	UPhysicalMaterial* SimplePhysMat = Instance->GetSimplePhysicalMaterial();
@@ -1280,7 +1303,7 @@ bool FInitBodiesHelperBase::CreateShapes_AssumesLocked(FBodyInstance* Instance)
 	// #PHYS2 Call interface AddGeometry
 	BodySetup->AddShapesToRigidActor_AssumesLocked(Instance, Instance->Scale3D, SimplePhysMat, ComplexPhysMats, ComplexPhysMatMasks, BodyCollisionData, FTransform::Identity);
 
-	FPhysicsInterface::SetIgnoreAnalyticCollisions_AssumesLocked(Instance->ActorHandle, CVarIgnoreAnalyticCollisionsOverride.GetValueOnGameThread() ? true : Instance->bIgnoreAnalyticCollisions);
+	FPhysicsInterface::SetIgnoreAnalyticCollisions_AssumesLocked(Instance->ActorHandle, CVarIgnoreAnalyticCollisionsOverride.GetValueOnAnyThread() ? true : Instance->bIgnoreAnalyticCollisions);
 
 	const int32 NumShapes = FPhysicsInterface::GetNumShapes(Instance->ActorHandle);
 	bInitFail |= NumShapes == 0;
@@ -1298,7 +1321,8 @@ UBodySetup* FBodyInstance::GetBodySetup() const
 	return nullptr;
 }
 
-const FString& GetBodyInstanceDebugName(FInitBodiesHelperBase& InitHelper)
+template<typename InAllocatorType>
+const FString& GetBodyInstanceDebugName(TInitBodiesHelperBase<InAllocatorType>& InitHelper)
 {
 	static FString NullName = TEXT("<NoName>");
 
@@ -1357,14 +1381,25 @@ bool ValidateTransformNaN(const FTransform& Transform, const FString& DebugName,
 
 // Takes actor ref arrays.
 // #PHYS2 this used to return arrays of low-level physics bodies, which would be added to scene in InitBodies. Should it still do that, rather then later iterate over BodyInstances to get phys actor refs?
-bool FInitBodiesHelperBase::CreateShapesAndActors()
+template<typename InAllocatorType>
+bool TInitBodiesHelperBase<InAllocatorType>::CreateShapesAndActors()
 {
 	SCOPE_CYCLE_COUNTER(STAT_CreatePhysicsShapesAndActors);
+	TRACE_CPUPROFILER_EVENT_SCOPE(CreatePhysicsShapesAndActors);
 
 	const int32 NumBodies = Bodies.Num();
 
-	// Ensure we have the AggGeom inside the body setup so we can calculate the number of shapes
-	BodySetup->CreatePhysicsMeshes();
+	const bool bIsInGameThread = IsInGameThread();
+	check(UE_CHAOS_ASYNC_INITBODY_ENABLED || bIsInGameThread);
+#if UE_CHAOS_ASYNC_INITBODY_ENABLED
+	// Current case where CreateShapesAndActors is not called from the GameThread is when experimental UE_CHAOS_ASYNC_INITBODY_ENABLED = 1
+	check(bIsInGameThread || (BodySetup->bCreatedPhysicsMeshes || BodySetup->bNeverNeedsCookedCollisionData));
+	if (bIsInGameThread)
+#endif
+	{
+		// Ensure we have the AggGeom inside the body setup so we can calculate the number of shapes
+		BodySetup->CreatePhysicsMeshes();
+	}
 
 	if (bSkipShapeCreationForEmptyBodySetup)
 	{
@@ -1429,6 +1464,7 @@ bool FInitBodiesHelperBase::CreateShapesAndActors()
 					UWorld* World = PrimitiveComp->GetWorld();
 					if (World && World->IsGameWorld())
 					{
+						check(bIsInGameThread);
 						//if we have a parent we will now do the weld and exit any further initialization
 						if (PrimitiveComp->WeldToImplementation(ParentPrimComponent, PrimitiveComp->GetAttachSocketName(), false))	//welded new simulated body so initialization is done
 						{
@@ -1446,8 +1482,9 @@ bool FInitBodiesHelperBase::CreateShapesAndActors()
 			Instance->OwnerComponent = nullptr;
 			Instance->SourceObject = nullptr;
 			Instance->BodySetup = nullptr;
-			Bodies.RemoveAt(BodyIdx);  // so we wont add it to the physx scene again later.
-			Transforms.RemoveAt(BodyIdx);
+			constexpr uint32 Count = 1;
+			Bodies.RemoveAt(BodyIdx, Count, EAllowShrinking::No);  // so we wont add it to the physx scene again later.
+			Transforms.RemoveAt(BodyIdx, Count, EAllowShrinking::No);
 			continue;
 		}
 
@@ -1481,7 +1518,7 @@ bool FInitBodiesHelperBase::CreateShapesAndActors()
 			Instance->OwnerComponent = nullptr;
 			Instance->SourceObject = nullptr;
 			Instance->BodySetup = nullptr;
-			Instance->ExternalCollisionProfileBodySetup = nullptr;
+			SetBodyInstanceExternalCollisionProfileBodySetup(Instance, nullptr);
 
 			continue;
 		}
@@ -1495,16 +1532,21 @@ bool FInitBodiesHelperBase::CreateShapesAndActors()
 	return true;
 }
 
-void FInitBodiesHelperBase::InitBodies()
+template<typename InAllocatorType>
+void TInitBodiesHelperBase<InAllocatorType>::InitBodies()
 {
 	LLM_SCOPE(ELLMTag::ChaosBody);
 
-	//check(IsInGameThread());
-
-	if (CreateShapesAndActors())
+	const bool bIsInGameThread = IsInGameThread();
+	check(UE_CHAOS_ASYNC_INITBODY_ENABLED || bIsInGameThread);
+	if (bIsInGameThread ? CreateShapesAndActors() : true)
 	{
 		FPhysicsCommand::ExecuteWrite(PhysScene, [&]()
 		{
+			if (!bIsInGameThread && !CreateShapesAndActors())
+			{
+				return;
+			}
 			// If an aggregate present, add to that
 			if (Aggregate.IsValid())
 			{
@@ -1543,11 +1585,11 @@ void FInitBodiesHelperBase::InitBodies()
 								Body_External.SetShapeSimCollisionEnabled(ShapeIndex, false);
 							}
 						}
-						if (BI->BodySetup.IsValid())
+						if (const UBodySetupCore* BS = BI->BodySetup.Get())
 						{
 							for (int32 ShapeIndex = 0; ShapeIndex < NumShapes; ++ShapeIndex)
 							{
-								Body_External.SetShapeCollisionTraceType(ShapeIndex, ChaosInterface::ConvertCollisionTraceFlag(BI->BodySetup->GetCollisionTraceFlag())) ;
+								Body_External.SetShapeCollisionTraceType(ShapeIndex, ChaosInterface::ConvertCollisionTraceFlag(BS->GetCollisionTraceFlag())) ;
 							}
 						}
 					}
@@ -1576,7 +1618,7 @@ void FInitBodiesHelperBase::InitBodies()
 			// set solver async delta time if any bodies are overriding the sim delta time
 			for (FBodyInstance* BI : Bodies)
 			{				
-				BI->UpdateSolverAsyncDeltaTime();			
+				UpdateBodyInstanceSolverAsyncDeltaTime(BI);
 			}
 
 			// Set up dynamic instance data
@@ -1632,7 +1674,7 @@ void FBodyInstance::InitBody(UBodySetup* Setup, const FTransform& Transform, UPr
 
 void FBodyInstance::InitBody(class UBodySetup* Setup, const FTransform& Transform, UPrimitiveComponent* PrimComp, UObject* InSourceObject, FPhysScene* InRBScene, const FInitBodySpawnParams& SpawnParams)
 {
-	if (CVarAllowCreatePhysxBodies.GetValueOnGameThread() == 0)
+	if (CVarAllowCreatePhysxBodies.GetValueOnAnyThread() == 0)
 	{
 		return;
 	}
@@ -1640,11 +1682,8 @@ void FBodyInstance::InitBody(class UBodySetup* Setup, const FTransform& Transfor
 	SCOPE_CYCLE_COUNTER(STAT_InitBody);
 	check(Setup);
 	
-	static TArray<FBodyInstance*> Bodies;
-	static TArray<FTransform> Transforms;
-
-	check(Bodies.Num() == 0);
-	check(Transforms.Num() == 0);
+	TArray<FBodyInstance*, BodyInstanceSingleElementArrayInlineAllocator> Bodies;
+	TArray<FTransform, BodyInstanceSingleElementArrayInlineAllocator> Transforms;
 
 	Bodies.Add(this);
 	Transforms.Add(Transform);
@@ -1658,18 +1697,15 @@ void FBodyInstance::InitBody(class UBodySetup* Setup, const FTransform& Transfor
 	bool bIsStatic = SpawnParams.bStaticPhysics;
 	if(bIsStatic)
 	{
-		FInitBodiesHelper<true> InitBodiesHelper(Bodies, Transforms, Setup, PrimComp, SourceObjectToUse, InRBScene, SpawnParams, SpawnParams.Aggregate);
-		InitBodiesHelper.InitBodies();
+		FInitSingleBodyHelperStatic InitBodyHelper(Bodies, Transforms, Setup, PrimComp, SourceObjectToUse, InRBScene, SpawnParams, SpawnParams.Aggregate);
+		InitBodyHelper.InitBodies();
 	}
 	else
 	{
-		FInitBodiesHelper<false> InitBodiesHelper(Bodies, Transforms, Setup, PrimComp, SourceObjectToUse, InRBScene, SpawnParams, SpawnParams.Aggregate);
-		InitBodiesHelper.InitBodies();
+		FInitSingleBodyHelperDynamic InitBodyHelper(Bodies, Transforms, Setup, PrimComp, SourceObjectToUse, InRBScene, SpawnParams, SpawnParams.Aggregate);
+		InitBodyHelper.InitBodies();
 	}
 
-	Bodies.Reset();
-	Transforms.Reset();
-
 	UpdateInterpolateWhenSubStepping();
 }
 
@@ -4351,7 +4387,7 @@ void FBodyInstance::BuildBodyFilterData(FBodyCollisionFilterData& OutFilterData,
 	//        should configure the bodies to reflect this desired behavior.
 	if(USkeletalMeshComponent* SkelMeshComp = Cast<USkeletalMeshComponent>(OwnerComponentInst))
 	{
-		if (CVarEnableDynamicPerBodyFilterHacks.GetValueOnGameThread() && bHACK_DisableCollisionResponse)
+		if (CVarEnableDynamicPerBodyFilterHacks.GetValueOnAnyThread() && bHACK_DisableCollisionResponse)
 		{
 			UseResponse.SetAllChannels(ECR_Ignore);
 			UseCollisionEnabled = ECollisionEnabled::PhysicsOnly;
@@ -4366,7 +4402,7 @@ void FBodyInstance::BuildBodyFilterData(FBodyCollisionFilterData& OutFilterData,
 			UseCollisionEnabled = ECollisionEnabled::PhysicsOnly;		// this will prevent object traces hitting this as well
 		}
 
-		const bool bDisableSkelComponentOverride = CVarEnableDynamicPerBodyFilterHacks.GetValueOnGameThread() && bHACK_DisableSkelComponentFilterOverriding;
+		const bool bDisableSkelComponentOverride = CVarEnableDynamicPerBodyFilterHacks.GetValueOnAnyThread() && bHACK_DisableSkelComponentFilterOverriding;
 		if (bDisableSkelComponentOverride)
 		{
 			// if we are disabling the skeletal component override, we want the original body instance collision response
@@ -4452,6 +4488,7 @@ void FBodyInstance::InitStaticBodies(const TArray<FBodyInstance*>& Bodies, const
 	check(BodySetup);
 	check(InRBScene);
 	check(Bodies.Num() > 0);
+	check(IsInGameThread());
 
 	static TArray<FBodyInstance*> BodiesStatic;
 	static TArray<FTransform> TransformsStatic;
@@ -4469,6 +4506,16 @@ void FBodyInstance::InitStaticBodies(const TArray<FBodyInstance*>& Bodies, const
 	TransformsStatic.Reset();
 }
 
+void FBodyInstance::InitStaticBodies(TArray<FBodyInstance*>&& Bodies, TArray<FTransform>&& Transforms, UBodySetup* BodySetup, UPrimitiveComponent* PrimitiveComp, FPhysScene* InRBScene)
+{
+	SCOPE_CYCLE_COUNTER(STAT_StaticInitBodies);
+	check(BodySetup);
+	check(InRBScene);
+	check(Bodies.Num() > 0);
+	FInitBodiesHelperWithData<true> InitBodiesHelper(MoveTemp(Bodies), MoveTemp(Transforms), BodySetup, PrimitiveComp, InRBScene, FInitBodySpawnParams(PrimitiveComp), FPhysicsAggregateHandle());
+	InitBodiesHelper.InitBodies();
+}
+
 int32 SimCollisionEnabled = 1;
 FAutoConsoleVariableRef CVarSimCollisionEnabled(TEXT("p.SimCollisionEnabled"), SimCollisionEnabled, TEXT("If 0 no sim collision will be used"));
 
diff --git a/Engine/Source/Runtime/Engine/Private/PhysicsEngine/Experimental/PhysScene_Chaos.cpp b/Engine/Source/Runtime/Engine/Private/PhysicsEngine/Experimental/PhysScene_Chaos.cpp
--- a/Engine/Source/Runtime/Engine/Private/PhysicsEngine/Experimental/PhysScene_Chaos.cpp
+++ b/Engine/Source/Runtime/Engine/Private/PhysicsEngine/Experimental/PhysScene_Chaos.cpp
@@ -29,6 +29,9 @@
 #include "Chaos/PendingSpatialData.h"
 #include "Chaos/PhysicsSolverBaseImpl.h"
 #include "Misc/CoreMisc.h"
+#if UE_CHAOS_ASYNC_INITBODY_ENABLED
+#include "Misc/ScopeRWLock.h"
+#endif
 
 #if WITH_EDITOR
 #include "Editor.h"
@@ -846,7 +849,7 @@ void FPhysScene_Chaos::AddReferencedObjects(FReferenceCollector& Collector)
 	Super::AddReferencedObjects(Collector);
 	Collector.AddReferencedObject(ChaosEventRelay);
 #if WITH_EDITOR
-
+	UE_CHAOS_ASYNC_INITBODY_READSCOPELOCK(PhysicsProxyComponentMapsLock);
 	for (auto& Pair : PhysicsProxyToComponentMap)
 	{
 		Collector.AddReferencedObject(Pair.Get<1>());
@@ -857,6 +860,7 @@ void FPhysScene_Chaos::AddReferencedObjects(FReferenceCollector& Collector)
 template<>
 UPrimitiveComponent* FPhysScene_Chaos::GetOwningComponent(const IPhysicsProxyBase* PhysicsProxy) const
 {
+	UE_CHAOS_ASYNC_INITBODY_READSCOPELOCK(PhysicsProxyComponentMapsLock);
 	if (const TObjectPtr<UPrimitiveComponent>* FoundComp = PhysicsProxyToComponentMap.Find(PhysicsProxy))
 	{
 		return *FoundComp;
@@ -1308,6 +1312,7 @@ void FPhysScene_Chaos::AddToComponentMaps(UPrimitiveComponent* Component, IPhysi
 {
 	if (Component != nullptr && InObject != nullptr)
 	{
+		UE_CHAOS_ASYNC_INITBODY_WRITESCOPELOCK(PhysicsProxyComponentMapsLock);
 		PhysicsProxyToComponentMap.Add(InObject, ObjectPtrWrap(Component));
 
 		TArray<IPhysicsProxyBase*>* ProxyArray = ComponentToPhysicsProxyMap.Find(Component);
@@ -1360,6 +1365,7 @@ void FPhysScene_Chaos::CreatePhysicsReplicationCache()
 
 void FPhysScene_Chaos::RemoveFromComponentMaps(IPhysicsProxyBase* InObject)
 {
+	UE_CHAOS_ASYNC_INITBODY_WRITESCOPELOCK(PhysicsProxyComponentMapsLock);
 	auto* const Component = PhysicsProxyToComponentMap.Find(InObject);
 	if (Component)
 	{
@@ -1444,6 +1450,7 @@ void FPhysScene_Chaos::OnWorldEndPlay()
 	PieModifiedObjects.Reset();
 #endif
 
+	UE_CHAOS_ASYNC_INITBODY_WRITESCOPELOCK(PhysicsProxyComponentMapsLock);
 	PhysicsProxyToComponentMap.Reset();
 	ComponentToPhysicsProxyMap.Reset();
 }
diff --git a/Engine/Source/Runtime/Engine/Public/Physics/Experimental/PhysScene_Chaos.h b/Engine/Source/Runtime/Engine/Public/Physics/Experimental/PhysScene_Chaos.h
--- a/Engine/Source/Runtime/Engine/Public/Physics/Experimental/PhysScene_Chaos.h
+++ b/Engine/Source/Runtime/Engine/Public/Physics/Experimental/PhysScene_Chaos.h
@@ -16,6 +16,10 @@
 #include "Chaos/Real.h"
 #include "UObject/ObjectKey.h"
 
+#if UE_CHAOS_ASYNC_INITBODY_ENABLED
+#include "Misc/ScopeRWLock.h"
+#endif
+
 #ifndef CHAOS_WITH_PAUSABLE_SOLVER
 #define CHAOS_WITH_PAUSABLE_SOLVER 1
 #endif
@@ -175,6 +179,7 @@ public:
 	template<class OwnerType>
 	OwnerType* GetOwningComponent(const IPhysicsProxyBase* PhysicsProxy) const
 	{ 
+		UE_CHAOS_ASYNC_INITBODY_READSCOPELOCK(PhysicsProxyComponentMapsLock);
 		auto* CompPtr = PhysicsProxyToComponentMap.Find(PhysicsProxy);
 		return CompPtr ? Cast<OwnerType>(*CompPtr) : nullptr;
 	}
@@ -185,6 +190,7 @@ public:
 	/** Given a component, returns its associated solver objects. */
 	const TArray<IPhysicsProxyBase*>* GetOwnedPhysicsProxies(UPrimitiveComponent* Comp) const
 	{
+		UE_CHAOS_ASYNC_INITBODY_READSCOPELOCK(PhysicsProxyComponentMapsLock);
 		return ComponentToPhysicsProxyMap.Find(Comp);
 	}
 
@@ -444,6 +450,12 @@ private:
 	TArray<TPair<TWeakObjectPtr<USkeletalMeshComponent>, FDeferredKinematicUpdateInfo>> DeferredKinematicUpdateSkelMeshes;
 
 	TSet<UPrimitiveComponent*> DeferredCreatePhysicsStateComponents;
+
+#if UE_CHAOS_ASYNC_INITBODY_ENABLED
+	// RWLock for for thread safe access to PhysicsProxyToComponentMap and ComponentToPhysicsProxyMap.
+	mutable FRWLock PhysicsProxyComponentMapsLock;
+#endif
+
 	//Body Instances
 	TUniquePtr<Chaos::TArrayCollectionArray<FBodyInstance*>> BodyInstances;
 	// Temp Interface
diff --git a/Engine/Source/Runtime/Experimental/Chaos/Private/Chaos/ChaosMarshallingManager.cpp b/Engine/Source/Runtime/Experimental/Chaos/Private/Chaos/ChaosMarshallingManager.cpp
--- a/Engine/Source/Runtime/Experimental/Chaos/Private/Chaos/ChaosMarshallingManager.cpp
+++ b/Engine/Source/Runtime/Experimental/Chaos/Private/Chaos/ChaosMarshallingManager.cpp
@@ -2,6 +2,9 @@
 
 #include "Chaos/ChaosMarshallingManager.h"
 #include "Chaos/PullPhysicsDataImp.h"
+#if UE_CHAOS_ASYNC_INITBODY_ENABLED
+#include "Misc/ScopeRWLock.h"
+#endif
 
 namespace Chaos
 {
@@ -48,6 +51,7 @@ void FChaosMarshallingManager::PreparePullData()
 
 void FChaosMarshallingManager::PrepareExternalQueue_External()
 {
+	// Here, we assume that MarshallingManagerLock is locked (when UE_CHAOS_ASYNC_INITBODY_ENABLED = 1)
 	if(!PushDataPool.Dequeue(ProducerData))
 	{
 		BackingBuffer.Add(MakeUnique<FPushPhysicsData>());
@@ -59,6 +63,7 @@ void FChaosMarshallingManager::PrepareExternalQueue_External()
 
 void FChaosMarshallingManager::Step_External(FReal ExternalDT, const int32 NumSteps, bool bInSolverSubstepped)
 {
+	// Here, we assume that MarshallingManagerLock is locked (when UE_CHAOS_ASYNC_INITBODY_ENABLED = 1)
 	ensure(NumSteps > 0);
 
 	FPushPhysicsData* FirstStepData = nullptr;
@@ -97,7 +102,7 @@ void FChaosMarshallingManager::Step_External(FReal ExternalDT, const int32 NumSt
 			ProducerData->CopySubstepData(*FirstStepData);
 		}
 
-		ExternalTime_External += ExternalDT;
+		ExternalTime_External.store(ExternalTime_External.load() + ExternalDT);
 		PrepareExternalQueue_External();
 	}
 
@@ -106,6 +111,7 @@ void FChaosMarshallingManager::Step_External(FReal ExternalDT, const int32 NumSt
 
 FPushPhysicsData* FChaosMarshallingManager::StepInternalTime_External()
 {
+	UE_CHAOS_ASYNC_INITBODY_WRITESCOPELOCK(MarshallingManagerLock);
 	if (Delay == 0)
 	{
 		if(ExternalQueue.Num())
@@ -123,6 +129,7 @@ FPushPhysicsData* FChaosMarshallingManager::StepInternalTime_External()
 
 void FChaosMarshallingManager::FreeData_Internal(FPushPhysicsData* PushData)
 {
+	UE_CHAOS_ASYNC_INITBODY_WRITESCOPELOCK(MarshallingManagerLock);
 	//TODO: we know entire manager is cleared, so we can probably just iterate over its pools and reset
 	//instead of going through dirty proxies. If perf matters fix this
 	FDirtyPropertiesManager* Manager = &PushData->DirtyPropertiesManager;
diff --git a/Engine/Source/Runtime/Experimental/Chaos/Private/Chaos/PhysicalMaterials.cpp b/Engine/Source/Runtime/Experimental/Chaos/Private/Chaos/PhysicalMaterials.cpp
--- a/Engine/Source/Runtime/Experimental/Chaos/Private/Chaos/PhysicalMaterials.cpp
+++ b/Engine/Source/Runtime/Experimental/Chaos/Private/Chaos/PhysicalMaterials.cpp
@@ -2,6 +2,9 @@
 
 #include "Chaos/PhysicalMaterials.h"
 #include "HAL/LowLevelMemTracker.h"
+#if UE_CHAOS_ASYNC_INITBODY_ENABLED
+#include "Misc/ScopeRWLock.h"
+#endif
 
 namespace Chaos
 {
@@ -56,35 +59,37 @@ namespace Chaos
 
 	FChaosPhysicsMaterial* FPhysicalMaterialManager::Resolve(FChaosMaterialHandle InHandle) const
 	{
+		UE_CHAOS_ASYNC_INITBODY_READSCOPELOCK(MaterialsLock);
 		return Materials.Get(InHandle);
 	}
 
 	const FChaosPhysicsMaterial* FPhysicalMaterialManager::Resolve(FChaosConstMaterialHandle InHandle) const
 	{
+		UE_CHAOS_ASYNC_INITBODY_READSCOPELOCK(MaterialsLock);
 		return Materials.Get(InHandle);
 	}
 
 	FChaosPhysicsMaterialMask* FPhysicalMaterialManager::Resolve(FChaosMaterialMaskHandle InHandle) const
 	{
+		UE_CHAOS_ASYNC_INITBODY_READSCOPELOCK(MaterialMasksLock);
 		return MaterialMasks.Get(InHandle);
 	}
 
 	const FChaosPhysicsMaterialMask* FPhysicalMaterialManager::Resolve(FChaosConstMaterialMaskHandle InHandle) const
 	{
+		UE_CHAOS_ASYNC_INITBODY_READSCOPELOCK(MaterialMasksLock);
 		return MaterialMasks.Get(InHandle);
 	}
 	
 	void FPhysicalMaterialManager::UpdateMaterial(FMaterialHandle InHandle)
 	{
-		check(IsInGameThread());
-
+		check(UE_CHAOS_ASYNC_INITBODY_ENABLED || IsInGameThread());
 		OnMaterialUpdated.Broadcast(InHandle);
 	}
 
 	void FPhysicalMaterialManager::UpdateMaterialMask(FMaterialMaskHandle InHandle)
 	{
-		check(IsInGameThread());
-
+		check(UE_CHAOS_ASYNC_INITBODY_ENABLED || IsInGameThread());
 		OnMaterialMaskUpdated.Broadcast(InHandle);
 	}
 
@@ -112,10 +117,12 @@ namespace Chaos
 	{
 		LLM_SCOPE(ELLMTag::ChaosMaterial);
 
-		check(IsInGameThread());
+		check(UE_CHAOS_ASYNC_INITBODY_ENABLED || IsInGameThread());
 		FMaterialHandle OutHandle;
-		OutHandle.InnerHandle = Materials.Create();
-
+		{
+			UE_CHAOS_ASYNC_INITBODY_WRITESCOPELOCK(MaterialsLock);
+			OutHandle.InnerHandle = Materials.Create();
+		}
 		OnMaterialCreated.Broadcast(OutHandle);
 
 		return OutHandle;
@@ -123,10 +130,12 @@ namespace Chaos
 
 	FMaterialMaskHandle FPhysicalMaterialManager::CreateMask()
 	{
-		check(IsInGameThread());
+		check(UE_CHAOS_ASYNC_INITBODY_ENABLED || IsInGameThread());
 		FMaterialMaskHandle OutHandle;
-		OutHandle.InnerHandle = MaterialMasks.Create();
-
+		{
+			UE_CHAOS_ASYNC_INITBODY_WRITESCOPELOCK(MaterialMasksLock);
+			OutHandle.InnerHandle = MaterialMasks.Create();
+		}
 		OnMaterialMaskCreated.Broadcast(OutHandle);
 
 		return OutHandle;
@@ -136,23 +145,27 @@ namespace Chaos
 	{
 		LLM_SCOPE(ELLMTag::ChaosMaterial);
 
-		check(IsInGameThread());
+		check(UE_CHAOS_ASYNC_INITBODY_ENABLED || IsInGameThread());
 		if(InHandle.InnerHandle.IsValid())
 		{
 			OnMaterialDestroyed.Broadcast(InHandle);
-
-			Materials.Destroy(InHandle.InnerHandle);
+			{
+				UE_CHAOS_ASYNC_INITBODY_WRITESCOPELOCK(MaterialsLock);
+				Materials.Destroy(InHandle.InnerHandle);
+			}
 		}
 	}
 
 	void FPhysicalMaterialManager::Destroy(FMaterialMaskHandle InHandle)
 	{
-		check(IsInGameThread());
+		check(UE_CHAOS_ASYNC_INITBODY_ENABLED || IsInGameThread());
 		if (InHandle.InnerHandle.IsValid())
 		{
 			OnMaterialMaskDestroyed.Broadcast(InHandle);
-
-			MaterialMasks.Destroy(InHandle.InnerHandle);
+			{
+				UE_CHAOS_ASYNC_INITBODY_WRITESCOPELOCK(MaterialMasksLock);
+				MaterialMasks.Destroy(InHandle.InnerHandle);
+			}
 		}
 	}
 }
\ No newline at end of file
diff --git a/Engine/Source/Runtime/Experimental/Chaos/Private/PBDRigidsSolver.cpp b/Engine/Source/Runtime/Experimental/Chaos/Private/PBDRigidsSolver.cpp
--- a/Engine/Source/Runtime/Experimental/Chaos/Private/PBDRigidsSolver.cpp
+++ b/Engine/Source/Runtime/Experimental/Chaos/Private/PBDRigidsSolver.cpp
@@ -28,6 +28,9 @@
 #include "Chaos/PullPhysicsDataImp.h"
 #include "Chaos/PhysicsSolverBaseImpl.h"
 #include "Chaos/ConvexOptimizer.h"
+#if UE_CHAOS_ASYNC_INITBODY_ENABLED
+#include "Misc/ScopeRWLock.h"
+#endif
 
 #include "ChaosDebugDraw/ChaosDDContext.h"
 #include "ChaosDebugDraw/ChaosDDScene.h"
@@ -742,6 +745,7 @@ namespace Chaos
 
 	void FPBDRigidsSolver::RegisterObject(FSingleParticlePhysicsProxy* Proxy)
 	{
+		UE_CHAOS_ASYNC_INITBODY_PHYSICSSCENE_WRITESCOPELOCK(GetExternalDataLock_External());
 		LLM_SCOPE(ELLMTag::ChaosBody);
 
 		UE_LOG(LogPBDRigidsSolver, Verbose, TEXT("FPBDRigidsSolver::RegisterObject()"));
@@ -796,6 +800,7 @@ namespace Chaos
 
 	void FPBDRigidsSolver::UnregisterObject(FSingleParticlePhysicsProxy* Proxy)
 	{
+		UE_CHAOS_ASYNC_INITBODY_PHYSICSSCENE_WRITESCOPELOCK(GetExternalDataLock_External());
 		UE_LOG(LogPBDRigidsSolver, Verbose, TEXT("FPBDRigidsSolver::UnregisterObject()"));
 
 		PullResultsManager->RemoveProxy_External(Proxy);
@@ -1363,6 +1368,7 @@ namespace Chaos
 		ensure(NumExternalSteps > 0);
 		//TODO: interpolate some data based on num steps
 
+		UE_CHAOS_ASYNC_INITBODY_WRITESCOPELOCK(MarshallingManager.GetMarshallingManagerLock());
 		FPushPhysicsData* PushData = MarshallingManager.GetProducerData_External();
 		const FReal DynamicsWeight = FReal(1) / FReal(NumExternalSteps);
 		FDirtySet* DirtyProxiesData = &PushData->DirtyProxiesDataBuffer;
diff --git a/Engine/Source/Runtime/Experimental/Chaos/Public/Chaos/ChaosMarshallingManager.h b/Engine/Source/Runtime/Experimental/Chaos/Public/Chaos/ChaosMarshallingManager.h
--- a/Engine/Source/Runtime/Experimental/Chaos/Public/Chaos/ChaosMarshallingManager.h
+++ b/Engine/Source/Runtime/Experimental/Chaos/Public/Chaos/ChaosMarshallingManager.h
@@ -8,6 +8,9 @@
 #include "Chaos/ParallelFor.h"
 #include "PhysicsProxy/SingleParticlePhysicsProxyFwd.h"
 #include "Chaos/SimCallbackObject.h"
+#if UE_CHAOS_ASYNC_INITBODY_ENABLED
+#include "Misc/ScopeRWLock.h"
+#endif
 
 class FGeometryCollectionResults;
 
@@ -275,25 +278,81 @@ public:
 		return ProducerData;
 	}
 
+	void AddDirtyProxy(IPhysicsProxyBase* ProxyBaseIn)
+	{
+		UE_CHAOS_ASYNC_INITBODY_WRITESCOPELOCK(MarshallingManagerLock);
+		ProducerData->DirtyProxiesDataBuffer.Add(ProxyBaseIn);
+	}
+
+	void RemoveDirtyProxy(IPhysicsProxyBase* ProxyBaseIn)
+	{
+		UE_CHAOS_ASYNC_INITBODY_WRITESCOPELOCK(MarshallingManagerLock);
+		ProducerData->DirtyProxiesDataBuffer.Remove(ProxyBaseIn);
+	}
+
+	void RemoveDirtyProxyIfNoShapesAreDirty(IPhysicsProxyBase* ProxyBaseIn)
+	{
+		UE_CHAOS_ASYNC_INITBODY_WRITESCOPELOCK(MarshallingManagerLock);
+		ProducerData->DirtyProxiesDataBuffer.RemoveIfNoShapesAreDirty(ProxyBaseIn);
+	}
+
+	const FDirtyProxiesBucketInfo& GetDirtyProxyBucketInfo_External()
+	{
+#if UE_CHAOS_ASYNC_INITBODY_ENABLED
+		ensureAlwaysMsgf(false, TEXT("This method is not safe when used with UE_CHAOS_ASYNC_INITBODY_ENABLED = 1"));
+#endif
+		return ProducerData->DirtyProxiesDataBuffer.GetDirtyProxyBucketInfo();
+	}
+
+	int32 GetDirtyProxyBucketInfoNum_External(EPhysicsProxyType Type)
+	{
+		UE_CHAOS_ASYNC_INITBODY_READSCOPELOCK(MarshallingManagerLock);
+		return ProducerData->DirtyProxiesDataBuffer.GetDirtyProxyBucketInfo().Num[(uint32)Type];
+	}
+
+	// Batch dirty proxies without checking DirtyIdx.
+	template <typename TProxiesArray>
+	void AddDirtyProxiesUnsafe(TProxiesArray& ProxiesArray)
+	{
+		UE_CHAOS_ASYNC_INITBODY_WRITESCOPELOCK(MarshallingManagerLock);
+		ProducerData->DirtyProxiesDataBuffer.AddMultipleUnsafe(ProxiesArray);
+	}
+
+	void AddDirtyProxyShape(IPhysicsProxyBase* ProxyBaseIn, int32 ShapeIdx)
+	{
+		UE_CHAOS_ASYNC_INITBODY_WRITESCOPELOCK(MarshallingManagerLock);
+		ProducerData->DirtyProxiesDataBuffer.AddShape(ProxyBaseIn, ShapeIdx);
+	}
+
+	void SetNumDirtyShapes(IPhysicsProxyBase* Proxy, int32 NumShapes)
+	{
+		UE_CHAOS_ASYNC_INITBODY_WRITESCOPELOCK(MarshallingManagerLock);
+		ProducerData->DirtyProxiesDataBuffer.SetNumDirtyShapes(Proxy, NumShapes);
+	}
+
 	void RegisterSimCallbackObject_External(ISimCallbackObject* SimCallbackObject)
 	{
-		GetProducerData_External()->SimCallbackObjectsToAdd.Add(SimCallbackObject);
+		UE_CHAOS_ASYNC_INITBODY_WRITESCOPELOCK(MarshallingManagerLock);
+		ProducerData->SimCallbackObjectsToAdd.Add(SimCallbackObject);
 	}
 
 	void RegisterSimCommand_External(FSimCallbackCommandObject* SimCommand)
 	{
-		GetProducerData_External()->SimCommands.Add(SimCommand);
+		UE_CHAOS_ASYNC_INITBODY_WRITESCOPELOCK(MarshallingManagerLock);
+		ProducerData->SimCommands.Add(SimCommand);
 	}
 
 	void UnregisterSimCallbackObject_External(ISimCallbackObject* SimCallbackObject)
 	{
+		UE_CHAOS_ASYNC_INITBODY_WRITESCOPELOCK(MarshallingManagerLock);
 		SimCallbackObject->bPendingDelete_External = true;
-		GetProducerData_External()->SimCallbackObjectsToRemove.Add(SimCallbackObject);
+		ProducerData->SimCallbackObjectsToRemove.Add(SimCallbackObject);
 	}
 
 	void AddSimCallbackInputData_External(ISimCallbackObject* SimCallbackObject, FSimCallbackInput* InputData)
 	{
-		GetProducerData_External()->SimCallbackInputs.Add(FSimCallbackInputAndObject{ SimCallbackObject, InputData });
+		UE_CHAOS_ASYNC_INITBODY_WRITESCOPELOCK(MarshallingManagerLock);
+		ProducerData->SimCallbackInputs.Add(FSimCallbackInputAndObject{ SimCallbackObject, InputData });
 	}
 	/** Step forward using the external delta time. Should only be called by external thread */
 	CHAOS_API void Step_External(FReal ExternalDT, const int32 NumSteps = 1, bool bSolverSubstepped = false);
@@ -343,12 +402,17 @@ public:
 
 	/** Return the size of the history queue */
 	int32 GetNumHistory_Internal() const {return HistoryQueue_Internal.Num();}
-		
+
+#if UE_CHAOS_ASYNC_INITBODY_ENABLED
+	/** Used for multithreaded access */
+	CHAOS_API FRWLock& GetMarshallingManagerLock() { return MarshallingManagerLock; }
+#endif
+
 private:
-	FReal ExternalTime_External;	//the global time external thread is currently at
-	int32 ExternalTimestamp_External; //the global timestamp external thread is currently at (1 per frame)
-	FReal SimTime_External;	//the global time the sim is at (once Step_External is called this time advances, even though the actual sim work has yet to be done)
-	int32 InternalStep_External; //the current internal step we are pushing work into. This is not synced with external timestamp as we may sub-step. This should match the solver step
+	std::atomic<FReal> ExternalTime_External;	//the global time external thread is currently at
+	std::atomic<int32> ExternalTimestamp_External; //the global timestamp external thread is currently at (1 per frame)
+	std::atomic<FReal> SimTime_External;	//the global time the sim is at (once Step_External is called this time advances, even though the actual sim work has yet to be done)
+	std::atomic<int32> InternalStep_External; //the current internal step we are pushing work into. This is not synced with external timestamp as we may sub-step. This should match the solver step
 	
 	//push
 	FPushPhysicsData* ProducerData;
@@ -363,11 +427,16 @@ private:
 	TQueue<FPullPhysicsData*,EQueueMode::Spsc> PullDataPool;	//the pull data pool to avoid reallocs. Pushed by external thread, popped by internal
 	TArray<TUniquePtr<FPullPhysicsData>> BackingPullBuffer;		//all pull data is cleaned up by this
 
-	int32 Delay;
+	std::atomic<int32> Delay;
 
 	int32 HistoryLength;	//how long to keep push data for
 
 	CHAOS_API void PrepareExternalQueue_External();
 	CHAOS_API void PreparePullData();
+
+#if UE_CHAOS_ASYNC_INITBODY_ENABLED
+	/** Used for multithreaded access */
+	FRWLock MarshallingManagerLock;
+#endif
 };
 }; // namespace Chaos
diff --git a/Engine/Source/Runtime/Experimental/Chaos/Public/Chaos/Defines.h b/Engine/Source/Runtime/Experimental/Chaos/Public/Chaos/Defines.h
--- a/Engine/Source/Runtime/Experimental/Chaos/Public/Chaos/Defines.h
+++ b/Engine/Source/Runtime/Experimental/Chaos/Public/Chaos/Defines.h
@@ -19,6 +19,18 @@
 #endif
 #endif
 
+// Experimental UE_CHAOS_ASYNC_INITBODY_ENABLED option: Set to 1 to allow InitBody to be called outside of GameThread
+#define UE_CHAOS_ASYNC_INITBODY_ENABLED 0 
+#if UE_CHAOS_ASYNC_INITBODY_ENABLED
+#define UE_CHAOS_ASYNC_INITBODY_WRITESCOPELOCK(x) FWriteScopeLock WriteLock(x)
+#define UE_CHAOS_ASYNC_INITBODY_READSCOPELOCK(x) FReadScopeLock ReadLock(x)
+#define UE_CHAOS_ASYNC_INITBODY_PHYSICSSCENE_WRITESCOPELOCK(x) Chaos::FPhysicsSceneGuardScopedWrite WriteLock(x)
+#else
+#define UE_CHAOS_ASYNC_INITBODY_WRITESCOPELOCK(x)
+#define UE_CHAOS_ASYNC_INITBODY_READSCOPELOCK(x)
+#define UE_CHAOS_ASYNC_INITBODY_PHYSICSSCENE_WRITESCOPELOCK(x)
+#endif
+
 #if COMPILE_WITHOUT_UNREAL_SUPPORT
 	#include <stdint.h>
 #else
diff --git a/Engine/Source/Runtime/Experimental/Chaos/Public/Chaos/Framework/PhysicsSolverBase.h b/Engine/Source/Runtime/Experimental/Chaos/Public/Chaos/Framework/PhysicsSolverBase.h
--- a/Engine/Source/Runtime/Experimental/Chaos/Public/Chaos/Framework/PhysicsSolverBase.h
+++ b/Engine/Source/Runtime/Experimental/Chaos/Public/Chaos/Framework/PhysicsSolverBase.h
@@ -325,38 +325,43 @@ namespace Chaos
 		void AddDirtyProxy(IPhysicsProxyBase * ProxyBaseIn)
 		{
 			check(ProxyBaseIn->GetMarkedDeleted() == false);
-			MarshallingManager.GetProducerData_External()->DirtyProxiesDataBuffer.Add(ProxyBaseIn);
+			MarshallingManager.AddDirtyProxy(ProxyBaseIn);
 		}
 		void RemoveDirtyProxy(IPhysicsProxyBase * ProxyBaseIn)
 		{
-			MarshallingManager.GetProducerData_External()->DirtyProxiesDataBuffer.Remove(ProxyBaseIn);
+			MarshallingManager.RemoveDirtyProxy(ProxyBaseIn);
 		}
 
 		void RemoveDirtyProxyIfNoShapesAreDirty(IPhysicsProxyBase* ProxyBaseIn)
 		{
-			MarshallingManager.GetProducerData_External()->DirtyProxiesDataBuffer.RemoveIfNoShapesAreDirty(ProxyBaseIn);
+			MarshallingManager.RemoveDirtyProxyIfNoShapesAreDirty(ProxyBaseIn);
 		}
 
 		const FDirtyProxiesBucketInfo& GetDirtyProxyBucketInfo_External()
 		{
-			return MarshallingManager.GetProducerData_External()->DirtyProxiesDataBuffer.GetDirtyProxyBucketInfo();
+			return MarshallingManager.GetDirtyProxyBucketInfo_External();
+		}
+
+		int32 GetDirtyProxyBucketInfoNum_External(EPhysicsProxyType Type)
+		{
+			return MarshallingManager.GetDirtyProxyBucketInfoNum_External(Type);
 		}
 
 		// Batch dirty proxies without checking DirtyIdx.
 		template <typename TProxiesArray>
 		void AddDirtyProxiesUnsafe(TProxiesArray& ProxiesArray)
 		{
-			MarshallingManager.GetProducerData_External()->DirtyProxiesDataBuffer.AddMultipleUnsafe(ProxiesArray);
+			MarshallingManager.AddDirtyProxiesUnsafe(ProxiesArray);
 		}
 
 		void AddDirtyProxyShape(IPhysicsProxyBase* ProxyBaseIn, int32 ShapeIdx)
 		{
-			MarshallingManager.GetProducerData_External()->DirtyProxiesDataBuffer.AddShape(ProxyBaseIn,ShapeIdx);
+			MarshallingManager.AddDirtyProxyShape(ProxyBaseIn, ShapeIdx);
 		}
 
 		void SetNumDirtyShapes(IPhysicsProxyBase* Proxy, int32 NumShapes)
 		{
-			MarshallingManager.GetProducerData_External()->DirtyProxiesDataBuffer.SetNumDirtyShapes(Proxy,NumShapes);
+			MarshallingManager.SetNumDirtyShapes(Proxy, NumShapes);
 		}
 
 		/** Creates a new sim callback object of the type given. Caller expected to free using FreeSimCallbackObject_External*/
@@ -418,7 +423,7 @@ namespace Chaos
 		void EnqueueCommandImmediate(Lambda&& Func)
 		{
 			//TODO: remove this check. Need to rename with _External
-			check(IsInGameThread());
+			check(UE_CHAOS_ASYNC_INITBODY_ENABLED || IsInGameThread());
 			RegisterSimOneShotCallback(MoveTemp(Func));
 		}
 
diff --git a/Engine/Source/Runtime/Experimental/Chaos/Public/Chaos/ParticleHandle.h b/Engine/Source/Runtime/Experimental/Chaos/Public/Chaos/ParticleHandle.h
--- a/Engine/Source/Runtime/Experimental/Chaos/Public/Chaos/ParticleHandle.h
+++ b/Engine/Source/Runtime/Experimental/Chaos/Public/Chaos/ParticleHandle.h
@@ -3047,16 +3047,13 @@ protected:
 	template <typename Lambda>
 	void ModifyGeometry(const Lambda& Func, const bool bDirectAccess = false)
 	{
-		ensure(IsInGameThread());
-
 		return ModifyGeometry(bDirectAccess ? EGeometryAccess::Direct : EGeometryAccess::DeepCopy, Func);
 	}
 
 	template <typename Lambda>
 	void ModifyGeometry(EGeometryAccess AccessType, const Lambda& Func)
 	{
-		ensure(IsInGameThread());
-
+		ensure(UE_CHAOS_ASYNC_INITBODY_ENABLED || IsInGameThread());
 		FPhysicsSolverBase* Solver = Proxy ? Proxy->GetSolverBase() : nullptr;
 
 		if(Solver == nullptr)
@@ -3801,7 +3798,7 @@ FORCEINLINE_DEBUGGABLE FAccelerationStructureHandle::FAccelerationStructureHandl
 	if (InGeometryParticle)
 	{
 		ensure(CachedUniqueIdx.IsValid());
-		ensure(IsInGameThread());
+		ensure(UE_CHAOS_ASYNC_INITBODY_ENABLED || IsInGameThread());
 		if (bUsePrefiltering)
 		{
 			UpdatePrePreFilter(*InGeometryParticle);
diff --git a/Engine/Source/Runtime/Experimental/Chaos/Public/Chaos/PhysicalMaterials.h b/Engine/Source/Runtime/Experimental/Chaos/Public/Chaos/PhysicalMaterials.h
--- a/Engine/Source/Runtime/Experimental/Chaos/Public/Chaos/PhysicalMaterials.h
+++ b/Engine/Source/Runtime/Experimental/Chaos/Public/Chaos/PhysicalMaterials.h
@@ -167,6 +167,11 @@ namespace Chaos
 
 		/** Handle-managed array of global material masks. This is pushed to all solvers who all maintain a copy */
 		THandleArray<FChaosPhysicsMaterialMask> MaterialMasks;
+
+#if UE_CHAOS_ASYNC_INITBODY_ENABLED
+		mutable FRWLock MaterialsLock;
+		mutable FRWLock MaterialMasksLock;
+#endif
 	};
 }
 
diff --git a/Engine/Source/Runtime/PhysicsCore/Private/ChaosScene.cpp b/Engine/Source/Runtime/PhysicsCore/Private/ChaosScene.cpp
--- a/Engine/Source/Runtime/PhysicsCore/Private/ChaosScene.cpp
+++ b/Engine/Source/Runtime/PhysicsCore/Private/ChaosScene.cpp
@@ -281,6 +281,7 @@ void FChaosScene::AddActorsToScene_AssumesLocked(TArray<FPhysicsActorHandle>& In
 	TRACE_CPUPROFILER_EVENT_SCOPE(FChaosScene::AddActorsToScene_AssumesLocked)
 
 	Chaos::FPhysicsSolver* Solver = GetSolver();
+	UE_CHAOS_ASYNC_INITBODY_PHYSICSSCENE_WRITESCOPELOCK(Solver->GetExternalDataLock_External());
 	Chaos::ISpatialAcceleration<Chaos::FAccelerationStructureHandle,Chaos::FReal,3>* SpatialAcceleration = GetSpacialAcceleration();
 	for(FPhysicsActorHandle& Handle : InHandles)
 	{
