Fixed static analysis following CL 35964441.

diff --git a/Engine/Source/Runtime/Engine/Classes/PhysicsEngine/BodyInstance.h b/Engine/Source/Runtime/Engine/Classes/PhysicsEngine/BodyInstance.h
--- a/Engine/Source/Runtime/Engine/Classes/PhysicsEngine/BodyInstance.h
+++ b/Engine/Source/Runtime/Engine/Classes/PhysicsEngine/BodyInstance.h
@@ -195,8 +195,16 @@ struct FInitBodiesHelperWithData : public FInitBodiesHelperBase
 	}
 
 	FInitBodiesHelperWithData(FInitBodiesHelperWithData&& InHelper)
-	: FInitBodiesHelperWithData(OwnedBodies, OwnedTransforms, InHelper.BodySetup, InHelper.PrimitiveComp, InHelper.PhysScene, InHelper.SpawnParams, InHelper.Aggregate), OwnedBodies(MoveTemp(InHelper.OwnedBodies)), OwnedTransforms(MoveTemp(InHelper.OwnedTransforms)) //-V1050
-	{}
+	: FInitBodiesHelperBase(OwnedBodies, OwnedTransforms, InHelper.BodySetup, InHelper.PrimitiveComp, InHelper.PhysScene, InHelper.SpawnParams, InHelper.Aggregate), OwnedBodies(MoveTemp(InHelper.OwnedBodies)), OwnedTransforms(MoveTemp(InHelper.OwnedTransforms)) //-V1050
+	{
+		//Compute all the needed constants
+		bStatic = bCompileStatic || SpawnParams.bStaticPhysics;
+		SkelMeshComp = bCompileStatic ? nullptr : Cast<USkeletalMeshComponent>(PrimitiveComp);
+		if (SpawnParams.bPhysicsTypeDeterminesSimulation)
+		{
+			this->UpdateSimulatingAndBlendWeight();
+		}
+	}
 
 	FInitBodiesHelperWithData& operator=(const FInitBodiesHelperWithData& InHelper) = delete;
 	FInitBodiesHelperWithData& operator=(FInitBodiesHelperWithData&& InHelper) = delete;
