Add support for using uncompressed textures with mesh texture color painting.
The "Mesh paint use compressed textures" project setting sets whether to use BC3 or ARGB8888 textures.
Compressed textures have a memory advantage. Uncompressed textures might be useful if we need to avoid artifacts from channel crosstalk.
Default behavior is to use compressed textures.
Changing the setting forces an editor restart, but existing textures will update correctly.

diff --git a/Engine/Plugins/MeshPainting/Source/MeshPaintingToolset/Private/MeshTexturePaintingTool.cpp b/Engine/Plugins/MeshPainting/Source/MeshPaintingToolset/Private/MeshTexturePaintingTool.cpp
--- a/Engine/Plugins/MeshPainting/Source/MeshPaintingToolset/Private/MeshTexturePaintingTool.cpp
+++ b/Engine/Plugins/MeshPainting/Source/MeshPaintingToolset/Private/MeshTexturePaintingTool.cpp
@@ -1305,7 +1305,7 @@ void UMeshTexturePaintingTool::SetAllTextureOverrides()
 
 					TextureData->PaintRenderTargetTextureAdapter = NewObject<UVirtualTextureAdapter>(GetTransientPackage(), NAME_None, RF_Transient);
 					TextureData->PaintRenderTargetTextureAdapter->Texture = TextureData->PaintRenderTargetTexture;
-					TextureData->PaintRenderTargetTextureAdapter->bUseCompressedFormat = true;
+					TextureData->PaintRenderTargetTextureAdapter->OverrideWithTextureFormat = Texture2D;
 					TextureData->PaintRenderTargetTextureAdapter->bUseDefaultTileSizes = false;
 					TextureData->PaintRenderTargetTextureAdapter->TileSize = VirtualTextureBuildSettings.TileSize;
 					TextureData->PaintRenderTargetTextureAdapter->TileBorderSize = VirtualTextureBuildSettings.TileBorderSize;
diff --git a/Engine/Source/Runtime/Engine/Classes/Engine/RendererSettings.h b/Engine/Source/Runtime/Engine/Classes/Engine/RendererSettings.h
--- a/Engine/Source/Runtime/Engine/Classes/Engine/RendererSettings.h
+++ b/Engine/Source/Runtime/Engine/Classes/Engine/RendererSettings.h
@@ -458,6 +458,13 @@ class URendererSettings : public UDeveloperSettings
 		ConfigRestartRequired = true))
 	uint32 MeshPaintVirtualTextureTileBorderSize;
 
+	UPROPERTY(config, EditAnywhere, Category = VirtualTextures, meta = (
+		EditCondition = "bVirtualTextures",
+		ConsoleVariable = "r.MeshPaintVirtualTexture.UseCompression", DisplayName = "Mesh paint use compressed textures",
+		ToolTip = "Whether to use a compressed texture format for storing mesh paint textures.",
+		ConfigRestartRequired = true))
+	bool MeshPaintVirtualTextureUseCompression;
+
 	UPROPERTY(config, EditAnywhere, Category = VirtualTextures, meta = (
 		EditCondition = "bVirtualTextures",
 		ConsoleVariable = "r.StaticMesh.DefaultMeshPaintTextureSupport", DisplayName = "Mesh paint default static mesh support",
diff --git a/Engine/Source/Runtime/Engine/Classes/VT/VirtualTextureAdapter.h b/Engine/Source/Runtime/Engine/Classes/VT/VirtualTextureAdapter.h
--- a/Engine/Source/Runtime/Engine/Classes/VT/VirtualTextureAdapter.h
+++ b/Engine/Source/Runtime/Engine/Classes/VT/VirtualTextureAdapter.h
@@ -32,9 +32,9 @@ public:
 	UPROPERTY(EditAnywhere, Category = Texture)
 	TObjectPtr<UTexture> Texture;
 
-	/** Use a compressed format for the virtual texture even when the wrapped Texture does not have a compressed format. */
-	UPROPERTY(EditAnywhere, Category = Texture)
-	bool bUseCompressedFormat = false;
+	/** Optional UTexture object that has the final texture format that we would like to use. */
+	UPROPERTY(EditAnywhere, AdvancedDisplay, Category = Texture)
+	TObjectPtr<class UTexture2D> OverrideWithTextureFormat;
 
 	/** Set to true if we want to use the default project virtual texture tile settings. */
 	UPROPERTY(EditAnywhere, Category = Texture)
diff --git a/Engine/Source/Runtime/Engine/Private/VT/MeshPaintVirtualTexture.cpp b/Engine/Source/Runtime/Engine/Private/VT/MeshPaintVirtualTexture.cpp
--- a/Engine/Source/Runtime/Engine/Private/VT/MeshPaintVirtualTexture.cpp
+++ b/Engine/Source/Runtime/Engine/Private/VT/MeshPaintVirtualTexture.cpp
@@ -29,6 +29,12 @@ static TAutoConsoleVariable<bool> CVarMeshPaintVirtualTextureEnable(
 	}),
 	ECVF_RenderThreadSafe);
 
+static TAutoConsoleVariable<bool> CVarMeshPaintVirtualTextureUseCompression(
+	TEXT("r.MeshPaintVirtualTexture.UseCompression"),
+	true,
+	TEXT("Whether mesh paint textures use a compressed format"),
+	ECVF_ReadOnly);
+
 static TAutoConsoleVariable<int32> CVarMeshPaintVirtualTextureTileSize(
 	TEXT("r.MeshPaintVirtualTexture.TileSize"),
 	32,
@@ -84,6 +90,11 @@ namespace MeshPaintVirtualTexture
 		return CVarMeshPaintVirtualTextureEnable.GetValueOnAnyThread();
 	}
 
+	static bool UseCompressedTextureFormat()
+	{
+		return CVarMeshPaintVirtualTextureUseCompression.GetValueOnGameThread();
+	}
+
 	static uint32 GetTileSize()
 	{
 		return FVirtualTextureBuildSettings::ClampAndAlignTileSize(CVarMeshPaintVirtualTextureTileSize.GetValueOnAnyThread());
@@ -262,6 +273,8 @@ UMeshPaintVirtualTexture::UMeshPaintVirtualTexture(const FObjectInitializer& Obj
 	VirtualTextureStreaming = true;
 	
 #if WITH_EDITORONLY_DATA
+	CompressionNone = !MeshPaintVirtualTexture::UseCompressedTextureFormat();
+
 	// Force alpha channel so that we the platform format is consistent for all content.
 	CompressionForceAlpha = true;
 #endif
diff --git a/Engine/Source/Runtime/Engine/Private/VT/VirtualTextureAdapter.cpp b/Engine/Source/Runtime/Engine/Private/VT/VirtualTextureAdapter.cpp
--- a/Engine/Source/Runtime/Engine/Private/VT/VirtualTextureAdapter.cpp
+++ b/Engine/Source/Runtime/Engine/Private/VT/VirtualTextureAdapter.cpp
@@ -62,29 +62,17 @@ namespace VirtualTextureAdapter
 
 
 	/** Get the final virtual texture format from the wrapped texture format. */
-	static EPixelFormat GetFinalFormat(EPixelFormat InSourceFormat, bool bInUseCompressedFormat)
+	static EPixelFormat GetFinalFormat(EPixelFormat InSourcePixelFormat, EPixelFormat InFinalPixelFormat)
 	{
-		if (!bInUseCompressedFormat || IsBlockCompressedFormat(InSourceFormat) || IsInteger(InSourceFormat))
-		{
-			return InSourceFormat;
-		}
-
-		if (IsHDR(InSourceFormat))
-		{
-			return PF_BC6H;
-		}
+		EPixelFormat FinalPixelFormat = InFinalPixelFormat != PF_Unknown ? InFinalPixelFormat : InSourcePixelFormat;
 
-		EPixelFormatChannelFlags ChannelFlags = GetPixelFormatValidChannels(InSourceFormat);
-		if (EnumHasAnyFlags(ChannelFlags, EPixelFormatChannelFlags::A))
-		{
-			return PF_DXT5;
-		}
-		if (!EnumHasAnyFlags(ChannelFlags, EPixelFormatChannelFlags::G | EPixelFormatChannelFlags::B | EPixelFormatChannelFlags::A))
+		// Can't override some formats.
+		if (IsBlockCompressedFormat(InSourcePixelFormat) || IsInteger(InSourcePixelFormat) || IsInteger(FinalPixelFormat))
 		{
-			return PF_BC4;
+			return InSourcePixelFormat;
 		}
 
-		return PF_DXT1;
+		return FinalPixelFormat;
 	}
 
 	/** Get the intermediate texture format that we use for transient intermediate targets. */
@@ -477,26 +465,26 @@ class FVirtualTextureAdapterRenderResource : public FVirtualTexture2DResource
 	uint32 NumSourceMips = 1;
 
 public:
-	FVirtualTextureAdapterRenderResource(UVirtualTextureAdapter const* InOwner, UTexture* InTexture, int32 InTileSize, int32 InTileBorderSize, bool bUseCompressedFormat)
+	FVirtualTextureAdapterRenderResource(UVirtualTextureAdapter const* InOwner, UTexture* InTexture, int32 InTileSize, int32 InTileBorderSize, EPixelFormat InFinalPixelFormat)
 	{
 		TextureName = InOwner->GetFName();
 		PackageName = InOwner->GetOutermost()->GetFName();
 
 		SourceResource = InTexture->GetResource();
 
-		EPixelFormat SourceFormat = PF_Unknown;
+		EPixelFormat SourcePixelFormat = PF_Unknown;
 		if (UTexture2D* Texture2D = Cast<UTexture2D>(InTexture))
 		{
-			SourceFormat = Texture2D->GetPixelFormat(0);
+			SourcePixelFormat = Texture2D->GetPixelFormat(0);
 			NumSourceMips = Texture2D->GetNumMips();
 			bSRGB = Texture2D->SRGB;
 		}
 		else if (UTextureRenderTarget2D* RenderTarget2D = Cast<UTextureRenderTarget2D>(InTexture))
 		{
-			SourceFormat = RenderTarget2D->GetFormat();
+			SourcePixelFormat = RenderTarget2D->GetFormat();
 			bSRGB = RenderTarget2D->SRGB;
 		}
-		Format = VirtualTextureAdapter::GetFinalFormat(SourceFormat, bUseCompressedFormat);
+		Format = VirtualTextureAdapter::GetFinalFormat(SourcePixelFormat, InFinalPixelFormat);
 
 		TileSize = InTileSize;
 		TileBorderSize = InTileBorderSize;
@@ -572,7 +560,9 @@ FTextureResource* UVirtualTextureAdapter::CreateResource()
 	const uint32 FinalTileSize = bUseDefaultTileSizes ? DefaultSettings.TileSize : FVirtualTextureBuildSettings::ClampAndAlignTileSize(TileSize);
 	const uint32 FinalTileBorderSize = bUseDefaultTileSizes ? DefaultSettings.TileBorderSize : FVirtualTextureBuildSettings::ClampAndAlignTileBorderSize(TileBorderSize);
 	
-	return new FVirtualTextureAdapterRenderResource(this, Texture, FinalTileSize, FinalTileBorderSize, bUseCompressedFormat);
+	const EPixelFormat FinalPixelFormat = OverrideWithTextureFormat ? OverrideWithTextureFormat->GetPixelFormat(0) : PF_Unknown;
+
+	return new FVirtualTextureAdapterRenderResource(this, Texture, FinalTileSize, FinalTileBorderSize, FinalPixelFormat);
 }
 
 
