[Cleanup] Fix up minor const-correctness issues in AnimationRuntime.

C-style casts were hiding a handful of missing `const`s.

diff --git a/Engine/Source/Runtime/Engine/Private/Animation/AnimationRuntime.cpp b/Engine/Source/Runtime/Engine/Private/Animation/AnimationRuntime.cpp
--- a/Engine/Source/Runtime/Engine/Private/Animation/AnimationRuntime.cpp
+++ b/Engine/Source/Runtime/Engine/Private/Animation/AnimationRuntime.cpp
@@ -164,8 +164,8 @@ FORCEINLINE void BlendPose<ETransformBlendMode::Overwrite>(const FCompactPose& S
 	if (bAnim_BlendPoseOverwrite_ISPC_Enabled)
 	{
 		ispc::BlendTransformOverwrite(
-			(ispc::FTransform*)&SourcePose.GetBones()[0],
-			(ispc::FTransform*)&ResultPose.GetBones()[0],
+			reinterpret_cast<const ispc::FTransform*>(SourcePose.GetBones().GetData()),
+			reinterpret_cast<ispc::FTransform*>(ResultPose.GetMutableBones().GetData()),
 			BlendWeight,
 			SourcePose.GetNumBones());
 	}
@@ -186,8 +186,8 @@ FORCEINLINE void BlendPose<ETransformBlendMode::Accumulate>(const FCompactPose&
 	if (bAnim_BlendPoseAccumulate_ISPC_Enabled)
 	{
 		ispc::BlendTransformAccumulate(
-			(ispc::FTransform*)&SourcePose.GetBones()[0],
-			(ispc::FTransform*)&ResultPose.GetBones()[0],
+			reinterpret_cast<const ispc::FTransform*>(SourcePose.GetBones().GetData()),
+			reinterpret_cast<ispc::FTransform*>(ResultPose.GetMutableBones().GetData()),
 			BlendWeight,
 			SourcePose.GetNumBones());
 	}
@@ -487,8 +487,8 @@ void FAnimationRuntime::BlendTwoPosesTogether(
 	
 	FAnimationPoseData AnimationPoseData = { ResultPose, ResultCurve, TempAttributes };
 	
-	const FAnimationPoseData SourceOnePoseData(*const_cast<FCompactPose*>(&SourcePose1), *const_cast<FBlendedCurve*>(&SourceCurve1), TempAttributes);
-	const FAnimationPoseData SourceTwoPosedata(*const_cast<FCompactPose*>(&SourcePose2), *const_cast<FBlendedCurve*>(&SourceCurve2), TempAttributes);
+	const FAnimationPoseData SourceOnePoseData(const_cast<FCompactPose&>(SourcePose1), const_cast<FBlendedCurve&>(SourceCurve1), TempAttributes);
+	const FAnimationPoseData SourceTwoPosedata(const_cast<FCompactPose&>(SourcePose2), const_cast<FBlendedCurve&>(SourceCurve2), TempAttributes);
 
 	BlendTwoPosesTogether(SourceOnePoseData, SourceTwoPosedata, WeightOfPose1, AnimationPoseData);
 }
@@ -523,8 +523,8 @@ void FAnimationRuntime::BlendTwoPosesTogetherPerBone(
 	UE::Anim::FStackAttributeContainer TempAttributes;
 
 	FAnimationPoseData AnimationPoseData = { ResultPose, ResultCurve, TempAttributes };
-	const FAnimationPoseData SourceOnePoseData(*const_cast<FCompactPose*>(&SourcePose1), *const_cast<FBlendedCurve*>(&SourceCurve1), TempAttributes);
-	const FAnimationPoseData SourceTwoPosedata(*const_cast<FCompactPose*>(&SourcePose2), *const_cast<FBlendedCurve*>(&SourceCurve2), TempAttributes);
+	const FAnimationPoseData SourceOnePoseData(const_cast<FCompactPose&>(SourcePose1), const_cast<FBlendedCurve&>(SourceCurve1), TempAttributes);
+	const FAnimationPoseData SourceTwoPosedata(const_cast<FCompactPose&>(SourcePose2), const_cast<FBlendedCurve&>(SourceCurve2), TempAttributes);
 
 	BlendTwoPosesTogetherPerBone(SourceOnePoseData, SourceTwoPosedata, WeightsOfSource2, AnimationPoseData);
 }
@@ -967,8 +967,8 @@ void FAnimationRuntime::LerpBoneTransforms(TArray<FTransform>& A, const TArray<F
 	if (bAnim_LerpBoneTransforms_ISPC_Enabled)
 	{
 		ispc::LerpBoneTransforms(
-			(ispc::FTransform*)A.GetData(),
-			(ispc::FTransform*)B.GetData(),
+			reinterpret_cast<ispc::FTransform*>(A.GetData()),
+			reinterpret_cast<const ispc::FTransform*>(B.GetData()),
 			Alpha,
 			RequiredBonesArray.GetData(),
 			RequiredBonesArray.Num()
@@ -1064,8 +1064,8 @@ void FAnimationRuntime::ConvertPoseToAdditive(FCompactPose& TargetPose, const FC
 	if (bAnim_ConvertPoseToAdditive_ISPC_Enabled)
 	{
 		ispc::ConvertPoseToAdditive(
-			(ispc::FTransform*)TargetPose.GetBones().GetData(),
-			(ispc::FTransform*)BasePose.GetBones().GetData(),
+			reinterpret_cast<ispc::FTransform*>(TargetPose.GetMutableBones().GetData()),
+			reinterpret_cast<const ispc::FTransform*>(BasePose.GetBones().GetData()),
 			BasePose.GetNumBones());
 	}
 	else
@@ -1089,8 +1089,8 @@ void FAnimationRuntime::ConvertPoseToMeshRotation(FCompactPose& LocalPose)
 	if (bAnim_ConvertPoseToMeshRotation_ISPC_Enabled)
 	{
 		ispc::ConvertPoseToMeshRotation(
-			(ispc::FTransform*)&LocalPose.GetBones()[0],
-			(int32*)&LocalPose.GetBoneContainer().GetCompactPoseParentBoneArray().GetData()[0],
+			reinterpret_cast<ispc::FTransform*>(LocalPose.GetMutableBones().GetData()),
+			reinterpret_cast<const int32*>(LocalPose.GetBoneContainer().GetCompactPoseParentBoneArray().GetData()),
 			LocalPose.GetNumBones());
 	}
 	else
@@ -1114,8 +1114,8 @@ void FAnimationRuntime::ConvertMeshRotationPoseToLocalSpace(FCompactPose& Pose)
 	if (bAnim_ConvertMeshRotationPoseToLocalSpace_ISPC_Enabled)
 	{
 		ispc::ConvertMeshRotationPoseToLocalSpace(
-			(ispc::FTransform*)&Pose.GetBones()[0],
-			(int32*)&Pose.GetBoneContainer().GetCompactPoseParentBoneArray().GetData()[0],
+			reinterpret_cast<ispc::FTransform*>(Pose.GetMutableBones().GetData()),
+			reinterpret_cast<const int32*>(Pose.GetBoneContainer().GetCompactPoseParentBoneArray().GetData()),
 			Pose.GetNumBones());
 	}
 	else
@@ -1135,7 +1135,7 @@ void FAnimationRuntime::AccumulateAdditivePose(FCompactPose& BasePose, const FCo
 {
 	UE::Anim::FStackAttributeContainer TempAttributes;
 	FAnimationPoseData BaseAnimationPoseData = { BasePose, BaseCurve, TempAttributes };
-	const FAnimationPoseData AdditiveAnimationPoseData(*const_cast<FCompactPose*>(&AdditivePose), *const_cast<FBlendedCurve*>(&AdditiveCurve), TempAttributes);
+	const FAnimationPoseData AdditiveAnimationPoseData(const_cast<FCompactPose&>(AdditivePose), const_cast<FBlendedCurve&>(AdditiveCurve), TempAttributes);
 
 	AccumulateAdditivePose(BaseAnimationPoseData, AdditiveAnimationPoseData, Weight, AdditiveType);
 }
@@ -1145,7 +1145,7 @@ void FAnimationRuntime::AccumulateLocalSpaceAdditivePose(FCompactPose& BasePose,
 {
 	UE::Anim::FStackAttributeContainer TempAttributes;
 	FAnimationPoseData BaseAnimationPoseData = { BasePose, BaseCurve, TempAttributes };
-	const FAnimationPoseData AdditiveAnimationPoseData(*const_cast<FCompactPose*>(&AdditivePose), *const_cast<FBlendedCurve*>(&AdditiveCurve), TempAttributes);
+	const FAnimationPoseData AdditiveAnimationPoseData(const_cast<FCompactPose&>(AdditivePose), const_cast<FBlendedCurve&>(AdditiveCurve), TempAttributes);
 
 	AccumulateAdditivePose(BaseAnimationPoseData, AdditiveAnimationPoseData, Weight, EAdditiveAnimationType::AAT_LocalSpaceBase);
 }
@@ -1159,7 +1159,7 @@ void FAnimationRuntime::AccumulateMeshSpaceRotationAdditiveToLocalPose(FCompactP
 {
 	UE::Anim::FStackAttributeContainer TempAttributes;
 	FAnimationPoseData BaseAnimationPoseData = { BasePose, BaseCurve, TempAttributes };
-	const FAnimationPoseData AdditiveAnimationPoseData(*const_cast<FCompactPose*>(&MeshSpaceRotationAdditive), *const_cast<FBlendedCurve*>(&AdditiveCurve), TempAttributes);
+	const FAnimationPoseData AdditiveAnimationPoseData(const_cast<FCompactPose&>(MeshSpaceRotationAdditive), const_cast<FBlendedCurve&>(AdditiveCurve), TempAttributes);
 
 	AccumulateAdditivePose(BaseAnimationPoseData, AdditiveAnimationPoseData, Weight, EAdditiveAnimationType::AAT_RotationOffsetMeshSpace);
 }
@@ -1200,8 +1200,8 @@ void FAnimationRuntime::AccumulateLocalSpaceAdditivePoseInternal(FCompactPose& B
 			if (bAnim_AccumulateLocalSpaceAdditivePose_ISPC_Enabled)
 			{
 				ispc::AccumulateWithAdditiveScale(
-					(ispc::FTransform*)&BasePose.GetBones()[0],
-					(ispc::FTransform*)&AdditivePose.GetBones()[0],
+					reinterpret_cast<ispc::FTransform*>(BasePose.GetMutableBones().GetData()),
+					reinterpret_cast<const ispc::FTransform*>(AdditivePose.GetBones().GetData()),
 					Weight,
 					BasePose.GetNumBones());
 			}
@@ -1742,9 +1742,9 @@ struct FBlendPosesPerBoneFilterScratchArea : public TThreadSingleton<FBlendPoses
 // Helper function to get FTransform from a PoseIndex and BoneIndex
 extern "C" const uint8* GetTransformFromArray(const uint8 *BlendPoseBase, const int32 PoseIndex, const int32 BoneIndex)
 {
-	const TArray<struct FCompactPose>& BlendPoses = *(const TArray<struct FCompactPose>*) BlendPoseBase;
+	const TArray<struct FCompactPose>& BlendPoses = reinterpret_cast<const TArray<struct FCompactPose>&>(BlendPoseBase);
 	const FTransform* BlendPose = &BlendPoses[PoseIndex][FCompactPoseBoneIndex(BoneIndex)];
-	return (const uint8*)BlendPose;
+	return reinterpret_cast<const uint8*>(BlendPose);
 }
 
 void FAnimationRuntime::BlendPosesPerBoneFilter(
@@ -1887,18 +1887,18 @@ FAnimationPoseData& OutAnimationPoseData, TArray<FPerBoneBlendWeight>& BoneBlend
 		if (bAnim_BlendPosesPerBoneFilter_ISPC_Enabled)
 		{
 			ispc::BlendPosesPerBoneFilterScaleRotation(
-				(ispc::FTransform*)OutPose.GetBones().GetData(),
-				(ispc::FTransform*)BasePose.GetBones().GetData(),
-				(const uint8*)&BlendPoses,
-				(ispc::FVector4*)SourceRotations.GetData(),
-				(ispc::FVector*)SourceScales.GetData(),
-				(ispc::FVector4*)TargetRotations.GetData(),
-				(ispc::FVector*)TargetScales.GetData(),
-				(ispc::FVector4*)BlendRotations.GetData(),
-				(ispc::FVector*)BlendScales.GetData(),
+				reinterpret_cast<ispc::FTransform*>(OutPose.GetMutableBones().GetData()),
+				reinterpret_cast<const ispc::FTransform*>(BasePose.GetBones().GetData()),
+				reinterpret_cast<const uint8*>(&BlendPoses),
+				reinterpret_cast<ispc::FVector4*>(SourceRotations.GetData()),
+				reinterpret_cast<ispc::FVector*>(SourceScales.GetData()),
+				reinterpret_cast<ispc::FVector4*>(TargetRotations.GetData()),
+				reinterpret_cast<ispc::FVector*>(TargetScales.GetData()),
+				reinterpret_cast<ispc::FVector4*>(BlendRotations.GetData()),
+				reinterpret_cast<ispc::FVector*>(BlendScales.GetData()),
 				MaxPoseWeights.GetData(),
-				(ispc::FPerBoneBlendWeight*)BoneBlendWeights.GetData(),
-				(int32*)BoneContainer.GetCompactPoseParentBoneArray().GetData(),
+				reinterpret_cast<const ispc::FPerBoneBlendWeight*>(BoneBlendWeights.GetData()),
+				reinterpret_cast<const int32*>(BoneContainer.GetCompactPoseParentBoneArray().GetData()),
 				BasePose.GetNumBones());
 		}
 		else
@@ -1965,15 +1965,15 @@ FAnimationPoseData& OutAnimationPoseData, TArray<FPerBoneBlendWeight>& BoneBlend
 		if (bAnim_BlendPosesPerBoneFilter_ISPC_Enabled)
 		{
 			ispc::BlendPosesPerBoneFilterRotation(
-				(ispc::FTransform*)OutPose.GetBones().GetData(),
-				(ispc::FTransform*)BasePose.GetBones().GetData(),
-				(const uint8*)&BlendPoses,
-				(ispc::FVector4*)SourceRotations.GetData(),
-				(ispc::FVector4*)TargetRotations.GetData(),
-				(ispc::FVector4*)BlendRotations.GetData(),
+				reinterpret_cast<ispc::FTransform*>(OutPose.GetMutableBones().GetData()),
+				reinterpret_cast<const ispc::FTransform*>(BasePose.GetBones().GetData()),
+				reinterpret_cast<const uint8*>(&BlendPoses),
+				reinterpret_cast<ispc::FVector4*>(SourceRotations.GetData()),
+				reinterpret_cast<ispc::FVector4*>(TargetRotations.GetData()),
+				reinterpret_cast<ispc::FVector4*>(BlendRotations.GetData()),
 				MaxPoseWeights.GetData(),
-				(ispc::FPerBoneBlendWeight*)BoneBlendWeights.GetData(),
-				(int32*)BoneContainer.GetCompactPoseParentBoneArray().GetData(),
+				reinterpret_cast<const ispc::FPerBoneBlendWeight*>(BoneBlendWeights.GetData()),
+				reinterpret_cast<const int32*>(BoneContainer.GetCompactPoseParentBoneArray().GetData()),
 				BasePose.GetNumBones());
 		}
 		else
@@ -2034,15 +2034,15 @@ FAnimationPoseData& OutAnimationPoseData, TArray<FPerBoneBlendWeight>& BoneBlend
 		if (bAnim_BlendPosesPerBoneFilter_ISPC_Enabled)
 		{
 			ispc::BlendPosesPerBoneFilterScale(
-				(ispc::FTransform*)OutPose.GetBones().GetData(),
-				(ispc::FTransform*)BasePose.GetBones().GetData(),
-				(const uint8*)&BlendPoses,
-				(ispc::FVector*)SourceScales.GetData(),
-				(ispc::FVector*)TargetScales.GetData(),
-				(ispc::FVector*)BlendScales.GetData(),
+				reinterpret_cast<ispc::FTransform*>(OutPose.GetMutableBones().GetData()),
+				reinterpret_cast<const ispc::FTransform*>(BasePose.GetBones().GetData()),
+				reinterpret_cast<const uint8*>(&BlendPoses),
+				reinterpret_cast<ispc::FVector*>(SourceScales.GetData()),
+				reinterpret_cast<ispc::FVector*>(TargetScales.GetData()),
+				reinterpret_cast<ispc::FVector*>(BlendScales.GetData()),
 				MaxPoseWeights.GetData(),
-				(ispc::FPerBoneBlendWeight*)BoneBlendWeights.GetData(),
-				(int32*)BoneContainer.GetCompactPoseParentBoneArray().GetData(),
+				reinterpret_cast<const ispc::FPerBoneBlendWeight*>(BoneBlendWeights.GetData()),
+				reinterpret_cast<const int32*>(BoneContainer.GetCompactPoseParentBoneArray().GetData()),
 				BasePose.GetNumBones());
 		}
 		else
@@ -2103,12 +2103,12 @@ FAnimationPoseData& OutAnimationPoseData, TArray<FPerBoneBlendWeight>& BoneBlend
 		if (bAnim_BlendPosesPerBoneFilter_ISPC_Enabled)
 		{
 			ispc::BlendPosesPerBoneFilter(
-				(ispc::FTransform*)OutPose.GetBones().GetData(),
-				(ispc::FTransform*)BasePose.GetBones().GetData(),
-				(const uint8*)&BlendPoses,
+				reinterpret_cast<ispc::FTransform*>(OutPose.GetMutableBones().GetData()),
+				reinterpret_cast<const ispc::FTransform*>(BasePose.GetBones().GetData()),
+				reinterpret_cast<const uint8*>(&BlendPoses),
 				MaxPoseWeights.GetData(),
-				(ispc::FPerBoneBlendWeight*)BoneBlendWeights.GetData(),
-				(int32*)BoneContainer.GetCompactPoseParentBoneArray().GetData(),
+				reinterpret_cast<const ispc::FPerBoneBlendWeight*>(BoneBlendWeights.GetData()),
+				reinterpret_cast<const int32*>(BoneContainer.GetCompactPoseParentBoneArray().GetData()),
 				BasePose.GetNumBones());
 		}
 		else
diff --git a/Engine/Source/Runtime/Engine/Private/Animation/AnimationRuntime.ispc b/Engine/Source/Runtime/Engine/Private/Animation/AnimationRuntime.ispc
--- a/Engine/Source/Runtime/Engine/Private/Animation/AnimationRuntime.ispc
+++ b/Engine/Source/Runtime/Engine/Private/Animation/AnimationRuntime.ispc
@@ -32,9 +32,9 @@ export void BlendTransformOverwrite(const uniform FTransform SourcePose[],
 }
 
 export void BlendTransformAccumulate(const uniform FTransform SourcePose[],
-									uniform FTransform ResultPose[],
-									const uniform float BlendWeight,
-									const uniform int NumBones)
+									 uniform FTransform ResultPose[],
+									 const uniform float BlendWeight,
+									 const uniform int NumBones)
 {
 	for(uniform int BoneIndex = 0; BoneIndex < NumBones; BoneIndex++)
 	{
@@ -52,8 +52,8 @@ export void BlendTransformAccumulate(const uniform FTransform SourcePose[],
 }
 
 export void ConvertPoseToAdditive(uniform FTransform TargetPose[],
-								const uniform FTransform BasePose[],
-								const uniform int NumBones)
+								  const uniform FTransform BasePose[],
+								  const uniform int NumBones)
 {
 	for(uniform int BoneIndex = 0; BoneIndex < NumBones; ++BoneIndex)
 	{
@@ -73,8 +73,8 @@ export void ConvertPoseToAdditive(uniform FTransform TargetPose[],
 
 
 export void ConvertPoseToMeshRotation(uniform FTransform LocalPoses[],
-									const uniform int ParentBones[],
-									const uniform int NumBones)
+									  const uniform int ParentBones[],
+									  const uniform int NumBones)
 {
 	for(uniform int BoneIndex = 1; BoneIndex < NumBones; ++BoneIndex)
 	{
@@ -172,18 +172,18 @@ extern "C" unmasked const uniform unsigned int8 *uniform GetTransformFromArray(c
 
 // blend poses with both mesh space rotation and scaling (we assume uniform scale)
 export void BlendPosesPerBoneFilterScaleRotation(uniform FTransform OutPose[],
-												const uniform FTransform BasePose[],
-												const uniform unsigned int8 *uniform BlendPoseBase,
-												uniform FVector4 SourceRotations[],
-												uniform FVector SourceScales[],
-												uniform FVector4 TargetRotations[],
-												uniform FVector TargetScales[],
-												uniform FVector4 BlendRotations[],
-												uniform FVector BlendScales[],
-												uniform float MaxPoseWeights[],
-												const uniform FPerBoneBlendWeight BoneBlendWeights[],
-												const uniform int ParentBoneArray[],
-												const uniform int NumBones)
+												 const uniform FTransform BasePose[],
+												 const uniform unsigned int8 *uniform BlendPoseBase,
+												 uniform FVector4 SourceRotations[],
+												 uniform FVector SourceScales[],
+												 uniform FVector4 TargetRotations[],
+												 uniform FVector TargetScales[],
+												 uniform FVector4 BlendRotations[],
+												 uniform FVector BlendScales[],
+												 uniform float MaxPoseWeights[],
+												 const uniform FPerBoneBlendWeight BoneBlendWeights[],
+												 const uniform int ParentBoneArray[],
+												 const uniform int NumBones)
 {
 	for(uniform int BoneIndex = 0; BoneIndex < NumBones; BoneIndex++)
 	{
@@ -266,15 +266,15 @@ export void BlendPosesPerBoneFilterScaleRotation(uniform FTransform OutPose[],
 
 // blend poses with both mesh space rotation (we assume uniform scale)
 export void BlendPosesPerBoneFilterRotation(uniform FTransform OutPose[],
-										const uniform FTransform BasePose[],
-										const uniform unsigned int8 *uniform BlendPoseBase,
-										uniform FVector4 SourceRotations[],
-										uniform FVector4 TargetRotations[],
-										uniform FVector4 BlendRotations[],
-										uniform float MaxPoseWeights[],
-										const uniform FPerBoneBlendWeight BoneBlendWeights[],
-										const uniform int ParentBoneArray[],
-										const uniform int NumBones)
+											const uniform FTransform BasePose[],
+											const uniform unsigned int8 *uniform BlendPoseBase,
+											uniform FVector4 SourceRotations[],
+											uniform FVector4 TargetRotations[],
+											uniform FVector4 BlendRotations[],
+											uniform float MaxPoseWeights[],
+											const uniform FPerBoneBlendWeight BoneBlendWeights[],
+											const uniform int ParentBoneArray[],
+											const uniform int NumBones)
 {
 	for(uniform int BoneIndex = 0; BoneIndex < NumBones; BoneIndex++)
 	{
@@ -339,15 +339,15 @@ export void BlendPosesPerBoneFilterRotation(uniform FTransform OutPose[],
 
 // blend poses with mesh space scaling (we assume uniform scale)
 export void BlendPosesPerBoneFilterScale(uniform FTransform OutPose[],
-										const uniform FTransform BasePose[],
-										const uniform unsigned int8 *uniform BlendPoseBase,
-										uniform FVector SourceScales[],
-										uniform FVector TargetScales[],
-										uniform FVector BlendScales[],
-										uniform float MaxPoseWeights[],
-										const uniform FPerBoneBlendWeight BoneBlendWeights[],
-										const uniform int ParentBoneArray[],
-										const uniform int NumBones)
+										 const uniform FTransform BasePose[],
+										 const uniform unsigned int8 *uniform BlendPoseBase,
+										 uniform FVector SourceScales[],
+										 uniform FVector TargetScales[],
+										 uniform FVector BlendScales[],
+										 uniform float MaxPoseWeights[],
+										 const uniform FPerBoneBlendWeight BoneBlendWeights[],
+										 const uniform int ParentBoneArray[],
+										 const uniform int NumBones)
 {
 	for(uniform int BoneIndex = 0; BoneIndex < NumBones; BoneIndex++)
 	{
@@ -504,11 +504,11 @@ export void BlendCurves(const uniform float SourceCurveWeights[],
 	}
 }
 
-export void LerpBoneTransforms(	uniform FTransform ATransformData[],
-								uniform FTransform BTransformData[],
-								const uniform float Alpha,
-								const uniform unsigned int16 RequiredBonesArray[],
-								const uniform int32 NumOfBones)
+export void LerpBoneTransforms(uniform FTransform ATransformData[],
+							   const uniform FTransform BTransformData[],
+							   const uniform float Alpha,
+							   const uniform unsigned int16 RequiredBonesArray[],
+							   const uniform int32 NumOfBones)
 {
 	// Note: Check FAnimWeight::IsFullWeight(Alpha) as well as !FAnimWeight::IsRelevant(Alpha) before this function runs
 
