Marked Cascade properties as transient that previously weren't as they should never be serialized

diff --git a/Engine/Source/Runtime/Engine/Classes/Particles/ParticleSystemComponent.h b/Engine/Source/Runtime/Engine/Classes/Particles/ParticleSystemComponent.h
--- a/Engine/Source/Runtime/Engine/Classes/Particles/ParticleSystemComponent.h
+++ b/Engine/Source/Runtime/Engine/Classes/Particles/ParticleSystemComponent.h
@@ -685,19 +685,19 @@ public:
 	UPROPERTY(BlueprintAssignable)
 	FParticleCollisionSignature OnParticleCollide;
 
-	UPROPERTY()
+	UPROPERTY(Transient, DuplicateTransient)
 	bool bOldPositionValid;
 
-	UPROPERTY()
+	UPROPERTY(Transient, DuplicateTransient)
 	FVector OldPosition;
 
-	UPROPERTY()
+	UPROPERTY(Transient, DuplicateTransient)
 	FVector PartSysVelocity;
 
-	UPROPERTY()
+	UPROPERTY(Transient, DuplicateTransient)
 	float WarmupTime;
 
-	UPROPERTY()
+	UPROPERTY(Transient, DuplicateTransient)
 	float WarmupTickRate;
 
 private:
