Use transactionally safe locks in VisualLogger

diff --git a/Engine/Source/Runtime/Engine/Private/VisualLogger/VisualLogger.cpp b/Engine/Source/Runtime/Engine/Private/VisualLogger/VisualLogger.cpp
--- a/Engine/Source/Runtime/Engine/Private/VisualLogger/VisualLogger.cpp
+++ b/Engine/Source/Runtime/Engine/Private/VisualLogger/VisualLogger.cpp
@@ -389,7 +389,7 @@ FVisualLogEntry* FVisualLogger::GetLastEntryForObject(const UObject* Object)
 {
 	const UObject* LogOwner = nullptr;
 	{
-		FReadScopeLock Lock(RedirectRWLock);
+		FTransactionallySafeReadScopeLock Lock(RedirectRWLock);
 		LogOwner = FindRedirectionInternal(Object);
 	}
 	if (LogOwner == nullptr)
@@ -405,7 +405,7 @@ FVisualLogEntry* FVisualLogger::GetEntryToWrite(const UObject* Object, const dou
 {
 	const UObject* LogOwner = nullptr;
 	{
-		FReadScopeLock Lock(RedirectRWLock);
+		FTransactionallySafeReadScopeLock Lock(RedirectRWLock);
 		LogOwner = FindRedirectionInternal(Object);
 	}
 	if (LogOwner == nullptr)
@@ -548,7 +548,7 @@ FVisualLogEntry* FVisualLogger::GetEntryToWriteInternal(const UObject* Object, c
 			CurrentEntry->bIsLocationValid = true;
 		}
 
-		FReadScopeLock RedirectScopeLock(RedirectRWLock);
+		FTransactionallySafeReadScopeLock RedirectScopeLock(RedirectRWLock);
 		FOwnerToChildrenRedirectionMap& RedirectionMap = GetRedirectionMap(LogOwner);
 		if (const IVisualLoggerDebugSnapshotInterface* DebugSnapshotInterface = Cast<const IVisualLoggerDebugSnapshotInterface>(LogOwner))
 		{
@@ -1048,7 +1048,7 @@ UObject* FVisualLogger::FindRedirectionInternal(const UObject* Object) const
 
 void FVisualLogger::CleanupRedirects()
 {
-	FWriteScopeLock Lock(RedirectRWLock);
+	FTransactionallySafeWriteScopeLock Lock(RedirectRWLock);
 	for (auto It = ChildToOwnerMap.CreateIterator(); It; ++It)
 	{
 		if(!It.Value().IsValid())
diff --git a/Engine/Source/Runtime/Engine/Public/VisualLogger/VisualLogger.h b/Engine/Source/Runtime/Engine/Public/VisualLogger/VisualLogger.h
--- a/Engine/Source/Runtime/Engine/Public/VisualLogger/VisualLogger.h
+++ b/Engine/Source/Runtime/Engine/Public/VisualLogger/VisualLogger.h
@@ -12,6 +12,8 @@
 #include "Traits/IsCharEncodingCompatibleWith.h"
 #include "UObject/ObjectKey.h"
 #include "Containers/Ticker.h"
+#include "Misc/TransactionallySafeRWLock.h"
+#include "Misc/TransactionallySafeRWScopeLock.h"
 
 #if ENABLE_VISUAL_LOG
 
@@ -675,7 +677,7 @@ public:
 		FVisualLogger& Logger = FVisualLogger::Get();
 		UObject* NewRedirection = nullptr;
 		{
-			FWriteScopeLock Lock(Logger.RedirectRWLock);
+			FTransactionallySafeWriteScopeLock Lock(Logger.RedirectRWLock);
 			NewRedirection = Logger.RedirectInternal(FromObject, ToObject);
 		}
 		UE_CVLOG(FromObject != nullptr && NewRedirection != nullptr, FromObject, LogVisual, Log, TEXT("Redirected '%s' to '%s'"), *FromObject->GetName(), *NewRedirection->GetName());
@@ -685,7 +687,7 @@ public:
 	static UObject* FindRedirection(const UObject* Object)
 	{ 
 		FVisualLogger& Logger = FVisualLogger::Get();
-		FReadScopeLock Lock(Logger.RedirectRWLock);
+		FTransactionallySafeReadScopeLock Lock(Logger.RedirectRWLock);
 		return Logger.FindRedirectionInternal(Object);
 	}
 
@@ -873,7 +875,7 @@ protected:
 	// redirect the traffic to
 	FChildToOwnerRedirectionMap ChildToOwnerMap;
 	// Read Write lock protecting redirection maps (ChildToOwnerMap and ObjectToWorldMap)
-	mutable FRWLock RedirectRWLock;
+	mutable FTransactionallySafeRWLock RedirectRWLock;
 	// if set all categories are blocked from logging
 	bool bBlockedAllCategories : 1;
 	// if set we are recording to file
