Fix up incorrect cast in animation runtime.

This fixes a bug just introduced at https://p4-swarm.epicgames.net/changes/35967063 which leads to a crash.

diff --git a/Engine/Source/Runtime/Engine/Private/Animation/AnimationRuntime.cpp b/Engine/Source/Runtime/Engine/Private/Animation/AnimationRuntime.cpp
--- a/Engine/Source/Runtime/Engine/Private/Animation/AnimationRuntime.cpp
+++ b/Engine/Source/Runtime/Engine/Private/Animation/AnimationRuntime.cpp
@@ -1742,7 +1742,7 @@ struct FBlendPosesPerBoneFilterScratchArea : public TThreadSingleton<FBlendPoses
 // Helper function to get FTransform from a PoseIndex and BoneIndex
 extern "C" const uint8* GetTransformFromArray(const uint8 *BlendPoseBase, const int32 PoseIndex, const int32 BoneIndex)
 {
-	const TArray<struct FCompactPose>& BlendPoses = reinterpret_cast<const TArray<struct FCompactPose>&>(BlendPoseBase);
+	const TArray<struct FCompactPose>& BlendPoses = *reinterpret_cast<const TArray<struct FCompactPose>*>(BlendPoseBase);
 	const FTransform* BlendPose = &BlendPoses[PoseIndex][FCompactPoseBoneIndex(BoneIndex)];
 	return reinterpret_cast<const uint8*>(BlendPose);
 }
