Added RDG async execute synchronization for material uniform expression and mateerial parameter collection access.

diff --git a/Engine/Source/Runtime/Engine/Private/Materials/MaterialRenderProxy.cpp b/Engine/Source/Runtime/Engine/Private/Materials/MaterialRenderProxy.cpp
--- a/Engine/Source/Runtime/Engine/Private/Materials/MaterialRenderProxy.cpp
+++ b/Engine/Source/Runtime/Engine/Private/Materials/MaterialRenderProxy.cpp
@@ -11,6 +11,7 @@
 #include "RendererInterface.h"
 #include "RenderingThread.h"
 #include "TextureResource.h"
+#include "RenderGraphBuilder.h"
 #include "VT/RuntimeVirtualTexture.h"
 
 #if WITH_EDITOR
@@ -550,6 +551,9 @@ void FMaterialRenderProxy::InvalidateUniformExpressionCache(bool bRecreateUnifor
 {
 	GUniformExpressionCacheAsyncUpdateTask.Wait();
 
+	// Async RDG tasks can call FMaterialShader::SetParameters which touch the material uniform expression cache.
+	FRDGBuilder::WaitForAsyncExecuteTask();
+
 #if WITH_EDITOR
 	FStaticLightingSystemInterface::OnMaterialInvalidated.Broadcast(this);
 #endif
@@ -587,6 +591,8 @@ const FMaterial* FMaterialRenderProxy::UpdateUniformExpressionCacheIfNeeded(FRHI
 	UE::TScopeLock Lock(Mutex);
 	if (Material && Material->GetRenderingThreadShaderMap() != UniformExpressionCache[InFeatureLevel].CachedUniformExpressionShaderMap)
 	{
+		// Async RDG tasks can call FMaterialShader::SetParameters which touch the material uniform expression cache.
+		FRDGBuilder::WaitForAsyncExecuteTask();
 		FMaterialRenderContext MaterialRenderContext(this, *Material, nullptr);
 		MaterialRenderContext.bShowSelection = GIsEditor;
 		EvaluateUniformExpressions(RHICmdList, UniformExpressionCache[InFeatureLevel], MaterialRenderContext, nullptr);
diff --git a/Engine/Source/Runtime/Engine/Private/Materials/MaterialShared.cpp b/Engine/Source/Runtime/Engine/Private/Materials/MaterialShared.cpp
--- a/Engine/Source/Runtime/Engine/Private/Materials/MaterialShared.cpp
+++ b/Engine/Source/Runtime/Engine/Private/Materials/MaterialShared.cpp
@@ -5341,6 +5341,9 @@ void SetShaderMapsOnMaterialResources_RenderThread(FRHICommandListImmediate& RHI
 #if WITH_EDITOR
 	bool bUpdateFeatureLevel[ERHIFeatureLevel::Num] = { false };
 
+	// Async RDG tasks can call FMaterialShader::SetParameters which touch the material uniform expression cache.
+	FRDGBuilder::WaitForAsyncExecuteTask();
+
 	for (auto& It : MaterialsToUpdate)
 	{
 		FMaterial* Material = It.Key;
@@ -5370,12 +5373,6 @@ void SetShaderMapsOnMaterialResources_RenderThread(FRHICommandListImmediate& RHI
 				{
 					MaterialProxy->CacheUniformExpressions(RHICmdList, true);
 					bFoundAnyInitializedMaterials = true;
-
-					/*const FMaterial& MaterialForRendering = *MaterialProxy->GetMaterial(MaterialFeatureLevel);
-					check(MaterialForRendering.GetRenderingThreadShaderMap());
-					check(!MaterialProxy->UniformExpressionCache[MaterialFeatureLevel].bUpToDate
-						|| MaterialProxy->UniformExpressionCache[MaterialFeatureLevel].CachedUniformExpressionShaderMap == MaterialForRendering.GetRenderingThreadShaderMap());
-					check(MaterialForRendering.GetRenderingThreadShaderMap()->IsValidForRendering());*/
 				}
 			}
 		}
diff --git a/Engine/Source/Runtime/Engine/Private/Materials/ParameterCollection.cpp b/Engine/Source/Runtime/Engine/Private/Materials/ParameterCollection.cpp
--- a/Engine/Source/Runtime/Engine/Private/Materials/ParameterCollection.cpp
+++ b/Engine/Source/Runtime/Engine/Private/Materials/ParameterCollection.cpp
@@ -11,6 +11,7 @@
 #include "Materials/MaterialParameterCollectionInstance.h"
 #include "Materials/MaterialParameterCollection.h"
 #include "Misc/App.h"
+#include "RenderGraphBuilder.h"
 
 int32 GDeferUpdateRenderStates = 1;
 FAutoConsoleVariableRef CVarDeferUpdateRenderStates(
@@ -78,7 +79,9 @@ void UMaterialParameterCollection::BeginDestroy()
 		FThreadSafeBool* Released = &ReleasedByRT;
 		ENQUEUE_RENDER_COMMAND(RemoveDefaultResourceCommand)(
 			[Resource, Id, Released](FRHICommandListImmediate& RHICmdList)
-			{	
+			{
+				// Async RDG tasks can call FMaterialShader::SetParameters which touch material parameter collections.
+				FRDGBuilder::WaitForAsyncExecuteTask();
 				GDefaultMaterialParameterCollectionInstances.RemoveSingle(Id, Resource);
 				*Released = true;
 			}
@@ -611,7 +614,9 @@ void UMaterialParameterCollection::UpdateDefaultResource(bool bRecreateUniformBu
 	FMaterialParameterCollectionInstanceResource* Resource = DefaultResource;
 	ENQUEUE_RENDER_COMMAND(UpdateDefaultResourceCommand)(
 		[Id, Resource](FRHICommandListImmediate& RHICmdList)
-		{	
+		{
+			// Async RDG tasks can call FMaterialShader::SetParameters which touch material parameter collections.
+			FRDGBuilder::WaitForAsyncExecuteTask();
 			GDefaultMaterialParameterCollectionInstances.Add(Id, Resource);
 		}
 	);
@@ -839,6 +844,11 @@ void FMaterialParameterCollectionInstanceResource::GameThread_UpdateContents(con
 	ENQUEUE_RENDER_COMMAND(UpdateCollectionCommand)(
 		[InGuid, Data, InOwnerName, Resource, bRecreateUniformBuffer](FRHICommandListImmediate& RHICmdList)
 		{
+			if (bRecreateUniformBuffer)
+			{
+				// Async RDG tasks can call FMaterialShader::SetParameters which touch material parameter collections.
+				FRDGBuilder::WaitForAsyncExecuteTask();
+			}
 			Resource->UpdateContents(InGuid, Data, InOwnerName, bRecreateUniformBuffer);
 		}
 	);
@@ -850,6 +860,8 @@ void FMaterialParameterCollectionInstanceResource::GameThread_Destroy()
 	ENQUEUE_RENDER_COMMAND(DestroyCollectionCommand)(
 		[Resource](FRHICommandListImmediate& RHICmdList)
 		{
+			// Async RDG tasks can call FMaterialShader::SetParameters which touch material parameter collections.
+			FRDGBuilder::WaitForAsyncExecuteTask();
 			Resource->UniformBuffer.SafeRelease();
 
 			// FRHIUniformBuffer instances take raw pointers to the layout struct.
diff --git a/Engine/Source/Runtime/Renderer/Private/RendererScene.cpp b/Engine/Source/Runtime/Renderer/Private/RendererScene.cpp
--- a/Engine/Source/Runtime/Renderer/Private/RendererScene.cpp
+++ b/Engine/Source/Runtime/Renderer/Private/RendererScene.cpp
@@ -1216,6 +1216,9 @@ void FScene::UpdateParameterCollections(const TArray<FMaterialParameterCollectio
 	ENQUEUE_RENDER_COMMAND(UpdateParameterCollectionsCommand)(
 		[this, InParameterCollections] (FRHICommandListBase&)
 	{
+		// Async RDG tasks can call FMaterialShader::SetParameters which touch material parameter collections.
+		FRDGBuilder::WaitForAsyncExecuteTask();
+
 		// Empty the scene's map so any unused uniform buffers will be released
 		ParameterCollections.Empty();
 
@@ -5398,7 +5401,9 @@ void FScene::Update(FRDGBuilder& GraphBuilder, const FUpdateParameters& Paramete
 		}
 	}
 
+	// Avoid overlapping prior scene render and async RDG execution tasks to simplify things as there are a lot of moving pieces.
 	FSceneRenderer::WaitForCleanUpTasks(GraphBuilder.RHICmdList);
+	FRDGBuilder::WaitForAsyncExecuteTask();
 
 	UE::Tasks::FTask UpdateUniformExpressionsTask;
 	FMaterialRenderProxy::UpdateDeferredCachedUniformExpressions(GraphBuilder.RHICmdList, EnumHasAnyFlags(Parameters.AsyncOps, EUpdateAllPrimitiveSceneInfosAsyncOps::CacheMaterialUniformExpressions) ? &UpdateUniformExpressionsTask : nullptr);
