Added FExternalPackageHelper::AsyncLoadObjectsFromExternalPackages which is the asynchronous version of LoadObjectsFromExternalPackages.
AWorldDataLayers::PostLoad now uses the asynchronous version to load data layer instance external packages.

diff --git a/Engine/Source/Runtime/Engine/Private/WorldPartition/DataLayer/WorldDataLayers.cpp b/Engine/Source/Runtime/Engine/Private/WorldPartition/DataLayer/WorldDataLayers.cpp
--- a/Engine/Source/Runtime/Engine/Private/WorldPartition/DataLayer/WorldDataLayers.cpp
+++ b/Engine/Source/Runtime/Engine/Private/WorldPartition/DataLayer/WorldDataLayers.cpp
@@ -1158,6 +1158,7 @@ TSet<TObjectPtr<UDataLayerInstance>>& AWorldDataLayers::GetDataLayerInstances()
 	if (IsUsingExternalPackageDataLayerInstances())
 	{
 		check(DataLayerInstances.IsEmpty());
+		check(AsyncLoadRequestExternalPackageDataLayerInstances.IsEmpty());
 		check(LoadedExternalPackageDataLayerInstances.IsEmpty());
 		return ExternalPackageDataLayerInstances;
 	}
@@ -1173,7 +1174,12 @@ TSet<TObjectPtr<UDataLayerInstance>>& AWorldDataLayers::GetDataLayerInstances()
 void AWorldDataLayers::InitializeExternalPackageDataLayerInstances()
 {
 	check(IsUsingExternalPackageDataLayerInstances());
+	if (AsyncLoadRequestExternalPackageDataLayerInstances.Num() != LoadedExternalPackageDataLayerInstances.Num())
+	{
+		FlushAsyncLoading(AsyncLoadRequestExternalPackageDataLayerInstances);
+	}
 	ExternalPackageDataLayerInstances.Append(LoadedExternalPackageDataLayerInstances);
+	AsyncLoadRequestExternalPackageDataLayerInstances.Reset();
 	LoadedExternalPackageDataLayerInstances.Reset();
 }
 #endif
@@ -1289,7 +1295,8 @@ void AWorldDataLayers::PostLoad()
 		if (!Level->bWasDuplicated && IsUsingExternalPackageDataLayerInstances())
 		{
 			// Load all folders for this level
-			FExternalPackageHelper::LoadObjectsFromExternalPackages<UDataLayerInstance>(this, [this](UDataLayerInstance* LoadedDataLayerInstance)
+			check(AsyncLoadRequestExternalPackageDataLayerInstances.IsEmpty());
+			AsyncLoadRequestExternalPackageDataLayerInstances = FExternalPackageHelper::AsyncLoadObjectsFromExternalPackages<UDataLayerInstance>(this, [this](UDataLayerInstance* LoadedDataLayerInstance)
 			{
 				check(IsValid(LoadedDataLayerInstance));
 				LoadedExternalPackageDataLayerInstances.Add(LoadedDataLayerInstance);
diff --git a/Engine/Source/Runtime/Engine/Public/ExternalPackageHelper.h b/Engine/Source/Runtime/Engine/Public/ExternalPackageHelper.h
--- a/Engine/Source/Runtime/Engine/Public/ExternalPackageHelper.h
+++ b/Engine/Source/Runtime/Engine/Public/ExternalPackageHelper.h
@@ -93,6 +93,9 @@ public:
 	template<typename T>
 	static void LoadObjectsFromExternalPackages(UObject* InOuter, TFunctionRef<void(T*)> Operation);
 
+	template<typename T>
+	static TArray<int> AsyncLoadObjectsFromExternalPackages(UObject* InOuter, TFunction<void(T*)> Operation);
+
 	enum class EGetExternalSaveableObjectsFlags : uint32
 	{
 		// No flags
@@ -135,13 +138,33 @@ public:
 private:
 	/** Get the external object package instance name. */
 	static ENGINE_API FString GetExternalObjectPackageInstanceName(const FString& OuterPackageName, const FString& ObjectPackageName);
+
+	struct FPackageInfo
+	{
+		FPackageInfo(const FString& InObjectPackageName)
+			: ObjectPackageName(InObjectPackageName)
+			, InstancePackageName(NAME_None)
+			, InstancePackage(nullptr)
+		{}
+		FString ObjectPackageName;
+		FName InstancePackageName;
+		UPackage* InstancePackage;
+	};
+
+	struct FLoadingRequests
+	{
+		bool bIsInstanced;
+		FLinkerInstancingContext InstancingContext;
+		TArray<FPackageInfo> PackageInfos;
+	};
+
+	template<typename T>
+	static void PrepareLoadingRequestsForExternalPackages(UObject* InOuter, bool bPrepareForAsync, FLoadingRequests& OutLoadingRequests);
 };
 
 template<typename T>
-void FExternalPackageHelper::LoadObjectsFromExternalPackages(UObject* InOuter, TFunctionRef<void(T*)> Operation)
+void FExternalPackageHelper::PrepareLoadingRequestsForExternalPackages(UObject* InOuter, bool bPrepareForAsync, FLoadingRequests& OutLoadingRequests)
 {
-	TRACE_CPUPROFILER_EVENT_SCOPE(FExternalPackageHelper::LoadObjectsFromExternalPackages);
-
 	check(InOuter);
 	UPackage* OutermostPackage = InOuter->IsA<UPackage>() ? CastChecked<UPackage>(InOuter) : InOuter->GetOutermostObject()->GetPackage();
 	const FString OutermostPackageName = !OutermostPackage->GetLoadedPath().IsEmpty() ? OutermostPackage->GetLoadedPath().GetPackageName() : OutermostPackage->GetName();
@@ -149,7 +172,6 @@ void FExternalPackageHelper::LoadObjectsFromExternalPackages(UObject* InOuter, T
 	const UExternalDataLayerAsset* ExternalDataLayerAsset = DataLayerInstanceProvider ? DataLayerInstanceProvider->GetRootExternalDataLayerAsset() : nullptr;
 	const FString RootPath = ExternalDataLayerAsset ? FExternalDataLayerHelper::GetExternalDataLayerLevelRootPath(ExternalDataLayerAsset, OutermostPackageName) : OutermostPackageName;
 	const FString ExternalObjectsPath = FExternalPackageHelper::GetExternalObjectsPath(RootPath);
-	TArray<FString> ObjectPackageNames;
 
 	// Do a synchronous scan of the world external objects path.			
 	IAssetRegistry& AssetRegistry = FModuleManager::LoadModuleChecked<FAssetRegistryModule>(TEXT("AssetRegistry")).Get();
@@ -159,18 +181,18 @@ void FExternalPackageHelper::LoadObjectsFromExternalPackages(UObject* InOuter, T
 	TArray<FTopLevelAssetPath> ClassPaths;
 	ClassPaths.Add(T::StaticClass()->GetClassPathName());
 
-	if(T::StaticClass() != UObject::StaticClass())
+	if (T::StaticClass() != UObject::StaticClass())
 	{
 		TArray<UClass*> DerivedClasses;
 		DerivedClasses.Add(T::StaticClass());
 		GetDerivedClasses(T::StaticClass(), DerivedClasses);
 
-		for(UClass* Class : DerivedClasses)
+		for (UClass* Class : DerivedClasses)
 		{
 			TArray<FCoreRedirectObjectName> PreviousRedirectedNames;
 			FCoreRedirects::FindPreviousNames(ECoreRedirectFlags::Type_Class, FCoreRedirectObjectName(Class->GetClassPathName()), PreviousRedirectedNames);
 
-			for(const FCoreRedirectObjectName& PreviousRedirectedName : PreviousRedirectedNames)
+			for (const FCoreRedirectObjectName& PreviousRedirectedName : PreviousRedirectedNames)
 			{
 				ClassPaths.Add(FTopLevelAssetPath(PreviousRedirectedName.PackageName, PreviousRedirectedName.ObjectName));
 			}
@@ -186,19 +208,16 @@ void FExternalPackageHelper::LoadObjectsFromExternalPackages(UObject* InOuter, T
 	TArray<FAssetData> Assets;
 	GetSortedAssets(Filter, Assets);
 
-	ObjectPackageNames.Reserve(Assets.Num());
+	OutLoadingRequests.PackageInfos.Reserve(Assets.Num());
 	for (const FAssetData& Asset : Assets)
 	{
-		ObjectPackageNames.Add(Asset.PackageName.ToString());
+		OutLoadingRequests.PackageInfos.Emplace(Asset.PackageName.ToString());
 	}
 
-	FLinkerInstancingContext InstancingContext;
-	TArray<UPackage*> InstancePackages;
-
-	UPackage* OuterPackage = InOuter->GetPackage();
+	const UPackage* OuterPackage = InOuter->GetPackage();
 	FName PackageResourceName = OuterPackage->GetLoadedPath().GetPackageFName();
-	const bool bInstanced = !PackageResourceName.IsNone() && (PackageResourceName != OuterPackage->GetFName());
-	if (bInstanced)
+	OutLoadingRequests.bIsInstanced = !PackageResourceName.IsNone() && (PackageResourceName != OuterPackage->GetFName());
+	if (OutLoadingRequests.bIsInstanced)
 	{
 		const FLinkerInstancingContext* OuterInstancingContext = nullptr;
 		if (FLinkerLoad* OuterLinker = InOuter->GetLinker())
@@ -216,16 +235,16 @@ void FExternalPackageHelper::LoadObjectsFromExternalPackages(UObject* InOuter, T
 			OuterPackages.Add(Package, &bIsAlreadyInSet);
 			if (!bIsAlreadyInSet)
 			{
-				InstancingContext.AddPackageMapping(Package->GetLoadedPath().GetPackageFName(), Package->GetFName());
+				OutLoadingRequests.InstancingContext.AddPackageMapping(Package->GetLoadedPath().GetPackageFName(), Package->GetFName());
 			}
 			ItObj = ItObj->GetOuter();
 		}
 
-		for (const FString& ObjectPackageName : ObjectPackageNames)
+		for (FPackageInfo& OutPackageInfo : OutLoadingRequests.PackageInfos)
 		{
 			FName InstancedName;
-			
-			FName ObjectPackageFName = *ObjectPackageName;
+
+			FName ObjectPackageFName = *OutPackageInfo.ObjectPackageName;
 			if (OuterInstancingContext)
 			{
 				InstancedName = OuterInstancingContext->RemapPackage(ObjectPackageFName);
@@ -234,26 +253,38 @@ void FExternalPackageHelper::LoadObjectsFromExternalPackages(UObject* InOuter, T
 			// Remap to the a instanced package if it wasn't remapped already by the outer instancing context
 			if (InstancedName == ObjectPackageFName || InstancedName.IsNone())
 			{
-				InstancedName = *GetExternalObjectPackageInstanceName(OuterPackage->GetName(), ObjectPackageName);
+				InstancedName = *GetExternalObjectPackageInstanceName(OuterPackage->GetName(), OutPackageInfo.ObjectPackageName);
 			}
 
-			InstancingContext.AddPackageMapping(ObjectPackageFName, InstancedName);
+			OutLoadingRequests.InstancingContext.AddPackageMapping(ObjectPackageFName, InstancedName);
 
-			// Create instance package
-			UPackage* InstancePackage = CreatePackage(*InstancedName.ToString());
-			// Propagate RF_Transient
-			if (OuterPackage->HasAnyFlags(RF_Transient))
+			OutPackageInfo.InstancePackageName = InstancedName;
+			if (!bPrepareForAsync)
 			{
-				InstancePackage->SetFlags(RF_Transient);
+				// Create instance package
+				OutPackageInfo.InstancePackage = CreatePackage(*InstancedName.ToString());
+				// Propagate RF_Transient
+				if (OuterPackage->HasAnyFlags(RF_Transient))
+				{
+					OutPackageInfo.InstancePackage->SetFlags(RF_Transient);
+				}
 			}
-			InstancePackages.Add(InstancePackage);
 		}
 	}
+}
+
+template<typename T>
+void FExternalPackageHelper::LoadObjectsFromExternalPackages(UObject* InOuter, TFunctionRef<void(T*)> Operation)
+{
+	TRACE_CPUPROFILER_EVENT_SCOPE(FExternalPackageHelper::LoadObjectsFromExternalPackages);
+
+	FLoadingRequests LoadingRequests;
+	PrepareLoadingRequestsForExternalPackages<T>(InOuter, false, LoadingRequests);
 
 	const ELoadFlags LoadFlags = InOuter->GetPackage()->HasAnyPackageFlags(PKG_PlayInEditor) ? LOAD_PackageForPIE : LOAD_None;
-	for (int32 i = 0; i < ObjectPackageNames.Num(); i++)
+	for (const FPackageInfo& PackageRequest : LoadingRequests.PackageInfos)
 	{
-		if (UPackage* Package = LoadPackage(bInstanced ? InstancePackages[i] : nullptr, *ObjectPackageNames[i], LoadFlags, nullptr, &InstancingContext))
+		if (UPackage* Package = LoadPackage(LoadingRequests.bIsInstanced ? PackageRequest.InstancePackage : nullptr, *PackageRequest.ObjectPackageName, LoadFlags, nullptr, &LoadingRequests.InstancingContext))
 		{
 			T* LoadedObject = nullptr;
 			ForEachObjectWithPackage(Package, [&LoadedObject](UObject* Object)
@@ -274,4 +305,71 @@ void FExternalPackageHelper::LoadObjectsFromExternalPackages(UObject* InOuter, T
 	}
 }
 
-#endif
+
+template<typename T>
+TArray<int32> FExternalPackageHelper::AsyncLoadObjectsFromExternalPackages(UObject* InOuter, TFunction<void(T*)> Operation)
+{
+	TRACE_CPUPROFILER_EVENT_SCOPE(FExternalPackageHelper::AsyncLoadObjectsFromExternalPackages);
+
+	FLoadingRequests LoadingRequests;
+	PrepareLoadingRequestsForExternalPackages<T>(InOuter, true, LoadingRequests);
+
+	TArray<int32> AsyncRequestIDs;
+	AsyncRequestIDs.Reserve(LoadingRequests.PackageInfos.Num());
+	const UPackage* OuterPackage = InOuter->GetPackage();
+	const ELoadFlags LoadFlags = OuterPackage->HasAnyPackageFlags(PKG_PlayInEditor) ? LOAD_PackageForPIE : LOAD_None;
+	const EPackageFlags PackageFlags = OuterPackage->HasAnyPackageFlags(PKG_PlayInEditor) ? PKG_PlayInEditor : PKG_None;
+	const int32 PIEInstanceID = OuterPackage->GetPIEInstanceID();
+	const TAsyncLoadPriority PackagePriority = 0;
+	const bool bPropagateTransientFlag = OuterPackage->HasAnyFlags(RF_Transient);
+
+	for (const FPackageInfo& PackageRequest : LoadingRequests.PackageInfos)
+	{
+		const FPackagePath PackagePath = FPackagePath::FromPackageNameChecked(PackageRequest.ObjectPackageName);
+		FName PackageName = LoadingRequests.bIsInstanced ? PackageRequest.InstancePackageName : *PackageRequest.ObjectPackageName;
+
+		int32 RequestID = LoadPackageAsync(PackagePath, PackageName, FLoadPackageAsyncDelegate::CreateLambda([Operation, bPropagateTransientFlag](const FName& PackageName, UPackage* LoadedPackage, EAsyncLoadingResult::Type Result)
+		{
+			if (LoadedPackage)
+			{
+				// Propagate flags
+				if (bPropagateTransientFlag)
+				{
+					LoadedPackage->SetFlags(RF_Transient);
+				}
+
+				T* LoadedObject = nullptr;
+				ForEachObjectWithPackage(LoadedPackage, [&LoadedObject, bPropagateTransientFlag](UObject* Object)
+				{
+					// Propagate flags
+					if (bPropagateTransientFlag)
+					{
+						Object->SetFlags(RF_Transient);
+					}
+
+					if (!LoadedObject)
+					{
+						if (T* TypedObj = Cast<T>(Object))
+						{
+							LoadedObject = TypedObj;
+							if (!bPropagateTransientFlag)
+							{
+								return false;
+							}
+						}
+					}
+					return true;
+				}, true, RF_NoFlags, EInternalObjectFlags::Unreachable);
+
+				if (ensure(LoadedObject))
+				{
+					Operation(LoadedObject);
+				}
+			}
+		}), PackageFlags, PIEInstanceID, PackagePriority, &LoadingRequests.InstancingContext, LoadFlags);
+		AsyncRequestIDs.Add(RequestID);
+	}
+	return AsyncRequestIDs;
+}
+
+#endif
\ No newline at end of file
diff --git a/Engine/Source/Runtime/Engine/Public/WorldPartition/DataLayer/WorldDataLayers.h b/Engine/Source/Runtime/Engine/Public/WorldPartition/DataLayer/WorldDataLayers.h
--- a/Engine/Source/Runtime/Engine/Public/WorldPartition/DataLayer/WorldDataLayers.h
+++ b/Engine/Source/Runtime/Engine/Public/WorldPartition/DataLayer/WorldDataLayers.h
@@ -311,17 +311,12 @@ private:
 	UPROPERTY()
 	bool bAllowRuntimeDataLayerEditing;
 
+	/** Used by Actor Editor Context */
 	UPROPERTY(Transient)
 	FActorPlacementDataLayers CurrentDataLayers;
 
+	/** Used by Actor Editor Context */
 	TArray<FActorPlacementDataLayers> CurrentDataLayersStack;
-#endif
-
-	UPROPERTY()
-	TObjectPtr<UExternalDataLayerInstance> RootExternalDataLayerInstance;
-
-	UPROPERTY()
-	TSet<TObjectPtr<UDataLayerInstance>> DataLayerInstances;
 
 	/** Data layer instances stored in their external package (only used when UseExternalPackageDataLayerInstances is True) */
 	UPROPERTY(Transient)
@@ -331,6 +326,16 @@ private:
 	UPROPERTY(Transient)
 	TArray<TObjectPtr<UDataLayerInstance>> LoadedExternalPackageDataLayerInstances;
 
+	/** AsyncLoad requests of data layer instance external packages */
+	TArray<int32> AsyncLoadRequestExternalPackageDataLayerInstances;
+#endif
+
+	UPROPERTY()
+	TObjectPtr<UExternalDataLayerInstance> RootExternalDataLayerInstance;
+
+	UPROPERTY()
+	TSet<TObjectPtr<UDataLayerInstance>> DataLayerInstances;
+
 	UPROPERTY(Transient)
 	TArray<TObjectPtr<UDataLayerInstance>> TransientDataLayerInstances;
 
