[UE] Fix exec cmd parsing to allow multiple instances
* Fix FParse::Value to skip trailing quotes so multiple calls into it won't immediately break

diff --git a/Engine/Source/Runtime/Core/Private/Misc/Parse.cpp b/Engine/Source/Runtime/Core/Private/Misc/Parse.cpp
--- a/Engine/Source/Runtime/Core/Private/Misc/Parse.cpp
+++ b/Engine/Source/Runtime/Core/Private/Misc/Parse.cpp
@@ -298,6 +298,11 @@ bool FParse::Value(
 
 	if (OptStreamGotTo)
 	{
+		if (bArgumentsQuoted && *ValueEndInStream == '"')
+		{
+			++ValueEndInStream;
+		}
+
 		*OptStreamGotTo = ValueEndInStream;
 	}
 
diff --git a/Engine/Source/Runtime/Core/Tests/Misc/ParseTests.cpp b/Engine/Source/Runtime/Core/Tests/Misc/ParseTests.cpp
--- a/Engine/Source/Runtime/Core/Tests/Misc/ParseTests.cpp
+++ b/Engine/Source/Runtime/Core/Tests/Misc/ParseTests.cpp
@@ -68,7 +68,7 @@ TEST_CASE("Parse::Value::ToBuffer", "[Parse][Smoke]")
 
 	SECTION("Key may appear mutiple times")
 	{
-		const TCHAR* Line = TEXT("rep=a1 rep=b2 rep=c3");
+		const TCHAR* Line = TEXT("rep=a1 rep=\"b2\" rep=c3");
 		const TCHAR* ExpectedResults[] = { TEXT("a1"), TEXT("b2"), TEXT("c3") };
 
 		const TCHAR* Cursor = Line;
diff --git a/Engine/Source/Runtime/Engine/Private/ParseExecCommands.cpp b/Engine/Source/Runtime/Engine/Private/ParseExecCommands.cpp
--- a/Engine/Source/Runtime/Engine/Private/ParseExecCommands.cpp
+++ b/Engine/Source/Runtime/Engine/Private/ParseExecCommands.cpp
@@ -56,13 +56,16 @@ namespace ParseExecCommands
 
 	TArray<FString> ParseExecCmdsFromCommandLine(const FString& InKey)
 	{
+		TArray<FString> Result;
+
 		FString Line;
-		if (FParse::Value(FCommandLine::Get(), *(InKey + TEXT("=")), Line, /*bShouldStopOnSeparator*/false))
+		const TCHAR* CommandLine = FCommandLine::Get();
+		while (CommandLine && FParse::Value(CommandLine, *(InKey + TEXT("=")), Line, /*bShouldStopOnSeparator*/false, &CommandLine))
 		{
-			return ParseExecCmds(Line);
+			Result.Append(ParseExecCmds(Line));
 		}
 
-		return TArray<FString>();
+		return Result;
 	}
 
 	void QueueDeferredCommands(const TArray<FString>& CommandArray)
