Added blend profile by name to AnimInstance

diff --git a/Engine/Source/Runtime/Engine/Classes/Animation/AnimInstance.h b/Engine/Source/Runtime/Engine/Classes/Animation/AnimInstance.h
--- a/Engine/Source/Runtime/Engine/Classes/Animation/AnimInstance.h
+++ b/Engine/Source/Runtime/Engine/Classes/Animation/AnimInstance.h
@@ -1140,6 +1140,11 @@ public:
 	//////////////////////////////////////////////////////////////////////////
 
 public:
+	
+	/** Returns a blend profile by name from our current skeleton. Null if not found. */
+	UFUNCTION(BlueprintPure, Category="Animation|Skeleton", meta=(BlueprintThreadSafe))
+	ENGINE_API const UBlendProfile* GetBlendProfileByName(FName InBlendProfileName) const;
+	
 	/** Returns the value of a named curve. */
 	UFUNCTION(BlueprintPure, Category="Animation|Curves", meta=(BlueprintThreadSafe))
 	ENGINE_API float GetCurveValue(FName CurveName) const;
diff --git a/Engine/Source/Runtime/Engine/Private/Animation/AnimInstance.cpp b/Engine/Source/Runtime/Engine/Private/Animation/AnimInstance.cpp
--- a/Engine/Source/Runtime/Engine/Private/Animation/AnimInstance.cpp
+++ b/Engine/Source/Runtime/Engine/Private/Animation/AnimInstance.cpp
@@ -4151,6 +4151,16 @@ const FAnimNode_AssetPlayerRelevancyBase* UAnimInstance::GetRelevantAssetPlayerI
 	return GetProxyOnAnyThread<FAnimInstanceProxy>().GetRelevantAssetPlayerInterfaceFromState(MachineIndex, StateIndex);
 }
 
+const UBlendProfile* UAnimInstance::GetBlendProfileByName(FName InBlendProfileName) const
+{
+	if (CurrentSkeleton)
+	{
+		return CurrentSkeleton->GetBlendProfile(InBlendProfileName);
+	}
+
+	return nullptr;
+}
+
 int32 UAnimInstance::GetStateMachineIndex(FName MachineName) const
 {
 	return GetProxyOnAnyThread<FAnimInstanceProxy>().GetStateMachineIndex(MachineName);
