Make `UNetConnection::Describe` transactionally safe.

diff --git a/Engine/Source/Programs/AutoRTFMTests/Private/UECoreTests.cpp b/Engine/Source/Programs/AutoRTFMTests/Private/UECoreTests.cpp
--- a/Engine/Source/Programs/AutoRTFMTests/Private/UECoreTests.cpp
+++ b/Engine/Source/Programs/AutoRTFMTests/Private/UECoreTests.cpp
@@ -1514,6 +1514,41 @@ TEST_CASE("UECore.FString")
 
 		REQUIRE(String == "Foo 'Stuff' BAR");
 	}
+
+	SECTION("Returned From Open")
+	{
+		SECTION("Copied New")
+		{
+			FString String;
+
+			AutoRTFM::Commit([&]
+				{
+					String = AutoRTFM::Open([&]
+						{
+							return TEXT("WOW");
+						});
+				});
+
+			REQUIRE(String == "WOW");
+		}
+
+		SECTION("Copied Old")
+		{
+			FString Other = TEXT("WOW");
+			FString String;
+
+			AutoRTFM::Commit([&]
+				{
+					String = AutoRTFM::Open([&]
+						{
+							return Other;
+						});
+				});
+
+			REQUIRE(Other == "WOW");
+			REQUIRE(String == "WOW");
+		}
+	}
 }
 
 TEST_CASE("UECore.TQueue")
diff --git a/Engine/Source/Runtime/Core/Public/Containers/UnrealString.h.inl b/Engine/Source/Runtime/Core/Public/Containers/UnrealString.h.inl
--- a/Engine/Source/Runtime/Core/Public/Containers/UnrealString.h.inl
+++ b/Engine/Source/Runtime/Core/Public/Containers/UnrealString.h.inl
@@ -2286,6 +2286,12 @@ public:
 		// This must match the GetTypeHash behavior of FStringView
 		return FCrc::Strihash_DEPRECATED(S.Len(), *S);
 	}
+
+	static void AutoRTFMAssignFromOpenToClosed(UE_STRING_CLASS& Closed, const UE_STRING_CLASS& Open)
+	{
+		const AutoRTFM::EContextStatus Status = AutoRTFM::Close([&] { Closed = Open; });
+		ensure(AutoRTFM::EContextStatus::OnTrack == Status);
+	}
 };
 
 template<> struct TIsZeroConstructType<UE_STRING_CLASS> { enum { Value = true }; };
diff --git a/Engine/Source/Runtime/Core/Public/Containers/UnrealStringIncludes.h.inl b/Engine/Source/Runtime/Core/Public/Containers/UnrealStringIncludes.h.inl
--- a/Engine/Source/Runtime/Core/Public/Containers/UnrealStringIncludes.h.inl
+++ b/Engine/Source/Runtime/Core/Public/Containers/UnrealStringIncludes.h.inl
@@ -41,5 +41,6 @@
 #include "Traits/IsCharType.h"
 #include "Traits/IsCharEncodingCompatibleWith.h"
 #include "Traits/IsCharEncodingSimplyConvertibleTo.h"
+#include "AutoRTFM/AutoRTFM.h"
 
 #include <type_traits>
diff --git a/Engine/Source/Runtime/Engine/Private/NetConnection.cpp b/Engine/Source/Runtime/Engine/Private/NetConnection.cpp
--- a/Engine/Source/Runtime/Engine/Private/NetConnection.cpp
+++ b/Engine/Source/Runtime/Engine/Private/NetConnection.cpp
@@ -1316,8 +1316,16 @@ bool UNetConnection::HandleReceiveNetUpgrade(uint32 RemoteNetworkVersion, EEngin
 
 FString UNetConnection::Describe()
 {
+	// `LowLevelGetRemoteAddress` just returns a human readable description of the remote address
+	// and will, on some implementations, call into system stuff to do that (think getting the IP
+	// address or stuff like that). So we need to call this in the open as a result.
+	FString RemoteAddress = AutoRTFM::Open([&]
+		{
+			return LowLevelGetRemoteAddress(true);
+		});
+
 	return FString::Printf( TEXT( "[UNetConnection] RemoteAddr: %s, Name: %s, Driver: %s, IsServer: %s, PC: %s, Owner: %s, UniqueId: %s" ),
-			*LowLevelGetRemoteAddress( true ),
+			*RemoteAddress,
 			*GetName(),
 			Driver ? *Driver->GetDescription() : TEXT( "NULL" ),
 			Driver && Driver->IsServer() ? TEXT( "YES" ) : TEXT( "NO" ),
diff --git a/Engine/Source/Runtime/Engine/Private/Tests/AutoRTFM/AutoRTFMNetDriverTest.cpp b/Engine/Source/Runtime/Engine/Private/Tests/AutoRTFM/AutoRTFMNetDriverTest.cpp
new file mode 100644
--- /dev/null
+++ b/Engine/Source/Runtime/Engine/Private/Tests/AutoRTFM/AutoRTFMNetDriverTest.cpp
@@ -0,0 +1,67 @@
+// Copyright Epic Games, Inc. All Rights Reserved.
+
+#include "Misc/AutomationTest.h"
+#include "Misc/App.h"
+#include "AutoRTFM/AutoRTFM.h"
+#include "AutoRTFMTestActor.h"
+#include "Engine/ActorChannel.h"
+#include "Engine/DemoNetDriver.h"
+
+#if WITH_DEV_AUTOMATION_TESTS
+
+#define TEST_CHECK_TRUE(b) do \
+{ \
+	if (!(b)) \
+	{ \
+		FString String(FString::Printf(TEXT("FAILED: %s:%u"), TEXT(__FILE__), __LINE__)); \
+		ExecutionInfo.AddEvent(FAutomationEvent(EAutomationEventType::Info, String)); \
+		return false; \
+	} \
+} while(false)
+
+IMPLEMENT_SIMPLE_AUTOMATION_TEST(FAutoRTFMNetDriverTests, "AutoRTFM + FTraceFilter", EAutomationTestFlags::EngineFilter | EAutomationTestFlags::ClientContext | EAutomationTestFlags::ServerContext | EAutomationTestFlags::CommandletContext)
+
+bool FAutoRTFMNetDriverTests::RunTest(const FString & Parameters)
+{
+	if (!AutoRTFM::ForTheRuntime::IsAutoRTFMRuntimeEnabled())
+	{
+		ExecutionInfo.AddEvent(FAutomationEvent(EAutomationEventType::Info, TEXT("SKIPPED 'FAutoRTFMNetDriverTests' test. AutoRTFM disabled.")));
+		return true;
+	}
+
+	UNetDriver* const Driver = NewObject<UDemoNetDriver>();
+
+	UNetConnection* const Connection = NewObject<UDemoNetConnection>();
+	Connection->Driver = Driver;
+	Driver->AddClientConnection(Connection);
+
+	UActorChannel* const ActorChannel = NewObject<UActorChannel>();
+	ActorChannel->OpenedLocally = true;
+	ActorChannel->Connection = Connection;
+	Connection->Channels.Add(ActorChannel);
+
+	FString Description;
+
+	AutoRTFM::ETransactionResult Result = AutoRTFM::Transact([&]
+		{
+			Description = ActorChannel->Describe();
+			AutoRTFM::AbortTransaction();
+		});
+
+	TEST_CHECK_TRUE(AutoRTFM::ETransactionResult::AbortedByRequest == Result);
+	TEST_CHECK_TRUE(Description.IsEmpty());
+
+	Result = AutoRTFM::Transact([&]
+		{
+			Description = ActorChannel->Describe();
+		});
+
+	TEST_CHECK_TRUE(AutoRTFM::ETransactionResult::Committed == Result);
+	TEST_CHECK_TRUE(!Description.IsEmpty());
+
+	return true;
+}
+
+#undef TEST_CHECK_TRUE
+
+#endif //WITH_DEV_AUTOMATION_TESTS
