[World Partition] Register the GC post reachability analysis callback once for the main world.

diff --git a/Engine/Source/Runtime/Engine/Private/WorldPartition/WorldPartition.cpp b/Engine/Source/Runtime/Engine/Private/WorldPartition/WorldPartition.cpp
--- a/Engine/Source/Runtime/Engine/Private/WorldPartition/WorldPartition.cpp
+++ b/Engine/Source/Runtime/Engine/Private/WorldPartition/WorldPartition.cpp
@@ -961,7 +961,10 @@ void UWorldPartition::RegisterDelegates()
 
 		if (!IsRunningCommandlet())
 		{
-			FCoreUObjectDelegates::PostReachabilityAnalysis.AddUObject(this, &UWorldPartition::OnGCPostReachabilityAnalysis);
+			if (IsMainWorldPartition())
+			{
+				FCoreUObjectDelegates::PostReachabilityAnalysis.AddUObject(this, &UWorldPartition::OnGCPostReachabilityAnalysis);
+			}
 			ExternalDirtyActorsTracker = MakeUnique<FWorldPartitionExternalDirtyActorsTracker>(this);
 		}
 	}
@@ -1007,7 +1010,10 @@ void UWorldPartition::UnregisterDelegates()
 		{
 			if (!IsEngineExitRequested())
 			{
-				FCoreUObjectDelegates::PostReachabilityAnalysis.RemoveAll(this);
+				if (IsMainWorldPartition())
+				{
+					FCoreUObjectDelegates::PostReachabilityAnalysis.RemoveAll(this);
+				}
 			}
 
 			ExternalDirtyActorsTracker.Reset();
