Project crashed due to Assertion failed: (Index >= 0) & (Index < ArrayNum) ...\Containers\Array.h Array index out of bounds: 0 into an array of size 0
#fix add cvar to disable clearing compressed data after saving out AnimationSequences during cook

diff --git a/Engine/Source/Runtime/Engine/Private/Animation/AnimSequence.cpp b/Engine/Source/Runtime/Engine/Private/Animation/AnimSequence.cpp
--- a/Engine/Source/Runtime/Engine/Private/Animation/AnimSequence.cpp
+++ b/Engine/Source/Runtime/Engine/Private/Animation/AnimSequence.cpp
@@ -474,32 +474,47 @@ void UAnimSequence::AddReferencedObjects(UObject* This, FReferenceCollector& Col
 }
 
 #if WITH_EDITOR
+
+int32 GAllowClearingCompressedDataDuringCook = 1;
+static FAutoConsoleVariableRef CVarAllowClearingCompressedDataDuringCook(
+	TEXT("a.AllowClearingCompressedDataDuringCook"),
+	GAllowClearingCompressedDataDuringCook,
+	TEXT("Values: 0/1\n")
+	TEXT("Controls whether or not to compressed data is cleared/free-ed during cook."),
+	ECVF_Default);
+
 void UAnimSequence::WillNeverCacheCookedPlatformDataAgain()
 {
-	PRAGMA_DISABLE_DEPRECATION_WARNINGS
-	bUseRawDataOnly = true;
-	PRAGMA_ENABLE_DEPRECATION_WARNINGS
+	if (GAllowClearingCompressedDataDuringCook)
+	{
+		PRAGMA_DISABLE_DEPRECATION_WARNINGS
+		bUseRawDataOnly = true;
+		PRAGMA_ENABLE_DEPRECATION_WARNINGS
 
-	UE::Anim::FAnimSequenceCompilingManager::Get().FinishCompilation({this});
-	// Clear out current platform, and any target platform data
-	CompressedData.Reset();
+		UE::Anim::FAnimSequenceCompilingManager::Get().FinishCompilation({this});
+		// Clear out current platform, and any target platform data
+		CompressedData.Reset();
 		
-	CacheTasksByKeyHash.Empty();
-	DataByPlatformKeyHash.Empty();
+		CacheTasksByKeyHash.Empty();
+		DataByPlatformKeyHash.Empty();
+	}
 }
 
 void UAnimSequence::ClearAllCachedCookedPlatformData()
 {
-	PRAGMA_DISABLE_DEPRECATION_WARNINGS
-	bUseRawDataOnly = true;
-	PRAGMA_ENABLE_DEPRECATION_WARNINGS
-	
-	// Delete any cache tasks first because the destructor will cancel the cache and build tasks,
-	// and drop their pointers to the data.
-	CacheTasksByKeyHash.Empty();
-	DataByPlatformKeyHash.Empty();
-	CompressedData.Reset();
-	DataKeyHash = FIoHash::Zero;
+	if (GAllowClearingCompressedDataDuringCook)
+	{
+		PRAGMA_DISABLE_DEPRECATION_WARNINGS
+		bUseRawDataOnly = true;
+		PRAGMA_ENABLE_DEPRECATION_WARNINGS
+		
+		// Delete any cache tasks first because the destructor will cancel the cache and build tasks,
+		// and drop their pointers to the data.
+		CacheTasksByKeyHash.Empty();
+		DataByPlatformKeyHash.Empty();
+		CompressedData.Reset();
+		DataKeyHash = FIoHash::Zero;
+	}
 }
 
 int64 UAnimSequence::GetUncompressedRawSize() const
