PR #12290: Wrong definition of macro UE_CVLOG_CIRCLE_THICK Visual Logger

diff --git a/Engine/Source/Runtime/Engine/Public/VisualLogger/VisualLogger.h b/Engine/Source/Runtime/Engine/Public/VisualLogger/VisualLogger.h
--- a/Engine/Source/Runtime/Engine/Public/VisualLogger/VisualLogger.h
+++ b/Engine/Source/Runtime/Engine/Public/VisualLogger/VisualLogger.h
@@ -93,10 +93,10 @@
 #define UE_CVLOG_ARROW(Condition, LogOwner, CategoryName, Verbosity, SegmentStart, SegmentEnd, Color, Format, ...) if(FVisualLogger::IsRecording() && Condition) {UE_VLOG_ARROW(LogOwner, CategoryName, Verbosity, SegmentStart, SegmentEnd, Color, Format, ##__VA_ARGS__);} 
 // Circle shape
 #define UE_VLOG_CIRCLE(LogOwner, CategoryName, Verbosity, Center, UpAxis, Radius, Color, Format, ...) if(FVisualLogger::IsRecording()) FVisualLogger::CircleLogf(LogOwner, CategoryName, ELogVerbosity::Verbosity, Center, UpAxis, Radius, Color, 0, Format, ##__VA_ARGS__)
-#define UE_CVLOG_CIRCLE(Condition, LogOwner, CategoryName, Verbosity, Center, UpAxis, Radius, Color, Format, ...) if(FVisualLogger::IsRecording() && Condition) {UE_VLOG_CIRCLE(LogOwner, CategoryName, Verbosity, Center, UpAxis, Radius, Color, 0, Format, ##__VA_ARGS__);} 
+#define UE_CVLOG_CIRCLE(Condition, LogOwner, CategoryName, Verbosity, Center, UpAxis, Radius, Color, Format, ...) if(FVisualLogger::IsRecording() && Condition) {UE_VLOG_CIRCLE(LogOwner, CategoryName, Verbosity, Center, UpAxis, Radius, Color, Format, ##__VA_ARGS__);} 
 // Circle shape
 #define UE_VLOG_CIRCLE_THICK(LogOwner, CategoryName, Verbosity, Center, UpAxis, Radius, Color, Thickness, Format, ...) if(FVisualLogger::IsRecording()) FVisualLogger::CircleLogf(LogOwner, CategoryName, ELogVerbosity::Verbosity, Center, UpAxis, Radius, Color, Thickness, Format, ##__VA_ARGS__)
-#define UE_CVLOG_CIRCLE_THICK(Condition, LogOwner, CategoryName, Verbosity, Center, UpAxis, Radius, Color, Thickness, Format, ...) if(FVisualLogger::IsRecording() && Condition) {UE_VLOG_CIRCLE(LogOwner, CategoryName, Verbosity, Center, UpAxis, Radius, Color, Thickness, Format, ##__VA_ARGS__);} 
+#define UE_CVLOG_CIRCLE_THICK(Condition, LogOwner, CategoryName, Verbosity, Center, UpAxis, Radius, Color, Thickness, Format, ...) if(FVisualLogger::IsRecording() && Condition) {UE_VLOG_CIRCLE_THICK(LogOwner, CategoryName, Verbosity, Center, UpAxis, Radius, Color, Thickness, Format, ##__VA_ARGS__);} 
 
 #define DECLARE_VLOG_EVENT(EventName) extern FVisualLogEventBase EventName;
 #define DEFINE_VLOG_EVENT(EventName, Verbosity, UserFriendlyDesc) FVisualLogEventBase EventName(TEXT(#EventName), TEXT(UserFriendlyDesc), ELogVerbosity::Verbosity); 
