Add project setting "Enable virtual textures for Post Processing Materials"
This allows VT sampling in PP materials but without any guarantees that things will work well.
The default setting is the current behavior which fails material compilation for such a set up.
If the new setting is enabled then PP materials that sample VT will compile.
Note that there is no VT feedback from the PP passes so the VT will only have good resolution pages in the cache if some other system is already triggering feedback.

diff --git a/Engine/Source/Runtime/Engine/Classes/Engine/RendererSettings.h b/Engine/Source/Runtime/Engine/Classes/Engine/RendererSettings.h
--- a/Engine/Source/Runtime/Engine/Classes/Engine/RendererSettings.h
+++ b/Engine/Source/Runtime/Engine/Classes/Engine/RendererSettings.h
@@ -424,6 +424,12 @@ class URendererSettings : public UDeveloperSettings
 		ToolTip = "Relax restriction on virtual textures contributing to Opacity Mask. In some edge cases this can lead to low resolution shadow edges."))
 		uint32 bEnableVirtualTextureOpacityMask : 1;
 
+	UPROPERTY(config, EditAnywhere, Category = VirtualTextures, meta = (
+		EditCondition = "bVirtualTextures",
+		DisplayName = "Enable virtual textures for Post Processing Materials",
+		ToolTip = "Relax restriction on virtual textures being sampled in Post Processing Materials. There is no feedback from the Post Processing passes, so virtual textures will only have pages in the cache that are triggered by feedback in another pass."))
+		uint32 bEnableVirtualTexturePostProcessing : 1;
+
 	UPROPERTY(config, EditAnywhere, Category = VirtualTextures, meta = (
 		EditCondition = "bVirtualTextures",
 		ConsoleVariable = "r.VT.TileSize", DisplayName = "Tile size",
diff --git a/Engine/Source/Runtime/Engine/Private/Materials/HLSLMaterialTranslator.cpp b/Engine/Source/Runtime/Engine/Private/Materials/HLSLMaterialTranslator.cpp
--- a/Engine/Source/Runtime/Engine/Private/Materials/HLSLMaterialTranslator.cpp
+++ b/Engine/Source/Runtime/Engine/Private/Materials/HLSLMaterialTranslator.cpp
@@ -7087,7 +7087,10 @@ int32 FHLSLMaterialTranslator::TextureSample(
 	{
 		if (Material->GetMaterialDomain() != MD_Surface && Material->GetMaterialDomain() != MD_DeferredDecal)
 		{
-			return Errorf(TEXT("Sampling a virtual texture is currently only supported for surface and decal materials."));
+			if (GetDefault<URendererSettings>()->bEnableVirtualTexturePostProcessing == false)
+			{
+				return Errorf(TEXT("Sampling a virtual texture is currently only supported for surface and decal materials."));
+			}
 		}
 	}
 
diff --git a/Engine/Source/Runtime/Engine/Private/RendererSettings.cpp b/Engine/Source/Runtime/Engine/Private/RendererSettings.cpp
--- a/Engine/Source/Runtime/Engine/Private/RendererSettings.cpp
+++ b/Engine/Source/Runtime/Engine/Private/RendererSettings.cpp
@@ -102,6 +102,7 @@ URendererSettings::URendererSettings(const FObjectInitializer& ObjectInitializer
 	WhiteChromaticityCoordinate = FVector2D::ZeroVector;
 	bUseLegacyLuminanceFactors = false;
 	bEnableVirtualTextureOpacityMask = false;
+	bEnableVirtualTexturePostProcessing = false;
 
 #if WITH_EDITOR
 	if (IsTemplate())
