[AutoRTFM] Update various ISPC calls for AutoRTFM compatibility.

Calls to ISPC-compiled code must be made in the open, as no closed versions of the functions exist. In places where ISPC overwrites a buffer that may require rollback, we explicitly mark the data via RecordOpenWrite.

Additionally, fixed a bug which would cause RecordOpenWrite to crash when no transaction was open (it would attempt to read from null via `CurrentTransaction->`).

diff --git a/Engine/Source/Runtime/Core/Private/AutoRTFM/API.cpp b/Engine/Source/Runtime/Core/Private/AutoRTFM/API.cpp
--- a/Engine/Source/Runtime/Core/Private/AutoRTFM/API.cpp
+++ b/Engine/Source/Runtime/Core/Private/AutoRTFM/API.cpp
@@ -583,8 +583,10 @@ extern "C" void autortfm_record_open_write(void* Ptr, size_t Size)
 	{
 		return RTFM_autortfm_record_open_write(Ptr, Size);
 	}
-
-	FContext::Get()->RecordWrite(Ptr, Size);
+	if (FTransaction* CurrentTransaction = FContext::Get()->GetCurrentTransaction())
+	{
+		CurrentTransaction->RecordWrite(Ptr, Size);
+	}
 }
 
 extern "C" UE_AUTORTFM_NOAUTORTFM void autortfm_register_open_function(void* OriginalFunction, void* NewFunction)
diff --git a/Engine/Source/Runtime/Engine/Private/Animation/AnimationRuntime.cpp b/Engine/Source/Runtime/Engine/Private/Animation/AnimationRuntime.cpp
--- a/Engine/Source/Runtime/Engine/Private/Animation/AnimationRuntime.cpp
+++ b/Engine/Source/Runtime/Engine/Private/Animation/AnimationRuntime.cpp
@@ -19,6 +19,7 @@
 #include "Animation/MirrorDataTable.h"
 #include "Animation/SkeletonRemappingRegistry.h"
 #include "Animation/SkeletonRemapping.h"
+#include "AutoRTFM/AutoRTFM.h"
 
 DEFINE_LOG_CATEGORY(LogAnimation);
 DEFINE_LOG_CATEGORY(LogRootMotion);
@@ -163,11 +164,17 @@ FORCEINLINE void BlendPose<ETransformBlendMode::Overwrite>(const FCompactPose& S
 #if INTEL_ISPC
 	if (bAnim_BlendPoseOverwrite_ISPC_Enabled)
 	{
-		ispc::BlendTransformOverwrite(
-			reinterpret_cast<const ispc::FTransform*>(SourcePose.GetBones().GetData()),
-			reinterpret_cast<ispc::FTransform*>(ResultPose.GetMutableBones().GetData()),
-			BlendWeight,
-			SourcePose.GetNumBones());
+		UE_AUTORTFM_OPEN2
+		{
+			AutoRTFM::RecordOpenWrite(
+				ResultPose.GetMutableBones().GetData(),
+				ResultPose.GetNumBones() * sizeof(FTransform));
+			ispc::BlendTransformOverwrite(
+				reinterpret_cast<const ispc::FTransform*>(SourcePose.GetBones().GetData()),
+				reinterpret_cast<ispc::FTransform*>(ResultPose.GetMutableBones().GetData()),
+				BlendWeight,
+				SourcePose.GetNumBones());
+		};
 	}
 	else
 #endif
@@ -185,11 +192,17 @@ FORCEINLINE void BlendPose<ETransformBlendMode::Accumulate>(const FCompactPose&
 #if INTEL_ISPC
 	if (bAnim_BlendPoseAccumulate_ISPC_Enabled)
 	{
-		ispc::BlendTransformAccumulate(
-			reinterpret_cast<const ispc::FTransform*>(SourcePose.GetBones().GetData()),
-			reinterpret_cast<ispc::FTransform*>(ResultPose.GetMutableBones().GetData()),
-			BlendWeight,
-			SourcePose.GetNumBones());
+		UE_AUTORTFM_OPEN2
+		{
+			AutoRTFM::RecordOpenWrite(
+				ResultPose.GetMutableBones().GetData(),
+				ResultPose.GetNumBones() * sizeof(FTransform));
+			ispc::BlendTransformAccumulate(
+				reinterpret_cast<const ispc::FTransform*>(SourcePose.GetBones().GetData()),
+				reinterpret_cast<ispc::FTransform*>(ResultPose.GetMutableBones().GetData()),
+				BlendWeight,
+				SourcePose.GetNumBones());
+		};
 	}
 	else
 #endif
@@ -966,13 +979,18 @@ void FAnimationRuntime::LerpBoneTransforms(TArray<FTransform>& A, const TArray<F
 #if INTEL_ISPC
 	if (bAnim_LerpBoneTransforms_ISPC_Enabled)
 	{
-		ispc::LerpBoneTransforms(
-			reinterpret_cast<ispc::FTransform*>(A.GetData()),
-			reinterpret_cast<const ispc::FTransform*>(B.GetData()),
-			Alpha,
-			RequiredBonesArray.GetData(),
-			RequiredBonesArray.Num()
-		);
+		UE_AUTORTFM_OPEN2
+		{
+			AutoRTFM::RecordOpenWrite(
+				A.GetData(),
+				A.Num() * sizeof(FTransform));
+			ispc::LerpBoneTransforms(
+				reinterpret_cast<ispc::FTransform*>(A.GetData()),
+				reinterpret_cast<const ispc::FTransform*>(B.GetData()),
+				Alpha,
+				RequiredBonesArray.GetData(),
+				RequiredBonesArray.Num());
+		};
 	}
 	else
 #endif
@@ -1063,10 +1081,16 @@ void FAnimationRuntime::ConvertPoseToAdditive(FCompactPose& TargetPose, const FC
 #if INTEL_ISPC
 	if (bAnim_ConvertPoseToAdditive_ISPC_Enabled)
 	{
-		ispc::ConvertPoseToAdditive(
-			reinterpret_cast<ispc::FTransform*>(TargetPose.GetMutableBones().GetData()),
-			reinterpret_cast<const ispc::FTransform*>(BasePose.GetBones().GetData()),
-			BasePose.GetNumBones());
+		UE_AUTORTFM_OPEN2
+		{
+			AutoRTFM::RecordOpenWrite(
+				TargetPose.GetMutableBones().GetData(),
+				TargetPose.GetNumBones() * sizeof(FTransform));
+			ispc::ConvertPoseToAdditive(
+				reinterpret_cast<ispc::FTransform*>(TargetPose.GetMutableBones().GetData()),
+				reinterpret_cast<const ispc::FTransform*>(BasePose.GetBones().GetData()),
+				BasePose.GetNumBones());
+		};
 	}
 	else
 #endif
@@ -1088,10 +1112,16 @@ void FAnimationRuntime::ConvertPoseToMeshRotation(FCompactPose& LocalPose)
 #if INTEL_ISPC
 	if (bAnim_ConvertPoseToMeshRotation_ISPC_Enabled)
 	{
-		ispc::ConvertPoseToMeshRotation(
-			reinterpret_cast<ispc::FTransform*>(LocalPose.GetMutableBones().GetData()),
-			reinterpret_cast<const int32*>(LocalPose.GetBoneContainer().GetCompactPoseParentBoneArray().GetData()),
-			LocalPose.GetNumBones());
+		UE_AUTORTFM_OPEN2
+		{
+			AutoRTFM::RecordOpenWrite(
+				LocalPose.GetMutableBones().GetData(),
+				LocalPose.GetNumBones() * sizeof(FTransform));
+			ispc::ConvertPoseToMeshRotation(
+				reinterpret_cast<ispc::FTransform*>(LocalPose.GetMutableBones().GetData()),
+				reinterpret_cast<const int32*>(LocalPose.GetBoneContainer().GetCompactPoseParentBoneArray().GetData()),
+				LocalPose.GetNumBones());
+		};
 	}
 	else
 #endif
@@ -1113,10 +1143,16 @@ void FAnimationRuntime::ConvertMeshRotationPoseToLocalSpace(FCompactPose& Pose)
 #if INTEL_ISPC
 	if (bAnim_ConvertMeshRotationPoseToLocalSpace_ISPC_Enabled)
 	{
-		ispc::ConvertMeshRotationPoseToLocalSpace(
-			reinterpret_cast<ispc::FTransform*>(Pose.GetMutableBones().GetData()),
-			reinterpret_cast<const int32*>(Pose.GetBoneContainer().GetCompactPoseParentBoneArray().GetData()),
-			Pose.GetNumBones());
+		UE_AUTORTFM_OPEN2
+		{
+			AutoRTFM::RecordOpenWrite(
+				Pose.GetMutableBones().GetData(),
+				Pose.GetNumBones() * sizeof(FTransform));
+			ispc::ConvertMeshRotationPoseToLocalSpace(
+				reinterpret_cast<ispc::FTransform*>(Pose.GetMutableBones().GetData()),
+				reinterpret_cast<const int32*>(Pose.GetBoneContainer().GetCompactPoseParentBoneArray().GetData()),
+				Pose.GetNumBones());
+		};
 	}
 	else
 #endif
@@ -1199,11 +1235,17 @@ void FAnimationRuntime::AccumulateLocalSpaceAdditivePoseInternal(FCompactPose& B
 #if INTEL_ISPC
 			if (bAnim_AccumulateLocalSpaceAdditivePose_ISPC_Enabled)
 			{
-				ispc::AccumulateWithAdditiveScale(
-					reinterpret_cast<ispc::FTransform*>(BasePose.GetMutableBones().GetData()),
-					reinterpret_cast<const ispc::FTransform*>(AdditivePose.GetBones().GetData()),
-					Weight,
-					BasePose.GetNumBones());
+				UE_AUTORTFM_OPEN2
+				{
+					AutoRTFM::RecordOpenWrite(
+						BasePose.GetMutableBones().GetData(),
+						BasePose.GetNumBones() * sizeof(FTransform));
+					ispc::AccumulateWithAdditiveScale(
+						reinterpret_cast<ispc::FTransform*>(BasePose.GetMutableBones().GetData()),
+						reinterpret_cast<const ispc::FTransform*>(AdditivePose.GetBones().GetData()),
+						Weight,
+						BasePose.GetNumBones());
+				};
 			}
 			else
 #endif
@@ -1767,7 +1809,7 @@ void FAnimationRuntime::BlendPosesPerBoneFilter(
 void FAnimationRuntime::BlendPosesPerBoneFilter(FCompactPose& BasePose, const TArray<FCompactPose>& BlendPoses, FBlendedCurve& BaseCurve, const TArray<FBlendedCurve>& BlendedCurves,
 	UE::Anim::FStackAttributeContainer& BaseAttributes,
 	const TArray<UE::Anim::FStackAttributeContainer>& BlendAttributes,
-FAnimationPoseData& OutAnimationPoseData, TArray<FPerBoneBlendWeight>& BoneBlendWeights, EBlendPosesPerBoneFilterFlags BlendFlags, enum ECurveBlendOption::Type CurveBlendOption)
+	FAnimationPoseData& OutAnimationPoseData, TArray<FPerBoneBlendWeight>& BoneBlendWeights, EBlendPosesPerBoneFilterFlags BlendFlags, enum ECurveBlendOption::Type CurveBlendOption)
 {
 	SCOPE_CYCLE_COUNTER(STAT_BlendPosesPerBoneFilter);
 
@@ -1775,7 +1817,6 @@ FAnimationPoseData& OutAnimationPoseData, TArray<FPerBoneBlendWeight>& BoneBlend
 	FBlendedCurve& OutCurve = OutAnimationPoseData.GetCurve();
 	UE::Anim::FStackAttributeContainer& OutAttributes = OutAnimationPoseData.GetAttributes();
 
-	
 	// if no blendpose, outpose = basepose
 	if (BlendPoses.Num() == 0)
 	{
@@ -1886,20 +1927,28 @@ FAnimationPoseData& OutAnimationPoseData, TArray<FPerBoneBlendWeight>& BoneBlend
 #if INTEL_ISPC
 		if (bAnim_BlendPosesPerBoneFilter_ISPC_Enabled)
 		{
-			ispc::BlendPosesPerBoneFilterScaleRotation(
-				reinterpret_cast<ispc::FTransform*>(OutPose.GetMutableBones().GetData()),
-				reinterpret_cast<const ispc::FTransform*>(BasePose.GetBones().GetData()),
-				reinterpret_cast<const uint8*>(&BlendPoses),
-				reinterpret_cast<ispc::FVector4*>(SourceRotations.GetData()),
-				reinterpret_cast<ispc::FVector*>(SourceScales.GetData()),
-				reinterpret_cast<ispc::FVector4*>(TargetRotations.GetData()),
-				reinterpret_cast<ispc::FVector*>(TargetScales.GetData()),
-				reinterpret_cast<ispc::FVector4*>(BlendRotations.GetData()),
-				reinterpret_cast<ispc::FVector*>(BlendScales.GetData()),
-				MaxPoseWeights.GetData(),
-				reinterpret_cast<const ispc::FPerBoneBlendWeight*>(BoneBlendWeights.GetData()),
-				reinterpret_cast<const int32*>(BoneContainer.GetCompactPoseParentBoneArray().GetData()),
-				BasePose.GetNumBones());
+			UE_AUTORTFM_OPEN2
+			{
+				AutoRTFM::RecordOpenWrite(
+					OutPose.GetMutableBones().GetData(),
+					OutPose.GetNumBones() * sizeof(FTransform));
+				// AutoRTFM doesn't track writes to Source/Target/Blend arrays, since these are part of the Scratch Area;
+				// they are not visible to callers, and don't need to persist.
+				ispc::BlendPosesPerBoneFilterScaleRotation(
+					reinterpret_cast<ispc::FTransform*>(OutPose.GetMutableBones().GetData()),
+					reinterpret_cast<const ispc::FTransform*>(BasePose.GetBones().GetData()),
+					reinterpret_cast<const uint8*>(&BlendPoses),
+					reinterpret_cast<ispc::FVector4*>(SourceRotations.GetData()),
+					reinterpret_cast<ispc::FVector*>(SourceScales.GetData()),
+					reinterpret_cast<ispc::FVector4*>(TargetRotations.GetData()),
+					reinterpret_cast<ispc::FVector*>(TargetScales.GetData()),
+					reinterpret_cast<ispc::FVector4*>(BlendRotations.GetData()),
+					reinterpret_cast<ispc::FVector*>(BlendScales.GetData()),
+					MaxPoseWeights.GetData(),
+					reinterpret_cast<const ispc::FPerBoneBlendWeight*>(BoneBlendWeights.GetData()),
+					reinterpret_cast<const int32*>(BoneContainer.GetCompactPoseParentBoneArray().GetData()),
+					BasePose.GetNumBones());
+			};
 		}
 		else
 #endif
@@ -1964,17 +2013,25 @@ FAnimationPoseData& OutAnimationPoseData, TArray<FPerBoneBlendWeight>& BoneBlend
 #if INTEL_ISPC
 		if (bAnim_BlendPosesPerBoneFilter_ISPC_Enabled)
 		{
-			ispc::BlendPosesPerBoneFilterRotation(
-				reinterpret_cast<ispc::FTransform*>(OutPose.GetMutableBones().GetData()),
-				reinterpret_cast<const ispc::FTransform*>(BasePose.GetBones().GetData()),
-				reinterpret_cast<const uint8*>(&BlendPoses),
-				reinterpret_cast<ispc::FVector4*>(SourceRotations.GetData()),
-				reinterpret_cast<ispc::FVector4*>(TargetRotations.GetData()),
-				reinterpret_cast<ispc::FVector4*>(BlendRotations.GetData()),
-				MaxPoseWeights.GetData(),
-				reinterpret_cast<const ispc::FPerBoneBlendWeight*>(BoneBlendWeights.GetData()),
-				reinterpret_cast<const int32*>(BoneContainer.GetCompactPoseParentBoneArray().GetData()),
-				BasePose.GetNumBones());
+			UE_AUTORTFM_OPEN2
+			{
+				AutoRTFM::RecordOpenWrite(
+					OutPose.GetMutableBones().GetData(),
+					OutPose.GetNumBones() * sizeof(FTransform));
+				// AutoRTFM doesn't track writes to Source/Target/Blend arrays, since these are part of the Scratch Area;
+				// they are not visible to callers, and don't need to persist.
+				ispc::BlendPosesPerBoneFilterRotation(
+					reinterpret_cast<ispc::FTransform*>(OutPose.GetMutableBones().GetData()),
+					reinterpret_cast<const ispc::FTransform*>(BasePose.GetBones().GetData()),
+					reinterpret_cast<const uint8*>(&BlendPoses),
+					reinterpret_cast<ispc::FVector4*>(SourceRotations.GetData()),
+					reinterpret_cast<ispc::FVector4*>(TargetRotations.GetData()),
+					reinterpret_cast<ispc::FVector4*>(BlendRotations.GetData()),
+					MaxPoseWeights.GetData(),
+					reinterpret_cast<const ispc::FPerBoneBlendWeight*>(BoneBlendWeights.GetData()),
+					reinterpret_cast<const int32*>(BoneContainer.GetCompactPoseParentBoneArray().GetData()),
+					BasePose.GetNumBones());
+			};
 		}
 		else
 #endif
@@ -2033,17 +2090,25 @@ FAnimationPoseData& OutAnimationPoseData, TArray<FPerBoneBlendWeight>& BoneBlend
 #if INTEL_ISPC
 		if (bAnim_BlendPosesPerBoneFilter_ISPC_Enabled)
 		{
-			ispc::BlendPosesPerBoneFilterScale(
-				reinterpret_cast<ispc::FTransform*>(OutPose.GetMutableBones().GetData()),
-				reinterpret_cast<const ispc::FTransform*>(BasePose.GetBones().GetData()),
-				reinterpret_cast<const uint8*>(&BlendPoses),
-				reinterpret_cast<ispc::FVector*>(SourceScales.GetData()),
-				reinterpret_cast<ispc::FVector*>(TargetScales.GetData()),
-				reinterpret_cast<ispc::FVector*>(BlendScales.GetData()),
-				MaxPoseWeights.GetData(),
-				reinterpret_cast<const ispc::FPerBoneBlendWeight*>(BoneBlendWeights.GetData()),
-				reinterpret_cast<const int32*>(BoneContainer.GetCompactPoseParentBoneArray().GetData()),
-				BasePose.GetNumBones());
+			UE_AUTORTFM_OPEN2
+			{
+				AutoRTFM::RecordOpenWrite(
+					OutPose.GetMutableBones().GetData(),
+					OutPose.GetNumBones() * sizeof(FTransform));
+				// AutoRTFM doesn't track writes to Source/Target/Blend arrays, since these are part of the Scratch Area;
+				// they are not visible to callers, and don't need to persist.
+				ispc::BlendPosesPerBoneFilterScale(
+					reinterpret_cast<ispc::FTransform*>(OutPose.GetMutableBones().GetData()),
+					reinterpret_cast<const ispc::FTransform*>(BasePose.GetBones().GetData()),
+					reinterpret_cast<const uint8*>(&BlendPoses),
+					reinterpret_cast<ispc::FVector*>(SourceScales.GetData()),
+					reinterpret_cast<ispc::FVector*>(TargetScales.GetData()),
+					reinterpret_cast<ispc::FVector*>(BlendScales.GetData()),
+					MaxPoseWeights.GetData(),
+					reinterpret_cast<const ispc::FPerBoneBlendWeight*>(BoneBlendWeights.GetData()),
+					reinterpret_cast<const int32*>(BoneContainer.GetCompactPoseParentBoneArray().GetData()),
+					BasePose.GetNumBones());
+			};
 		}
 		else
 #endif
@@ -2102,14 +2167,20 @@ FAnimationPoseData& OutAnimationPoseData, TArray<FPerBoneBlendWeight>& BoneBlend
 #if INTEL_ISPC
 		if (bAnim_BlendPosesPerBoneFilter_ISPC_Enabled)
 		{
-			ispc::BlendPosesPerBoneFilter(
-				reinterpret_cast<ispc::FTransform*>(OutPose.GetMutableBones().GetData()),
-				reinterpret_cast<const ispc::FTransform*>(BasePose.GetBones().GetData()),
-				reinterpret_cast<const uint8*>(&BlendPoses),
-				MaxPoseWeights.GetData(),
-				reinterpret_cast<const ispc::FPerBoneBlendWeight*>(BoneBlendWeights.GetData()),
-				reinterpret_cast<const int32*>(BoneContainer.GetCompactPoseParentBoneArray().GetData()),
-				BasePose.GetNumBones());
+			UE_AUTORTFM_OPEN2
+			{
+				AutoRTFM::RecordOpenWrite(
+					OutPose.GetMutableBones().GetData(),
+					OutPose.GetNumBones() * sizeof(FTransform));
+				ispc::BlendPosesPerBoneFilter(
+					reinterpret_cast<ispc::FTransform*>(OutPose.GetMutableBones().GetData()),
+					reinterpret_cast<const ispc::FTransform*>(BasePose.GetBones().GetData()),
+					reinterpret_cast<const uint8*>(&BlendPoses),
+					MaxPoseWeights.GetData(),
+					reinterpret_cast<const ispc::FPerBoneBlendWeight*>(BoneBlendWeights.GetData()),
+					reinterpret_cast<const int32*>(BoneContainer.GetCompactPoseParentBoneArray().GetData()),
+					BasePose.GetNumBones());
+			};
 		}
 		else
 #endif
diff --git a/Engine/Source/Runtime/Engine/Private/Animation/BonePose.cpp b/Engine/Source/Runtime/Engine/Private/Animation/BonePose.cpp
--- a/Engine/Source/Runtime/Engine/Private/Animation/BonePose.cpp
+++ b/Engine/Source/Runtime/Engine/Private/Animation/BonePose.cpp
@@ -5,6 +5,7 @@
 #include "AnimationRuntime.h"
 #include "AnimEncoding.h"
 #include "Animation/AnimSequenceHelpers.h"
+#include "AutoRTFM/AutoRTFM.h"
 #if INTEL_ISPC
 #include "BonePose.ispc.generated.h"
 
@@ -29,7 +30,15 @@ void FCompactPose::NormalizeRotations()
 	if (bAnim_BonePose_ISPC_Enabled)
 	{
 #if INTEL_ISPC
-		ispc::NormalizeRotations((ispc::FTransform*)this->Bones.GetData(), this->Bones.Num());
+		UE_AUTORTFM_OPEN2
+		{
+			AutoRTFM::RecordOpenWrite(
+				this->Bones.GetData(),
+				this->Bones.Num() * sizeof(FTransform));
+			ispc::NormalizeRotations(
+				reinterpret_cast<ispc::FTransform*>(this->Bones.GetData()), 
+				this->Bones.Num());
+		};
 #endif
 	}
 	else
@@ -47,7 +56,15 @@ void FCompactPose::ResetToAdditiveIdentity()
 	if (bAnim_BonePose_ISPC_Enabled)
 	{
 #if INTEL_ISPC
-		ispc::ResetToAdditiveIdentity((ispc::FTransform*)this->Bones.GetData(), this->Bones.Num());
+		UE_AUTORTFM_OPEN2
+		{
+			AutoRTFM::RecordOpenWrite(
+				this->Bones.GetData(),
+				this->Bones.Num() * sizeof(FTransform));
+			ispc::ResetToAdditiveIdentity(
+				reinterpret_cast<ispc::FTransform*>(this->Bones.GetData()),
+				this->Bones.Num());
+		};
 #endif
 	}
 	else
@@ -65,7 +82,15 @@ void FCompactHeapPose::NormalizeRotations()
 	if (bAnim_BonePose_ISPC_Enabled)
 	{
 #if INTEL_ISPC
-		ispc::NormalizeRotations((ispc::FTransform*)this->Bones.GetData(), this->Bones.Num());
+		UE_AUTORTFM_OPEN2
+		{
+			AutoRTFM::RecordOpenWrite(
+				this->Bones.GetData(),
+				this->Bones.Num() * sizeof(FTransform));
+			ispc::NormalizeRotations(
+				reinterpret_cast<ispc::FTransform*>(this->Bones.GetData()),
+				this->Bones.Num());
+		};
 #endif
 	}
 	else
@@ -83,7 +108,15 @@ void FCompactHeapPose::ResetToAdditiveIdentity()
 	if (bAnim_BonePose_ISPC_Enabled)
 	{
 #if INTEL_ISPC
-		ispc::ResetToAdditiveIdentity((ispc::FTransform*)this->Bones.GetData(), this->Bones.Num());
+		UE_AUTORTFM_OPEN2
+		{
+			AutoRTFM::RecordOpenWrite(
+				this->Bones.GetData(),
+				this->Bones.Num() * sizeof(FTransform));
+			ispc::ResetToAdditiveIdentity(
+				reinterpret_cast<ispc::FTransform*>(this->Bones.GetData()),
+				this->Bones.Num());
+		};
 #endif
 	}
 	else
diff --git a/Engine/Source/Runtime/Engine/Private/Collision/KAggregateGeom.cpp b/Engine/Source/Runtime/Engine/Private/Collision/KAggregateGeom.cpp
--- a/Engine/Source/Runtime/Engine/Private/Collision/KAggregateGeom.cpp
+++ b/Engine/Source/Runtime/Engine/Private/Collision/KAggregateGeom.cpp
@@ -19,6 +19,7 @@
 #include "Chaos/Levelset.h"
 #include "PhysicsEngine/TaperedCapsuleElem.h"
 #include "Chaos/WeightedLatticeImplicitObject.h"
+#include "AutoRTFM/AutoRTFM.h"
 #if INTEL_ISPC
 #include "KAggregateGeom.ispc.generated.h"
 #endif
@@ -350,16 +351,18 @@ FBox FKBoxElem::CalcAABB(const FTransform& BoneTM, float Scale) const
 	{
 #if INTEL_ISPC
 		FBox LocalBox(ForceInit);
-		ispc::BoxCalcAABB(
-			(ispc::FBox&)LocalBox,
-			(ispc::FTransform&)BoneTM,
-			Scale,
-			(ispc::FRotator&)Rotation,
-			(ispc::FVector&)Center,
-			X,
-			Y,
-			Z);
-
+		UE_AUTORTFM_OPEN2
+		{
+			ispc::BoxCalcAABB(
+				reinterpret_cast<ispc::FBox&>(LocalBox),
+				reinterpret_cast<const ispc::FTransform&>(BoneTM),
+				Scale,
+				reinterpret_cast<const ispc::FRotator&>(Rotation),
+				reinterpret_cast<const ispc::FVector&>(Center),
+				X,
+				Y,
+				Z);
+		};
 		return LocalBox;
 #endif
 	}
@@ -391,15 +394,17 @@ FBox FKSphylElem::CalcAABB(const FTransform& BoneTM, float Scale) const
 	{
 #if INTEL_ISPC
 		FBox Result(ForceInit);
-		ispc::SPhylCalcAABB(
-			(ispc::FBox&)Result,
-			(ispc::FTransform&)BoneTM,
-			Scale,
-			(ispc::FRotator&)Rotation,
-			(ispc::FVector&)Center,
-			Radius,
-			Length);
-
+		UE_AUTORTFM_OPEN2
+		{
+			ispc::SPhylCalcAABB(
+				reinterpret_cast<ispc::FBox&>(Result),
+				reinterpret_cast<const ispc::FTransform&>(BoneTM),
+				Scale,
+				reinterpret_cast<const ispc::FRotator&>(Rotation),
+				reinterpret_cast<const ispc::FVector&>(Center),
+				Radius,
+				Length);
+		};
 		return Result;
 #endif
 	}
diff --git a/Engine/Source/Runtime/Engine/Private/SkinnedAsset.cpp b/Engine/Source/Runtime/Engine/Private/SkinnedAsset.cpp
--- a/Engine/Source/Runtime/Engine/Private/SkinnedAsset.cpp
+++ b/Engine/Source/Runtime/Engine/Private/SkinnedAsset.cpp
@@ -9,6 +9,7 @@
 #include "Materials/MaterialInterface.h"
 #include "SkeletalRenderGPUSkin.h"
 #include "PSOPrecache.h"
+#include "AutoRTFM/AutoRTFM.h"
 #if WITH_EDITORONLY_DATA
 #include "Streaming/UVChannelDensity.h"
 #endif
@@ -468,14 +469,20 @@ void USkinnedAsset::FillComponentSpaceTransforms(const TArray<FTransform>& InBon
 	if (IsISPCEnabled())
 	{
 #if INTEL_ISPC
-		ispc::FillComponentSpaceTransforms(
-			(ispc::FTransform*)&ComponentSpaceData[0],
-			(ispc::FTransform*)&LocalTransformsData[0],
-			InFillComponentSpaceTransformsRequiredBones.GetData(),
-			(const uint8*)GetRefSkeleton().GetRefBoneInfo().GetData(),
-			sizeof(FMeshBoneInfo),
-			offsetof(FMeshBoneInfo, ParentIndex),
-			InFillComponentSpaceTransformsRequiredBones.Num());
+		UE_AUTORTFM_OPEN2
+		{
+			AutoRTFM::RecordOpenWrite(
+				ComponentSpaceData,
+				OutComponentSpaceTransforms.Num() * sizeof(FTransform));
+			ispc::FillComponentSpaceTransforms(
+				reinterpret_cast<ispc::FTransform*>(ComponentSpaceData),
+				reinterpret_cast<const ispc::FTransform*>(LocalTransformsData),
+				InFillComponentSpaceTransformsRequiredBones.GetData(),
+				reinterpret_cast<const uint8*>(GetRefSkeleton().GetRefBoneInfo().GetData()),
+				sizeof(FMeshBoneInfo),
+				offsetof(FMeshBoneInfo, ParentIndex),
+				InFillComponentSpaceTransformsRequiredBones.Num());
+		};
 #endif
 	}
 	else
