Fix static analysis warning
#rnx

diff --git a/Engine/Source/Runtime/Engine/Public/ScreenRendering.h b/Engine/Source/Runtime/Engine/Public/ScreenRendering.h
--- a/Engine/Source/Runtime/Engine/Public/ScreenRendering.h
+++ b/Engine/Source/Runtime/Engine/Public/ScreenRendering.h
@@ -132,12 +132,14 @@ public:
 
 	static bool ShouldCompilePermutation(const FGlobalShaderPermutationParameters& Parameters)
 	{
+#if PLATFORM_DESKTOP
 		// this is a desktop-only shader so far (really only needed in the editor)
-		if ((PLATFORM_DESKTOP) && FScreenPS::ShouldCompilePermutation(Parameters))
+		if (FScreenPS::ShouldCompilePermutation(Parameters))
 		{
 			UE::StereoRenderUtils::FStereoShaderAspects Aspects(Parameters.Platform);
 			return Aspects.IsMobileMultiViewEnabled();
 		}
+#endif
 		return false;
 	}
 
