Changed preview skeleton forward axis name to be more clear

diff --git a/Engine/Source/Editor/Persona/Private/AnimationEditorViewportClient.cpp b/Engine/Source/Editor/Persona/Private/AnimationEditorViewportClient.cpp
--- a/Engine/Source/Editor/Persona/Private/AnimationEditorViewportClient.cpp
+++ b/Engine/Source/Editor/Persona/Private/AnimationEditorViewportClient.cpp
@@ -2120,7 +2120,7 @@ void FAnimationViewportClient::DrawRootMotionTrajectory(UDebugSkelMeshComponent*
 			const float CurrentTime = MeshComponent->GetPosition();
 			const USkeleton* Skeleton = AnimSequenceBase->GetSkeleton();
 			check(Skeleton);
-			EAxis::Type SkeletonForwardAxis = Skeleton->GetForwardAxis();
+			EAxis::Type SkeletonForwardAxis = Skeleton->GetPreviewForwardAxis();
 
 			const FColor TrajectoryColor = FColor::Black.WithAlpha(64);
 
diff --git a/Engine/Source/Runtime/Engine/Classes/Animation/Skeleton.h b/Engine/Source/Runtime/Engine/Classes/Animation/Skeleton.h
--- a/Engine/Source/Runtime/Engine/Classes/Animation/Skeleton.h
+++ b/Engine/Source/Runtime/Engine/Classes/Animation/Skeleton.h
@@ -310,7 +310,7 @@ protected:
 
 	// Preview axis to consider as "forward" for the skeleton. Only used for preview purposes.
 	UPROPERTY(EditAnywhere, Category = Preview)
-	TEnumAsByte<EAxis::Type> ForwardAxis;
+	TEnumAsByte<EAxis::Type> PreviewForwardAxis;
 #endif
 
 	/** Reference Skeleton */
@@ -349,7 +349,7 @@ public:
 #if WITH_EDITOR
 	ENGINE_API virtual void PreEditUndo() override;
 	ENGINE_API virtual void PostEditUndo() override;
-	ENGINE_API EAxis::Type GetForwardAxis() const { return ForwardAxis; }
+	ENGINE_API EAxis::Type GetPreviewForwardAxis() const { return PreviewForwardAxis; }
 #endif
 	ENGINE_API virtual void BeginDestroy() override;
 
diff --git a/Engine/Source/Runtime/Engine/Private/Animation/Skeleton.cpp b/Engine/Source/Runtime/Engine/Private/Animation/Skeleton.cpp
--- a/Engine/Source/Runtime/Engine/Private/Animation/Skeleton.cpp
+++ b/Engine/Source/Runtime/Engine/Private/Animation/Skeleton.cpp
@@ -336,7 +336,7 @@ USkeleton::USkeleton(const FObjectInitializer& ObjectInitializer)
 	CachedSoftObjectPtr = TSoftObjectPtr<USkeleton>(this);
 
 #if WITH_EDITORONLY_DATA
-	ForwardAxis = EAxis::Y;
+	PreviewForwardAxis = EAxis::Y;
 #endif
 
 	if (HasAnyFlags(RF_ClassDefaultObject))
