Keep track of the most recently used FInputDeviceId for a given FPlatformUserId in the UInputDeviceSubsystem. This is a good place to keep track of this because it is populated by an input device pre-processor, so it will only populate the data if an input event has been fired from the message handler. This means that every input device plugin does not need to implement this.

Deprecate bEnableInputDeviceSubsystem on the input device subsystem, we always want to create it.

#rnx

diff --git a/Engine/Config/BaseInput.ini b/Engine/Config/BaseInput.ini
--- a/Engine/Config/BaseInput.ini
+++ b/Engine/Config/BaseInput.ini
@@ -18,7 +18,6 @@ TriggerKeyIsConsideredPressed=0.75
 TriggerKeyIsConsideredReleased=0.25
 InitialButtonRepeatDelay=0.2
 ButtonRepeatDelay=0.1
-bEnableInputDeviceSubsystem=true
 
 ;-----------------------------------------------------------------------------------------
 ; Axis properties
diff --git a/Engine/Source/Runtime/Engine/Classes/GameFramework/InputDeviceSubsystem.h b/Engine/Source/Runtime/Engine/Classes/GameFramework/InputDeviceSubsystem.h
--- a/Engine/Source/Runtime/Engine/Classes/GameFramework/InputDeviceSubsystem.h
+++ b/Engine/Source/Runtime/Engine/Classes/GameFramework/InputDeviceSubsystem.h
@@ -135,7 +135,7 @@ public:
 	//~ Begin UEngineSubsystem interface
 	ENGINE_API virtual void Initialize(FSubsystemCollectionBase& Collection) override;
 	ENGINE_API virtual void Deinitialize() override;
-	ENGINE_API virtual bool ShouldCreateSubsystem(UObject* Outer) const;
+	ENGINE_API virtual bool ShouldCreateSubsystem(UObject* Outer) const override;
 	//~ End UEngineSubsystem interface
 	
 	//~ Begin FTickableGameObject interface	
@@ -195,6 +195,14 @@ public:
 	UFUNCTION(BlueprintCallable, Category = "Input Devices")
 	ENGINE_API FHardwareDeviceIdentifier GetMostRecentlyUsedHardwareDevice(const FPlatformUserId InUserId) const;
 
+	/**
+	 * Returns the most recently used FInputDeviceId for the given platform user id.
+	 * 
+	 * This will be INPUTDEVICEID_NONE if there is no known device for the given user.
+	 */
+	UFUNCTION(BlueprintCallable, Category = "Input Devices", meta = (ReturnDisplayName = "Latest Device Id"))
+	ENGINE_API FInputDeviceId GetMostRecentlyUsedInputDeviceId(const FPlatformUserId InUserId) const;
+	
 	UFUNCTION(BlueprintCallable, Category = "Input Devices")
 	ENGINE_API FHardwareDeviceIdentifier GetInputDeviceHardwareIdentifier(const FInputDeviceId InputDevice) const;
 
@@ -252,8 +260,25 @@ protected:
 	/** A map of an input device to it's most recent hardware device identifier */
 	TMap<FInputDeviceId, FHardwareDeviceIdentifier> LatestInputDeviceIdentifiers;
 
+	/**
+	 * Data store about the input device that was most recently used by a specific FPlatformUserId
+	 */
+	struct FLatestInputDeviceData
+	{
+		/** The input device ID of the latest device used */
+		FInputDeviceId LatestDeviceId = INPUTDEVICEID_NONE;
+
+		/** The latest hardware device ID for the input device last used by the platform user. */
+		FHardwareDeviceIdentifier HardwareDeviceId = FHardwareDeviceIdentifier::Invalid;
+
+		bool operator==(const FLatestInputDeviceData& Other) const
+		{
+			return LatestDeviceId == Other.LatestDeviceId && HardwareDeviceId == Other.HardwareDeviceId;
+		}
+	};
+	
 	/** A map of platform user's to their most recent hardware device identifier */
-	TMap<FPlatformUserId, FHardwareDeviceIdentifier> LatestUserDeviceIdentifiers;
+	TMap<FPlatformUserId, FLatestInputDeviceData> LatestUserDeviceIdentifiers;
 
 	/** An input processor that is used to determine the most recently used hardware device for each user*/
 	TSharedPtr<class FInputDeviceSubsystemProcessor> InputPreprocessor;
diff --git a/Engine/Source/Runtime/Engine/Classes/GameFramework/InputSettings.h b/Engine/Source/Runtime/Engine/Classes/GameFramework/InputSettings.h
--- a/Engine/Source/Runtime/Engine/Classes/GameFramework/InputSettings.h
+++ b/Engine/Source/Runtime/Engine/Classes/GameFramework/InputSettings.h
@@ -84,7 +84,8 @@ class UInputSettings
 	 * 
 	 * @see UInputDeviceSubsystem
 	 */
-	UPROPERTY(config, EditAnywhere, Category = "Input")
+	UE_DEPRECATED(5.5, "bEnableInputDeviceSubsystem is deprecated, ths input device subsystem will now always be created.")
+	UPROPERTY(config, meta = (DeprecatedProperty, DeprecationMessage = "bEnableInputDeviceSubsystem is deprecated, ths input device subsystem will now always be created."))
 	uint8 bEnableInputDeviceSubsystem:1;
 
 	/**
diff --git a/Engine/Source/Runtime/Engine/Private/GameFramework/InputDeviceSubsystem.cpp b/Engine/Source/Runtime/Engine/Private/GameFramework/InputDeviceSubsystem.cpp
--- a/Engine/Source/Runtime/Engine/Private/GameFramework/InputDeviceSubsystem.cpp
+++ b/Engine/Source/Runtime/Engine/Private/GameFramework/InputDeviceSubsystem.cpp
@@ -222,14 +222,18 @@ bool UInputDeviceSubsystem::ShouldCreateSubsystem(UObject* Outer) const
 		return false;
 	}
 
+	// As of 5.5 this option is disabled. Some input interfaces depend on this subsystem to get info about
+	// the currently connected devices, so we are enforcing it as a requirement from now on.
+PRAGMA_DISABLE_DEPRECATION_WARNINGS
 	// There is a setting to turn off this subsystem entirely.
 	const bool bShouldCreate = GetDefault<UInputSettings>()->bEnableInputDeviceSubsystem;
 	if (!bShouldCreate)
 	{
-		UE_LOG(LogInputDeviceProperties, Log, TEXT("UInputSettings::bEnableInputDeviceSubsystem is false, the Input Device Subsystem will NOT be created!"));
+		UE_LOG(LogInputDeviceProperties, Error, TEXT("UInputSettings::bEnableInputDeviceSubsystem is false, but it is required to be true as of 5.5. Ignoring..."));
 	}
+PRAGMA_ENABLE_DEPRECATION_WARNINGS
 	
-	return bShouldCreate && Super::ShouldCreateSubsystem(Outer);
+	return Super::ShouldCreateSubsystem(Outer);
 }
 
 UWorld* UInputDeviceSubsystem::GetTickableGameObjectWorld() const
@@ -488,14 +492,24 @@ void UInputDeviceSubsystem::RemoveAllDeviceProperties()
 
 FHardwareDeviceIdentifier UInputDeviceSubsystem::GetMostRecentlyUsedHardwareDevice(const FPlatformUserId InUserId) const
 {
-	if (const FHardwareDeviceIdentifier* FoundDevice = LatestUserDeviceIdentifiers.Find(InUserId))
+	if (const FLatestInputDeviceData* FoundDevice = LatestUserDeviceIdentifiers.Find(InUserId))
 	{
-		return *FoundDevice;
+		return FoundDevice->HardwareDeviceId;
 	}
 		
 	return FHardwareDeviceIdentifier::Invalid;
 }
 
+FInputDeviceId UInputDeviceSubsystem::GetMostRecentlyUsedInputDeviceId(const FPlatformUserId InUserId) const
+{
+	if (const FLatestInputDeviceData* FoundDevice = LatestUserDeviceIdentifiers.Find(InUserId))
+	{
+		return FoundDevice->LatestDeviceId;
+	}
+		
+	return INPUTDEVICEID_NONE;
+}
+
 FHardwareDeviceIdentifier UInputDeviceSubsystem::GetInputDeviceHardwareIdentifier(const FInputDeviceId InputDevice) const
 {
 	if (const FHardwareDeviceIdentifier* FoundDevice = LatestInputDeviceIdentifiers.Find(InputDevice))
@@ -511,19 +525,22 @@ void UInputDeviceSubsystem::SetMostRecentlyUsedHardwareDevice(const FInputDevice
 	const FPlatformUserId OwningUserId = IPlatformInputDeviceMapper::Get().GetUserForInputDevice(InDeviceId);
 
 	// If this hardware is the same as what the platform user already has, then there is no need to fire this event
-	if (FHardwareDeviceIdentifier* ExistingDevice = LatestUserDeviceIdentifiers.Find(OwningUserId))
+	if (FLatestInputDeviceData* ExistingDevice = LatestUserDeviceIdentifiers.Find(OwningUserId))
 	{
+		// We always want to update the device ID
+		ExistingDevice->LatestDeviceId = InDeviceId;
+		
 		// If they are exactly the same, do nothing
-		if (InHardwareId == *ExistingDevice)
+		if (InHardwareId == ExistingDevice->HardwareDeviceId)
 		{
 			return;
 		}
 		// If they are the same device but different input API's, just update the data but don't broadcast the delegate
 		// This can happen if there are multiple input interfaces sending messages for the same device (such as a keyboard)
 		// Which your gameplay code doesn't really care about switching which API the input came from
-		else if (InHardwareId.HardwareDeviceIdentifier == ExistingDevice->HardwareDeviceIdentifier)
+		else if (InHardwareId.HardwareDeviceIdentifier == ExistingDevice->HardwareDeviceId.HardwareDeviceIdentifier)
 		{
-			*ExistingDevice = InHardwareId;
+			ExistingDevice->HardwareDeviceId = InHardwareId;
 			return;
 		}
 	}
@@ -532,7 +549,7 @@ void UInputDeviceSubsystem::SetMostRecentlyUsedHardwareDevice(const FInputDevice
 	LatestInputDeviceIdentifiers.Add(InDeviceId, InHardwareId);
 
 	// Keep a map to platform users so that we can easily get their most recent hardware
-	LatestUserDeviceIdentifiers.Add(OwningUserId, InHardwareId);
+	LatestUserDeviceIdentifiers.Add(OwningUserId, { .LatestDeviceId = InDeviceId, .HardwareDeviceId = InHardwareId });
 
 	if (OnInputHardwareDeviceChanged.IsBound())
 	{
diff --git a/Engine/Source/Runtime/Engine/Private/UserInterface/InputSettings.cpp b/Engine/Source/Runtime/Engine/Private/UserInterface/InputSettings.cpp
--- a/Engine/Source/Runtime/Engine/Private/UserInterface/InputSettings.cpp
+++ b/Engine/Source/Runtime/Engine/Private/UserInterface/InputSettings.cpp
@@ -27,7 +27,9 @@ UInputSettings::UInputSettings(const FObjectInitializer& ObjectInitializer)
 	, bEnableLegacyInputScales(true)
 	, bEnableMotionControls(true)
 	, bFilterInputByPlatformUser(false)
+PRAGMA_DISABLE_DEPRECATION_WARNINGS
 	, bEnableInputDeviceSubsystem(true)
+PRAGMA_ENABLE_DEPRECATION_WARNINGS
 	, bShouldFlushPressedKeysOnViewportFocusLost(true)
 	, bEnableDynamicComponentInputBinding(true)
 	, DefaultViewportMouseCaptureMode(EMouseCaptureMode::CapturePermanently_IncludingInitialMouseDown)
@@ -187,6 +189,14 @@ void UInputSettings::PostReloadConfig(FProperty* PropertyThatWasLoaded)
 	PopulateAxisConfigs();
 #endif
 	AddInternationalConsoleKey();
+
+PRAGMA_DISABLE_DEPRECATION_WARNINGS
+	if (!bEnableInputDeviceSubsystem)
+	{
+		bEnableInputDeviceSubsystem = true;
+		UE_LOG(LogInput, Warning, TEXT("UInputSettings::bEnableInputDeviceSubsystem was detected as false from a config, but it is required to be true as of 5.5. Updating the value to be true."));
+	}
+PRAGMA_ENABLE_DEPRECATION_WARNINGS
 }
 
 #if WITH_EDITOR
