Add ENGINE_API tags to attenuation override functions on UAudioComponent

diff --git a/Engine/Source/Runtime/Engine/Classes/Components/AudioComponent.h b/Engine/Source/Runtime/Engine/Classes/Components/AudioComponent.h
--- a/Engine/Source/Runtime/Engine/Classes/Components/AudioComponent.h
+++ b/Engine/Source/Runtime/Engine/Classes/Components/AudioComponent.h
@@ -212,7 +212,7 @@ class UAudioComponent : public USceneComponent, public ISoundParameterController
 	uint8 bOverrideAttenuation:1;
 
 	UFUNCTION(BlueprintSetter)
-	void SetOverrideAttenuation(bool bInOverrideAttenuation);
+	ENGINE_API void SetOverrideAttenuation(bool bInOverrideAttenuation);
 
 	/** Whether or not to override the sound's subtitle priority. */
 	UPROPERTY(EditAnywhere, BlueprintReadWrite, Category = Subtitles, meta = (InlineEditConditionToggle))
@@ -366,14 +366,14 @@ public:
 	TObjectPtr<USoundAttenuation> AttenuationSettings;
 
 	UFUNCTION(BlueprintSetter)
-	void SetAttenuationSettings(USoundAttenuation* InAttenuationSettings);
+	ENGINE_API void SetAttenuationSettings(USoundAttenuation* InAttenuationSettings);
 
 	/** If bOverrideSettings is true, the attenuation properties to use for sounds generated by this component */
 	UPROPERTY(EditAnywhere, BlueprintReadWrite, Category = Attenuation, BlueprintSetter = SetAttenuationOverrides, meta = (EditCondition = "bOverrideAttenuation", DisplayAfter = "bOverrideAttenuation", EditConditionHides))
 	struct FSoundAttenuationSettings AttenuationOverrides;
 
 	UFUNCTION(BlueprintSetter)
-	void SetAttenuationOverrides(const FSoundAttenuationSettings& InAttenuationOverrides);
+	ENGINE_API void SetAttenuationOverrides(const FSoundAttenuationSettings& InAttenuationOverrides);
 
 	/** What sound concurrency to use for sounds generated by this audio component */
 	UPROPERTY()
