Chaos:add per-body control of position/velocity/projection iterations. The iterations can be set on physics section of the detail panel. During solve of rigid bodies, an island group takes maximum of the iterations from all its particles as its iteration count.

diff --git a/Engine/Source/Editor/DetailCustomizations/Private/BodyInstanceCustomization.cpp b/Engine/Source/Editor/DetailCustomizations/Private/BodyInstanceCustomization.cpp
--- a/Engine/Source/Editor/DetailCustomizations/Private/BodyInstanceCustomization.cpp
+++ b/Engine/Source/Editor/DetailCustomizations/Private/BodyInstanceCustomization.cpp
@@ -1396,8 +1396,6 @@ void FBodyInstanceCustomizationHelper::CustomizeDetails(IDetailLayoutBuilder& De
 		BodyInstanceHandler->GetChildHandle(GET_MEMBER_NAME_CHECKED(FBodyInstance, bLockXRotation))->MarkHiddenByCustomization();
 		BodyInstanceHandler->GetChildHandle(GET_MEMBER_NAME_CHECKED(FBodyInstance, bLockYRotation))->MarkHiddenByCustomization();
 		BodyInstanceHandler->GetChildHandle(GET_MEMBER_NAME_CHECKED(FBodyInstance, bLockZRotation))->MarkHiddenByCustomization();
-		BodyInstanceHandler->GetChildHandle(GET_MEMBER_NAME_CHECKED(FBodyInstance, PositionSolverIterationCount))->MarkHiddenByCustomization();
-		BodyInstanceHandler->GetChildHandle(GET_MEMBER_NAME_CHECKED(FBodyInstance, VelocitySolverIterationCount))->MarkHiddenByCustomization();
 		BodyInstanceHandler->GetChildHandle(GET_MEMBER_NAME_CHECKED(FBodyInstance, CustomDOFPlaneNormal))->MarkHiddenByCustomization();
 
 		//Add the rest
diff --git a/Engine/Source/Programs/HeadlessChaos/Private/GeometryCollection/GeometryCollectionTestFramework.cpp b/Engine/Source/Programs/HeadlessChaos/Private/GeometryCollection/GeometryCollectionTestFramework.cpp
--- a/Engine/Source/Programs/HeadlessChaos/Private/GeometryCollection/GeometryCollectionTestFramework.cpp
+++ b/Engine/Source/Programs/HeadlessChaos/Private/GeometryCollection/GeometryCollectionTestFramework.cpp
@@ -177,6 +177,9 @@ namespace GeometryCollectionTest
 			SimulationParams.AngularDamping = 0;
 			SimulationParams.UseCCD = false;
 			SimulationParams.UseMACD = false;
+			SimulationParams.PositionSolverIterations = 8;
+			SimulationParams.VelocitySolverIterations = 1;
+			SimulationParams.ProjectionSolverIterations = 1;
 
 			Chaos::FErrorReporter ErrorReporter;
 			BuildSimulationData(ErrorReporter, *RestCollection.Get(), SimulationParams.Shared);
diff --git a/Engine/Source/Runtime/Engine/Classes/PhysicsEngine/BodyInstance.h b/Engine/Source/Runtime/Engine/Classes/PhysicsEngine/BodyInstance.h
--- a/Engine/Source/Runtime/Engine/Classes/PhysicsEngine/BodyInstance.h
+++ b/Engine/Source/Runtime/Engine/Classes/PhysicsEngine/BodyInstance.h
@@ -321,13 +321,17 @@ struct FBodyInstance : public FBodyInstanceCore
 	int16 InstanceBoneIndex;
 
 	/** [PhysX Only] This physics body's solver iteration count for position. Increasing this will be more CPU intensive, but better stabilized.  */
-	UPROPERTY(EditAnywhere, AdvancedDisplay, BlueprintReadOnly, Category=Physics)
+	UPROPERTY(EditAnywhere, AdvancedDisplay, BlueprintReadOnly, Category=Physics, meta = (ClampMin = "1", UIMin = "1"))
 	uint8 PositionSolverIterationCount;
 
 	/** [PhysX Only] This physics body's solver iteration count for velocity. Increasing this will be more CPU intensive, but better stabilized. */
-	UPROPERTY(EditAnywhere, AdvancedDisplay, BlueprintReadOnly, Category = Physics)
+	UPROPERTY(EditAnywhere, AdvancedDisplay, BlueprintReadOnly, Category = Physics, meta = (ClampMin = "1", UIMin = "1"))
 	uint8 VelocitySolverIterationCount;
 
+	/** [PhysX Only] This physics body's solver iteration count for projection. Increasing this will be more CPU intensive, but better stabilized. */
+	UPROPERTY(EditAnywhere, AdvancedDisplay, BlueprintReadOnly, Category = Physics, meta = (ClampMin = "1", UIMin = "1"))
+	uint8 ProjectionSolverIterationCount;
+
 private:
 	/** Enum indicating what type of object this should be considered as when it moves */
 	UPROPERTY(EditAnywhere, Category=Custom)
@@ -982,9 +986,19 @@ public:
 	 */
 	ENGINE_API void SetUseMACD(bool bInUseMACD);
 
+	ENGINE_API void SetPositionSolverIterationCount(uint8 PositionSolverIterationCountIn);
+
+	ENGINE_API void SetVelocitySolverIterationCount(uint8 VelocitySolverIterationCountIn);
+
+	ENGINE_API void SetProjectionSolverIterationCount(uint8 ProjectionSolverIterationCountIn);
+
 	/** [EXPERIMENTAL] Whether Motion-Aware Collision Detection is enabled */
 	bool GetUseMACD() const { return bUseMACD != 0; }
 
+	uint8 GetPositionSolverIterationCount() const { return PositionSolverIterationCount; }
+	uint8 GetVelocitySolverIterationCount() const { return VelocitySolverIterationCount; }
+	uint8 GetProjectionSolverIterationCount() const { return ProjectionSolverIterationCount; }
+
 	/** Disable/Re-Enable this body in the solver,  when disable, the body won't be part of the simulation ( regardless if it's dynamic or kinematic ) and no collision will occur 
 	* this can be used for performance control situation for example
 	*/
diff --git a/Engine/Source/Runtime/Engine/Private/PhysicsEngine/BodyInstance.cpp b/Engine/Source/Runtime/Engine/Private/PhysicsEngine/BodyInstance.cpp
--- a/Engine/Source/Runtime/Engine/Private/PhysicsEngine/BodyInstance.cpp
+++ b/Engine/Source/Runtime/Engine/Private/PhysicsEngine/BodyInstance.cpp
@@ -330,7 +330,8 @@ FBodyInstance::FBodyInstance()
 	: InstanceBodyIndex(INDEX_NONE)
 	, InstanceBoneIndex(INDEX_NONE)
 	, PositionSolverIterationCount(8)
-	, VelocitySolverIterationCount(1)
+	, VelocitySolverIterationCount(2)
+	, ProjectionSolverIterationCount(1)
 	, ObjectType(ECC_WorldStatic)
 	, MaskFilter(0)
 	, CollisionEnabled(ECollisionEnabled::QueryAndPhysics)
@@ -1279,6 +1280,10 @@ void TInitBodiesHelperBase<InAllocatorType>::CreateActor_AssumesLocked(FBodyInst
 		FPhysicsInterface::SetSmoothEdgeCollisionsEnabled_AssumesLocked(Instance->ActorHandle, Instance->bSmoothEdgeCollisions);
 		FPhysicsInterface::SetInertiaConditioningEnabled_AssumesLocked(Instance->ActorHandle, Instance->IsInertiaConditioningEnabled());
 
+		FPhysicsInterface::SetPositionSolverIterationCount_AssumesLocked(Instance->ActorHandle, Instance->PositionSolverIterationCount);
+		FPhysicsInterface::SetVelocitySolverIterationCount_AssumesLocked(Instance->ActorHandle, Instance->VelocitySolverIterationCount);
+		FPhysicsInterface::SetProjectionSolverIterationCount_AssumesLocked(Instance->ActorHandle, Instance->ProjectionSolverIterationCount);
+
 		// Set sleep event notification
 		FPhysicsInterface::SetSendsSleepNotifies_AssumesLocked(Instance->ActorHandle, Instance->bGenerateWakeEvents);
 	}
@@ -2438,6 +2443,9 @@ void FBodyInstance::UpdateInstanceSimulatePhysics()
 		FPhysicsInterface::SetIsKinematic_AssumesLocked(Actor, !bUseSimulate);
 		FPhysicsInterface::SetCcdEnabled_AssumesLocked(Actor, bUseCCD);
 		FPhysicsInterface::SetMACDEnabled_AssumesLocked(Actor, bUseMACD);
+		FPhysicsInterface::SetPositionSolverIterationCount_AssumesLocked(Actor, PositionSolverIterationCount);
+		FPhysicsInterface::SetVelocitySolverIterationCount_AssumesLocked(Actor, VelocitySolverIterationCount);
+		FPhysicsInterface::SetProjectionSolverIterationCount_AssumesLocked(Actor, ProjectionSolverIterationCount);
 
 		if(bSimulatePhysics)
 		{
@@ -3807,6 +3815,54 @@ void FBodyInstance::SetUseMACD(bool bInUseMACD)
 	}
 }
 
+void FBodyInstance::SetPositionSolverIterationCount(uint8 PositionSolverIterationCountIn)
+{
+	if (PositionSolverIterationCount != PositionSolverIterationCountIn)
+	{
+		PositionSolverIterationCount = PositionSolverIterationCountIn;
+
+		FPhysicsCommand::ExecuteWrite(ActorHandle, [this, PositionSolverIterationCountIn](const FPhysicsActorHandle& Actor)
+			{
+				if (FPhysicsInterface::IsValid(Actor) && FPhysicsInterface::IsRigidBody(Actor))
+				{
+					FPhysicsInterface::SetPositionSolverIterationCount_AssumesLocked(Actor, PositionSolverIterationCount);
+				}
+			});
+	}
+}
+
+void FBodyInstance::SetVelocitySolverIterationCount(uint8 VelocitySolverIterationCountIn)
+{
+	if (VelocitySolverIterationCount != VelocitySolverIterationCountIn)
+	{
+		VelocitySolverIterationCount = VelocitySolverIterationCountIn;
+
+		FPhysicsCommand::ExecuteWrite(ActorHandle, [this, VelocitySolverIterationCountIn](const FPhysicsActorHandle& Actor)
+			{
+				if (FPhysicsInterface::IsValid(Actor) && FPhysicsInterface::IsRigidBody(Actor))
+				{
+					FPhysicsInterface::SetVelocitySolverIterationCount_AssumesLocked(Actor, VelocitySolverIterationCount);
+				}
+			});
+	}
+}
+
+void FBodyInstance::SetProjectionSolverIterationCount(uint8 ProjectionSolverIterationCountIn)
+{
+	if (ProjectionSolverIterationCount != ProjectionSolverIterationCountIn)
+	{
+		ProjectionSolverIterationCount = ProjectionSolverIterationCountIn;
+
+		FPhysicsCommand::ExecuteWrite(ActorHandle, [this, ProjectionSolverIterationCountIn](const FPhysicsActorHandle& Actor)
+			{
+				if (FPhysicsInterface::IsValid(Actor) && FPhysicsInterface::IsRigidBody(Actor))
+				{
+					FPhysicsInterface::SetProjectionSolverIterationCount_AssumesLocked(Actor, ProjectionSolverIterationCount);
+				}
+			});
+	}
+}
+
 void FBodyInstance::SetPhysicsDisabled(bool bSetDisabled)
 {
 	FPhysicsCommand::ExecuteWrite(ActorHandle, [&](const FPhysicsActorHandle& Actor)
@@ -4335,8 +4391,10 @@ void FBodyInstance::InitDynamicProperties_AssumesLocked()
 
 		uint32 PositionIterCount = FMath::Clamp<uint8>(PositionSolverIterationCount, 1, 255);
 		uint32 VelocityIterCount = FMath::Clamp<uint8>(VelocitySolverIterationCount, 1, 255);
+		uint32 ProjectionIterCount = FMath::Clamp<uint8>(ProjectionSolverIterationCount, 1, 255);
 		FPhysicsInterface::SetSolverPositionIterationCount_AssumesLocked(ActorHandle, PositionIterCount);
 		FPhysicsInterface::SetSolverVelocityIterationCount_AssumesLocked(ActorHandle, VelocityIterCount);
+		FPhysicsInterface::SetSolverProjectionIterationCount_AssumesLocked(ActorHandle, ProjectionIterCount);
 
 		CreateDOFLock();
 		if(FPhysicsInterface::IsInScene(ActorHandle) && !IsRigidBodyKinematic_AssumesLocked(ActorHandle))
diff --git a/Engine/Source/Runtime/Engine/Private/SkeletalMeshComponentPhysics.cpp b/Engine/Source/Runtime/Engine/Private/SkeletalMeshComponentPhysics.cpp
--- a/Engine/Source/Runtime/Engine/Private/SkeletalMeshComponentPhysics.cpp
+++ b/Engine/Source/Runtime/Engine/Private/SkeletalMeshComponentPhysics.cpp
@@ -855,6 +855,11 @@ void USkeletalMeshComponent::InstantiatePhysicsAssetBodies_Internal(const UPhysi
 
 			BodyInst->bStartAwake = UseRootBodyIndex >= 0 ? BodyInstance.bStartAwake : true;	//We don't allow customization here. Just use whatever the component is set to
 
+			//Copying code from BodyInstance here, every body instance get same iteration count
+			BodyInst->PositionSolverIterationCount = BodyInstance.PositionSolverIterationCount;
+			BodyInst->VelocitySolverIterationCount = BodyInstance.VelocitySolverIterationCount;
+			BodyInst->ProjectionSolverIterationCount = BodyInstance.ProjectionSolverIterationCount;
+
 			if(BodyIdx == UseRootBodyIndex)
 			{
 				BodyInst->DOFMode = BodyInstance.DOFMode;
diff --git a/Engine/Source/Runtime/Experimental/Chaos/Private/Chaos/ClusterUnionManager.cpp b/Engine/Source/Runtime/Experimental/Chaos/Private/Chaos/ClusterUnionManager.cpp
--- a/Engine/Source/Runtime/Experimental/Chaos/Private/Chaos/ClusterUnionManager.cpp
+++ b/Engine/Source/Runtime/Experimental/Chaos/Private/Chaos/ClusterUnionManager.cpp
@@ -637,6 +637,9 @@ namespace Chaos
 		// Cluster uses MACD is any children require MACD (or it is already enabled on this cluster)
 		bool bIsMACD = Cluster->InternalCluster->MACDEnabled();
 
+		Private::FIterationSettings SolverIterationSettings = Cluster->InternalCluster->IterationSettings();
+
+
 		// Use the minimum sleep multiplier of all member particles
 		FRealSingle MinSleepThresholdMultiplier = TNumericLimits<FRealSingle>::Max();
 
@@ -661,6 +664,8 @@ namespace Chaos
 			bIsMACD |= Particle->MACDEnabled();
 			MinSleepThresholdMultiplier = FMath::Min(MinSleepThresholdMultiplier, Particle->SleepThresholdMultiplier());
 
+			SolverIterationSettings = Private::FIterationSettings::Merge(SolverIterationSettings, Particle->IterationSettings());
+
 			if (!Cluster->ChildProperties.Contains(Particle))
 			{
 				FClusterUnionParticleProperties Properties;
@@ -732,6 +737,7 @@ namespace Chaos
 
 		Cluster->InternalCluster->SetOneWayInteraction(bIsOneWayInteraction);
 		Cluster->InternalCluster->SetMACDEnabled(bIsMACD);
+		Cluster->InternalCluster->SetIterationSettings(SolverIterationSettings);
 
 		if (MinSleepThresholdMultiplier != TNumericLimits<FRealSingle>::Max())
 		{
diff --git a/Engine/Source/Runtime/Experimental/Chaos/Private/Chaos/Island/IslandGroup.cpp b/Engine/Source/Runtime/Experimental/Chaos/Private/Chaos/Island/IslandGroup.cpp
--- a/Engine/Source/Runtime/Experimental/Chaos/Private/Chaos/Island/IslandGroup.cpp
+++ b/Engine/Source/Runtime/Experimental/Chaos/Private/Chaos/Island/IslandGroup.cpp
@@ -92,10 +92,10 @@ namespace Chaos
 		void FPBDIslandConstraintGroupSolver::SetIterationSettings(const FIterationSettings& InDefaultIterations)
 		{
 			Iterations = InDefaultIterations;
-			//for (FPBDIsland* Island : Islands)
-			//{
-			//	Iterations = FIterationSettings::Merge(Iterations, Island->GetIterationSettings());
-			//}
+			for (FPBDIsland* Island : Islands)
+			{
+				Iterations = FIterationSettings::Merge(Iterations, Island->GetIterationSettings());
+			}
 		}
 
 	}	// namespace Private
diff --git a/Engine/Source/Runtime/Experimental/Chaos/Private/Chaos/Island/IslandManager.cpp b/Engine/Source/Runtime/Experimental/Chaos/Private/Chaos/Island/IslandManager.cpp
--- a/Engine/Source/Runtime/Experimental/Chaos/Private/Chaos/Island/IslandManager.cpp
+++ b/Engine/Source/Runtime/Experimental/Chaos/Private/Chaos/Island/IslandManager.cpp
@@ -6,6 +6,7 @@
 #include "Chaos/PBDConstraintContainer.h"
 #include "Chaos/PBDRigidsSOAs.h"
 #include "Chaos/PerParticleGravity.h"
+#include "Chaos/Evolution/IterationSettings.h"
 #include "Chaos/PhysicsMaterialUtilities.h"
 
 #include "ChaosStats.h"
@@ -33,6 +34,7 @@ DECLARE_CYCLE_STAT(TEXT("IslandManager::Merge"), STAT_IslandManager_MergeIslands
 DECLARE_CYCLE_STAT(TEXT("IslandManager::Split"), STAT_IslandManager_SplitIslands, STATGROUP_ChaosIslands);
 DECLARE_CYCLE_STAT(TEXT("IslandManager::Levels"), STAT_IslandManager_AssignLevels, STATGROUP_ChaosIslands);
 DECLARE_CYCLE_STAT(TEXT("IslandManager::Finalize"), STAT_IslandManager_Finalize, STATGROUP_ChaosIslands);
+DECLARE_CYCLE_STAT(TEXT("IslandManager::ComputeIterationSettings"), STAT_IslandManager_ComputeIterationSettings, STATGROUP_ChaosIslands);
 DECLARE_CYCLE_STAT(TEXT("IslandManager::Validate"), STAT_IslandManager_Validate, STATGROUP_ChaosIslands);
 
 namespace Chaos::CVars
@@ -79,6 +81,10 @@ namespace Chaos::CVars
 	/* Cvar to adjust the sleep angular threshold for floating particles */
 	FRealSingle IsolatedParticleSleepAngularThresholdMultiplier = 1.0f;
 	FAutoConsoleVariableRef CVarChaosSolverIsolatedParticleSleepAngularThresholdMultiplier(TEXT("p.Chaos.Solver.Sleep.IsolatedParticle.AngularMultiplier"), IsolatedParticleSleepAngularThresholdMultiplier, TEXT("A multiplier applied to SleepAngularThreshold for floating particles"));
+	
+	/** Cvar to enable/disable computing max iterations if island is dirty */
+	bool bChaosSolverComputeIterationSettings = true;
+	FAutoConsoleVariableRef CVarChaosSolverComputeIterationSettingsEnabled(TEXT("p.Chaos.Solver.ComputeIterationSettings.Enabled"), bChaosSolverComputeIterationSettings, TEXT("Recompute iteration settings every time an island is changed"));
 }
 
 
@@ -1348,6 +1354,7 @@ namespace Chaos::Private
 			{
 				Node->Island = Island;
 				Node->IslandArrayIndex = Island->Nodes.Add(Node);
+				Island->SetIterationSettings(FIterationSettings::Merge(Island->GetIterationSettings(), Node->GetIterationSettings()));
 			}
 
 			Island->Flags.bItemsAdded = true;
@@ -1745,6 +1752,7 @@ namespace Chaos::Private
 				{
 					Node->Island = ParentIsland;
 					Node->IslandArrayIndex = NextIslandArrayIndex++;
+					ParentIsland->SetIterationSettings(FIterationSettings::Merge(ParentIsland->GetIterationSettings(), Node->GetIterationSettings()));
 				}
 			}
 
@@ -2067,6 +2075,23 @@ namespace Chaos::Private
 				Island->SleepCounter = 0;
 			}
 
+			if (Island->Flags.bItemsAdded || Island->Flags.bItemsRemoved)
+			{
+				//Compute max of iteration settings by iterating through all the nodes:
+				SCOPE_CYCLE_COUNTER(STAT_IslandManager_ComputeIterationSettings);
+				if (Chaos::CVars::bChaosSolverComputeIterationSettings)
+				{
+					Island->SetIterationSettings(FIterationSettings(0, 0, 0));
+					for (FPBDIslandParticle* Node : Island->Nodes)
+					{
+						if (Node->Flags.bIsDynamic)
+						{
+							Island->SetIterationSettings(FIterationSettings::Merge(Island->GetIterationSettings(), Node->GetIterationSettings()));
+						}
+					}
+				}
+			}
+
 			Island->Flags.bWasSleeping = Island->Flags.bIsSleeping;
 			Island->Flags.bCheckSleep = false;
 			Island->Flags.bItemsAdded = false;
diff --git a/Engine/Source/Runtime/Experimental/Chaos/Private/PBDRigidsSolver.cpp b/Engine/Source/Runtime/Experimental/Chaos/Private/PBDRigidsSolver.cpp
--- a/Engine/Source/Runtime/Experimental/Chaos/Private/PBDRigidsSolver.cpp
+++ b/Engine/Source/Runtime/Experimental/Chaos/Private/PBDRigidsSolver.cpp
@@ -2740,6 +2740,7 @@ TRACE_COUNTER_SET(ChaosTraceCounter_##Name, Value)
 		Rigid->SetDisabled(DynamicMisc.Disabled());
 		Rigid->SetCollisionConstraintFlags(DynamicMisc.CollisionConstraintFlags());
 		Rigid->SetControlFlags(DynamicMisc.ControlFlags());
+		Rigid->SetIterationSettings(DynamicMisc.IterationSettings());
 
 		GetEvolution()->SetParticleObjectState(Rigid, DynamicMisc.ObjectState());
 		GetEvolution()->SetParticleSleepType(Rigid, DynamicMisc.SleepType());
diff --git a/Engine/Source/Runtime/Experimental/Chaos/Private/PhysicsProxy/GeometryCollectionPhysicsProxy.cpp b/Engine/Source/Runtime/Experimental/Chaos/Private/PhysicsProxy/GeometryCollectionPhysicsProxy.cpp
--- a/Engine/Source/Runtime/Experimental/Chaos/Private/PhysicsProxy/GeometryCollectionPhysicsProxy.cpp
+++ b/Engine/Source/Runtime/Experimental/Chaos/Private/PhysicsProxy/GeometryCollectionPhysicsProxy.cpp
@@ -1604,6 +1604,7 @@ void FGeometryCollectionPhysicsProxy::InitializeBodiesPT(Chaos::FPBDRigidsSolver
 				Handle->SetGravityGroupIndex(Parameters.GravityGroupIndex);
 				Handle->SetCCDEnabled(Parameters.UseCCD);
 				Handle->SetMACDEnabled(Parameters.UseMACD);
+				Handle->SetIterationSettings(Chaos::Private::FIterationSettings(Parameters.PositionSolverIterations, Parameters.VelocitySolverIterations, Parameters.ProjectionSolverIterations));
 				Handle->SetOneWayInteraction(bIsOneWayInteraction);
 				Handle->SetInertiaConditioningEnabled(Parameters.UseInertiaConditioning);
 				Handle->SetLinearEtherDrag(Parameters.LinearDamping);
diff --git a/Engine/Source/Runtime/Experimental/Chaos/Public/Chaos/Island/IslandManager.h b/Engine/Source/Runtime/Experimental/Chaos/Public/Chaos/Island/IslandManager.h
--- a/Engine/Source/Runtime/Experimental/Chaos/Public/Chaos/Island/IslandManager.h
+++ b/Engine/Source/Runtime/Experimental/Chaos/Public/Chaos/Island/IslandManager.h
@@ -6,6 +6,8 @@
 #include "Chaos/Framework/Handles.h"
 #include "Chaos/Island/IslandManagerFwd.h"
 #include "Chaos/ParticleHandleFwd.h"
+#include "Chaos/ParticleHandle.h"
+#include "Chaos/Evolution/IterationSettings.h"
 #include "Chaos/Serializable.h"
 
 namespace Chaos
@@ -51,6 +53,22 @@ namespace Chaos::Private
 		// An integer ID for an island (for debug only). @see FPBDIsland::GetIslandId()
 		CHAOS_API int32 GetIslandId() const;
 
+		/**
+		* Get the iterations settings for this group of islands
+		*/
+		FIterationSettings GetIterationSettings() const
+		{
+			return Particle->CastToRigidParticle()->IterationSettings();
+		}
+
+		/**
+		 * Set the iterations settings for this group of islands
+		*/
+		void SetIterationSettings(const FIterationSettings& InIterations)
+		{
+			Particle->CastToRigidParticle()->SetIterationSettings(InIterations);
+		}
+
 	private:
 		friend class FPBDIslandManager;
 		friend class FPBDIsland;
@@ -289,6 +307,22 @@ namespace Chaos::Private
 		CHAOS_API void Reuse();
 		CHAOS_API void Trash();
 
+		/**
+			* Get the iterations settings for this group of islands
+		*/
+		FIterationSettings GetIterationSettings() const
+		{
+			return Iterations;
+		}
+
+		/**
+		 * Set the iterations settings for this group of islands
+		*/
+		void SetIterationSettings(const FIterationSettings& InIterations)
+		{
+			Iterations = InIterations;
+		}
+
 	private:
 		friend class FPBDIslandManager;
 
@@ -344,6 +378,8 @@ namespace Chaos::Private
 		int32 ResimFrame = INDEX_NONE;
 
 		FFlags Flags;
+
+		FIterationSettings Iterations;
 	};
 
 	/**
diff --git a/Engine/Source/Runtime/Experimental/Chaos/Public/Chaos/ParticleDirtyFlags.h b/Engine/Source/Runtime/Experimental/Chaos/Public/Chaos/ParticleDirtyFlags.h
--- a/Engine/Source/Runtime/Experimental/Chaos/Public/Chaos/ParticleDirtyFlags.h
+++ b/Engine/Source/Runtime/Experimental/Chaos/Public/Chaos/ParticleDirtyFlags.h
@@ -20,6 +20,7 @@
 #include "UObject/UE5MainStreamObjectVersion.h"
 #include "UObject/UE5ReleaseStreamObjectVersion.h"
 #include "Framework/PhysicsProxyBase.h"
+#include "Chaos/Evolution/IterationSettings.h"
 #include "PBDJointConstraintTypes.h"
 #include "PBDSuspensionConstraintTypes.h"
 
@@ -545,6 +546,12 @@ public:
 	bool MACDEnabled() const { return MControlFlags.GetMACDEnabled(); }
 	void SetMACDEnabled(bool bInCCDEnabled) { MControlFlags.SetMACDEnabled(bInCCDEnabled); }
 
+	void SetPositionSolverIterationCount(const int32 PositionSolverIterationCountIn) { MIterationSettings.SetNumPositionIterations(PositionSolverIterationCountIn); }
+	void SetVelocitySolverIterationCount(const int32 VelocitySolverIterationCountIn) { MIterationSettings.SetNumVelocityIterations(VelocitySolverIterationCountIn); }
+	void SetProjectionSolverIterationCount(const int32 ProjectionSolverIterationCountIn) { MIterationSettings.SetNumProjectionIterations(ProjectionSolverIterationCountIn); }
+	Private::FIterationSettings IterationSettings() const { return MIterationSettings; }
+	void SetIterationSettings(const Private::FIterationSettings& SolverIterationSettings) { MIterationSettings = SolverIterationSettings; }
+
 	bool Disabled() const { return bDisabled; }
 	void SetDisabled(bool bInDisabled) { bDisabled = bInDisabled; }
 
@@ -587,6 +594,8 @@ private:
 	uint32 MCollisionConstraintFlag = 0;
 	FRigidParticleControlFlags MControlFlags;
 
+	Chaos::Private::FIterationSettings MIterationSettings;
+
 	bool bDisabled;
 };
 
diff --git a/Engine/Source/Runtime/Experimental/Chaos/Public/Chaos/ParticleHandle.h b/Engine/Source/Runtime/Experimental/Chaos/Public/Chaos/ParticleHandle.h
--- a/Engine/Source/Runtime/Experimental/Chaos/Public/Chaos/ParticleHandle.h
+++ b/Engine/Source/Runtime/Experimental/Chaos/Public/Chaos/ParticleHandle.h
@@ -128,6 +128,7 @@ void PBDRigidParticleDefaultConstruct(FConcrete& Concrete, const FPBDRigidPartic
 	Concrete.SetGravityGroupIndex(0);
 	Concrete.SetCCDEnabled(Params.bCCDEnabled);
 	Concrete.SetMACDEnabled(false);
+	Concrete.SetIterationSettings(Private::FIterationSettings(INDEX_NONE, INDEX_NONE, INDEX_NONE));
 	Concrete.SetDisabled(Params.bDisabled);
 	Concrete.SetSleepType(ESleepType::MaterialSleep);
 }
@@ -1306,6 +1307,31 @@ public:
 		PBDRigidParticles->ControlFlags(ParticleIdx).SetMACDEnabled(bEnabled);
 	}
 
+	inline Private::FIterationSettings IterationSettings() const
+	{
+		return PBDRigidParticles->ParticleIterationCounts(ParticleIdx);
+	}
+
+	inline void SetIterationSettings(const Private::FIterationSettings& IterationSettingsIn)
+	{
+		PBDRigidParticles->ParticleIterationCounts(ParticleIdx) = IterationSettingsIn;
+	}
+
+	inline void SetPositionSolverIterations(const int32 PositionIterationsIn)
+	{
+		PBDRigidParticles->ParticleIterationCounts(ParticleIdx).SetNumPositionIterations(PositionIterationsIn);
+	}
+
+	inline void SetVelocitySolverIterations(const int32 VelocityIterationsIn)
+	{
+		PBDRigidParticles->ParticleIterationCounts(ParticleIdx).SetNumVelocityIterations(VelocityIterationsIn);
+	}
+
+	inline void SetProjectionSolverIterations(const int32 ProjectionIterationsIn)
+	{
+		PBDRigidParticles->ParticleIterationCounts(ParticleIdx).SetNumProjectionIterations(ProjectionIterationsIn);
+	}
+
 	inline bool OneWayInteraction() const
 	{ 
 		return ControlFlags().GetOneWayInteractionEnabled();
@@ -1844,6 +1870,17 @@ public:
 		return false;
 	}
 
+	Private::FIterationSettings IterationSettings() const
+	{
+		if (MHandle->CastToRigidParticle())
+		{
+			return MHandle->CastToRigidParticle()->IterationSettings();
+		}
+
+		return Private::FIterationSettings(0, 0, 0);
+	
+	}
+
 	bool HasCollisionConstraintFlag(const ECollisionConstraintFlags Flag)  const
 	{
 		if (MHandle->CastToRigidParticle())
@@ -3397,6 +3434,28 @@ public:
 		MMiscData.Modify(true, MDirtyFlags, Proxy, [bInEnabled](auto& Data) { Data.SetMACDEnabled(bInEnabled); });
 	}
 
+	Private::FIterationSettings IterationSettings() const { return MMiscData.Read().IterationSettings(); }
+
+	void SetIterationSettings(const Private::FIterationSettings& SolverIterationSettingsIn)
+	{
+		MMiscData.Modify(true, MDirtyFlags, Proxy, [SolverIterationSettingsIn](auto& Data) { Data.SetIterationSettings(SolverIterationSettingsIn); });
+	}
+
+	void SetPositionSolverIterations(const uint32 PositionSolverIterationCount)
+	{
+		MMiscData.Modify(true, MDirtyFlags, Proxy, [PositionSolverIterationCount](auto& Data) { Data.SetPositionSolverIterationCount((int32)PositionSolverIterationCount); });
+	}
+
+	void SetVelocitySolverIterations(const uint32 VelocitySolverIterationCount)
+	{
+		MMiscData.Modify(true, MDirtyFlags, Proxy, [VelocitySolverIterationCount](auto& Data) { Data.SetVelocitySolverIterationCount((int32)VelocitySolverIterationCount); });
+	}
+
+	void SetProjectionSolverIterations(const uint32 ProjectionSolverIterationCount)
+	{
+		MMiscData.Modify(true, MDirtyFlags, Proxy, [ProjectionSolverIterationCount](auto& Data) { Data.SetProjectionSolverIterationCount((int32)ProjectionSolverIterationCount); });
+	}
+
 	bool InertiaConditioningEnabled() const { return MMiscData.Read().InertiaConditioningEnabled(); }
 	void SetInertiaConditioningEnabled(bool bInEnabled)
 	{
diff --git a/Engine/Source/Runtime/Experimental/Chaos/Public/Chaos/RigidParticles.h b/Engine/Source/Runtime/Experimental/Chaos/Public/Chaos/RigidParticles.h
--- a/Engine/Source/Runtime/Experimental/Chaos/Public/Chaos/RigidParticles.h
+++ b/Engine/Source/Runtime/Experimental/Chaos/Public/Chaos/RigidParticles.h
@@ -13,6 +13,7 @@
 #include "Chaos/Rotation.h"
 #include "Chaos/RigidParticleControlFlags.h"
 #include "HAL/LowLevelMemTracker.h"
+#include "Chaos/Evolution/IterationSettings.h"
 #include "UObject/FortniteMainBranchObjectVersion.h"
 
 namespace Chaos
@@ -113,6 +114,7 @@ public:
 		, MSleepType(MoveTemp(Other.MSleepType))
 		, MSleepCounter(MoveTemp(Other.MSleepCounter))
 		, MDisableCounter(MoveTemp(Other.MDisableCounter))
+		, MParticleIterationCounts(MoveTemp(Other.MParticleIterationCounts))
 	{
 		RegisterArrays();
 	}
@@ -144,6 +146,7 @@ public:
 		TArrayCollection::AddArray(&MSleepType);
 		TArrayCollection::AddArray(&MSleepCounter);
 		TArrayCollection::AddArray(&MDisableCounter);
+		TArrayCollection::AddArray(&MParticleIterationCounts);
 
 	}
 
@@ -247,6 +250,9 @@ public:
 	FORCEINLINE int8 DisableCounter(const int32 Index) const { return MDisableCounter[Index]; }
 	FORCEINLINE int8& DisableCounter(const int32 Index) { return MDisableCounter[Index]; }
 
+	FORCEINLINE Private::FIterationSettings ParticleIterationCounts(const int32 Index) const { return MParticleIterationCounts[Index]; }
+	FORCEINLINE Private::FIterationSettings& ParticleIterationCounts(const int32 Index) { return MParticleIterationCounts[Index]; }
+
 	// @todo(chaos): This data should be marshalled via the proxies like everything else. There is probably a particle recycling bug here.
 	FORCEINLINE TArray<TSleepData<T, d>>& GetSleepData() { return MSleepData; }
 	FORCEINLINE	void AddSleepData(TGeometryParticleHandle<T, d>* Particle, bool Sleeping)
@@ -423,6 +429,7 @@ private:
 	TArrayCollectionArray<ESleepType> MSleepType;
 	TArrayCollectionArray<int8> MSleepCounter;
 	TArrayCollectionArray<int8> MDisableCounter;
+	TArrayCollectionArray<Private::FIterationSettings> MParticleIterationCounts;
 
 	TArray<TSleepData<T, d>> MSleepData;
 	FRWLock SleepDataLock;
diff --git a/Engine/Source/Runtime/Experimental/Chaos/Public/GeometryCollection/GeometryCollectionSimulationCoreTypes.h b/Engine/Source/Runtime/Experimental/Chaos/Public/GeometryCollection/GeometryCollectionSimulationCoreTypes.h
--- a/Engine/Source/Runtime/Experimental/Chaos/Public/GeometryCollection/GeometryCollectionSimulationCoreTypes.h
+++ b/Engine/Source/Runtime/Experimental/Chaos/Public/GeometryCollection/GeometryCollectionSimulationCoreTypes.h
@@ -217,6 +217,9 @@ PRAGMA_ENABLE_DEPRECATION_WARNINGS
 		, bOptimizeConvexes(true)
 		, bUseSimplicialsWhenAvailable(false)
 		, bUseDamagePropagation(false)
+		, PositionSolverIterations(8)
+		, VelocitySolverIterations(1)
+		, ProjectionSolverIterations(1)
 		, BreakDamagePropagationFactor(1.0f)
 		, ShockDamagePropagationFactor(0.0f)
 		, LinearDamping(0.01f)
@@ -291,6 +294,9 @@ PRAGMA_ENABLE_DEPRECATION_WARNINGS
 		, bOptimizeConvexes(Other.bOptimizeConvexes)
 		, bUseSimplicialsWhenAvailable(Other.bUseSimplicialsWhenAvailable)
 		, bUseDamagePropagation(Other.bUseDamagePropagation)
+		, PositionSolverIterations(Other.PositionSolverIterations)
+		, VelocitySolverIterations(Other.VelocitySolverIterations)
+		, ProjectionSolverIterations(Other.ProjectionSolverIterations)
 		, BreakDamagePropagationFactor(Other.BreakDamagePropagationFactor)
 		, ShockDamagePropagationFactor(Other.ShockDamagePropagationFactor)
 		, LinearDamping(Other.LinearDamping)
@@ -394,13 +400,17 @@ PRAGMA_ENABLE_DEPRECATION_WARNINGS
 	bool UseInertiaConditioning : 1;
 	bool UseCCD : 1;
 	bool UseMACD : 1;
-
 	bool bEnableStrainOnCollision : 1;
 	bool bUseStaticMeshCollisionForTraces : 1;
 	bool bOptimizeConvexes : 1;
 	bool bUseSimplicialsWhenAvailable : 1;
 
 	bool bUseDamagePropagation : 1;
+
+	uint8 PositionSolverIterations;
+	uint8 VelocitySolverIterations;
+	uint8 ProjectionSolverIterations;
+
 	float BreakDamagePropagationFactor;
 	float ShockDamagePropagationFactor;
 
diff --git a/Engine/Source/Runtime/Experimental/Chaos/Public/PhysicsProxy/SingleParticlePhysicsProxy.h b/Engine/Source/Runtime/Experimental/Chaos/Public/PhysicsProxy/SingleParticlePhysicsProxy.h
--- a/Engine/Source/Runtime/Experimental/Chaos/Public/PhysicsProxy/SingleParticlePhysicsProxy.h
+++ b/Engine/Source/Runtime/Experimental/Chaos/Public/PhysicsProxy/SingleParticlePhysicsProxy.h
@@ -507,6 +507,39 @@ public:
 			});
 	}
 
+	void SetPositionSolverIterations(const int32 PositionSolverIterationsIn)
+	{
+		Write([PositionSolverIterationsIn](auto* Particle)
+			{
+				if (auto Rigid = Particle->CastToRigidParticle())
+				{
+					return Rigid->SetPositionSolverIterations(PositionSolverIterationsIn);
+				}
+			});
+	}
+
+	void SetVelocitySolverIterations(const int32 VelocitySolverIterationsIn)
+	{
+		Write([VelocitySolverIterationsIn](auto* Particle)
+			{
+				if (auto Rigid = Particle->CastToRigidParticle())
+				{
+					return Rigid->SetVelocitySolverIterations(VelocitySolverIterationsIn);
+				}
+			});
+	}
+
+	void SetProjectionSolverIterations(const int32 ProjectionSolverIterationsIn)
+	{
+		Write([ProjectionSolverIterationsIn](auto* Particle)
+			{
+				if (auto Rigid = Particle->CastToRigidParticle())
+				{
+					return Rigid->SetProjectionSolverIterations(ProjectionSolverIterationsIn);
+				}
+			});
+	}
+
 	bool OneWayInteraction() const
 	{
 		return Read([](auto* Particle)
@@ -1329,6 +1362,46 @@ public:
 		return false;
 	}
 
+	Private::FIterationSettings IterationSettings() const
+	{
+		VerifyContext();
+		if (const TPBDRigidParticle<FReal, 3>*Rigid = GetParticle_LowLevel()->CastToRigidParticle())
+		{
+			return Rigid->IterationSettings();
+		}
+
+		return Private::FIterationSettings(0, 0, 0);
+	}
+
+	void SetPositionSolverIterationCount(uint32 PositionSolverIterationCountIn)
+	{
+		VerifyContext();
+		if (TPBDRigidParticle<FReal, 3>*Rigid = GetParticle_LowLevel()->CastToRigidParticle())
+		{
+			Rigid->SetPositionSolverIterations(PositionSolverIterationCountIn);
+		}
+	}
+
+
+	void SetVelocitySolverIterationCount(uint32 VelocitySolverIterationCountIn)
+	{
+		VerifyContext();
+		if (TPBDRigidParticle<FReal, 3>*Rigid = GetParticle_LowLevel()->CastToRigidParticle())
+		{
+			Rigid->SetVelocitySolverIterations(VelocitySolverIterationCountIn);
+		}
+	}
+
+
+	void SetProjectionSolverIterationCount(uint32 ProjectionSolverIterationCountIn)
+	{
+		VerifyContext();
+		if (TPBDRigidParticle<FReal, 3>*Rigid = GetParticle_LowLevel()->CastToRigidParticle())
+		{
+			Rigid->SetProjectionSolverIterations(ProjectionSolverIterationCountIn);
+		}
+	}
+
 	void SetMaxLinearSpeedSq(FReal InNewSpeed)
 	{
 		VerifyContext();
diff --git a/Engine/Source/Runtime/Experimental/ChaosVisualDebugger/Private/DataWrappers/ChaosVDParticleDataWrapper.cpp b/Engine/Source/Runtime/Experimental/ChaosVisualDebugger/Private/DataWrappers/ChaosVDParticleDataWrapper.cpp
--- a/Engine/Source/Runtime/Experimental/ChaosVisualDebugger/Private/DataWrappers/ChaosVDParticleDataWrapper.cpp
+++ b/Engine/Source/Runtime/Experimental/ChaosVisualDebugger/Private/DataWrappers/ChaosVDParticleDataWrapper.cpp
@@ -27,6 +27,9 @@ bool FChaosVDFRigidParticleControlFlags::Serialize(FArchive& Ar)
 	Ar << bInertiaConditioningEnabled;
 	Ar << GravityGroupIndex;
 	Ar << bMACDEnabled;
+	Ar << PositionSolverIterationCount;
+	Ar << VelocitySolverIterationCount;
+	Ar << ProjectionSolverIterationCount;
 
 	return !Ar.IsError();
 }
diff --git a/Engine/Source/Runtime/Experimental/ChaosVisualDebugger/Public/DataWrappers/ChaosVDParticleDataWrapper.h b/Engine/Source/Runtime/Experimental/ChaosVisualDebugger/Public/DataWrappers/ChaosVDParticleDataWrapper.h
--- a/Engine/Source/Runtime/Experimental/ChaosVisualDebugger/Public/DataWrappers/ChaosVDParticleDataWrapper.h
+++ b/Engine/Source/Runtime/Experimental/ChaosVisualDebugger/Public/DataWrappers/ChaosVDParticleDataWrapper.h
@@ -85,7 +85,10 @@ struct FChaosVDFRigidParticleControlFlags : public FChaosVDWrapperDataBase
 		  bOneWayInteractionEnabled(false),
 		  bInertiaConditioningEnabled(false), 
 		  GravityGroupIndex(0),
-		  bMACDEnabled(false)
+		  bMACDEnabled(false), 
+		  PositionSolverIterationCount(8),
+		  VelocitySolverIterationCount(1),
+		  ProjectionSolverIterationCount(1)
 	{
 	}
 
@@ -116,6 +119,12 @@ struct FChaosVDFRigidParticleControlFlags : public FChaosVDWrapperDataBase
 	int32 GravityGroupIndex;
 	UPROPERTY(VisibleAnywhere, Category = "Particle Control Flags")
 	bool bMACDEnabled;
+	UPROPERTY(VisibleAnywhere, Category = "Particle Control Flags")
+	uint8 PositionSolverIterationCount;
+	UPROPERTY(VisibleAnywhere, Category = "Particle Control Flags")
+	uint8 VelocitySolverIterationCount;
+	UPROPERTY(VisibleAnywhere, Category = "Particle Control Flags")
+	uint8 ProjectionSolverIterationCount;
 };
 
 inline FArchive& operator<<(FArchive& Ar, FChaosVDFRigidParticleControlFlags& Data)
diff --git a/Engine/Source/Runtime/Experimental/GeometryCollectionEngine/Private/GeometryCollection/GeometryCollectionComponent.cpp b/Engine/Source/Runtime/Experimental/GeometryCollectionEngine/Private/GeometryCollection/GeometryCollectionComponent.cpp
--- a/Engine/Source/Runtime/Experimental/GeometryCollectionEngine/Private/GeometryCollection/GeometryCollectionComponent.cpp
+++ b/Engine/Source/Runtime/Experimental/GeometryCollectionEngine/Private/GeometryCollection/GeometryCollectionComponent.cpp
@@ -3942,6 +3942,9 @@ PRAGMA_ENABLE_DEPRECATION_WARNINGS
 		SimulationParameters.UseInertiaConditioning = BodyInstance.IsInertiaConditioningEnabled();
 		SimulationParameters.UseCCD = BodyInstance.bUseCCD;
 		SimulationParameters.UseMACD = BodyInstance.GetUseMACD();
+		SimulationParameters.PositionSolverIterations = BodyInstance.GetPositionSolverIterationCount();
+		SimulationParameters.VelocitySolverIterations = BodyInstance.GetVelocitySolverIterationCount();
+		SimulationParameters.ProjectionSolverIterations = BodyInstance.GetProjectionSolverIterationCount();
 		SimulationParameters.LinearDamping = BodyInstance.LinearDamping;
 		SimulationParameters.AngularDamping = BodyInstance.AngularDamping;
 		SimulationParameters.InitialOverlapDepenetrationVelocity = BodyInstance.GetMaxDepenetrationVelocity();
diff --git a/Engine/Source/Runtime/PhysicsCore/Private/ChaosEngineInterface.cpp b/Engine/Source/Runtime/PhysicsCore/Private/ChaosEngineInterface.cpp
--- a/Engine/Source/Runtime/PhysicsCore/Private/ChaosEngineInterface.cpp
+++ b/Engine/Source/Runtime/PhysicsCore/Private/ChaosEngineInterface.cpp
@@ -548,6 +548,22 @@ void FChaosEngineInterface::SetMACDEnabled_AssumesLocked(const FPhysicsActorHand
 	InActorReference->GetGameThreadAPI().SetMACDEnabled(bIsMACDEnabled);
 }
 
+void FChaosEngineInterface::SetPositionSolverIterationCount_AssumesLocked(const FPhysicsActorHandle& InActorReference, uint8 PositionSolverIterationCount)
+{
+	InActorReference->GetGameThreadAPI().SetPositionSolverIterationCount(PositionSolverIterationCount);
+}
+
+void FChaosEngineInterface::SetVelocitySolverIterationCount_AssumesLocked(const FPhysicsActorHandle& InActorReference, uint8 VelocitySolverIterationCount)
+{
+	InActorReference->GetGameThreadAPI().SetVelocitySolverIterationCount(VelocitySolverIterationCount);
+}
+
+void FChaosEngineInterface::SetProjectionSolverIterationCount_AssumesLocked(const FPhysicsActorHandle& InActorReference, uint8 ProjectionSolverIterationCount)
+{
+	InActorReference->GetGameThreadAPI().SetProjectionSolverIterationCount(ProjectionSolverIterationCount);
+}
+
+
 void FChaosEngineInterface::SetIgnoreAnalyticCollisions_AssumesLocked(const FPhysicsActorHandle& InActorReference,bool bIgnoreAnalyticCollisions)
 {
 	InActorReference->GetGameThreadAPI().SetIgnoreAnalyticCollisions(bIgnoreAnalyticCollisions);
@@ -1266,28 +1282,25 @@ void FChaosEngineInterface::SetStabilizationEnergyThreshold_AssumesLocked(const
 	// #todo : Implement
 }
 
-uint32 FChaosEngineInterface::GetSolverPositionIterationCount_AssumesLocked(const FPhysicsActorHandle& InHandle)
-{
-	// #todo : Implement
-	return 0;
-}
-
 void FChaosEngineInterface::SetSolverPositionIterationCount_AssumesLocked(const FPhysicsActorHandle& InHandle,uint32 InSolverIterationCount)
 {
-	// #todo : Implement
+	Chaos::FRigidBodyHandle_External& Body_External = InHandle->GetGameThreadAPI();
+	Body_External.SetPositionSolverIterationCount(InSolverIterationCount);
 }
 
-uint32 FChaosEngineInterface::GetSolverVelocityIterationCount_AssumesLocked(const FPhysicsActorHandle& InHandle)
+void FChaosEngineInterface::SetSolverVelocityIterationCount_AssumesLocked(const FPhysicsActorHandle& InHandle, uint32 InSolverIterationCount)
 {
-	// #todo : Implement
-	return 0;
+	Chaos::FRigidBodyHandle_External& Body_External = InHandle->GetGameThreadAPI();
+	Body_External.SetVelocitySolverIterationCount(InSolverIterationCount);
 }
 
-void FChaosEngineInterface::SetSolverVelocityIterationCount_AssumesLocked(const FPhysicsActorHandle& InHandle,uint32 InSolverIterationCount)
+void FChaosEngineInterface::SetSolverProjectionIterationCount_AssumesLocked(const FPhysicsActorHandle& InHandle, uint32 InSolverIterationCount)
 {
-	// #todo : Implement
+	Chaos::FRigidBodyHandle_External& Body_External = InHandle->GetGameThreadAPI();
+	Body_External.SetProjectionSolverIterationCount(InSolverIterationCount);
 }
 
+
 float FChaosEngineInterface::GetWakeCounter_AssumesLocked(const FPhysicsActorHandle& InHandle)
 {
 	// #todo : Implement
diff --git a/Engine/Source/Runtime/PhysicsCore/Public/Chaos/ChaosEngineInterface.h b/Engine/Source/Runtime/PhysicsCore/Public/Chaos/ChaosEngineInterface.h
--- a/Engine/Source/Runtime/PhysicsCore/Public/Chaos/ChaosEngineInterface.h
+++ b/Engine/Source/Runtime/PhysicsCore/Public/Chaos/ChaosEngineInterface.h
@@ -8,6 +8,7 @@
 #include "PhysicsInterfaceDeclaresCore.h"
 #include "PhysicsInterfaceWrapperShared.h"
 #include "PhysicsInterfaceTypesCore.h"
+#include "Chaos/Evolution/IterationSettings.h"
 #include "PhysicsProxy/SingleParticlePhysicsProxyFwd.h"
 #include "ChaosEngineInterface.generated.h"
 
@@ -407,6 +408,9 @@ public:
 	static PHYSICSCORE_API void SetIsKinematic_AssumesLocked(const FPhysicsActorHandle& InActorReference,bool bIsKinematic);
 	static PHYSICSCORE_API void SetCcdEnabled_AssumesLocked(const FPhysicsActorHandle& InActorReference,bool bIsCcdEnabled);
 	static PHYSICSCORE_API void SetMACDEnabled_AssumesLocked(const FPhysicsActorHandle& InActorReference, bool bIsMACDEnabled);
+	static PHYSICSCORE_API void SetPositionSolverIterationCount_AssumesLocked(const FPhysicsActorHandle& InActorReference, uint8 PositionSolverIterationCount);
+	static PHYSICSCORE_API void SetVelocitySolverIterationCount_AssumesLocked(const FPhysicsActorHandle& InActorReference, uint8 VelocitySolverIterationCount);
+	static PHYSICSCORE_API void SetProjectionSolverIterationCount_AssumesLocked(const FPhysicsActorHandle& InActorReference, uint8 ProjectionSolverIterationCount);
 	static PHYSICSCORE_API void SetIgnoreAnalyticCollisions_AssumesLocked(const FPhysicsActorHandle& InActorReference,bool bIsCcdEnabled);
 
 	static PHYSICSCORE_API FTransform GetGlobalPose_AssumesLocked(const FPhysicsActorHandle& InActorReference);
@@ -480,10 +484,9 @@ public:
 
 	static PHYSICSCORE_API float GetStabilizationEnergyThreshold_AssumesLocked(const FPhysicsActorHandle& InHandle);
 	static PHYSICSCORE_API void SetStabilizationEnergyThreshold_AssumesLocked(const FPhysicsActorHandle& InHandle,float InThreshold);
-	static PHYSICSCORE_API uint32 GetSolverPositionIterationCount_AssumesLocked(const FPhysicsActorHandle& InHandle);
 	static PHYSICSCORE_API void SetSolverPositionIterationCount_AssumesLocked(const FPhysicsActorHandle& InHandle,uint32 InSolverIterationCount);
-	static PHYSICSCORE_API uint32 GetSolverVelocityIterationCount_AssumesLocked(const FPhysicsActorHandle& InHandle);
 	static PHYSICSCORE_API void SetSolverVelocityIterationCount_AssumesLocked(const FPhysicsActorHandle& InHandle,uint32 InSolverIterationCount);
+	static PHYSICSCORE_API void SetSolverProjectionIterationCount_AssumesLocked(const FPhysicsActorHandle& InHandle, uint32 InSolverIterationCount);
 	static PHYSICSCORE_API float GetWakeCounter_AssumesLocked(const FPhysicsActorHandle& InHandle);
 	static PHYSICSCORE_API void SetWakeCounter_AssumesLocked(const FPhysicsActorHandle& InHandle,float InWakeCounter);
 
