Added function to check whether a notify comes from a source that is blending out

diff --git a/Engine/Source/Runtime/Engine/Classes/Animation/AnimationAsset.h b/Engine/Source/Runtime/Engine/Classes/Animation/AnimationAsset.h
--- a/Engine/Source/Runtime/Engine/Classes/Animation/AnimationAsset.h
+++ b/Engine/Source/Runtime/Engine/Classes/Animation/AnimationAsset.h
@@ -404,7 +404,8 @@ struct FAnimTickRecord
 	bool bRequestedInertialization = false;
 	bool bOverridePositionWhenJoiningSyncGroupAsLeader = false;
 	bool bIsExclusiveLeader = false;
-	
+	bool bActiveContext = true;
+
 	const UMirrorDataTable* MirrorDataTable = nullptr;
 
 	TSharedPtr<TArray<TUniquePtr<const UE::Anim::IAnimNotifyEventContextDataInterface>>> ContextData;
diff --git a/Engine/Source/Runtime/Engine/Private/Animation/AnimNotifyLibrary.cpp b/Engine/Source/Runtime/Engine/Private/Animation/AnimNotifyLibrary.cpp
--- a/Engine/Source/Runtime/Engine/Private/Animation/AnimNotifyLibrary.cpp
+++ b/Engine/Source/Runtime/Engine/Private/Animation/AnimNotifyLibrary.cpp
@@ -60,3 +60,8 @@ float UAnimNotifyLibrary::GetCurrentAnimationNotifyStateTimeRatio(const FAnimNot
 
 	return 0.0f;
 }
+
+bool UAnimNotifyLibrary::IsBlendingOut(const FAnimNotifyEventReference& EventReference)
+{
+	return !EventReference.IsActiveContext();
+}
\ No newline at end of file
diff --git a/Engine/Source/Runtime/Engine/Private/Animation/AnimNotifyQueue.cpp b/Engine/Source/Runtime/Engine/Private/Animation/AnimNotifyQueue.cpp
--- a/Engine/Source/Runtime/Engine/Private/Animation/AnimNotifyQueue.cpp
+++ b/Engine/Source/Runtime/Engine/Private/Animation/AnimNotifyQueue.cpp
@@ -21,6 +21,7 @@ void FAnimNotifyEventReference::GatherTickRecordData(const FAnimTickRecord& InTi
 {
 	ContextData = InTickRecord.ContextData;
 	CurrentAnimTime = InTickRecord.TimeAccumulator != nullptr ? *InTickRecord.TimeAccumulator : 0.0f;
+	bActiveContext = InTickRecord.bActiveContext;
 }
 
 bool FAnimNotifyQueue::PassesFiltering(const FAnimNotifyEvent* Notify) const
diff --git a/Engine/Source/Runtime/Engine/Private/AnimationAsset.cpp b/Engine/Source/Runtime/Engine/Private/AnimationAsset.cpp
--- a/Engine/Source/Runtime/Engine/Private/AnimationAsset.cpp
+++ b/Engine/Source/Runtime/Engine/Private/AnimationAsset.cpp
@@ -285,6 +285,7 @@ FAnimTickRecord::FAnimTickRecord(UPoseAsset* InPoseAsset, float InFinalBlendWeig
 
 void FAnimTickRecord::GatherContextData(const FAnimationUpdateContext& InContext)
 {
+	bActiveContext = InContext.IsActive();
 	if(InContext.GetSharedContext())
 	{
 		TArray<TUniquePtr<const UE::Anim::IAnimNotifyEventContextDataInterface>> NewContextData;
diff --git a/Engine/Source/Runtime/Engine/Public/Animation/AnimNotifyLibrary.h b/Engine/Source/Runtime/Engine/Public/Animation/AnimNotifyLibrary.h
--- a/Engine/Source/Runtime/Engine/Public/Animation/AnimNotifyLibrary.h
+++ b/Engine/Source/Runtime/Engine/Public/Animation/AnimNotifyLibrary.h
@@ -63,4 +63,13 @@ public:
 	 */
 	UFUNCTION(BlueprintPure, Category = "Utilities|Animation|Notifies" , meta = (ScriptMethod))
 	static ENGINE_API float GetCurrentAnimationNotifyStateTimeRatio(const FAnimNotifyEventReference& EventReference);
+
+	/**
+	 * Gets whether this notify comes from a source that is blending out.
+	 *
+	 * @param EventReference		The event to inspect
+	 * @return  Whether this notify comes from a source that is blending out.
+	 */
+	UFUNCTION(BlueprintPure, Category = "Utilities|Animation|Notifies" , meta = (ScriptMethod))
+	static ENGINE_API bool IsBlendingOut(const FAnimNotifyEventReference& EventReference);
 };
diff --git a/Engine/Source/Runtime/Engine/Public/Animation/AnimNotifyQueue.h b/Engine/Source/Runtime/Engine/Public/Animation/AnimNotifyQueue.h
--- a/Engine/Source/Runtime/Engine/Public/Animation/AnimNotifyQueue.h
+++ b/Engine/Source/Runtime/Engine/Public/Animation/AnimNotifyQueue.h
@@ -108,6 +108,12 @@ struct FAnimNotifyEventReference
 		return CurrentAnimTime;
 	}
 
+	// Whether the context this notify was fired from is active or not (active == not blending out).
+	bool IsActiveContext() const
+	{
+		return bActiveContext;
+	}
+
 private:
 	// Context data gleaned from the tick record
 	TSharedPtr<TArray<TUniquePtr<const UE::Anim::IAnimNotifyEventContextDataInterface>>> ContextData;
@@ -123,6 +129,9 @@ private:
 
 	// The recorded time from the tick record that this notify event was fired at
 	float CurrentAnimTime = 0.0f;
+
+	// Whether the context this notify was fired from is active or not (active == not blending out).
+	bool bActiveContext = false;
 };
 
 USTRUCT()
