Add null check for debug material.
The material could be null if it is badly configured, or if GEditor is false, which would cause it not to be loaded in UEngine::InitializeObjectReferences().
Note that the case where this actually crashes in UE5.4 is already fixed by the refactor to mesh paint debug rendering. But adding the null check is still correct, and might fix some other edge case that we haven't encountered.

diff --git a/Engine/Source/Runtime/Engine/Private/MeshPaintVisualize.cpp b/Engine/Source/Runtime/Engine/Private/MeshPaintVisualize.cpp
--- a/Engine/Source/Runtime/Engine/Private/MeshPaintVisualize.cpp
+++ b/Engine/Source/Runtime/Engine/Private/MeshPaintVisualize.cpp
@@ -129,7 +129,7 @@ namespace MeshPaintVisualize
 		}
 
 #if WITH_EDITORONLY_DATA
-		if (GPaintMode == EMeshPaintVisualizePaintMode::TextureAsset && GTextureAsset.IsValid())
+		if (GPaintMode == EMeshPaintVisualizePaintMode::TextureAsset && GTextureAsset.IsValid() && GEngine->TexturePaintingMaskMaterial)
 		{
 			FColoredTexturedMaterialRenderProxy* TextureColorVisualizationMaterialInstance = new FColoredTexturedMaterialRenderProxy(
 				GEngine->TexturePaintingMaskMaterial->GetRenderProxy(),
