Fix gpu crash that can occur when mesh painting with uncompressed textures.
It was possible that the virtual texture adapter would try and RHICopyTexture between RGB16 and RGB8 formats.
Detect that case and use compute shader copy instead.

diff --git a/Engine/Source/Runtime/Engine/Private/VT/VirtualTextureAdapter.cpp b/Engine/Source/Runtime/Engine/Private/VT/VirtualTextureAdapter.cpp
--- a/Engine/Source/Runtime/Engine/Private/VT/VirtualTextureAdapter.cpp
+++ b/Engine/Source/Runtime/Engine/Private/VT/VirtualTextureAdapter.cpp
@@ -101,9 +101,13 @@ namespace VirtualTextureAdapter
 	{
 		FGlobalShaderMap* GlobalShaderMap = GetGlobalShaderMap(GMaxRHIFeatureLevel);
 
-		// Need to do a compute shader copy when source UV extends outside texture (because of tile borders). This is to get correct clamping behavior.
-		// For other cases we can just rely on the final RHICopyTexture (which can also handle direct copy of compressed formats).
-		const bool bUseCopyStep = InLevel == 0 && (InUVRange.Min.X < 0 || InUVRange.Min.Y < 0 || InUVRange.Max.X > 1 || InUVRange.Max.Y > 1);
+		// Need a compute shader copy on the first mip level when:
+		// * Source/Dest formats are incompatible, or
+		// * Source UV extends outside texture (because of tile borders). Then we need correct clamping behavior.
+		// For other cases we can just rely on the final RHICopyTexture (which can also handle the direct copy of compressed formats).
+		const bool bCanUseRHICopyTexture = InSourceFormat == InDestFormat || IsBlockCompressedFormat(InDestFormat);
+		const bool bCopyRequiresClamping = InUVRange.Min.X < 0 || InUVRange.Min.Y < 0 || InUVRange.Max.X > 1 || InUVRange.Max.Y > 1;
+		const bool bUseCopyStep = InLevel == 0 && (!bCanUseRHICopyTexture || bCopyRequiresClamping);
 		// Need downsampling steps if we are copying to higher mip level.
 		// Note that if source texture has mips then we should have passed in the correct SRV and modified InLevel to avoid the downsampling.
 		const bool bUseDownsampleStep = InLevel > 0 && !bUseCopyStep;
