[Material Editor] introduced FloatToUInt and UIntToFloat material expressions.

<

diff --git a/Engine/Source/Runtime/Engine/Private/Materials/HLSLMaterialTranslator.cpp b/Engine/Source/Runtime/Engine/Private/Materials/HLSLMaterialTranslator.cpp
--- a/Engine/Source/Runtime/Engine/Private/Materials/HLSLMaterialTranslator.cpp
+++ b/Engine/Source/Runtime/Engine/Private/Materials/HLSLMaterialTranslator.cpp
@@ -4853,6 +4853,12 @@ int32 FHLSLMaterialTranslator::ForceCast(int32 Code, EMaterialValueType DestType
 			return AddInlinedCodeChunk(DestType, *FiniteCode);
 		}
 	}
+	else if (IsNumericType(SourceType) && IsNumericType(DestType))
+	{
+		EMaterialCastFlags CastFlags = GetForceCastFlags(ForceCastFlags) | EMaterialCastFlags::AllowInteger;
+		FString FiniteCode = CastValue(GetParameterCode(Code), SourceType, DestType, CastFlags);
+		return AddInlinedCodeChunk(DestType, *FiniteCode);
+	}
 	else if ((SourceType & MCT_StaticBool) && IsFloatNumericType(DestType))
 	{
 		FString StaticBoolToFloat = GetParameterCode(Code).Equals("true") ? "1.0f" : "0.0f";
diff --git a/Engine/Source/Runtime/Engine/Private/Materials/MaterialExpressions.cpp b/Engine/Source/Runtime/Engine/Private/Materials/MaterialExpressions.cpp
--- a/Engine/Source/Runtime/Engine/Private/Materials/MaterialExpressions.cpp
+++ b/Engine/Source/Runtime/Engine/Private/Materials/MaterialExpressions.cpp
@@ -105,6 +105,7 @@
 #include "Materials/MaterialExpressionDataDrivenShaderPlatformInfoSwitch.h"
 #include "Materials/MaterialExpressionRequiredSamplersSwitch.h"
 #include "Materials/MaterialExpressionFloor.h"	
+#include "Materials/MaterialExpressionFloatToUInt.h"	
 #include "Materials/MaterialExpressionFmod.h"
 #include "Materials/MaterialExpressionFontSample.h"
 #include "Materials/MaterialExpressionFontSampleParameter.h"
@@ -9048,6 +9049,136 @@ uint32 UMaterialExpressionMaterialAttributeLayers::GetInputType(int32 InputIndex
 }
 #endif // WITH_EDITOR
 
+
+// -----
+
+UMaterialExpressionFloatToUInt::UMaterialExpressionFloatToUInt(const FObjectInitializer& ObjectInitializer)
+: Super(ObjectInitializer)
+{
+	// Structure to hold one-time initialization
+	struct FConstructorStatics
+	{
+		FText NAME_Math;
+		FConstructorStatics()
+		: NAME_Math(LOCTEXT( "Math", "Math" ))
+		{
+		}
+	};
+	static FConstructorStatics ConstructorStatics;
+
+	#if WITH_EDITORONLY_DATA
+	MenuCategories.Add(ConstructorStatics.NAME_Math);
+	#endif
+}
+
+#if WITH_EDITOR
+
+int32 UMaterialExpressionFloatToUInt::Compile(class FMaterialCompiler* Compiler, int32 OutputIndex)
+{
+	if (!Input.GetTracedInput().Expression)
+	{
+		return Compiler->Errorf(TEXT("Missing FloatToUInt input"));
+	}
+
+	int32 Value = Input.Compile(Compiler);
+	switch (Mode)
+	{
+		case EFloatToIntMode::Truncate: Value = Compiler->Truncate(Value); break;
+		case EFloatToIntMode::Floor: Value = Compiler->Floor(Value); break;
+		case EFloatToIntMode::Round: Value = Compiler->Round(Value); break;
+		case EFloatToIntMode::Ceil: Value = Compiler->Ceil(Value); break;
+		default: check(false);
+	}
+	
+	EMaterialValueType Type = Compiler->GetParameterType(Value);
+	int NumComponents = GetNumComponents(Type);
+	if (NumComponents <= 0 || NumComponents > 4)
+	{
+		return Compiler->Errorf(TEXT("Input FloatToUInt is not a scalar or vector"));
+	}
+
+	static const EMaterialValueType UIntTypes[] = { MCT_UInt1, MCT_UInt2, MCT_UInt3, MCT_UInt4 };
+	return Compiler->ForceCast(Value, UIntTypes[NumComponents - 1]);
+}
+
+void UMaterialExpressionFloatToUInt::GetCaption(TArray<FString>& OutCaptions) const
+{
+	OutCaptions.Add(TEXT("FloatToUInt"));
+}
+
+uint32 UMaterialExpressionFloatToUInt::GetInputType(int32 InputIndex)
+{
+	return MCT_Float;
+}
+
+uint32 UMaterialExpressionFloatToUInt::GetOutputType(int32 OutputIndex)
+{
+	return MCT_UInt;
+}
+
+
+#endif // WITH_EDITOR
+
+
+// -----
+
+UMaterialExpressionUIntToFloat::UMaterialExpressionUIntToFloat(const FObjectInitializer& ObjectInitializer)
+: Super(ObjectInitializer)
+{
+	// Structure to hold one-time initialization
+	struct FConstructorStatics
+	{
+		FText NAME_Math;
+		FConstructorStatics()
+		: NAME_Math(LOCTEXT( "Math", "Math" ))
+		{
+		}
+	};
+	static FConstructorStatics ConstructorStatics;
+
+	#if WITH_EDITORONLY_DATA
+	MenuCategories.Add(ConstructorStatics.NAME_Math);
+	#endif
+}
+
+#if WITH_EDITOR
+int32 UMaterialExpressionUIntToFloat::Compile(class FMaterialCompiler* Compiler, int32 OutputIndex)
+{
+ 	if (!Input.GetTracedInput().Expression)
+	{
+		return Compiler->Errorf(TEXT("Missing UIntToFloat input"));
+	}
+
+	int32 Value = Input.Compile(Compiler);
+	EMaterialValueType Type = Compiler->GetParameterType(Value);
+	int NumComponents = GetNumComponents(Type);
+	if (NumComponents <= 0 || NumComponents > 4)
+	{
+		return Compiler->Errorf(TEXT("Input FloatToUInt is not a scalar or vector"));
+	}
+
+	static const EMaterialValueType FloatTypes[] = { MCT_Float1, MCT_Float2, MCT_Float3, MCT_Float4 };
+	return Compiler->ForceCast(Value, FloatTypes[NumComponents - 1]);
+}
+
+void UMaterialExpressionUIntToFloat::GetCaption(TArray<FString>& OutCaptions) const
+{
+	OutCaptions.Add(TEXT("UIntToFloat"));
+}
+
+uint32 UMaterialExpressionUIntToFloat::GetInputType(int32 InputIndex)
+{
+	return MCT_UInt;
+}
+
+uint32 UMaterialExpressionUIntToFloat::GetOutputType(int32 OutputIndex)
+{
+	return MCT_Float;
+}
+
+#endif // WITH_EDITOR
+
+
 // -----
 
 UMaterialExpressionFloor::UMaterialExpressionFloor(const FObjectInitializer& ObjectInitializer)
diff --git a/Engine/Source/Runtime/Engine/Public/Materials/MaterialExpressionFloatToUInt.h b/Engine/Source/Runtime/Engine/Public/Materials/MaterialExpressionFloatToUInt.h
new file mode 100644
--- /dev/null
+++ b/Engine/Source/Runtime/Engine/Public/Materials/MaterialExpressionFloatToUInt.h
@@ -0,0 +1,57 @@
+// Copyright Epic Games, Inc. All Rights Reserved.
+
+#pragma once
+
+#include "CoreMinimal.h"
+#include "UObject/ObjectMacros.h"
+#include "MaterialExpressionIO.h"
+#include "Materials/MaterialExpression.h"
+#include "MaterialExpressionFloatToUInt.generated.h"
+
+UENUM()
+enum class EFloatToIntMode : uint8
+{
+	Truncate UMETA(DisplayName="Truncate"),
+	Floor UMETA(DisplayName="Floor"),
+	Round UMETA(DisplayName="Round"),
+	Ceil UMETA(DisplayName="Ceil")
+};
+
+UCLASS(MinimalAPI, collapsecategories, hidecategories=Object)
+class UMaterialExpressionFloatToUInt : public UMaterialExpression
+{
+	GENERATED_UCLASS_BODY()
+
+	UPROPERTY()
+	FExpressionInput Input;
+
+	UPROPERTY(EditAnywhere, Category=FloatToInt)
+	EFloatToIntMode  Mode = EFloatToIntMode::Truncate;
+
+	//~ Begin UMaterialExpression Interface
+#if WITH_EDITOR
+	virtual int32 Compile(class FMaterialCompiler* Compiler, int32 OutputIndex) override;
+	virtual void GetCaption(TArray<FString>& OutCaptions) const override;
+	virtual uint32 GetInputType(int32 InputIndex) override;
+	virtual uint32 GetOutputType(int32 OutputIndex) override;
+#endif
+	//~ End UMaterialExpression Interface
+};
+
+UCLASS(MinimalAPI, collapsecategories, hidecategories=Object)
+class UMaterialExpressionUIntToFloat : public UMaterialExpression
+{
+	GENERATED_UCLASS_BODY()
+
+	UPROPERTY()
+	FExpressionInput Input;
+
+	//~ Begin UMaterialExpression Interface
+#if WITH_EDITOR
+	virtual int32 Compile(class FMaterialCompiler* Compiler, int32 OutputIndex) override;
+	virtual void GetCaption(TArray<FString>& OutCaptions) const override;
+	virtual uint32 GetInputType(int32 InputIndex) override;
+	virtual uint32 GetOutputType(int32 OutputIndex) override;
+#endif
+	//~ End UMaterialExpression Interface
+};
