PR #11518: Add default data validation to Data Table rows, this will not do anything unless IsDataValid is overridden

diff --git a/Engine/Source/Runtime/Engine/Classes/Engine/DataTable.h b/Engine/Source/Runtime/Engine/Classes/Engine/DataTable.h
--- a/Engine/Source/Runtime/Engine/Classes/Engine/DataTable.h
+++ b/Engine/Source/Runtime/Engine/Classes/Engine/DataTable.h
@@ -58,6 +58,16 @@ struct FTableRowBase
 	 * @param InRowName						The name of the row we're performing fix-up on
 	 */
 	virtual void OnDataTableChanged(const UDataTable* InDataTable, const FName InRowName) {}
+
+#if WITH_EDITOR
+	/** Generic function to validate data table row.
+	 *
+	 * @param	Context	the context holding validation warnings/errors.
+	 * @return	Valid if this table row has data validation rules set up for it and the data for this object is valid.
+	 *			Returns Invalid if it does not pass the rules. Returns NotValidated if no rules are set for this object.
+	 */
+	virtual EDataValidationResult IsDataValid(FDataValidationContext& Context) const { return EDataValidationResult::NotValidated; }
+#endif // WITH_EDITOR
 };
 
 
@@ -126,6 +136,7 @@ public:
 	
 #if WITH_EDITOR
 	ENGINE_API virtual void PostEditChangeProperty(FPropertyChangedEvent& PropertyChangedEvent) override;
+	ENGINE_API virtual EDataValidationResult IsDataValid(FDataValidationContext& Context) const override;
 protected:
 	ENGINE_API virtual void ThreadedPostLoadAssetRegistryTagsOverride(FPostLoadAssetRegistryTagsContext& Context) const override;
 public:
diff --git a/Engine/Source/Runtime/Engine/Private/DataTable.cpp b/Engine/Source/Runtime/Engine/Private/DataTable.cpp
--- a/Engine/Source/Runtime/Engine/Private/DataTable.cpp
+++ b/Engine/Source/Runtime/Engine/Private/DataTable.cpp
@@ -396,6 +396,22 @@ void UDataTable::ThreadedPostLoadAssetRegistryTagsOverride(FPostLoadAssetRegistr
 		}
 	}
 }
+
+EDataValidationResult UDataTable::IsDataValid(FDataValidationContext& Context) const
+{
+	EDataValidationResult Result = Super::IsDataValid(Context);
+
+	if (RowStruct && RowStruct->IsChildOf(FTableRowBase::StaticStruct()))
+	{
+		for (const TPair<FName, uint8*>& TableRowPair : RowMap)
+		{
+			const FTableRowBase* CurRow = reinterpret_cast<FTableRowBase*>(TableRowPair.Value);
+			Result = CombineDataValidationResults(Result, CurRow->IsDataValid(Context));
+		}
+	}
+	
+	return Result;
+}
 #endif // WITH_EDITOR
 
 UScriptStruct& UDataTable::GetEmptyUsingStruct() const
