- Fix Asan issue, the scene proxy data should not be accessed like this as it might be swapped out on the render thread

#rnx
#tests Running Asan build

diff --git a/Engine/Source/Runtime/Engine/Private/Particles/ParticleComponents.cpp b/Engine/Source/Runtime/Engine/Private/Particles/ParticleComponents.cpp
--- a/Engine/Source/Runtime/Engine/Private/Particles/ParticleComponents.cpp
+++ b/Engine/Source/Runtime/Engine/Private/Particles/ParticleComponents.cpp
@@ -8419,6 +8419,8 @@ uint32 UParticleSystemComponent::GetApproxMemoryUsage()const
 		}
 	}
 
+	// This is buggy we are peeking into the scene proxy data and a command might be in flight to update the dynamic data
+#if 0
 	FParticleSystemSceneProxy* PSysSceneProxy = (FParticleSystemSceneProxy*)SceneProxy;
 	if (PSysSceneProxy != NULL)
 	{
@@ -8443,6 +8445,7 @@ uint32 UParticleSystemComponent::GetApproxMemoryUsage()const
 		#endif
 		}
 	}
+#endif
 
 	return MemUsage;
 }
