Fix potential data loss of changed properties on placed components of a derived blueprint class that does not modify the component (which causes it to be marked as unnecessary in the ICH).
Right now during reinstancing it shares ICH templates with the REINST class but not the ICH itself, which means it uses the "unnecessary" template during reinstancing.
This has been a problem for years but was made worse in 5.3 by CL 31455096 which renamed these unnecessary components, which caused the name check in FActorComponentInstanceSourceInfo::MatchesComponent to fail.
This change causes it to ignore the invalid templates, but a follow up change is required to address the core issue.

diff --git a/Engine/Source/Runtime/Engine/Private/BlueprintGeneratedClass.cpp b/Engine/Source/Runtime/Engine/Private/BlueprintGeneratedClass.cpp
--- a/Engine/Source/Runtime/Engine/Private/BlueprintGeneratedClass.cpp
+++ b/Engine/Source/Runtime/Engine/Private/BlueprintGeneratedClass.cpp
@@ -1519,6 +1519,12 @@ UObject* UBlueprintGeneratedClass::FindArchetype(const UClass* ArchetypeClass, c
 				{
 					Archetype = ICH->GetOverridenComponentTemplate(ComponentKey);
 
+					// Exclude archetypes that were previously destroyed by ValidateTemplates to work around issues with REINST classes sharing templates
+					if (Archetype && Archetype->HasAnyFlags(RF_Transient))
+					{
+						Archetype = nullptr;
+					}
+
 					if (GEventDrivenLoaderEnabled && EVENT_DRIVEN_ASYNC_LOAD_ACTIVE_AT_RUNTIME)
 					{
 						if (Archetype && Archetype->HasAnyFlags(RF_NeedLoad))
