Texture Collection material fixes
- Validate that the TextureCollection is valid before attempting to use it.
- We don't support Samplers from Texture Collection objects yet, so throw an error saying to change the Sampler Source field from "From texture source"
- Allow Texture Collection parameters to Custom nodes.
- The expressions "Ceil", "Floor", "Round", and "Truncate" all should say they support Uint outputs.

diff --git a/Engine/Source/Runtime/Engine/Private/Materials/HLSLMaterialTranslator.cpp b/Engine/Source/Runtime/Engine/Private/Materials/HLSLMaterialTranslator.cpp
--- a/Engine/Source/Runtime/Engine/Private/Materials/HLSLMaterialTranslator.cpp
+++ b/Engine/Source/Runtime/Engine/Private/Materials/HLSLMaterialTranslator.cpp
@@ -14980,6 +14980,10 @@ int32 FHLSLMaterialTranslator::CustomExpression( class UMaterialExpressionCustom
 				InputParamDecl += InputNameStr;
 				InputParamDecl += TEXT("Sampler ");
 				break;
+			case MCT_TextureCollection:
+				InputParamDecl += TEXT("FResourceCollection ");
+				InputParamDecl += InputNameStr;
+				break;
 			default:
 				return Errorf(TEXT("Bad type %s for %s input %s"), DescribeType(GetParameterType(CompiledInputs[i])), *Custom->Description, *InputNameStr);
 				break;
diff --git a/Engine/Source/Runtime/Engine/Private/Materials/MaterialExpressions.cpp b/Engine/Source/Runtime/Engine/Private/Materials/MaterialExpressions.cpp
--- a/Engine/Source/Runtime/Engine/Private/Materials/MaterialExpressions.cpp
+++ b/Engine/Source/Runtime/Engine/Private/Materials/MaterialExpressions.cpp
@@ -2845,6 +2845,11 @@ int32 UMaterialExpressionTextureSample::Compile(class FMaterialCompiler* Compile
 			return CoordinateIndex;
 		};
 
+		if ((TextureType & MCT_TextureCollection) != 0 && SamplerSource == SSM_FromTextureAsset)
+		{
+			return CompilerError(Compiler, TEXT("Texture Collections do not provide a sampler, please choose something other than 'From texture asset'"));
+		}
+
 		if (TextureType & (MCT_TextureCollection | MCT_TextureMeshPaint))
 		{
 			// There's no UTexture object to get here
@@ -9211,6 +9216,11 @@ int32 UMaterialExpressionFloor::Compile(class FMaterialCompiler* Compiler, int32
 	return Compiler->Floor(Input.Compile(Compiler));
 }
 
+uint32 UMaterialExpressionFloor::GetOutputType(int32 OutputIndex)
+{
+	return MCT_Float | MCT_UInt;
+}
+
 void UMaterialExpressionFloor::GetCaption(TArray<FString>& OutCaptions) const
 {
 	OutCaptions.Add(TEXT("Floor"));
@@ -9246,6 +9256,10 @@ int32 UMaterialExpressionCeil::Compile(class FMaterialCompiler* Compiler, int32
 	return Compiler->Ceil(Input.Compile(Compiler));
 }
 
+uint32 UMaterialExpressionCeil::GetOutputType(int32 OutputIndex)
+{
+	return MCT_Float | MCT_UInt;
+}
 
 void UMaterialExpressionCeil::GetCaption(TArray<FString>& OutCaptions) const
 {
@@ -9285,6 +9299,11 @@ int32 UMaterialExpressionRound::Compile(class FMaterialCompiler* Compiler, int32
 	return Compiler->Round(Input.Compile(Compiler));
 }
 
+uint32 UMaterialExpressionRound::GetOutputType(int32 OutputIndex)
+{
+	return MCT_Float | MCT_UInt;
+}
+
 void UMaterialExpressionRound::GetCaption(TArray<FString>& OutCaptions) const
 {
 	OutCaptions.Add(TEXT("Round"));
@@ -9328,6 +9347,11 @@ int32 UMaterialExpressionTruncate::Compile(class FMaterialCompiler* Compiler, in
 	return Compiler->Truncate(Input.Compile(Compiler));
 }
 
+uint32 UMaterialExpressionTruncate::GetOutputType(int32 OutputIndex)
+{
+	return MCT_Float | MCT_UInt;
+}
+
 void UMaterialExpressionTruncate::GetCaption(TArray<FString>& OutCaptions) const
 {
 	OutCaptions.Add(TEXT("Truncate"));
@@ -11128,6 +11152,12 @@ UMaterialExpressionTextureCollectionParameter::UMaterialExpressionTextureCollect
 #if WITH_EDITOR
 int32 UMaterialExpressionTextureCollectionParameter::Compile(FMaterialCompiler* Compiler, int32 OutputIndex)
 {
+	FString ErrorMessage;
+	if (!TextureCollectionIsValid(TextureCollection, ErrorMessage))
+	{
+		return CompilerError(Compiler, *ErrorMessage);
+	}
+
 	const int32 TextureCollectionCodeIndex = Compiler->TextureCollectionParameter(ParameterName, TextureCollection);
 
 	if (OutputIndex == 1)
@@ -11206,6 +11236,17 @@ bool UMaterialExpressionTextureCollectionParameter::SetParameterValue(const FNam
 	return false;
 }
 
+bool UMaterialExpressionTextureCollectionParameter::TextureCollectionIsValid(UTextureCollection* InTextureCollection, FString& OutMessage)
+{
+	if (!InTextureCollection)
+	{
+		OutMessage = TEXT("Requires valid texture collection");
+		return false;
+	}
+
+	return true;
+}
+
 bool UMaterialExpressionTextureCollectionParameter::SetParameterValue(const FName& InParameterName, UTextureCollection* InValue, EMaterialExpressionSetParameterValueFlags Flags)
 {
 	if (InParameterName == ParameterName)
diff --git a/Engine/Source/Runtime/Engine/Public/Materials/MaterialExpressionCeil.h b/Engine/Source/Runtime/Engine/Public/Materials/MaterialExpressionCeil.h
--- a/Engine/Source/Runtime/Engine/Public/Materials/MaterialExpressionCeil.h
+++ b/Engine/Source/Runtime/Engine/Public/Materials/MaterialExpressionCeil.h
@@ -21,6 +21,7 @@ class UMaterialExpressionCeil : public UMaterialExpression
 	//~ Begin UMaterialExpression Interface
 #if WITH_EDITOR
 	virtual int32 Compile(class FMaterialCompiler* Compiler, int32 OutputIndex) override;
+	virtual uint32 GetOutputType(int32 OutputIndex) override;
 	virtual void GetCaption(TArray<FString>& OutCaptions) const override;
 	virtual bool GenerateHLSLExpression(FMaterialHLSLGenerator& Generator, UE::HLSLTree::FScope& Scope, int32 OutputIndex, UE::HLSLTree::FExpression const*& OutExpression) const override;
 #endif
diff --git a/Engine/Source/Runtime/Engine/Public/Materials/MaterialExpressionFloor.h b/Engine/Source/Runtime/Engine/Public/Materials/MaterialExpressionFloor.h
--- a/Engine/Source/Runtime/Engine/Public/Materials/MaterialExpressionFloor.h
+++ b/Engine/Source/Runtime/Engine/Public/Materials/MaterialExpressionFloor.h
@@ -21,6 +21,7 @@ class UMaterialExpressionFloor : public UMaterialExpression
 	//~ Begin UMaterialExpression Interface
 #if WITH_EDITOR
 	virtual int32 Compile(class FMaterialCompiler* Compiler, int32 OutputIndex) override;
+	virtual uint32 GetOutputType(int32 OutputIndex) override;
 	virtual void GetCaption(TArray<FString>& OutCaptions) const override;
 	virtual bool GenerateHLSLExpression(FMaterialHLSLGenerator& Generator, UE::HLSLTree::FScope& Scope, int32 OutputIndex, UE::HLSLTree::FExpression const*& OutExpression) const override;
 #endif
diff --git a/Engine/Source/Runtime/Engine/Public/Materials/MaterialExpressionRound.h b/Engine/Source/Runtime/Engine/Public/Materials/MaterialExpressionRound.h
--- a/Engine/Source/Runtime/Engine/Public/Materials/MaterialExpressionRound.h
+++ b/Engine/Source/Runtime/Engine/Public/Materials/MaterialExpressionRound.h
@@ -19,6 +19,7 @@ class UMaterialExpressionRound : public UMaterialExpression
 	//~ Begin UMaterialExpression Interface
 #if WITH_EDITOR
 	virtual int32 Compile(class FMaterialCompiler* Compiler, int32 OutputIndex) override;
+	virtual uint32 GetOutputType(int32 OutputIndex) override;
 	virtual void GetCaption(TArray<FString>& OutCaptions) const override;
 	virtual void GetExpressionToolTip(TArray<FString>& OutToolTip) override;
 	virtual bool GenerateHLSLExpression(FMaterialHLSLGenerator& Generator, UE::HLSLTree::FScope& Scope, int32 OutputIndex, UE::HLSLTree::FExpression const*& OutExpression) const override;
diff --git a/Engine/Source/Runtime/Engine/Public/Materials/MaterialExpressionTextureCollectionParameter.h b/Engine/Source/Runtime/Engine/Public/Materials/MaterialExpressionTextureCollectionParameter.h
--- a/Engine/Source/Runtime/Engine/Public/Materials/MaterialExpressionTextureCollectionParameter.h
+++ b/Engine/Source/Runtime/Engine/Public/Materials/MaterialExpressionTextureCollectionParameter.h
@@ -49,6 +49,7 @@ class UMaterialExpressionTextureCollectionParameter : public UMaterialExpression
 	//~ End UMaterialExpression Interface
 
 #if WITH_EDITOR
+	bool TextureCollectionIsValid(UTextureCollection* InTextureCollection, FString& OutMessage);
 	bool SetParameterValue(const FName& InParameterName, UTextureCollection* InValue, EMaterialExpressionSetParameterValueFlags Flags = EMaterialExpressionSetParameterValueFlags::None);
 #endif
 };
diff --git a/Engine/Source/Runtime/Engine/Public/Materials/MaterialExpressionTruncate.h b/Engine/Source/Runtime/Engine/Public/Materials/MaterialExpressionTruncate.h
--- a/Engine/Source/Runtime/Engine/Public/Materials/MaterialExpressionTruncate.h
+++ b/Engine/Source/Runtime/Engine/Public/Materials/MaterialExpressionTruncate.h
@@ -19,6 +19,7 @@ class UMaterialExpressionTruncate : public UMaterialExpression
 	//~ Begin UMaterialExpression Interface
 #if WITH_EDITOR
 	virtual int32 Compile(class FMaterialCompiler* Compiler, int32 OutputIndex) override;
+	virtual uint32 GetOutputType(int32 OutputIndex) override;
 	virtual void GetCaption(TArray<FString>& OutCaptions) const override;
 	virtual void GetExpressionToolTip(TArray<FString>& OutToolTip) override;
 	virtual bool GenerateHLSLExpression(FMaterialHLSLGenerator& Generator, UE::HLSLTree::FScope& Scope, int32 OutputIndex, UE::HLSLTree::FExpression const*& OutExpression) const override;
