Add HasMeshPaintTexture() to material interface.
Allows fast test for whether a material is using the MeshPaintTextureObject node.
This will be useful to speed up some MeshPaintTool operations such as detecting whether materials should be used for seam painting masks.

diff --git a/Engine/Source/Runtime/Engine/Private/Materials/Material.cpp b/Engine/Source/Runtime/Engine/Private/Materials/Material.cpp
--- a/Engine/Source/Runtime/Engine/Private/Materials/Material.cpp
+++ b/Engine/Source/Runtime/Engine/Private/Materials/Material.cpp
@@ -6965,6 +6965,11 @@ bool UMaterial::WritesToRuntimeVirtualTexture() const
 	return GetCachedExpressionData().bHasRuntimeVirtualTextureOutput;
 }
 
+bool UMaterial::HasMeshPaintTexture() const
+{
+	return GetCachedExpressionData().bHasMeshPaintTexture;
+}
+
 bool UMaterial::HasVertexInterpolator() const
 {
 	return GetCachedExpressionData().bHasVertexInterpolator;
diff --git a/Engine/Source/Runtime/Engine/Private/Materials/MaterialCachedData.cpp b/Engine/Source/Runtime/Engine/Private/Materials/MaterialCachedData.cpp
--- a/Engine/Source/Runtime/Engine/Private/Materials/MaterialCachedData.cpp
+++ b/Engine/Source/Runtime/Engine/Private/Materials/MaterialCachedData.cpp
@@ -29,6 +29,7 @@
 #include "Materials/MaterialExpressionStaticBool.h"
 #include "Materials/MaterialExpressionLandscapeGrassOutput.h"
 #include "Materials/MaterialExpressionUserSceneTexture.h"
+#include "Materials/MaterialExpressionMeshPaintTextureObject.h"
 #include "Materials/MaterialFunctionInterface.h"
 #include "Materials/MaterialParameterCollection.h"
 #include "MaterialHLSLTree.h"
@@ -61,6 +62,7 @@ FMaterialCachedExpressionData::FMaterialCachedExpressionData()
 	, bHasPerInstanceRandom(false)
 	, bHasVertexInterpolator(false)
 	, bHasCustomizedUVs(false)
+	, bHasMeshPaintTexture(false)
 {
 	QualityLevelsUsed.AddDefaulted(EMaterialQualityLevel::Num);
 #if WITH_EDITORONLY_DATA
@@ -476,6 +478,10 @@ void FMaterialCachedExpressionData::UpdateForExpressions(const FMaterialCachedEx
 		{
 			bHasVertexInterpolator = true;
 		}
+		else if (Expression->IsA(UMaterialExpressionMeshPaintTextureObject::StaticClass()))
+		{
+			bHasMeshPaintTexture = true;
+		}
 		else if (UMaterialExpressionMaterialAttributeLayers* LayersExpression = Cast<UMaterialExpressionMaterialAttributeLayers>(Expression))
 		{
 			checkf(Association == GlobalParameter, TEXT("UMaterialExpressionMaterialAttributeLayers can't be nested"));
diff --git a/Engine/Source/Runtime/Engine/Private/Materials/MaterialInstance.cpp b/Engine/Source/Runtime/Engine/Private/Materials/MaterialInstance.cpp
--- a/Engine/Source/Runtime/Engine/Private/Materials/MaterialInstance.cpp
+++ b/Engine/Source/Runtime/Engine/Private/Materials/MaterialInstance.cpp
@@ -2020,6 +2020,11 @@ bool UMaterialInstance::WritesToRuntimeVirtualTexture() const
 	return Parent ? Parent->WritesToRuntimeVirtualTexture() : false;
 }
 
+bool UMaterialInstance::HasMeshPaintTexture() const
+{
+	return Parent ? Parent->HasMeshPaintTexture() : false;
+}
+
 const FMaterialResource* UMaterialInstance::GetMaterialResource(ERHIFeatureLevel::Type InFeatureLevel, EMaterialQualityLevel::Type QualityLevel) const
 {
 	return const_cast<UMaterialInstance*>(this)->GetMaterialResource(InFeatureLevel, QualityLevel);
diff --git a/Engine/Source/Runtime/Engine/Private/Materials/MaterialInterface.cpp b/Engine/Source/Runtime/Engine/Private/Materials/MaterialInterface.cpp
--- a/Engine/Source/Runtime/Engine/Private/Materials/MaterialInterface.cpp
+++ b/Engine/Source/Runtime/Engine/Private/Materials/MaterialInterface.cpp
@@ -1436,6 +1436,11 @@ bool UMaterialInterface::WritesToRuntimeVirtualTexture() const
 	return false;
 }
 
+bool UMaterialInterface::HasMeshPaintTexture() const
+{
+	return false;
+}
+
 FMaterialShadingModelField UMaterialInterface::GetShadingModels() const
 {
 	return MSM_DefaultLit;
diff --git a/Engine/Source/Runtime/Engine/Public/MaterialCachedData.h b/Engine/Source/Runtime/Engine/Public/MaterialCachedData.h
--- a/Engine/Source/Runtime/Engine/Public/MaterialCachedData.h
+++ b/Engine/Source/Runtime/Engine/Public/MaterialCachedData.h
@@ -372,6 +372,9 @@ struct FMaterialCachedExpressionData
 	UPROPERTY()
 	uint32 bHasCustomizedUVs : 1;
 
+	UPROPERTY()
+	uint32 bHasMeshPaintTexture : 1;
+
 	UPROPERTY()
 	uint32 PropertyConnectedBitmask_DEPRECATED = 0;
 
diff --git a/Engine/Source/Runtime/Engine/Public/Materials/Material.h b/Engine/Source/Runtime/Engine/Public/Materials/Material.h
--- a/Engine/Source/Runtime/Engine/Public/Materials/Material.h
+++ b/Engine/Source/Runtime/Engine/Public/Materials/Material.h
@@ -1218,6 +1218,7 @@ public:
 	ENGINE_API virtual bool HasVertexInterpolator() const override;
 	ENGINE_API virtual bool HasCustomizedUVs() const override;
 	ENGINE_API virtual bool WritesToRuntimeVirtualTexture() const override;
+	ENGINE_API virtual bool HasMeshPaintTexture() const override;
 
 	ENGINE_API virtual FGraphEventArray PrecachePSOs(const FPSOPrecacheVertexFactoryDataList& VertexFactoryDataList, const FPSOPrecacheParams& PreCacheParams, EPSOPrecachePriority Priority, TArray<FMaterialPSOPrecacheRequestID>& OutMaterialPSORequestIDs) override;
 
diff --git a/Engine/Source/Runtime/Engine/Public/Materials/MaterialInstance.h b/Engine/Source/Runtime/Engine/Public/Materials/MaterialInstance.h
--- a/Engine/Source/Runtime/Engine/Public/Materials/MaterialInstance.h
+++ b/Engine/Source/Runtime/Engine/Public/Materials/MaterialInstance.h
@@ -891,6 +891,7 @@ public:
 	ENGINE_API virtual bool IsDitheredLODTransition() const override;
 	ENGINE_API virtual bool IsMasked() const override;
 	ENGINE_API virtual bool WritesToRuntimeVirtualTexture() const override;
+	ENGINE_API virtual bool HasMeshPaintTexture() const override;
 	ENGINE_API virtual FDisplacementScaling GetDisplacementScaling() const override;
 	ENGINE_API virtual bool IsDisplacementFadeEnabled() const override;
 	ENGINE_API virtual FDisplacementFadeRange GetDisplacementFadeRange() const override;
diff --git a/Engine/Source/Runtime/Engine/Public/Materials/MaterialInterface.h b/Engine/Source/Runtime/Engine/Public/Materials/MaterialInterface.h
--- a/Engine/Source/Runtime/Engine/Public/Materials/MaterialInterface.h
+++ b/Engine/Source/Runtime/Engine/Public/Materials/MaterialInterface.h
@@ -942,6 +942,7 @@ public:
 	ENGINE_API virtual bool IsMasked() const;
 	ENGINE_API virtual bool IsDeferredDecal() const;
 	ENGINE_API virtual bool WritesToRuntimeVirtualTexture() const;
+	ENGINE_API virtual bool HasMeshPaintTexture() const;
 	ENGINE_API virtual FDisplacementScaling GetDisplacementScaling() const;
 	ENGINE_API virtual bool IsDisplacementFadeEnabled() const;
 	ENGINE_API virtual FDisplacementFadeRange GetDisplacementFadeRange() const;
