Skeletal Mesh: Update the skeletal mesh morph target info when the source mesh description's morph target layout changes.

#rnx

diff --git a/Engine/Source/Runtime/Engine/Classes/Engine/SkeletalMesh.h b/Engine/Source/Runtime/Engine/Classes/Engine/SkeletalMesh.h
--- a/Engine/Source/Runtime/Engine/Classes/Engine/SkeletalMesh.h
+++ b/Engine/Source/Runtime/Engine/Classes/Engine/SkeletalMesh.h
@@ -584,6 +584,12 @@ public:
 	{
 		FCommitMeshDescriptionParams() {}
 
+		/** Ensure that there is a placeholder on the skeletal mesh to hold the morph target delta data prior to
+		 *  the mesh rebuild. Set this to off if the morph target placeholder creation is not necessary. If the
+		 *  UMorphTarget already exists, then nothing is done.
+		 */
+		bool bUpdateMorphTargets = true;
+		
 		/** Update the alternate skin weight profile list based on presence of non-default skin weight attributes in the  
 		 *  committed mesh. If a skin weight attribute exists on the mesh but not in the list of profiles, then the profile
 		 *  is added. Likewise, if a profile exists but there's no corresponding skin weight attribute then that profile
@@ -2651,7 +2657,7 @@ public:
 	ENGINE_API void UnregisterAllMorphTarget();
 
 	/** Initialize MorphSets look up table : MorphTargetIndexMap */
-	ENGINE_API void InitMorphTargets();
+	ENGINE_API void InitMorphTargets(bool bInKeepEmptyMorphTargets = false);
 
 #if WITH_EDITOR
 	/** Remove the morph targets with the specified names */
diff --git a/Engine/Source/Runtime/Engine/Private/SkeletalMesh.cpp b/Engine/Source/Runtime/Engine/Private/SkeletalMesh.cpp
--- a/Engine/Source/Runtime/Engine/Private/SkeletalMesh.cpp
+++ b/Engine/Source/Runtime/Engine/Private/SkeletalMesh.cpp
@@ -2613,6 +2613,72 @@ bool USkeletalMesh::CommitMeshDescription(
 	
 	if (SourceModel.HasMeshDescription())
 	{
+		if (InParams.bUpdateMorphTargets)
+		{
+			static FCriticalSection MorphTargetUpdateMutex;
+
+			// Since MorphTargets/MorphTargetIndexMap are USkeletalMesh members, we want to 
+			// avoid multiple threads all mutating them at the same time, in case we have a 
+			// geometry processor that is committing multiple meshes across differing LODs 
+			// simultaneously.
+			FScopeLock ScopeLock(&MorphTargetUpdateMutex);
+
+			TArray<TObjectPtr<UMorphTarget>>& ExistingMorphTargets = GetMorphTargets();
+			TSet<FName> ExistingMorphTargetNames;
+			for (TObjectPtr<UMorphTarget> MorphTarget: ExistingMorphTargets)
+			{
+				ExistingMorphTargetNames.Add(MorphTarget->GetFName());
+			}
+			
+			TSet<FName> ValidMorphTargetNames;
+			for (const FSkeletalMeshSourceModel& OtherSourceModels: GetAllSourceModels())
+			{
+				ValidMorphTargetNames.Append(OtherSourceModels.GetMorphTargetNames());
+			}
+
+			// Add in a dummy UMorphTarget placeholder for any morph target that is being added.
+			bool bMorphTargetsChanged = false;
+			for (const FName& MorphTargetName: ValidMorphTargetNames)
+			{
+				if (!ExistingMorphTargetNames.Contains(MorphTargetName))
+				{
+					UMorphTarget* MorphTarget = NewObject<UMorphTarget>(this, MorphTargetName);
+					MorphTarget->BaseSkelMesh = this;
+					MorphTarget->ClearInternalFlags(EInternalObjectFlags::Async);
+					
+					ExistingMorphTargets.Add(MorphTarget);
+					
+					bMorphTargetsChanged = true;
+				}
+			}
+
+			// Remove any existing morph targets that don't have a corresponding representation on
+			// any of the source models.
+			if (ExistingMorphTargets.RemoveAll([&ValidMorphTargetNames](const TObjectPtr<UMorphTarget>& InMorphTarget)
+				{ return !ValidMorphTargetNames.Contains(InMorphTarget->GetFName()); }) != 0)
+			{
+				bMorphTargetsChanged = true;
+			}
+
+			if (bMorphTargetsChanged)
+			{
+				constexpr bool bKeepEmptyMorphTargets = true;
+				InitMorphTargets(bKeepEmptyMorphTargets);
+
+				// Ensure all components are working from the latest morph target data.
+				if (IsInGameThread())
+				{
+					for (TObjectIterator<USkeletalMeshComponent> It; It; ++It)
+					{
+						if (It->GetSkeletalMeshAsset() == this)
+						{
+							It->RefreshMorphTargets();
+						}
+					}
+				}
+			}
+		}
+	
 		if (InParams.bUpdateSkinWeightProfiles)
 		{
 			static FCriticalSection ProfileUpdateMutex;
@@ -2622,34 +2688,35 @@ bool USkeletalMesh::CommitMeshDescription(
 			// that is committing multiple meshes across differing LODs simultaneously.
 			FScopeLock ScopeLock(&ProfileUpdateMutex);
 			
-			FSkeletalMeshConstAttributes Attributes(*SourceModel.GetMeshDescription());
-
 			TArray<FSkinWeightProfileInfo>& ExistingProfiles = GetSkinWeightProfiles();
-			TMap<FName, int32> ProfileIndexes;
+			TSet<FName> ExistingProfileNames;
 
-			for (int32 Index = 0; Index < ExistingProfiles.Num(); Index++)
+			for (const FSkinWeightProfileInfo& ProfileInfo: ExistingProfiles)
 			{
-				ProfileIndexes.Add(ExistingProfiles[Index].Name, Index);
+				ExistingProfileNames.Add(ProfileInfo.Name);
 			}
 
-			// Get all profiles from all LODs, since we may have some that weren't defined 
-			TSet<FName> ValidProfiles;
+			// Get all profiles from the models on all LODs, since we may have some that aren't
+			// defined on the skeletal mesh's list of profiles.
+			TSet<FName> ValidProfileNames;
 			for (const FSkeletalMeshSourceModel& OtherSourceModels: GetAllSourceModels())
 			{
-				ValidProfiles.Append(OtherSourceModels.GetSkinWeightProfileNames());
+				ValidProfileNames.Append(OtherSourceModels.GetSkinWeightProfileNames());
 			}
-			for (const FName& ProfileName: ValidProfiles)
+			for (const FName& ProfileName: ValidProfileNames)
 			{
-				if (!ProfileIndexes.Contains(ProfileName))
+				if (!ExistingProfileNames.Contains(ProfileName))
 				{
 					FSkinWeightProfileInfo& NewProfile = ExistingProfiles.AddDefaulted_GetRef();
 					NewProfile.Name = ProfileName;
 				}
 			}
 
-			ExistingProfiles.RemoveAll([ValidProfiles](const FSkinWeightProfileInfo& InProfileInfo)
+			// Remove all profiles listed on the skeletal mesh that no longer have a correspondence 
+			// on the source models. 
+			ExistingProfiles.RemoveAll([&ValidProfileNames](const FSkinWeightProfileInfo& InProfileInfo)
 			{
-				return !ValidProfiles.Contains(InProfileInfo.Name);
+				return !ValidProfileNames.Contains(InProfileInfo.Name);
 			});
 		}
 		
@@ -4286,7 +4353,7 @@ void USkeletalMesh::UnregisterMorphTarget(UMorphTarget* MorphTarget, bool bInval
 	}
 }
 
-void USkeletalMesh::InitMorphTargets()
+void USkeletalMesh::InitMorphTargets(bool bInKeepEmptyMorphTargets)
 {
 	TRACE_CPUPROFILER_EVENT_SCOPE(USkeletalMesh::InitMorphTargets);
 	GetMorphTargetIndexMap().Empty();
@@ -4294,21 +4361,22 @@ void USkeletalMesh::InitMorphTargets()
 	TArray<TObjectPtr<UMorphTarget>>& MorphTargetsLocal = GetMorphTargets();
 	for (int32 Index = 0; Index < MorphTargetsLocal.Num(); ++Index)
 	{
-		UMorphTarget* MorphTarget = MorphTargetsLocal[Index];
-		// if we don't have a valid data, just remove it
-		if (!MorphTarget->HasValidData())
+		const UMorphTarget* MorphTarget = MorphTargetsLocal[Index];
+		
+		// If asked to remove empty morph targets and the morph target doesn't have any data, just remove it.
+		if (!bInKeepEmptyMorphTargets && !MorphTarget->HasValidData())
 		{
 			MorphTargetsLocal.RemoveAt(Index);
 			--Index;
 			continue;
 		}
 
-		FName const ShapeName = MorphTarget->GetFName();
+		const FName ShapeName = MorphTarget->GetFName();
 		if (GetMorphTargetIndexMap().Find(ShapeName) == nullptr)
 		{
 			GetMorphTargetIndexMap().Add(ShapeName, Index);
 
-			// Note: we dont register as morph target curves here as curves metadata can now be
+			// Note: we don't register as morph target curves here as curves metadata can now be
 			// specified on this mesh, which can now opt out of the morph flag being set
 		}
 	}
