Substrate - renamed matrial function to be more UE style.

#rnx
#tests none

#virtualized

diff --git a/Engine/Source/Runtime/Engine/Private/Materials/Material.cpp b/Engine/Source/Runtime/Engine/Private/Materials/Material.cpp
--- a/Engine/Source/Runtime/Engine/Private/Materials/Material.cpp
+++ b/Engine/Source/Runtime/Engine/Private/Materials/Material.cpp
@@ -3606,7 +3606,7 @@ void UMaterial::ConvertMaterialToSubstrateMaterial()
 		else if (!bUseMaterialAttributes && !EditorOnly->FrontMaterial.IsConnected() && GetExpressions().IsEmpty())
 		{
 			// Empty material: Create by default a slab node
-			UMaterialFunction* DefaultMF = LoadObject<UMaterialFunction>(nullptr, TEXT("/Engine/Functions/Substrate/SMF_UE4Disney.SMF_UE4Disney"));
+			UMaterialFunction* DefaultMF = LoadObject<UMaterialFunction>(nullptr, TEXT("/Engine/Functions/Substrate/SMF_UE4Legacy.SMF_UE4Legacy")); 
 			if (DefaultMF)
 			{
 				DefaultMF->UpdateFromFunctionResource();
