[World Partition Runtime HashSet] Handle the case where there is no valid partition objects.

diff --git a/Engine/Source/Runtime/Engine/Private/WorldPartition/RuntimeHashSet/WorldPartitionRuntimeHashSetStreamingGeneration.cpp b/Engine/Source/Runtime/Engine/Private/WorldPartition/RuntimeHashSet/WorldPartitionRuntimeHashSetStreamingGeneration.cpp
--- a/Engine/Source/Runtime/Engine/Private/WorldPartition/RuntimeHashSet/WorldPartitionRuntimeHashSetStreamingGeneration.cpp
+++ b/Engine/Source/Runtime/Engine/Private/WorldPartition/RuntimeHashSet/WorldPartitionRuntimeHashSetStreamingGeneration.cpp
@@ -17,6 +17,11 @@ bool UWorldPartitionRuntimeHashSet::GenerateRuntimePartitionsStreamingDescs(cons
 	UWorld* OuterWorld = GetTypedOuter<UWorld>();
 	const bool bIsMainWorldPartition = (World == OuterWorld);
 
+	if (RuntimePartitions.IsEmpty())
+	{
+		return false;
+	}
+
 	//
 	// Split actor sets into their corresponding runtime partition implementation
 	//
