Hiding dynamic branch checkbox on the static parameter.

- This isn't working as intended, and is confusing users.

#rnx

diff --git a/Engine/Source/Runtime/Engine/Public/Materials/MaterialExpressionStaticBoolParameter.h b/Engine/Source/Runtime/Engine/Public/Materials/MaterialExpressionStaticBoolParameter.h
--- a/Engine/Source/Runtime/Engine/Public/Materials/MaterialExpressionStaticBoolParameter.h
+++ b/Engine/Source/Runtime/Engine/Public/Materials/MaterialExpressionStaticBoolParameter.h
@@ -18,7 +18,7 @@ class UMaterialExpressionStaticBoolParameter : public UMaterialExpressionParamet
 	uint32 DefaultValue:1;
 
 	/**Change Parameter from "static bool" to (dynamic) bool type which enables it to be used with dynamic branching*/
-	UPROPERTY(EditAnywhere, Category = MaterialExpressionStaticBoolParameter)
+	UPROPERTY()
 	uint32 DynamicBranch:1;
 
 public:
