[World Partition Runtime HashSet] Handle the case where there is no valid partition objects.

diff --git a/Engine/Source/Runtime/Engine/Private/WorldPartition/RuntimeHashSet/WorldPartitionRuntimeHashSet.cpp b/Engine/Source/Runtime/Engine/Private/WorldPartition/RuntimeHashSet/WorldPartitionRuntimeHashSet.cpp
--- a/Engine/Source/Runtime/Engine/Private/WorldPartition/RuntimeHashSet/WorldPartitionRuntimeHashSet.cpp
+++ b/Engine/Source/Runtime/Engine/Private/WorldPartition/RuntimeHashSet/WorldPartitionRuntimeHashSet.cpp
@@ -251,7 +251,7 @@ void UWorldPartitionRuntimeHashSet::FlushStreamingContent()
 
 bool UWorldPartitionRuntimeHashSet::IsValidGrid(FName GridName, const UClass* ActorClass) const
 {
-	return ResolveRuntimePartition(GridName) != nullptr;
+	return GridName.IsNone() || !!ResolveRuntimePartition(GridName);
 }
 #endif
 
