Fix compilation errors in EngineTest when UE_WITH_IRIS is false

#rnx

diff --git a/Engine/Source/Runtime/Engine/Private/Net/Experimental/Iris/DataStreamChannel.cpp b/Engine/Source/Runtime/Engine/Private/Net/Experimental/Iris/DataStreamChannel.cpp
--- a/Engine/Source/Runtime/Engine/Private/Net/Experimental/Iris/DataStreamChannel.cpp
+++ b/Engine/Source/Runtime/Engine/Private/Net/Experimental/Iris/DataStreamChannel.cpp
@@ -83,11 +83,13 @@ void UDataStreamChannel::ReInit()
 {
 	check(Connection);
 
+#if UE_WITH_IRIS
 	if (UReplicationSystem* ReplicationSystem = Connection->Driver->GetReplicationSystem())
 	{
 		bIsReadyToHandshake = 1U;
 		ReplicationSystem->InitDataStreams(Connection->GetConnectionId(), DataStreamManager);
 	}
+#endif
 }
 
 bool UDataStreamChannel::CleanUp(const bool bForDestroy, EChannelCloseReason CloseReason)
@@ -484,10 +486,12 @@ bool UDataStreamChannel::IsPacketWindowFull() const
 void UDataStreamChannel::AddReferencedObjects(UObject* Object, FReferenceCollector& Collector)
 {
 	UDataStreamChannel* Channel = CastChecked<UDataStreamChannel>(Object);
+#if UE_WITH_IRIS
 	if (Channel->DataStreamManager)
 	{
 		Collector.AddReferencedObject(Channel->DataStreamManager);
 	}
+#endif
 
 	Super::AddReferencedObjects(Channel, Collector);
 }
@@ -502,10 +506,12 @@ void UDataStreamChannel::AppendMustBeMappedGuids(FOutBunch* Bunch)
 
 bool UDataStreamChannel::HasAcknowledgedAllReliableData() const
 {
+#if UE_WITH_IRIS
 	if (DataStreamManager)
 	{
 		return DataStreamManager->HasAcknowledgedAllReliableData();
 	}
+#endif
 
 	return true;
 }
diff --git a/Engine/Source/Runtime/Engine/Private/Net/Experimental/Iris/DataStreamChannel.h b/Engine/Source/Runtime/Engine/Private/Net/Experimental/Iris/DataStreamChannel.h
--- a/Engine/Source/Runtime/Engine/Private/Net/Experimental/Iris/DataStreamChannel.h
+++ b/Engine/Source/Runtime/Engine/Private/Net/Experimental/Iris/DataStreamChannel.h
@@ -80,7 +80,9 @@ private:
 
 	void WriteData(UE::Net::EDataStreamWriteMode WriteMode);
 
+#if UE_WITH_IRIS
 	TObjectPtr<UDataStreamManager> DataStreamManager = nullptr;
+#endif
 
 	TResizableCircularQueue<FDataStreamChannelRecord> WriteRecords;
 
