PR #12309: Adding Highpass filter to Audio Component

diff --git a/Engine/Source/Runtime/Engine/Classes/Components/AudioComponent.h b/Engine/Source/Runtime/Engine/Classes/Components/AudioComponent.h
--- a/Engine/Source/Runtime/Engine/Classes/Components/AudioComponent.h
+++ b/Engine/Source/Runtime/Engine/Classes/Components/AudioComponent.h
@@ -225,9 +225,13 @@ class UAudioComponent : public USceneComponent, public ISoundParameterController
 	/** Whether or not to apply a low-pass filter to the sound that plays in this audio component. */
 	UPROPERTY(EditAnywhere, BlueprintReadWrite, Category = Sound, meta = (InlineEditConditionToggle, DisplayAfter = "PitchMultiplier"))
 	uint8 bEnableLowPassFilter : 1;
+	
+	/** Whether or not to apply a low-pass filter to the sound that plays in this audio component. */
+	UPROPERTY(EditAnywhere, BlueprintReadWrite, Category = Sound, meta = (InlineEditConditionToggle, DisplayAfter = "LowPassFilterFrequency"))
+	uint8 bEnableHighPassFilter : 1;
 
 	/** Whether or not to override the priority of the given sound with the value provided. */
-	UPROPERTY(EditAnywhere, BlueprintReadWrite, Category = Sound, meta = (InlineEditConditionToggle, DisplayAfter = "LowPassFilterFrequency"))
+	UPROPERTY(EditAnywhere, BlueprintReadWrite, Category = Sound, meta = (InlineEditConditionToggle, DisplayAfter = "HighPassFilterFrequency"))
 	uint8 bOverridePriority:1;
 
 	/** If true, subtitles in the sound data will be ignored. */
@@ -331,7 +335,7 @@ public:
 	int32 EnvelopeFollowerReleaseTime;
 
 	/** If enabled, overrides the priority of the selected sound with the value provided. */
-	UPROPERTY(EditAnywhere, BlueprintReadWrite, Category = Sound, meta = (ClampMin = "0.0", UIMin = "0.0", EditCondition = "bOverridePriority", DisplayAfter = "LowPassFilterFrequency"))
+	UPROPERTY(EditAnywhere, BlueprintReadWrite, Category = Sound, meta = (ClampMin = "0.0", UIMin = "0.0", EditCondition = "bOverridePriority", DisplayAfter = "HighPassFilterFrequency"))
 	float Priority;
 
 	/** Used by the subtitle manager to prioritize subtitles wave instances spawned by this component. */
@@ -357,6 +361,10 @@ public:
 	/** If enabled, the frequency of the Lowpass Filter (in Hz) to apply to this voice. A frequency of 0.0 is the device sample rate and will bypass the filter. */
 	UPROPERTY(EditAnywhere, BlueprintReadWrite, Category = Sound, meta = (ClampMin = "0.0", UIMin = "0.0", EditCondition = "bEnableLowPassFilter", DisplayAfter = "PitchMultiplier"))
 	float LowPassFilterFrequency;
+	
+	/** If enabled, the frequency of the Highpass Filter (in Hz) to apply to this voice. A frequency of 0.0 is the device sample rate and will bypass the filter. */
+	UPROPERTY(EditAnywhere, BlueprintReadWrite, Category = Sound, meta = (ClampMin = "0.0", UIMin = "0.0", EditCondition = "bEnableHighPassFilter", DisplayAfter = "LowPassFilterFrequency"))
+	float HighPassFilterFrequency;
 
 	/** A count of how many times we've started playing */
 	int32 ActiveCount;
@@ -668,12 +676,25 @@ public:
 	 */
 	UFUNCTION(BlueprintCallable, Category = "Audio|Components|Audio")
 	ENGINE_API void SetLowPassFilterEnabled(bool InLowPassFilterEnabled);
+	
+	/** When set to TRUE, enables an additional High Pass Filter Frequency to be calculated in with the
+	 *  sound instance’s HPF total, allowing designers to set filter settings for the targeted Audio Component’s
+	 *  sound instance.
+	 */
+	UFUNCTION(BlueprintCallable, Category = "Audio|Components|Audio")
+	ENGINE_API void SetHighPassFilterEnabled(bool InHighPassFilterEnabled);
 
 	/** Sets a cutoff frequency, in Hz, for the targeted Audio Component’s sound’s Low Pass Filter calculation.
 	 *  The lowest cutoff frequency from all of the sound instance’s possible LPF calculations wins.
 	 */
 	UFUNCTION(BlueprintCallable, Category = "Audio|Components|Audio")
 	ENGINE_API void SetLowPassFilterFrequency(float InLowPassFilterFrequency);
+	
+	/** Sets a cutoff frequency, in Hz, for the targeted Audio Component’s sound’s High Pass Filter calculation.
+	 *  The highest cutoff frequency from all of the sound instance’s possible HPF calculations wins.
+	 */
+	UFUNCTION(BlueprintCallable, Category = "Audio|Components|Audio")
+	ENGINE_API void SetHighPassFilterFrequency(float InHighPassFilterFrequency);
 
 	/** Sets whether or not to output the audio to bus only. */
 	UFUNCTION(BlueprintCallable, Category = "Audio|Components|Audio")
diff --git a/Engine/Source/Runtime/Engine/Private/ActiveSound.cpp b/Engine/Source/Runtime/Engine/Private/ActiveSound.cpp
--- a/Engine/Source/Runtime/Engine/Private/ActiveSound.cpp
+++ b/Engine/Source/Runtime/Engine/Private/ActiveSound.cpp
@@ -85,6 +85,7 @@ FActiveSound::FActiveSound()
 	, bWarnedAboutOrphanedLooping(false)
 #endif
 	, bEnableLowPassFilter(false)
+	, bEnableHighPassFilter(false)
 	, bUpdatePlayPercentage(false)
 	, bUpdateSingleEnvelopeValue(false)
 	, bUpdateMultiEnvelopeValue(false)
@@ -112,6 +113,7 @@ FActiveSound::FActiveSound()
 	, VolumeMultiplier(1.0f)
 	, PitchMultiplier(1.0f)
 	, LowPassFilterFrequency(MAX_FILTER_FREQUENCY)
+	, HighPassFilterFrequency(MIN_FILTER_FREQUENCY)
 	, CurrentOcclusionFilterFrequency(MAX_FILTER_FREQUENCY)
 	, CurrentOcclusionVolumeAttenuation(1.0f)
 	, SubtitlePriority(DEFAULT_SUBTITLE_PRIORITY)
@@ -874,6 +876,8 @@ void FActiveSound::UpdateWaveInstances(TArray<FWaveInstance*> &InWaveInstances,
 	ParseParams.Pitch *= GetPitch() * Sound->GetPitchMultiplier();
 	ParseParams.bEnableLowPassFilter = bEnableLowPassFilter;
 	ParseParams.LowPassFilterFrequency = LowPassFilterFrequency;
+	ParseParams.bEnableHighPassFilter = bEnableHighPassFilter;
+	ParseParams.HighPassFilterFrequency = HighPassFilterFrequency;
 	ParseParams.SoundClass = GetSoundClass();
 	ParseParams.bIsPaused = bIsPaused;
 
diff --git a/Engine/Source/Runtime/Engine/Private/Audio.cpp b/Engine/Source/Runtime/Engine/Private/Audio.cpp
--- a/Engine/Source/Runtime/Engine/Private/Audio.cpp
+++ b/Engine/Source/Runtime/Engine/Private/Audio.cpp
@@ -421,6 +421,11 @@ void FSoundSource::SetFilterFrequency()
 		{
 			// Set the HPFFrequency to highest provided value
 			HPFFrequency = WaveInstance->AttenuationHighpassFilterFrequency;
+
+			if (WaveInstance->bEnableHighPassFilter)
+			{
+				HPFFrequency = FMath::Max(HPFFrequency, WaveInstance->HighPassFilterFrequency);
+			}
 		}
 		break;
 	}
@@ -864,6 +869,7 @@ FWaveInstance::FWaveInstance(const UPTRINT InWaveInstanceHash, FActiveSound& InA
 	, bAlreadyNotifiedHook(false)
 	, bUseSpatialization(false)
 	, bEnableLowPassFilter(false)
+	, bEnableHighPassFilter(false)
 	, bIsOccluded(false)
 	, bIsUISound(false)
 	, bIsMusic(false)
@@ -881,6 +887,7 @@ FWaveInstance::FWaveInstance(const UPTRINT InWaveInstanceHash, FActiveSound& InA
 	, SourceDataOverridePluginSettings(nullptr)
 	, OutputTarget(EAudioOutputTarget::Speaker)
 	, LowPassFilterFrequency(MAX_FILTER_FREQUENCY)
+	, HighPassFilterFrequency(MIN_FILTER_FREQUENCY)
 	, SoundClassFilterFrequency(MAX_FILTER_FREQUENCY)
 	, OcclusionFilterFrequency(MAX_FILTER_FREQUENCY)
 	, AmbientZoneFilterFrequency(MAX_FILTER_FREQUENCY)
diff --git a/Engine/Source/Runtime/Engine/Private/Components/AudioComponent.cpp b/Engine/Source/Runtime/Engine/Private/Components/AudioComponent.cpp
--- a/Engine/Source/Runtime/Engine/Private/Components/AudioComponent.cpp
+++ b/Engine/Source/Runtime/Engine/Private/Components/AudioComponent.cpp
@@ -108,7 +108,9 @@ UAudioComponent::UAudioComponent(const FObjectInitializer& ObjectInitializer)
 	PitchModulationMin = 1.f;
 	PitchModulationMax = 1.f;
 	bEnableLowPassFilter = false;
+	bEnableHighPassFilter = false;
 	LowPassFilterFrequency = MAX_FILTER_FREQUENCY;
+	HighPassFilterFrequency = MIN_FILTER_FREQUENCY;
 	OcclusionCheckInterval = 0.1f;
 	ActiveCount = 0;
 
@@ -750,6 +752,8 @@ void UAudioComponent::PlayInternal(const PlayInternalRequestData& InPlayRequestD
 
 	NewActiveSound.bEnableLowPassFilter = bEnableLowPassFilter;
 	NewActiveSound.LowPassFilterFrequency = LowPassFilterFrequency;
+	NewActiveSound.bEnableHighPassFilter = bEnableHighPassFilter;
+	NewActiveSound.HighPassFilterFrequency = HighPassFilterFrequency;
 	NewActiveSound.RequestedStartTime = FMath::Max(0.f, InPlayRequestData.StartTime);
 
 	if (bOverrideSubtitlePriority)
@@ -1599,6 +1603,18 @@ void UAudioComponent::SetLowPassFilterEnabled(bool InLowPassFilterEnabled)
 	}
 }
 
+void UAudioComponent::SetHighPassFilterEnabled(bool InHighPassFilterEnabled)
+{
+	if (FAudioDevice* AudioDevice = GetAudioDevice())
+	{
+		DECLARE_CYCLE_STAT(TEXT("FAudioThreadTask.SetHighPassFilterFrequency"), STAT_AudioSetHighPassFilterEnabled, STATGROUP_AudioThreadCommands);
+		AudioDevice->SendCommandToActiveSounds(AudioComponentID, [InHighPassFilterEnabled](FActiveSound& ActiveSound)
+		{
+			ActiveSound.bEnableHighPassFilter = InHighPassFilterEnabled;
+		}, GET_STATID(STAT_AudioSetHighPassFilterEnabled));
+	}
+}
+
 void UAudioComponent::SetLowPassFilterFrequency(float InLowPassFilterFrequency)
 {
 	if (FAudioDevice* AudioDevice = GetAudioDevice())
@@ -1611,6 +1627,18 @@ void UAudioComponent::SetLowPassFilterFrequency(float InLowPassFilterFrequency)
 	}
 }
 
+void UAudioComponent::SetHighPassFilterFrequency(float InHighPassFilterFrequency)
+{
+	if (FAudioDevice* AudioDevice = GetAudioDevice())
+	{
+		DECLARE_CYCLE_STAT(TEXT("FAudioThreadTask.SetHighPassFilterFrequency"), STAT_AudioSetHighPassFilterFrequency, STATGROUP_AudioThreadCommands);
+		AudioDevice->SendCommandToActiveSounds(AudioComponentID, [InHighPassFilterFrequency](FActiveSound& ActiveSound)
+		{
+			ActiveSound.HighPassFilterFrequency = InHighPassFilterFrequency;
+		}, GET_STATID(STAT_AudioSetHighPassFilterFrequency));
+	}
+}
+
 void UAudioComponent::SetOutputToBusOnly(bool bInOutputToBusOnly)
 {
 	if (FAudioDevice* AudioDevice = GetAudioDevice())
diff --git a/Engine/Source/Runtime/Engine/Private/SoundWave.cpp b/Engine/Source/Runtime/Engine/Private/SoundWave.cpp
--- a/Engine/Source/Runtime/Engine/Private/SoundWave.cpp
+++ b/Engine/Source/Runtime/Engine/Private/SoundWave.cpp
@@ -3351,8 +3351,10 @@ void USoundWave::Parse(FAudioDevice* AudioDevice, const UPTRINT NodeWaveInstance
 	WaveInstance->SetOcclusionAttenuation(ParseParams.OcclusionAttenuation);
 	WaveInstance->SetPitch(ParseParams.Pitch * Pitch);
 	WaveInstance->bEnableLowPassFilter = ParseParams.bEnableLowPassFilter;
+	WaveInstance->bEnableHighPassFilter = ParseParams.bEnableHighPassFilter;
 	WaveInstance->bIsOccluded = ParseParams.bIsOccluded;
 	WaveInstance->LowPassFilterFrequency = ParseParams.LowPassFilterFrequency;
+	WaveInstance->HighPassFilterFrequency = ParseParams.HighPassFilterFrequency;
 	WaveInstance->OcclusionFilterFrequency = ParseParams.OcclusionFilterFrequency;
 	WaveInstance->AttenuationLowpassFilterFrequency = ParseParams.AttenuationLowpassFilterFrequency;
 	WaveInstance->AttenuationHighpassFilterFrequency = ParseParams.AttenuationHighpassFilterFrequency;
diff --git a/Engine/Source/Runtime/Engine/Public/ActiveSound.h b/Engine/Source/Runtime/Engine/Public/ActiveSound.h
--- a/Engine/Source/Runtime/Engine/Public/ActiveSound.h
+++ b/Engine/Source/Runtime/Engine/Public/ActiveSound.h
@@ -172,6 +172,9 @@ struct FSoundParseParameters
 
 	// The lowpass filter frequency to apply (if enabled)
 	float LowPassFilterFrequency;
+	
+	// The highpass filter frequency to apply (if enabled)
+	float HighPassFilterFrequency;
 
 	// The lowpass filter frequency to apply due to distance attenuation
 	float AttenuationLowpassFilterFrequency;
@@ -206,6 +209,9 @@ struct FSoundParseParameters
 
 	// Whether we have enabled low-pass filtering of this sound
 	uint8 bEnableLowPassFilter:1;
+	
+	// Whether we have enabled high-pass filtering of this sound
+	uint8 bEnableHighPassFilter:1;
 
 	// Whether this sound is occluded
 	uint8 bIsOccluded:1;
@@ -253,6 +259,7 @@ struct FSoundParseParameters
 		, AudioLinkSettingsOverride(nullptr)
 		, SourceEffectChain(nullptr)
 		, LowPassFilterFrequency(MAX_FILTER_FREQUENCY)
+		, HighPassFilterFrequency(MIN_FILTER_FREQUENCY)
 		, AttenuationLowpassFilterFrequency(MAX_FILTER_FREQUENCY)
 		, AttenuationHighpassFilterFrequency(MIN_FILTER_FREQUENCY)
 		, OcclusionFilterFrequency(MAX_FILTER_FREQUENCY)
@@ -264,6 +271,7 @@ struct FSoundParseParameters
 		, bUseSpatialization(false)
 		, bLooping(false)
 		, bEnableLowPassFilter(false)
+		, bEnableHighPassFilter(false)
 		, bIsOccluded(false)
 		, bIsPaused(false)
 		, bEnableRetrigger(false)
@@ -501,6 +509,9 @@ public:
 
 	/** Whether or not we have a low-pass filter enabled on this active sound. */
 	uint8 bEnableLowPassFilter : 1;
+	
+	/** Whether or not we have a low-pass filter enabled on this active sound. */
+	uint8 bEnableHighPassFilter : 1;
 
 	/** Whether or not this active sound will update play percentage. Based on set delegates on audio component. */
 	uint8 bUpdatePlayPercentage:1;
@@ -563,6 +574,9 @@ public:
 
 	/** The low-pass filter frequency to apply if bEnableLowPassFilter is true. */
 	float LowPassFilterFrequency;
+	
+	/** The high-pass filter frequency to apply if bEnableHighPassFilter is true. */
+	float HighPassFilterFrequency;
 
 	/** Fader that tracks component volume */
 	Audio::FVolumeFader ComponentVolumeFader;
diff --git a/Engine/Source/Runtime/Engine/Public/Audio.h b/Engine/Source/Runtime/Engine/Public/Audio.h
--- a/Engine/Source/Runtime/Engine/Public/Audio.h
+++ b/Engine/Source/Runtime/Engine/Public/Audio.h
@@ -92,7 +92,6 @@ struct FWaveInstance;
 struct FReverbSettings;
 struct FSampleLoop;
 struct FSoundWaveTimecodeInfo;
-
 enum ELoopingMode
 {
 	/** One shot sound */
@@ -300,6 +299,9 @@ private:
 public:
 	/** Whether or not to enable the low pass filter */
 	uint32 bEnableLowPassFilter:1;
+	
+	/** Whether or not to enable the high pass filter */
+	uint32 bEnableHighPassFilter:1;
 
 	/** Whether or not the sound is occluded. */
 	uint32 bIsOccluded:1;
@@ -351,6 +353,9 @@ public:
 
 	/** The low pass filter frequency to use */
 	float LowPassFilterFrequency;
+	
+	/** The high pass filter frequency to use */
+	float HighPassFilterFrequency;
 
 	/** The low pass filter frequency to use from sound class. */
 	float SoundClassFilterFrequency;
