Fix material conversion when the input normal is not a float3/float4.

diff --git a/Engine/Source/Runtime/Engine/Private/Materials/MaterialExpressionSubstrate.cpp b/Engine/Source/Runtime/Engine/Private/Materials/MaterialExpressionSubstrate.cpp
--- a/Engine/Source/Runtime/Engine/Private/Materials/MaterialExpressionSubstrate.cpp
+++ b/Engine/Source/Runtime/Engine/Private/Materials/MaterialExpressionSubstrate.cpp
@@ -151,6 +151,10 @@ static int32 CompileWithDefaultNormalWS(class FMaterialCompiler* Compiler, FExpr
 			// Nothing is plug in from the linked input, so specify world space normal the BSDF node expects.
 			return Compiler->VertexNormal();
 		}
+
+		// Ensure the normal has always a valid float3 type
+		NormalCodeChunk = Compiler->ForceCast(NormalCodeChunk, MCT_Float3, MFCF_ExactMatch | MFCF_ReplicateValue);
+
 		// Transform into world space normal if needed. BSDF nodes always expects world space normal as input.
 		return bConvertToRequestedSpace ? Compiler->TransformNormalFromRequestedBasisToWorld(NormalCodeChunk) : NormalCodeChunk;
 	}
@@ -168,6 +172,10 @@ static int32 CompileWithDefaultTangentWS(class FMaterialCompiler* Compiler, FExp
 			// Nothing is plug in from the linked input, so specify world space tangent the BSDF node expects.
 			return Compiler->VertexTangent();
 		}
+
+		// Ensure the tangent has always a valid float3 type
+		TangentCodeChunk = Compiler->ForceCast(TangentCodeChunk, MCT_Float3, MFCF_ExactMatch | MFCF_ReplicateValue);
+
 		// Transform into world space tangent if needed. BSDF nodes always expects world space tangent as input.
 		return bConvertToRequestedSpace ? Compiler->TransformNormalFromRequestedBasisToWorld(TangentCodeChunk) : TangentCodeChunk;
 	}
@@ -3549,7 +3557,9 @@ int32 UMaterialExpressionSubstrateConvertMaterialAttributes::Compile(class FMate
 	const bool bHasAnisotropy = IsMaterialAttributeInputConnected(Cached, MP_Anisotropy);
 
 	// Regular normal basis
-	int32 NormalCodeChunk = Compiler->TransformNormalFromRequestedBasisToWorld(MaterialAttributes.CompileWithDefault(Compiler, FMaterialAttributeDefinitionMap::GetID(MP_Normal)));
+	int32 NormalCodeChunk = MaterialAttributes.CompileWithDefault(Compiler, FMaterialAttributeDefinitionMap::GetID(MP_Normal));
+	NormalCodeChunk = Compiler->ForceCast(NormalCodeChunk, MCT_Float3, MFCF_ExactMatch | MFCF_ReplicateValue);
+	NormalCodeChunk = Compiler->TransformNormalFromRequestedBasisToWorld(NormalCodeChunk);
 
 	// When computing NormalCodeChunk, we invoke TransformNormalFromRequestedBasisToWorld which requires input to be float or float3.
 	// Certain material do not respect this requirement. We handle here a simple recovery when source material doesn't have a valid 
@@ -3557,7 +3567,13 @@ int32 UMaterialExpressionSubstrateConvertMaterialAttributes::Compile(class FMate
 	// to the user, but the compilation will succeed.
 	if (NormalCodeChunk == INDEX_NONE) { NormalCodeChunk = Compiler->VertexNormal(); }
 
-	int32 TangentCodeChunk = bHasAnisotropy ? Compiler->TransformNormalFromRequestedBasisToWorld(MaterialAttributes.CompileWithDefault(Compiler, FMaterialAttributeDefinitionMap::GetID(MP_Tangent))) : INDEX_NONE;
+	int32 TangentCodeChunk = INDEX_NONE; 
+	if (bHasAnisotropy)
+	{
+		TangentCodeChunk = MaterialAttributes.CompileWithDefault(Compiler, FMaterialAttributeDefinitionMap::GetID(MP_Tangent));
+		TangentCodeChunk = Compiler->ForceCast(TangentCodeChunk, MCT_Float3, MFCF_ExactMatch | MFCF_ReplicateValue);
+		TangentCodeChunk = Compiler->TransformNormalFromRequestedBasisToWorld(TangentCodeChunk);
+	}
 	const FSubstrateRegisteredSharedLocalBasis NewRegisteredSharedLocalBasis = SubstrateCompilationInfoCreateSharedLocalBasis(Compiler, NormalCodeChunk, TangentCodeChunk);
 	const FString BasisIndexMacro = Compiler->GetSubstrateSharedLocalBasisIndexMacro(NewRegisteredSharedLocalBasis);
 
