Improve CookStatsFile to include a more complete set of metrics

To better track and understand cook performance for the content worker, we require a good bit more information that what the cook profile currently dumped out via -CookStatsFile=... provides. This change aims to make the metrics provided there more completel, such that they have parity with the stats we print to the log (but in a format that makes them easier to parse for tools).

The most noteworthy change is that we break the format slightly. Previously, all steps from the cook profile lived at the top-level of the output json object. With this change, they now sit in a subobject "CookProfile", so they do not pollute the other stats.

[RN] Cooks running with -CookStatsFile=... now produce a more comprehensive set of stats. The format has changed slightly, such that the previous members of the json object now live in a subobject called "CookProfile"

diff --git a/Engine/Source/Editor/UnrealEd/Private/Cooker/CookProfiling.cpp b/Engine/Source/Editor/UnrealEd/Private/Cooker/CookProfiling.cpp
--- a/Engine/Source/Editor/UnrealEd/Private/Cooker/CookProfiling.cpp
+++ b/Engine/Source/Editor/UnrealEd/Private/Cooker/CookProfiling.cpp
@@ -20,6 +20,7 @@
 #include "PackageBuildDependencyTracker.h"
 #include "Policies/CondensedJsonPrintPolicy.h"
 #include "Serialization/ArchiveUObject.h"
+#include "Serialization/JsonSerializer.h"
 #include "Serialization/JsonWriter.h"
 #include "Templates/Casts.h"
 #include "UObject/GCObject.h"
@@ -972,6 +973,25 @@ void SendLogCookStats(ECookMode::Type CookMode)
 
 	FCookStatsManager::LogCookStats(LogStatsFunc);
 
+	FString CookStatsFileName;
+	FString CookStatsJsonString;
+	TSharedPtr<TJsonWriter<TCHAR, TCondensedJsonPrintPolicy<TCHAR>>> CookStatsWriter{};
+	FStringBuilderBase Builder;
+	if (FParse::Value(FCommandLine::Get(), TEXT("-CookStatsFile="), CookStatsFileName))
+	{
+		uint32 MultiprocessId = UE::GetMultiprocessId();
+		if (MultiprocessId == 0)
+		{
+			CookStatsWriter = TJsonWriterFactory<TCHAR, TCondensedJsonPrintPolicy<TCHAR>>::Create(&CookStatsJsonString).ToSharedPtr();
+			CookStatsWriter->WriteObjectStart();
+		}
+		else
+		{
+			// Suppress the file creation on CookWorkers
+			// TODO: Replicate the information back to the CookDirector instead, UE-185774
+		}
+	}
+
 	UE_LOG(LogCookStats, Display, TEXT("Misc Cook Stats"));
 	UE_LOG(LogCookStats, Display, TEXT("==============="));
 	for (FString& StatCategory : StatCategories)
@@ -984,6 +1004,32 @@ void SendLogCookStats(ECookMode::Type CookMode)
 		{
 			UE_LOG(LogCookStats, Display, TEXT("    %s=%s"), *StatKeyValue.Key, *StatKeyValue.Value);
 		}
+
+		if (CookStatsWriter)
+		{
+			CookStatsWriter->WriteObjectStart(StatCategory);
+			for (const FCookStatsManager::StringKeyValue& StatKeyValue : StatsInCategory)
+			{
+				CookStatsWriter->WriteValue(*StatKeyValue.Key, *StatKeyValue.Value);
+			}
+			CookStatsWriter->WriteObjectEnd();
+		}
+	}
+
+	if (CookStatsWriter)
+	{
+		CookStatsWriter->WriteObjectStart(TEXT("MiscCookStats"));
+		for (FString& StatCategory : StatCategories)
+		{
+			CookStatsWriter->WriteObjectStart(*StatCategory);
+			TArray<FCookStatsManager::StringKeyValue>& StatsInCategory = StatsInCategories.FindOrAdd(StatCategory);
+			for (const FCookStatsManager::StringKeyValue& StatKeyValue : StatsInCategory)
+			{
+				CookStatsWriter->WriteValue(*StatKeyValue.Key, *StatKeyValue.Value);
+			}
+			CookStatsWriter->WriteObjectEnd();
+		}
+		CookStatsWriter->WriteObjectEnd();
 	}
 
 	// DDC Usage stats are custom formatted, and the above code just accumulated them into a TSet. Now log it with our special formatting for readability.
@@ -997,44 +1043,18 @@ void SendLogCookStats(ECookMode::Type CookMode)
 			UE_LOG(LogCookStats, Display, TEXT("%s.%s=%s"), *ProfileEntry.Path, *ProfileEntry.Key, *ProfileEntry.Value);
 		}
 
-		FString CookStatsFileName;
-		if (FParse::Value(FCommandLine::Get(), TEXT("-CookStatsFile="), CookStatsFileName))
+		if (CookStatsWriter)
 		{
-			uint32 MultiprocessId = UE::GetMultiprocessId();
-			if (MultiprocessId != 0)
-			{
-				// Suppress the file creation on CookWorkers
-				// TODO: Replicate the information back to the CookDirector instead, UE-185774
-				CookStatsFileName.Empty();
-			}
-		}
-
-		if (!CookStatsFileName.IsEmpty())
-		{
-			FString JsonString;
-			TSharedRef<TJsonWriter<TCHAR, TCondensedJsonPrintPolicy<TCHAR>>> JsonWriter = TJsonWriterFactory<TCHAR, TCondensedJsonPrintPolicy<TCHAR> >::Create(&JsonString);
-			JsonWriter->WriteObjectStart();
+			CookStatsWriter->WriteObjectStart("CookProfile");
 			for (const auto& ProfileEntry : CookProfileData)
 			{
-				JsonWriter->WriteObjectStart(ProfileEntry.Key);
-				JsonWriter->WriteValue(TEXT("Path"), ProfileEntry.Path);
-				JsonWriter->WriteValue(TEXT("Value"), ProfileEntry.Value);
-				JsonWriter->WriteObjectEnd();
-			}
-			JsonWriter->WriteObjectEnd();
-			JsonWriter->Close();
-			TUniquePtr<FArchive> JsonFile(IFileManager::Get().CreateFileWriter(*CookStatsFileName));
-			if (!JsonFile)
-			{
-				UE_LOG(LogCookStats, Warning, TEXT("Could not write to CookStatsFile %s."), *CookStatsFileName);
-			}
-			else
-			{
-				JsonFile->Serialize(TCHAR_TO_ANSI(*JsonString), JsonString.Len());
-				JsonFile->Close();
+				CookStatsWriter->WriteObjectStart(ProfileEntry.Key);
+				CookStatsWriter->WriteValue(TEXT("Path"), ProfileEntry.Path);
+				CookStatsWriter->WriteValue(TEXT("Value"), ProfileEntry.Value);
+				CookStatsWriter->WriteObjectEnd();
 			}
+			CookStatsWriter->WriteObjectEnd();
 		}
-
 	}
 	if (DDCSummaryStats.Num() > 0)
 	{
@@ -1045,6 +1065,16 @@ void SendLogCookStats(ECookMode::Type CookMode)
 		{
 			UE_LOG(LogCookStats, Display, TEXT("%-16s=%10s"), *Attr.Key, *Attr.Value);
 		}
+
+		if (CookStatsWriter)
+		{
+			CookStatsWriter->WriteObjectStart("DDCSummaryStats");
+			for (const auto& Attr : DDCSummaryStats)
+			{
+				CookStatsWriter->WriteValue(*Attr.Key, *Attr.Value);
+			}
+			CookStatsWriter->WriteObjectEnd();
+		}
 	}
 
 	DumpDerivedDataBuildRemoteExecutorStats();
@@ -1064,6 +1094,21 @@ void SendLogCookStats(ECookMode::Type CookMode)
 				*Stat.AssetType, Stat.LoadTimeSec + Stat.BuildTimeSec, Stat.GameThreadTimeSec,
 				Stat.BuildCount, Stat.LoadSizeMB + Stat.BuildSizeMB);
 		}
+
+		if (CookStatsWriter)
+		{
+			CookStatsWriter->WriteObjectStart("DDCResourceStats");
+			for (const FDerivedDataCacheResourceStat& Stat : DDCResourceUsageStats)
+			{
+				CookStatsWriter->WriteObjectStart(*Stat.AssetType);
+				CookStatsWriter->WriteValue(TEXT("TotalTimeSec"), Stat.LoadTimeSec + Stat.BuildTimeSec);
+				CookStatsWriter->WriteValue(TEXT("GameThreadTimeSec"), Stat.GameThreadTimeSec);
+				CookStatsWriter->WriteValue(TEXT("AssetsBuilt"), Stat.BuildCount);
+				CookStatsWriter->WriteValue(TEXT("MBProcessed"), Stat.LoadSizeMB + Stat.BuildSizeMB);
+				CookStatsWriter->WriteObjectEnd();
+			}
+			CookStatsWriter->WriteObjectEnd();
+		}
 	}
 
 	DumpBuildDependencyTrackerStats();
@@ -1072,6 +1117,36 @@ void SendLogCookStats(ECookMode::Type CookMode)
 	{
 		UE::Virtualization::IVirtualizationSystem::Get().DumpStats();
 	}
+
+	if (CookStatsWriter)
+	{
+		FShaderCompilerStats ShaderCompilerStats;
+		GShaderCompilingManager->GetLocalStats(ShaderCompilerStats);
+
+		TSharedPtr<FJsonObject> JsonShaderCompilerStats = ShaderCompilerStats.ToJson();
+		FJsonSerializer::Serialize(
+			MakeShared<FJsonValueObject>(JsonShaderCompilerStats),
+			TEXT("ShaderCompilerStats"),
+			CookStatsWriter.ToSharedRef(),
+			false
+		);
+	}
+
+	if (!CookStatsFileName.IsEmpty())
+	{
+		CookStatsWriter->WriteObjectEnd();
+		CookStatsWriter->Close();
+		TUniquePtr<FArchive> CookStatsJsonFile(IFileManager::Get().CreateFileWriter(*CookStatsFileName));
+		if (!CookStatsJsonFile)
+		{
+			UE_LOG(LogCookStats, Warning, TEXT("Could not write to CookStatsFile %s."), *CookStatsFileName);
+		}
+		else
+		{
+			CookStatsJsonFile->Serialize(TCHAR_TO_ANSI(*CookStatsJsonString), CookStatsJsonString.Len());
+			CookStatsJsonFile->Close();
+		}
+	}
 }
 
 }
diff --git a/Engine/Source/Runtime/Engine/Private/ShaderCompiler/ShaderCompilerStats.cpp b/Engine/Source/Runtime/Engine/Private/ShaderCompiler/ShaderCompilerStats.cpp
--- a/Engine/Source/Runtime/Engine/Private/ShaderCompiler/ShaderCompilerStats.cpp
+++ b/Engine/Source/Runtime/Engine/Private/ShaderCompiler/ShaderCompilerStats.cpp
@@ -9,6 +9,8 @@
 #include "AnalyticsEventAttribute.h"
 #include "CoreMinimal.h"
 #include "DistributedBuildControllerInterface.h"
+#include "Dom/JsonObject.h"
+#include "JsonObjectConverter.h"
 #include "Misc/ConfigCacheIni.h"
 #include "Misc/EngineVersion.h"
 #include "ProfilingDebugging/CookStats.h"
@@ -17,7 +19,6 @@
 #include "ProfilingDebugging/StallDetector.h"
 #include "Serialization/CompactBinaryWriter.h"
 
-
 static int32 GLogShaderCompilerStats = 0;
 static FAutoConsoleVariableRef CVarLogShaderCompilerStats(
 	TEXT("r.LogShaderCompilerStats"),
@@ -66,7 +67,7 @@ void FShaderCompilerStats::WriteStats(FOutputDevice* Ar)
 				for (const auto& Pair : Stats)
 				{
 					const FString& Path = Pair.Key;
-					const FShaderCompilerStats::FShaderStats& SingleStats = Pair.Value;
+					const FShaderStats& SingleStats = Pair.Value;
 
 					StatWriter.AddColumn(*Path);
 					StatWriter.AddColumn(TEXT("%u"), Platform);
@@ -148,8 +149,8 @@ void FShaderCompilerStats::WriteStats(FOutputDevice* Ar)
 				for (const auto& Pair : Stats)
 				{
 					const FString& Path = Pair.Key;
-					const FShaderCompilerStats::FShaderStats& SingleStats = Pair.Value;
-					for (const FShaderCompilerStats::FShaderCompilerSinglePermutationStat& Stat : SingleStats.PermutationCompilations)
+					const FShaderStats& SingleStats = Pair.Value;
+					for (const FShaderCompilerSinglePermutationStat& Stat : SingleStats.PermutationCompilations)
 					{
 						StatWriter.AddColumn(*Path);
 						StatWriter.AddColumn(TEXT("%u"), Platform);
@@ -577,6 +578,51 @@ void FShaderCompilerStats::Aggregate(FShaderCompilerStats& Other)
 	MaterialCounters += Other.MaterialCounters;
 }
 
+
+TSharedPtr<FJsonObject> FShaderCompilerStats::ToJson()
+{
+	TSharedPtr<FJsonObject> RootObject(MakeShared<FJsonObject>());
+	{
+		TSharedRef<FJsonObject> JsonObject = MakeShared<FJsonObject>();
+		FJsonObjectConverter::UStructToJsonObject(FShaderCompilerCounters::StaticStruct(), &Counters, JsonObject, 0, 0, nullptr, EJsonObjectConversionFlags::SkipStandardizeCase);
+		RootObject->SetObjectField("Counters", JsonObject);
+	}
+	{
+		TSharedRef<FJsonObject> JsonObject = MakeShared<FJsonObject>();
+		FJsonObjectConverter::UStructToJsonObject(FShaderCompilerMaterialCounters::StaticStruct(), &MaterialCounters, JsonObject, 0, 0, nullptr, EJsonObjectConversionFlags::SkipStandardizeCase);
+		RootObject->SetObjectField("MaterialCounters", JsonObject);
+	}
+
+	TArray<TSharedPtr<FJsonValue>> CompileStatsArray;
+	for (TSparseArray<ShaderCompilerStats>::TConstIterator It(CompileStats); It; ++It)
+	{
+		if (!CompileStats.IsValidIndex(It.GetIndex()))
+		{
+			continue;
+		}
+
+		TSharedPtr<FJsonObject> InnerObject(MakeShared<FJsonObject>());
+		for (const TPair<FString, FShaderStats>& Pair : *It)
+		{
+			TSharedRef<FJsonObject> JsonObject = MakeShared<FJsonObject>();
+			FJsonObjectConverter::UStructToJsonObject(FShaderStats::StaticStruct(), &Pair.Value, JsonObject, 0, 0, nullptr, EJsonObjectConversionFlags::SkipStandardizeCase);
+			InnerObject->SetObjectField(Pair.Key, JsonObject);
+		}
+		CompileStatsArray.Push(MakeShared<FJsonValueObject>(InnerObject));
+	}
+	RootObject->SetArrayField("CompileStats", CompileStatsArray);
+
+	TSharedPtr<FJsonObject> ShaderTimingsObject(MakeShared<FJsonObject>());
+	for (const TPair<FString, FShaderTimings>& TimingPair : ShaderTimings)
+	{
+		TSharedRef<FJsonObject> JsonObject = MakeShared<FJsonObject>();
+		FJsonObjectConverter::UStructToJsonObject(FShaderTimings::StaticStruct(), &TimingPair.Value, JsonObject, 0, 0, nullptr, EJsonObjectConversionFlags::SkipStandardizeCase);
+		ShaderTimingsObject->SetObjectField(TimingPair.Key, JsonObject);
+	}
+	RootObject->SetObjectField("ShaderTimings", ShaderTimingsObject);
+	return RootObject;
+}
+
 void FShaderCompilerStats::WriteToCompactBinary(FCbWriter& Writer)
 {
 	FScopeLock Lock(&CompileStatsLock);
@@ -666,12 +712,12 @@ void FShaderCompilerStats::ReadFromCompactBinary(FCbObjectView& Reader)
 {
 	FScopeLock Lock(&CompileStatsLock);
 	FMemoryView CountersMem = Reader["Counters"].AsBinaryView();
-	check(CountersMem.GetSize() == sizeof(FCounters));
-	Counters = *reinterpret_cast<const FCounters*>(CountersMem.GetData());
+	check(CountersMem.GetSize() == sizeof(FShaderCompilerCounters));
+	Counters = *reinterpret_cast<const FShaderCompilerCounters*>(CountersMem.GetData());
 
 	FMemoryView MaterialCountersMem = Reader["MaterialCounters"].AsBinaryView();
-	check(MaterialCountersMem.GetSize() == sizeof(FMaterialCounters));
-	MaterialCounters = *reinterpret_cast<const FMaterialCounters*>(MaterialCountersMem.GetData());
+	check(MaterialCountersMem.GetSize() == sizeof(FShaderCompilerMaterialCounters));
+	MaterialCounters = *reinterpret_cast<const FShaderCompilerMaterialCounters*>(MaterialCountersMem.GetData());
 
 	FCbArrayView CompileStatIndicesView = Reader["CompileStatIndices"].AsArrayView();
 	FCbArrayView CompileStatsView = Reader["CompileStats"].AsArrayView();
@@ -900,7 +946,7 @@ void FShaderCompilerStats::RegisterDistributedBuildStats(const FDistributedBuild
 	Counters.MaxActiveAgentCores = FMath::Max(Counters.MaxActiveAgentCores, InStats.MaxActiveAgentCores);
 }
 
-void FShaderCompilerStats::FMaterialCounters::WriteStatSummary(const TCHAR* AggregatedSuffix)
+void FShaderCompilerMaterialCounters::WriteStatSummary(const TCHAR* AggregatedSuffix)
 {
 	auto CalcTimePercentage = [&](double Val) {
 		return  (int)round(Val / FMath::Max(1e-6, MaterialTranslateTotalTimeSec) * 100);
@@ -917,7 +963,7 @@ void FShaderCompilerStats::FMaterialCounters::WriteStatSummary(const TCHAR* Aggr
 	UE_LOG(LogShaderCompilers, Display, TEXT("Material Cache Hits: %d (%d%%)"), MaterialCacheHits, HitsPercentage);
 }
 
-void FShaderCompilerStats::FMaterialCounters::GatherAnalytics(TArray<FAnalyticsEventAttribute>& Attributes)
+void FShaderCompilerMaterialCounters::GatherAnalytics(TArray<FAnalyticsEventAttribute>& Attributes)
 {
 	Attributes.Emplace(TEXT("Material_NumMaterialsCooked"), NumMaterialsCooked);
 	Attributes.Emplace(TEXT("Material_MaterialTranslateCalls"), MaterialTranslateCalls);
@@ -957,7 +1003,7 @@ void FShaderCompilerStats::RegisterCookedShaders(uint32 NumCooked, float Compile
 		CompileStats.Insert(Platform, Stats);
 	}
 
-	FShaderCompilerStats::FShaderStats& Stats = CompileStats[Platform].FindOrAdd(MaterialPath);
+	FShaderStats& Stats = CompileStats[Platform].FindOrAdd(MaterialPath);
 	Stats.CompileTime += CompileTime;
 	bool bFound = false;
 	for (FShaderCompilerSinglePermutationStat& Stat : Stats.PermutationCompilations)
@@ -995,7 +1041,7 @@ void FShaderCompilerStats::RegisterCompiledShaders(uint32 NumCompiled, EShaderPl
 		ShaderCompilerStats Stats;
 		CompileStats.Insert(Platform, Stats);
 	}
-	FShaderCompilerStats::FShaderStats& Stats = CompileStats[Platform].FindOrAdd(MaterialPath);
+	FShaderStats& Stats = CompileStats[Platform].FindOrAdd(MaterialPath);
 
 	bool bFound = false;
 	for (FShaderCompilerSinglePermutationStat& Stat : Stats.PermutationCompilations)
diff --git a/Engine/Source/Runtime/Engine/Public/ShaderCompiler.h b/Engine/Source/Runtime/Engine/Public/ShaderCompiler.h
--- a/Engine/Source/Runtime/Engine/Public/ShaderCompiler.h
+++ b/Engine/Source/Runtime/Engine/Public/ShaderCompiler.h
@@ -29,6 +29,8 @@
 #include "SceneTypes.h"
 #include "UObject/StrongObjectPtr.h"
 
+#include "ShaderCompiler.generated.h"
+
 class FAsyncCompilationNotification;
 class FCbObjectView;
 class FCbWriter;
@@ -36,10 +38,11 @@ class FVertexFactoryType;
 class IDistributedBuildController;
 class FMaterialShaderMap;
 class FShaderCompileJob;
-class FShaderCompilerStats;
+struct FShaderCompilerStats;
 class FShaderKeyGenerator;
 class FShaderPipelineCompileJob;
 class UMaterialInterface;
+class FJsonObject;
 struct FAnalyticsEventAttribute;
 
 DECLARE_LOG_CATEGORY_EXTERN(LogShaderCompilers, Log, All);
@@ -374,89 +377,360 @@ using FShaderMapFinalizeResults = FShaderMapCompileResults;
 
 struct FDistributedBuildStats;
 
-class FShaderCompilerStats
+USTRUCT()
+struct FShaderCompilerCounters
 {
-public:
-	struct FShaderCompilerSinglePermutationStat
+	GENERATED_BODY()
+	/** This tracks accumulated wait time from local workers during the lifetime of the stats.
+	 *
+	 * Wait time is only counted for local workers that are alive and not between their invocations
+	 */
+	UPROPERTY()
+	double AccumulatedLocalWorkerIdleTime = 0.0;
+
+	/** How many times we registered idle time? */
+	UPROPERTY()
+	double TimesLocalWorkersWereIdle = 0;
+
+	/** Number of jobs assigned to workers, no matter if they completed or not - used to average pending time. */
+	UPROPERTY()
+	int64 JobsAssigned = 0;
+
+	/** Total number jobs completed. */
+	UPROPERTY()
+	int64 JobsCompleted = 0;
+
+	/** Amount of time a job had to spent in pending queue (i.e. waiting to be assigned to a worker). */
+	UPROPERTY()
+	double AccumulatedPendingTime = 0;
+
+	/** Max amount of time any single job was pending (waiting to be assigned to a worker). */
+	UPROPERTY()
+	double MaxPendingTime = 0;
+
+	/** Amount of time job spent being processed by the worker. */
+	UPROPERTY()
+	double AccumulatedJobExecutionTime = 0;
+
+	/** Max amount of time any single job spent being processed by the worker. */
+	UPROPERTY()
+	double MaxJobExecutionTime = 0;
+
+	/** Amount of time job spent being processed overall. */
+	UPROPERTY()
+	double AccumulatedJobLifeTime = 0;
+
+	/** Max amount of time any single job spent being processed overall. */
+	UPROPERTY()
+	double MaxJobLifeTime = 0;
+
+	/** Time spent in tasks generated in FShaderJobCache::SubmitJobs, plus stall time on mutex locks in those tasks */
+	UPROPERTY()
+	double AccumulatedTaskSubmitJobs = 0.0;
+	UPROPERTY()
+	double AccumulatedTaskSubmitJobsStall = 0.0;
+
+	/** Number of local job batches seen. */
+	UPROPERTY()
+	int64 LocalJobBatchesSeen = 0;
+
+	/** Total jobs in local job batches. */
+	UPROPERTY()
+	int64 TotalJobsReportedInLocalJobBatches = 0;
+
+	/** Number of distributed job batches seen. */
+	UPROPERTY()
+	int64 DistributedJobBatchesSeen = 0;
+
+	/** Total jobs in local job batches. */
+	UPROPERTY()
+	int64 TotalJobsReportedInDistributedJobBatches = 0;
+
+	/** Size of the smallest output shader code. */
+	UPROPERTY()
+	int32 MinShaderCodeSize = 0;
+
+	/** Size of the largest output shader code. */
+	UPROPERTY()
+	int32 MaxShaderCodeSize = 0;
+
+	/** Total accumulated size of all output shader codes. */
+	UPROPERTY()
+	uint64 AccumulatedShaderCodeSize = 0;
+
+	/** Number of accumulated output shader codes. */
+	UPROPERTY()
+	uint64 NumAccumulatedShaderCodes = 0;
+
+	/** Total number of DDC misses on shader maps. */
+	UPROPERTY()
+	uint32 ShaderMapDDCMisses = 0;
+
+	/** Total number of DDC hits on shader maps. */
+	UPROPERTY()
+	uint32 ShaderMapDDCHits = 0;
+
+	/** Total number of job cache query attempts. */
+	UPROPERTY()
+	uint64 TotalCacheSearchAttempts = 0;
+
+	/** Total number of hits in the job cache (i.e. input hashes seen >1 time) */
+	UPROPERTY()
+	uint64 TotalCacheHits = 0;
+
+	/** Total number of duplicate jobs (input hash matches an in-flight job, processed when in-flight job completes) */
+	UPROPERTY()
+	uint32 TotalCacheDuplicates = 0;
+
+	/** Total number of DDC queries in the job cache (per-shader DDC). */
+	UPROPERTY()
+	uint32 TotalCacheDDCQueries = 0;
+
+	/** Total number of DDC hits in the job cache (per shader DDC, as opposed to shader map DDC stats above). */
+	UPROPERTY()
+	uint32 TotalCacheDDCHits = 0;
+
+	/** Total number of unique input hashes seen in job cache queries */
+	UPROPERTY()
+	uint64 UniqueCacheInputHashes = 0;
+
+	/** Total number of unique job outputs stored in the cache.
+	  * Outputs are deduplicated based on a content hash so this number is in practice smaller than UniqueCacheInputHashes.
+	  */
+	UPROPERTY()
+	uint64 UniqueCacheOutputs = 0;
+
+	/** Total amount of memory currently used by the job cache */
+	UPROPERTY()
+	uint64 CacheMemUsed = 0;
+
+	/** Memory budget allocated for the job cache */
+	UPROPERTY()
+	uint64 CacheMemBudget = 0;
+
+	/** Maximum number of remote agents used during compilation. */
+	UPROPERTY()
+	uint32 MaxRemoteAgents = 0;
+
+	/** Maximum number of CPU cores active across all remote agents. */
+	UPROPERTY()
+	uint32 MaxActiveAgentCores = 0;
+
+	FShaderCompilerCounters& operator+=(const FShaderCompilerCounters& Other)
 	{
-		FShaderCompilerSinglePermutationStat(FString PermutationString, uint32 Compiled, uint32 Cooked)
-			: PermutationString(PermutationString)
-			, Compiled(Compiled)
-			, Cooked(Cooked)
-			, CompiledDouble(0)
-			, CookedDouble(0)
-
-		{}
-		FString PermutationString;
-		uint32 Compiled;
-		uint32 Cooked;
-		uint32 CompiledDouble;
-		uint32 CookedDouble;
-	};
-	struct FShaderStats
+		AccumulatedLocalWorkerIdleTime += Other.AccumulatedLocalWorkerIdleTime;
+		TimesLocalWorkersWereIdle += Other.TimesLocalWorkersWereIdle;
+		JobsAssigned += Other.JobsAssigned;
+		JobsCompleted += Other.JobsCompleted;
+		AccumulatedPendingTime += Other.AccumulatedPendingTime;
+		MaxPendingTime = FMath::Max(Other.MaxPendingTime, MaxPendingTime);
+		AccumulatedJobExecutionTime += Other.AccumulatedJobExecutionTime;
+		MaxJobExecutionTime = FMath::Max(Other.MaxJobExecutionTime, MaxJobExecutionTime);
+		AccumulatedJobLifeTime += Other.AccumulatedJobLifeTime;
+		MaxJobLifeTime = FMath::Max(Other.MaxJobLifeTime, MaxJobLifeTime);
+		AccumulatedTaskSubmitJobs += Other.AccumulatedTaskSubmitJobs;
+		AccumulatedTaskSubmitJobsStall += Other.AccumulatedTaskSubmitJobsStall;
+		LocalJobBatchesSeen += Other.LocalJobBatchesSeen;
+		TotalJobsReportedInLocalJobBatches += Other.TotalJobsReportedInLocalJobBatches;
+		DistributedJobBatchesSeen += Other.DistributedJobBatchesSeen;
+		TotalJobsReportedInDistributedJobBatches += Other.TotalJobsReportedInDistributedJobBatches;
+		if (Other.MinShaderCodeSize > 0)
+		{
+			MinShaderCodeSize = (MinShaderCodeSize > 0 ? FMath::Min(MinShaderCodeSize, Other.MinShaderCodeSize) : Other.MinShaderCodeSize);
+		}
+		MaxShaderCodeSize = FMath::Max(Other.MaxShaderCodeSize, MaxShaderCodeSize);
+		AccumulatedShaderCodeSize += Other.AccumulatedShaderCodeSize;
+		NumAccumulatedShaderCodes += Other.NumAccumulatedShaderCodes;
+		ShaderMapDDCMisses += Other.ShaderMapDDCMisses;
+		ShaderMapDDCHits += Other.ShaderMapDDCHits;
+		TotalCacheSearchAttempts += Other.TotalCacheSearchAttempts;
+		TotalCacheHits += Other.TotalCacheHits;
+		TotalCacheDuplicates += Other.TotalCacheDuplicates;
+		TotalCacheDDCQueries += Other.TotalCacheDDCQueries;
+		TotalCacheDDCHits += Other.TotalCacheDDCHits;
+		UniqueCacheInputHashes += Other.UniqueCacheInputHashes;
+		UniqueCacheOutputs += Other.UniqueCacheOutputs;
+		CacheMemUsed += Other.CacheMemUsed;
+		CacheMemBudget += Other.CacheMemBudget;
+		MaxRemoteAgents += Other.MaxRemoteAgents;
+		MaxActiveAgentCores += Other.MaxActiveAgentCores;
+
+		return *this;
+	}
+};
+
+USTRUCT()
+struct FShaderCompilerMaterialCounters
+{
+	GENERATED_BODY()
+	/** The total number of materials cooked.  This corresponds to UMaterialInterface::Presave() */
+	int32 NumMaterialsCooked = 0;
+
+	/** The total number of materials that have been translated.  */
+	UPROPERTY()
+	int32 MaterialTranslateCalls = 0;
+
+	/** The total time in seconds to translate all materials.  */
+	UPROPERTY()
+	double MaterialTranslateTotalTimeSec = 0.0;
+
+	/** The total time spent actually translating materials (rather than for instance accessing the DDC cache). */
+	UPROPERTY()
+	double MaterialTranslateTranslationOnlyTimeSec = 0.0;
+
+	/** The total time spent serializing DDC results. */
+	UPROPERTY()
+	double MaterialTranslateSerializationOnlyTimeSec = 0.0;
+
+	/** The total number times a material translation was skipped because the the results were in the DDC. */
+	UPROPERTY()
+	int32 MaterialCacheHits = 0;
+
+	FShaderCompilerMaterialCounters& operator+=(const FShaderCompilerMaterialCounters& Other)
+	{
+		NumMaterialsCooked += Other.NumMaterialsCooked;
+		MaterialTranslateCalls += Other.MaterialTranslateCalls;
+		MaterialTranslateTotalTimeSec += Other.MaterialTranslateTotalTimeSec;
+		MaterialTranslateTranslationOnlyTimeSec += Other.MaterialTranslateTranslationOnlyTimeSec;
+		MaterialTranslateSerializationOnlyTimeSec += Other.MaterialTranslateSerializationOnlyTimeSec;
+		MaterialCacheHits += Other.MaterialCacheHits;
+
+		return *this;
+	}
+
+	void WriteStatSummary(const TCHAR* AggregatedSuffix);
+	void GatherAnalytics(TArray<FAnalyticsEventAttribute>& Attributes);
+};
+
+/** Structure used to describe compiling time of a shader type (for all the instances of it that we have seen). */
+USTRUCT()
+struct FShaderTimings
+{
+	GENERATED_BODY()
+
+	UPROPERTY()
+	float MinCompileTime = 0.0f;
+	UPROPERTY()
+	float MaxCompileTime = 0.0f;
+	UPROPERTY()
+	float TotalCompileTime = 0.0f;
+	UPROPERTY()
+	float TotalPreprocessTime = 0.0f;
+	UPROPERTY()
+	int32 NumCompiled = 0;
+	UPROPERTY()
+	float AverageCompileTime = 0.0f;	// stored explicitly as an optimization
+
+	FShaderTimings& operator+=(const FShaderTimings& Other)
 	{
-		TArray<FShaderCompilerSinglePermutationStat> PermutationCompilations;
-		uint32 Compiled = 0;
-		uint32 Cooked = 0;
-		uint32 CompiledDouble = 0;
-		uint32 CookedDouble = 0;
-		float CompileTime = 0.f;
-
-		FShaderStats& operator+=(const FShaderStats& Other)
+		MinCompileTime = FMath::Min(MinCompileTime, Other.MinCompileTime);
+		MaxCompileTime = FMath::Max(MaxCompileTime, Other.MaxCompileTime);
+		TotalCompileTime += Other.TotalCompileTime;
+		TotalPreprocessTime += Other.TotalPreprocessTime;
+		NumCompiled += Other.NumCompiled;
+		if (NumCompiled)
 		{
-			if (Compiled)
-			{
-				CompiledDouble += Other.Compiled;
-			}
-			else
-			{
-				Compiled += Other.Compiled;
-			}
-
-			if (Cooked)
-			{
-				CookedDouble += Other.Cooked;
-			}
-			else
-			{
-				Cooked += Other.Cooked;
-			}
-
-			CompiledDouble += Other.CompiledDouble;
-			CookedDouble += Other.CookedDouble;
-			CompileTime += Other.CompileTime;
-
-			PermutationCompilations.Append(Other.PermutationCompilations);
-
-			return *this;
+			AverageCompileTime = TotalCompileTime / static_cast<float>(NumCompiled);
 		}
-	};
-	using ShaderCompilerStats = TMap<FString, FShaderStats>;
+		return *this;
+	}
+};
+
+USTRUCT()
+struct FShaderCompilerSinglePermutationStat
+{
+	GENERATED_BODY()
+
+	FShaderCompilerSinglePermutationStat()
+		: PermutationString()
+		, Compiled(0)
+		, Cooked(0)
+		, CompiledDouble(0)
+		, CookedDouble(0)
+	{}
+
+	FShaderCompilerSinglePermutationStat(FString PermutationString, uint32 Compiled, uint32 Cooked)
+		: PermutationString(PermutationString)
+		, Compiled(Compiled)
+		, Cooked(Cooked)
+		, CompiledDouble(0)
+		, CookedDouble(0)
+
+	{}
+
+	UPROPERTY()
+	FString PermutationString;
+
+	UPROPERTY()
+	uint32 Compiled;
+
+	UPROPERTY()
+	uint32 Cooked;
+
+	UPROPERTY()
+	uint32 CompiledDouble;
+
+	UPROPERTY()
+	uint32 CookedDouble;
+};
+
+USTRUCT()
+struct FShaderStats
+{
+	GENERATED_BODY()
 
-	/** Structure used to describe compiling time of a shader type (for all the instances of it that we have seen). */
-	struct FShaderTimings
+	UPROPERTY()
+	TArray<FShaderCompilerSinglePermutationStat> PermutationCompilations;
+
+	UPROPERTY()
+	uint32 Compiled = 0;
+
+	UPROPERTY()
+	uint32 Cooked = 0;
+
+	UPROPERTY()
+	uint32 CompiledDouble = 0;
+
+	UPROPERTY()
+	uint32 CookedDouble = 0;
+
+	UPROPERTY()
+	float CompileTime = 0.f;
+
+	FShaderStats& operator+=(const FShaderStats& Other)
 	{
-		float MinCompileTime = 0.0f;
-		float MaxCompileTime = 0.0f;
-		float TotalCompileTime = 0.0f;
-		float TotalPreprocessTime = 0.0f;
-		int32 NumCompiled = 0;
-		float AverageCompileTime = 0.0f;	// stored explicitly as an optimization
-
-		FShaderTimings& operator+=(const FShaderTimings& Other)
+		if (Compiled)
 		{
-			MinCompileTime = FMath::Min(MinCompileTime, Other.MinCompileTime);
-			MaxCompileTime = FMath::Max(MaxCompileTime, Other.MaxCompileTime);
-			TotalCompileTime += Other.TotalCompileTime;
-			TotalPreprocessTime += Other.TotalPreprocessTime;
-			NumCompiled += Other.NumCompiled;
-			if (NumCompiled)
-			{
-				AverageCompileTime = TotalCompileTime / static_cast<float>(NumCompiled);
-			}
-			return *this;
+			CompiledDouble += Other.Compiled;
 		}
-	};
+		else
+		{
+			Compiled += Other.Compiled;
+		}
+
+		if (Cooked)
+		{
+			CookedDouble += Other.Cooked;
+		}
+		else
+		{
+			Cooked += Other.Cooked;
+		}
+
+		CompiledDouble += Other.CompiledDouble;
+		CookedDouble += Other.CookedDouble;
+		CompileTime += Other.CompileTime;
+
+		PermutationCompilations.Append(Other.PermutationCompilations);
+
+		return *this;
+	}
+};
+
+struct FShaderCompilerStats
+{
+public:
+	using ShaderCompilerStats = TMap<FString, FShaderStats>;
 
 	void IncrementMaterialCook();
 	void IncrementMaterialTranslated(double InTotalTime, double InTranslationOnlyTime, double InSerializeTime);
@@ -472,6 +746,7 @@ public:
 	ENGINE_API void Aggregate(FShaderCompilerStats& Other);
 	ENGINE_API void WriteToCompactBinary(FCbWriter& Writer);
 	ENGINE_API void ReadFromCompactBinary(FCbObjectView& Reader);
+	ENGINE_API TSharedPtr<FJsonObject> ToJson();
 	inline void SetMultiProcessAggregated() { bMultiProcessAggregated = true; }
 
 	void AddDDCMiss(uint32 NumMisses);
@@ -511,192 +786,8 @@ private:
 	FCriticalSection CompileStatsLock;
 	TSparseArray<ShaderCompilerStats> CompileStats;
 
-	struct FCounters
-	{
-		/** This tracks accumulated wait time from local workers during the lifetime of the stats.
-		 *
-		 * Wait time is only counted for local workers that are alive and not between their invocations
-		 */
-		double AccumulatedLocalWorkerIdleTime = 0.0;
-
-		/** How many times we registered idle time? */
-		double TimesLocalWorkersWereIdle = 0;
-
-		/** Number of jobs assigned to workers, no matter if they completed or not - used to average pending time. */
-		int64 JobsAssigned = 0;
-
-		/** Total number jobs completed. */
-		int64 JobsCompleted = 0;
-
-		/** Amount of time a job had to spent in pending queue (i.e. waiting to be assigned to a worker). */
-		double AccumulatedPendingTime = 0;
-
-		/** Max amount of time any single job was pending (waiting to be assigned to a worker). */
-		double MaxPendingTime = 0;
-
-		/** Amount of time job spent being processed by the worker. */
-		double AccumulatedJobExecutionTime = 0;
-
-		/** Max amount of time any single job spent being processed by the worker. */
-		double MaxJobExecutionTime = 0;
-
-		/** Amount of time job spent being processed overall. */
-		double AccumulatedJobLifeTime = 0;
-
-		/** Max amount of time any single job spent being processed overall. */
-		double MaxJobLifeTime = 0;
-
-		/** Time spent in tasks generated in FShaderJobCache::SubmitJobs, plus stall time on mutex locks in those tasks */
-		double AccumulatedTaskSubmitJobs = 0.0;
-		double AccumulatedTaskSubmitJobsStall = 0.0;
-
-		/** Number of local job batches seen. */
-		int64 LocalJobBatchesSeen = 0;
-
-		/** Total jobs in local job batches. */
-		int64 TotalJobsReportedInLocalJobBatches = 0;
-
-		/** Number of distributed job batches seen. */
-		int64 DistributedJobBatchesSeen = 0;
-
-		/** Total jobs in local job batches. */
-		int64 TotalJobsReportedInDistributedJobBatches = 0;
-
-		/** Size of the smallest output shader code. */
-		int32 MinShaderCodeSize = 0;
-
-		/** Size of the largest output shader code. */
-		int32 MaxShaderCodeSize = 0;
-
-		/** Total accumulated size of all output shader codes. */
-		uint64 AccumulatedShaderCodeSize = 0;
-
-		/** Number of accumulated output shader codes. */
-		uint64 NumAccumulatedShaderCodes = 0;
-
-		/** Total number of DDC misses on shader maps. */
-		uint32 ShaderMapDDCMisses = 0;
-
-		/** Total number of DDC hits on shader maps. */
-		uint32 ShaderMapDDCHits = 0;
-
-		/** Total number of job cache query attempts. */
-		uint64 TotalCacheSearchAttempts = 0;
-
-		/** Total number of hits in the job cache (i.e. input hashes seen >1 time) */
-		uint64 TotalCacheHits = 0;
-
-		/** Total number of duplicate jobs (input hash matches an in-flight job, processed when in-flight job completes) */
-		uint32 TotalCacheDuplicates = 0;
-
-		/** Total number of DDC queries in the job cache (per-shader DDC). */
-		uint32 TotalCacheDDCQueries = 0;
-
-		/** Total number of DDC hits in the job cache (per shader DDC, as opposed to shader map DDC stats above). */
-		uint32 TotalCacheDDCHits = 0;
-
-		/** Total number of unique input hashes seen in job cache queries */
-		uint64 UniqueCacheInputHashes = 0;
-
-		/** Total number of unique job outputs stored in the cache.
-		  * Outputs are deduplicated based on a content hash so this number is in practice smaller than UniqueCacheInputHashes.
-		  */
-		uint64 UniqueCacheOutputs = 0;
-
-		/** Total amount of memory currently used by the job cache */
-		uint64 CacheMemUsed = 0;
-
-		/** Memory budget allocated for the job cache */
-		uint64 CacheMemBudget = 0;
-
-		/** Maximum number of remote agents used during compilation. */
-		uint32 MaxRemoteAgents = 0;
-
-		/** Maximum number of CPU cores active across all remote agents. */
-		uint32 MaxActiveAgentCores = 0;
-
-		FCounters& operator+=(const FCounters& Other)
-		{
-			AccumulatedLocalWorkerIdleTime += Other.AccumulatedLocalWorkerIdleTime;
-			TimesLocalWorkersWereIdle += Other.TimesLocalWorkersWereIdle;
-			JobsAssigned += Other.JobsAssigned;
-			JobsCompleted += Other.JobsCompleted;
-			AccumulatedPendingTime += Other.AccumulatedPendingTime;
-			MaxPendingTime = FMath::Max(Other.MaxPendingTime, MaxPendingTime);
-			AccumulatedJobExecutionTime += Other.AccumulatedJobExecutionTime;
-			MaxJobExecutionTime = FMath::Max(Other.MaxJobExecutionTime, MaxJobExecutionTime);
-			AccumulatedJobLifeTime += Other.AccumulatedJobLifeTime;
-			MaxJobLifeTime = FMath::Max(Other.MaxJobLifeTime, MaxJobLifeTime);
-			AccumulatedTaskSubmitJobs += Other.AccumulatedTaskSubmitJobs;
-			AccumulatedTaskSubmitJobsStall += Other.AccumulatedTaskSubmitJobsStall;
-			LocalJobBatchesSeen += Other.LocalJobBatchesSeen;
-			TotalJobsReportedInLocalJobBatches += Other.TotalJobsReportedInLocalJobBatches;
-			DistributedJobBatchesSeen += Other.DistributedJobBatchesSeen;
-			TotalJobsReportedInDistributedJobBatches += Other.TotalJobsReportedInDistributedJobBatches;
-			if (Other.MinShaderCodeSize > 0)
-			{
-				MinShaderCodeSize = (MinShaderCodeSize > 0 ? FMath::Min(MinShaderCodeSize, Other.MinShaderCodeSize) : Other.MinShaderCodeSize);
-			}
-			MaxShaderCodeSize = FMath::Max(Other.MaxShaderCodeSize, MaxShaderCodeSize);
-			AccumulatedShaderCodeSize += Other.AccumulatedShaderCodeSize;
-			NumAccumulatedShaderCodes += Other.NumAccumulatedShaderCodes;
-			ShaderMapDDCMisses += Other.ShaderMapDDCMisses;
-			ShaderMapDDCHits += Other.ShaderMapDDCHits;
-			TotalCacheSearchAttempts += Other.TotalCacheSearchAttempts;
-			TotalCacheHits += Other.TotalCacheHits;
-			TotalCacheDuplicates += Other.TotalCacheDuplicates;
-			TotalCacheDDCQueries += Other.TotalCacheDDCQueries;
-			TotalCacheDDCHits += Other.TotalCacheDDCHits;
-			UniqueCacheInputHashes += Other.UniqueCacheInputHashes;
-			UniqueCacheOutputs += Other.UniqueCacheOutputs;
-			CacheMemUsed += Other.CacheMemUsed;
-			CacheMemBudget += Other.CacheMemBudget;
-			MaxRemoteAgents += Other.MaxRemoteAgents;
-			MaxActiveAgentCores += Other.MaxActiveAgentCores;
-
-			return *this;
-		}
-	};
-
-	FCounters Counters;
-
-	struct FMaterialCounters
-	{
-		/** The total number of materials cooked.  This corresponds to UMaterialInterface::Presave() */
-		int32 NumMaterialsCooked = 0;
-
-		/** The total number of materials that have been translated.  */
-		int32 MaterialTranslateCalls = 0;
-
-		/** The total time in seconds to translate all materials.  */
-		double MaterialTranslateTotalTimeSec = 0.0;
-
-		/** The total time spent actually translating materials (rather than for instance accessing the DDC cache). */
-		double MaterialTranslateTranslationOnlyTimeSec = 0.0;
-
-		/** The total time spent serializing DDC results. */
-		double MaterialTranslateSerializationOnlyTimeSec = 0.0;
-
-		/** The total number times a material translation was skipped because the the results were in the DDC. */
-		int32 MaterialCacheHits = 0;
-		
-		FMaterialCounters& operator+=(const FMaterialCounters& Other)
-		{
-			NumMaterialsCooked += Other.NumMaterialsCooked;
-			MaterialTranslateCalls += Other.MaterialTranslateCalls;
-			MaterialTranslateTotalTimeSec += Other.MaterialTranslateTotalTimeSec;
-			MaterialTranslateTranslationOnlyTimeSec += Other.MaterialTranslateTranslationOnlyTimeSec;
-			MaterialTranslateSerializationOnlyTimeSec += Other.MaterialTranslateSerializationOnlyTimeSec;
-			MaterialCacheHits += Other.MaterialCacheHits;
-
-			return *this;
-		}
-
-		void WriteStatSummary(const TCHAR* AggregatedSuffix);
-		void GatherAnalytics(TArray<FAnalyticsEventAttribute>& Attributes);
-	};
-
-	FMaterialCounters MaterialCounters;
+	FShaderCompilerCounters Counters;
+	FShaderCompilerMaterialCounters MaterialCounters;
 
 	/** Accumulates the job lifetimes without overlaps */
 	TArray<TInterval<double>> JobLifeTimeIntervals;
diff --git a/Engine/Source/Runtime/Json/Public/Serialization/JsonWriter.h b/Engine/Source/Runtime/Json/Public/Serialization/JsonWriter.h
--- a/Engine/Source/Runtime/Json/Public/Serialization/JsonWriter.h
+++ b/Engine/Source/Runtime/Json/Public/Serialization/JsonWriter.h
@@ -259,7 +259,7 @@ public:
 
 	// WARNING: THIS IS DANGEROUS. Use this only if you know for a fact that the Value is valid JSON!
 	// Use this to insert the results of a different JSON Writer in.
-	void WriteRawJSONValue( const FString& Identifier, const FString& Value )
+	void WriteRawJSONValue( FStringView Identifier, FStringView Value )
 	{
 		check( Stack.Top() == EJson::Object );
 		WriteIdentifier( Identifier );
