Fix for required bones computation
    - ensure that the bone indices are sorted when ensuring that parent have been added
    - mesh bone reduction: invalidate DDC when removing bones as the LODModels are updated

#rnx

diff --git a/Engine/Source/Developer/MeshBoneReduction/Private/MeshBoneReduction.cpp b/Engine/Source/Developer/MeshBoneReduction/Private/MeshBoneReduction.cpp
--- a/Engine/Source/Developer/MeshBoneReduction/Private/MeshBoneReduction.cpp
+++ b/Engine/Source/Developer/MeshBoneReduction/Private/MeshBoneReduction.cpp
@@ -520,6 +520,8 @@ public:
 			}
 
 			delete SrcModel;
+
+			SkeletalMesh->InvalidateDeriveDataCacheGUID();
 		}
 		else
 		{
diff --git a/Engine/Source/Runtime/Engine/Private/Components/SkeletalMeshComponent.cpp b/Engine/Source/Runtime/Engine/Private/Components/SkeletalMeshComponent.cpp
--- a/Engine/Source/Runtime/Engine/Private/Components/SkeletalMeshComponent.cpp
+++ b/Engine/Source/Runtime/Engine/Private/Components/SkeletalMeshComponent.cpp
@@ -1945,6 +1945,7 @@ void USkeletalMeshComponent::ComputeRequiredBones(TArray<FBoneIndexType>& OutReq
 
 	// Ensure that we have a complete hierarchy down to those bones.
 	FAnimationRuntime::EnsureParentsPresent(OutRequiredBones, SkelMesh->GetRefSkeleton());
+	OutRequiredBones.Sort();
 
 	OutFillComponentSpaceTransformsRequiredBones.Reset(OutRequiredBones.Num() + NeededBonesForFillComponentSpaceTransforms.Num());
 	OutFillComponentSpaceTransformsRequiredBones = OutRequiredBones;
