[shaders] bytecode sharing between shader and shadermap DDC buckets
- create utility classes (serialization context objects) which allow serialization of shader code blobs independently of "shader object" (i.e. whatever type of thing we're pushing to the cache, shader job or shader map) structure
- commonize bytecode format in cache entries for shadermap and shader DDC (FShaderCodeResource struct)
- rework per-shader caching to use the serialization helpers to separate serialization of job output struct data and bytecode
- rework material and global shadermap DDC serialization to use the serialization helpers to separate serialization of shadermap structural data and bytecode

diff --git a/Engine/Plugins/Compositing/OpenColorIO/Source/OpenColorIO/Private/OpenColorIOShaderMap.cpp b/Engine/Plugins/Compositing/OpenColorIO/Source/OpenColorIO/Private/OpenColorIOShaderMap.cpp
--- a/Engine/Plugins/Compositing/OpenColorIO/Source/OpenColorIO/Private/OpenColorIOShaderMap.cpp
+++ b/Engine/Plugins/Compositing/OpenColorIO/Source/OpenColorIO/Private/OpenColorIOShaderMap.cpp
@@ -626,7 +626,8 @@ bool FOpenColorIOShaderMap::Serialize(FArchive& Ar)
 	// Note: This is saved to the DDC, not into packages (except when cooked)
 	// Backwards compatibility therefore will not work based on the version of Ar
 	// Instead, just bump OPENCOLORIO_DERIVEDDATA_VER
-	return Super::Serialize(Ar, {});
+	FShaderSerializeContext Ctx(Ar);
+	return Super::Serialize(Ctx);
 }
 
 #if WITH_EDITOR
diff --git a/Engine/Plugins/FX/Niagara/Source/NiagaraEditor/Private/NiagaraCompilationTasks.cpp b/Engine/Plugins/FX/Niagara/Source/NiagaraEditor/Private/NiagaraCompilationTasks.cpp
--- a/Engine/Plugins/FX/Niagara/Source/NiagaraEditor/Private/NiagaraCompilationTasks.cpp
+++ b/Engine/Plugins/FX/Niagara/Source/NiagaraEditor/Private/NiagaraCompilationTasks.cpp
@@ -275,7 +275,8 @@ namespace NiagaraCompilationTasksImpl
 
 		TArray<uint8> BinaryData;
 		FMemoryWriter Ar(BinaryData, true);
-		ShaderMap->Serialize(Ar, FShaderMapBase::FSerializationContext{});
+		FShaderSerializeContext Ctx(Ar);
+		ShaderMap->Serialize(Ctx);
 
 		if (!BinaryData.IsEmpty() && !Ar.IsError())
 		{
@@ -294,7 +295,8 @@ namespace NiagaraCompilationTasksImpl
 
 		FMemoryReaderView Ar(SharedBuffer.GetView(), true);
 		FNiagaraShaderMapRef ShaderMap = new FNiagaraShaderMap(FNiagaraShaderMap::WorkerThread);
-		ShaderMap->Serialize(Ar, FShaderMapBase::FSerializationContext{});
+		FShaderSerializeContext Ctx(Ar);
+		ShaderMap->Serialize(Ctx);
 
 		if (Ar.IsError())
 		{
diff --git a/Engine/Plugins/FX/Niagara/Source/NiagaraShader/Private/NiagaraShader.cpp b/Engine/Plugins/FX/Niagara/Source/NiagaraShader/Private/NiagaraShader.cpp
--- a/Engine/Plugins/FX/Niagara/Source/NiagaraShader/Private/NiagaraShader.cpp
+++ b/Engine/Plugins/FX/Niagara/Source/NiagaraShader/Private/NiagaraShader.cpp
@@ -546,9 +546,9 @@ void FNiagaraShaderMap::LoadFromDerivedDataCache(const FNiagaraShaderScript* Scr
 				COOK_STAT(Timer.AddHit(CachedData.Num()));
 				InOutShaderMap = new FNiagaraShaderMap();
 				FMemoryReader Ar(CachedData, true);
-
+				FShaderSerializeContext Ctx(Ar);
 				// Deserialize from the cached data
-				if (InOutShaderMap->Serialize(Ar, FShaderMapBase::FSerializationContext{}))
+				if (InOutShaderMap->Serialize(Ctx))
 				{
 					check(InOutShaderMap->GetShaderMapId() == ShaderMapId);
 
@@ -578,7 +578,8 @@ void FNiagaraShaderMap::SaveToDerivedDataCache(const FNiagaraShaderScript* Scrip
 	COOK_STAT(auto Timer = NiagaraShaderCookStats::UsageStats.TimeSyncWork());
 	TArray<uint8> SaveData;
 	FMemoryWriter Ar(SaveData, true);
-	Serialize(Ar, FShaderMapBase::FSerializationContext{});
+	FShaderSerializeContext Ctx(Ar);
+	Serialize(Ctx);
 
 	GetDerivedDataCacheRef().Put(*GetNiagaraShaderMapKeyString(GetContent()->ShaderMapId, GetShaderPlatform()), SaveData, Script ? Script->GetFriendlyName() : TEXT(""));
 	COOK_STAT(Timer.AddMiss(SaveData.Num()));
@@ -1057,12 +1058,12 @@ FNiagaraShaderMap::~FNiagaraShaderMap()
 	check(!bRegistered);
 }
 
-bool FNiagaraShaderMap::Serialize(FArchive& Ar, const FShaderMapBase::FSerializationContext& Ctx)
+bool FNiagaraShaderMap::Serialize(FShaderSerializeContext& Ctx)
 {
 	// Note: This is saved to the DDC, not into packages (except when cooked)
 	// Backwards compatibility therefore will not work based on the version of Ar
 	// Instead, just bump NIAGARASHADERMAP_DERIVEDDATA_VER
-	return Super::Serialize(Ar, Ctx);
+	return Super::Serialize(Ctx);
 }
 
 #if WITH_EDITOR
diff --git a/Engine/Plugins/FX/Niagara/Source/NiagaraShader/Private/NiagaraShared.cpp b/Engine/Plugins/FX/Niagara/Source/NiagaraShader/Private/NiagaraShared.cpp
--- a/Engine/Plugins/FX/Niagara/Source/NiagaraShader/Private/NiagaraShared.cpp
+++ b/Engine/Plugins/FX/Niagara/Source/NiagaraShader/Private/NiagaraShared.cpp
@@ -344,7 +344,8 @@ void FNiagaraShaderScript::SerializeShaderMap(FArchive& Ar)
 				{
 					GameThreadShaderMap->AssociateWithAsset(BaseVMScript->GetOutermost()->GetFName());
 				}
-				GameThreadShaderMap->Serialize(Ar, FShaderMapBase::FSerializationContext{});
+				FShaderSerializeContext Ctx(Ar);
+				GameThreadShaderMap->Serialize(Ctx);
 			}
 			//else if (GameThreadShaderMap != nullptr && !GameThreadShaderMap->CompiledSuccessfully())
 			//{
@@ -361,7 +362,9 @@ void FNiagaraShaderScript::SerializeShaderMap(FArchive& Ar)
 			if (bValid)
 			{
 				FNiagaraShaderMapRef LoadedShaderMap = new FNiagaraShaderMap();
-				bool bLoaded = LoadedShaderMap->Serialize(Ar, FShaderMapBase::FSerializationContext{ /*bLoadedByCookedMaterial*/ true });
+				FShaderSerializeContext Ctx(Ar);
+				Ctx.bLoadingCooked = true;
+				bool bLoaded = LoadedShaderMap->Serialize(Ctx);
 
 				// Toss the loaded shader data if this is a server only instance
 				//@todo - don't cook it in the first place
diff --git a/Engine/Plugins/FX/Niagara/Source/NiagaraShader/Public/NiagaraShared.h b/Engine/Plugins/FX/Niagara/Source/NiagaraShader/Public/NiagaraShared.h
--- a/Engine/Plugins/FX/Niagara/Source/NiagaraShader/Public/NiagaraShared.h
+++ b/Engine/Plugins/FX/Niagara/Source/NiagaraShader/Public/NiagaraShared.h
@@ -599,13 +599,15 @@ public:
 #endif // WITH_EDITOR
 
 	/** Serializes the shader map. */
-	UE_DEPRECATED(5.5, "Use overload taking a FShaderMapBase::FSerializationContext")
-	NIAGARASHADER_API bool Serialize(FArchive& Ar, bool bInlineShaderResources = true, bool bLoadedByCookedMaterial = false)
+	UE_DEPRECATED(5.5, "Use overload taking a FShaderSerializeContext")
+	NIAGARASHADER_API bool Serialize(FArchive& Ar, bool bInlineShaderResources = true, bool bLoadingCooked = false)
 	{
-		return Serialize(Ar, FShaderMapBase::FSerializationContext{ bLoadedByCookedMaterial });
+		FShaderSerializeContext Ctx(Ar);
+		Ctx.bLoadingCooked = bLoadingCooked;
+		return Serialize(Ctx);
 	}
 
-	NIAGARASHADER_API bool Serialize(FArchive& Ar, const FShaderMapBase::FSerializationContext& Ctx);
+	NIAGARASHADER_API bool Serialize(FShaderSerializeContext& Ctx);
 
 #if WITH_EDITOR
 	/** Saves this shader map to the derived data cache. */
diff --git a/Engine/Plugins/Runtime/ComputeFramework/Source/ComputeFramework/Private/ComputeKernelShaderMap.cpp b/Engine/Plugins/Runtime/ComputeFramework/Source/ComputeFramework/Private/ComputeKernelShaderMap.cpp
--- a/Engine/Plugins/Runtime/ComputeFramework/Source/ComputeFramework/Private/ComputeKernelShaderMap.cpp
+++ b/Engine/Plugins/Runtime/ComputeFramework/Source/ComputeFramework/Private/ComputeKernelShaderMap.cpp
@@ -671,7 +671,8 @@ bool FComputeKernelShaderMap::Serialize(FArchive& Ar)
 	// Note: This is saved to the DDC, not into packages (except when cooked)
 	// Backwards compatibility therefore will not work based on the version of Ar
 	// Instead, just bump COMPUTEKERNEL_DERIVEDDATA_VER
-	return Super::Serialize(Ar, {});
+	FShaderSerializeContext Ctx(Ar);
+	return Super::Serialize(Ctx);
 }
 
 #if WITH_EDITOR
diff --git a/Engine/Shaders/Public/ShaderVersion.ush b/Engine/Shaders/Public/ShaderVersion.ush
--- a/Engine/Shaders/Public/ShaderVersion.ush
+++ b/Engine/Shaders/Public/ShaderVersion.ush
@@ -3,4 +3,4 @@
 // Each time the console command is executed it generates a new GUID. As this file is included
 // in Platform.ush (which should be included in any shader) it allows to invalidate the shader DDC.
 // 
-#pragma message("UESHADERMETADATA_VERSION D032387F-3E24-4A2D-903A-377331E33FA8")
+#pragma message("UESHADERMETADATA_VERSION 0D63C246-532E-4D9D-83F3-906A518818CC")
\ No newline at end of file
diff --git a/Engine/Source/Developer/ShaderPreprocessor/Private/ShaderPreprocessor.cpp b/Engine/Source/Developer/ShaderPreprocessor/Private/ShaderPreprocessor.cpp
--- a/Engine/Source/Developer/ShaderPreprocessor/Private/ShaderPreprocessor.cpp
+++ b/Engine/Source/Developer/ShaderPreprocessor/Private/ShaderPreprocessor.cpp
@@ -867,7 +867,6 @@ bool PreprocessShader(
 
 	OutPreprocessedShader = FString(Output.GetSourceViewWide());
 
-	Output.MoveDirectives(ShaderOutput.PragmaDirectives);
 	for (FShaderCompilerError& Error : Output.EditErrors())
 	{
 		ShaderOutput.Errors.Add(MoveTemp(Error));
diff --git a/Engine/Source/Runtime/Core/Private/UObject/FortniteMainBranchObjectVersion.cpp b/Engine/Source/Runtime/Core/Private/UObject/FortniteMainBranchObjectVersion.cpp
--- a/Engine/Source/Runtime/Core/Private/UObject/FortniteMainBranchObjectVersion.cpp
+++ b/Engine/Source/Runtime/Core/Private/UObject/FortniteMainBranchObjectVersion.cpp
@@ -6,9 +6,9 @@ TMap<FGuid, FGuid> FFortniteMainBranchObjectVersion::GetSystemGuids()
 	TMap<FGuid, FGuid> SystemGuids;
 	const FDevSystemGuids& DevGuids = FDevSystemGuids::Get();
 
-	SystemGuids.Add(DevGuids.GLOBALSHADERMAP_DERIVEDDATA_VER, FGuid("8D3A12924CDF4D9F84C0A900E9445CAD"));
+	SystemGuids.Add(DevGuids.GLOBALSHADERMAP_DERIVEDDATA_VER, FGuid("E1B72EA109BB4C2A995F0F65E5BD353D"));
 	SystemGuids.Add(DevGuids.LANDSCAPE_MOBILE_COOK_VERSION, FGuid("32D02EF867C74B71A0D4E0FA41392732"));
-	SystemGuids.Add(DevGuids.MATERIALSHADERMAP_DERIVEDDATA_VER, FGuid("BB56D95CD4FD48E38C4CD7869B44B987"));
+	SystemGuids.Add(DevGuids.MATERIALSHADERMAP_DERIVEDDATA_VER, FGuid("0E6202C151DE402CA451802671AE9347"));
 	SystemGuids.Add(DevGuids.NANITE_DERIVEDDATA_VER, FGuid("DB3ECDFAF27A4344ACB6B1967C21148D"));
 	SystemGuids.Add(DevGuids.NIAGARASHADERMAP_DERIVEDDATA_VER, FGuid("C417A50271B3427D98785388266E3B66"));
 	SystemGuids.Add(DevGuids.Niagara_LatestScriptCompileVersion, FGuid("F8D995224324470E97ED718D33939A0E"));
diff --git a/Engine/Source/Runtime/Engine/Private/Materials/Material.cpp b/Engine/Source/Runtime/Engine/Private/Materials/Material.cpp
--- a/Engine/Source/Runtime/Engine/Private/Materials/Material.cpp
+++ b/Engine/Source/Runtime/Engine/Private/Materials/Material.cpp
@@ -856,13 +856,13 @@ void ProcessSerializedInlineShaderMaps(UMaterialInterface* Owner, TArray<FMateri
 	}
 
 #if WITH_EDITORONLY_DATA
-	const bool bLoadedByCookedMaterial = FPlatformProperties::RequiresCookedData() || Owner->GetOutermost()->bIsCookedForEditor;
+	const bool bLoadingCooked = FPlatformProperties::RequiresCookedData() || Owner->GetOutermost()->bIsCookedForEditor;
 #else
-	const bool bLoadedByCookedMaterial = FPlatformProperties::RequiresCookedData();
+	const bool bLoadingCooked = FPlatformProperties::RequiresCookedData();
 #endif
 	for (FMaterialResource& Resource : LoadedResources)
 	{
-		Resource.RegisterInlineShaderMap(bLoadedByCookedMaterial);
+		Resource.RegisterInlineShaderMap(bLoadingCooked);
 	}
 
 	const bool bDiscardUnusedQualityLevels = CVarDiscardUnusedQualityLevels.GetValueOnAnyThread() != 0;
diff --git a/Engine/Source/Runtime/Engine/Private/Materials/MaterialShader.cpp b/Engine/Source/Runtime/Engine/Private/Materials/MaterialShader.cpp
--- a/Engine/Source/Runtime/Engine/Private/Materials/MaterialShader.cpp
+++ b/Engine/Source/Runtime/Engine/Private/Materials/MaterialShader.cpp
@@ -52,6 +52,7 @@
 #include "Serialization/CompactBinarySerialization.h"
 #include "Serialization/CompactBinaryWriter.h"
 #include "Serialization/MemoryReader.h"
+#include "ShaderSerialization.h"
 #endif
 
 CSV_DECLARE_CATEGORY_MODULE_EXTERN(RENDERCORE_API, Shaders);
@@ -972,7 +973,7 @@ void FSubstrateCompilationConfig::Serialize(FArchive& Ar)
 }
 #endif
 
-void FMaterialShaderMapId::Serialize(FArchive& Ar, bool bLoadedByCookedMaterial)
+void FMaterialShaderMapId::Serialize(FArchive& Ar, bool bLoadingCooked)
 {
 	SCOPED_LOADTIMER(FMaterialShaderMapId_Serialize);
 
@@ -992,9 +993,9 @@ void FMaterialShaderMapId::Serialize(FArchive& Ar, bool bLoadedByCookedMaterial)
 
 #if WITH_EDITOR
 	const bool bIsSavingCooked = Ar.IsSaving() && Ar.IsCooking();
-	bIsCookedId = bLoadedByCookedMaterial;
+	bIsCookedId = bLoadingCooked;
 
-	if (!bIsSavingCooked && !bLoadedByCookedMaterial)
+	if (!bIsSavingCooked && !bLoadingCooked)
 	{
 		uint32 UsageInt = Usage;
 		Ar << UsageInt;
@@ -1022,7 +1023,7 @@ void FMaterialShaderMapId::Serialize(FArchive& Ar, bool bLoadedByCookedMaterial)
 	}
 
 #if WITH_EDITOR
-	if (!bIsSavingCooked && !bLoadedByCookedMaterial)
+	if (!bIsSavingCooked && !bLoadingCooked)
 	{
 		if (Ar.CustomVer(FRenderingObjectVersion::GUID) < FRenderingObjectVersion::MaterialShaderMapIdSerialization)
 		{
@@ -1832,13 +1833,12 @@ TSharedRef<FMaterialShaderMap::FAsyncLoadContext> FMaterialShaderMap::BeginLoadF
 
 	struct FMaterialShaderMapAsyncLoadContext : public FMaterialShaderMap::FAsyncLoadContext
 	{
-		FString	         DataKey;
-		FSharedBuffer    CachedData;
-		FIoHash          CachedDataHash;
-		EShaderPlatform  Platform;
-		const FMaterial* Material = nullptr;
-		FRequestOwner    RequestOwner{UE::DerivedData::EPriority::Normal};
-		TRefCountPtr<FMaterialShaderMap> ShaderMap;
+		FString								DataKey;
+		FShaderCacheLoadContext				LoadContext;
+		EShaderPlatform						Platform;
+		const FMaterial*					Material = nullptr;
+		FRequestOwner						RequestOwner{UE::DerivedData::EPriority::Normal};
+		TRefCountPtr<FMaterialShaderMap>	ShaderMap;
 
 		bool IsReady() const override
 		{
@@ -1858,16 +1858,17 @@ TSharedRef<FMaterialShaderMap::FAsyncLoadContext> FMaterialShaderMap::BeginLoadF
 			TRACE_CPUPROFILER_EVENT_SCOPE(FMaterialShaderMap::FinishLoadFromDerivedDataCache);
 			COOK_STAT(auto Timer = MaterialShaderCookStats::UsageStats.TimeSyncWork());
 
-			if (CachedData)
+			if (LoadContext)
 			{
 				TRACE_COUNTER_INCREMENT(Shaders_FMaterialShaderMapDDCHits);
-				TRACE_COUNTER_ADD(Shaders_FMaterialShaderMapDDCBytesReceived, int64(CachedData.GetSize()));
-				COOK_STAT(Timer.AddHit(int64(CachedData.GetSize())));
-				ShaderMap = new FMaterialShaderMap();
-				FMemoryReaderView Ar(CachedData, /*bIsPersistent*/ true);
 
+				int64 BytesReceived = LoadContext.GetSerializedSize();
+				TRACE_COUNTER_ADD(Shaders_FMaterialShaderMapDDCBytesReceived, BytesReceived);
+				COOK_STAT(Timer.AddHit(BytesReceived));
+
+				ShaderMap = new FMaterialShaderMap();
 				// Deserialize from the cached data
-				ShaderMap->Serialize(Ar, FShaderMapBase::FSerializationContext{});
+				ShaderMap->Serialize(LoadContext);
 
 				check(Material != nullptr);
 
@@ -1877,7 +1878,7 @@ TSharedRef<FMaterialShaderMap::FAsyncLoadContext> FMaterialShaderMap::BeginLoadF
 				{
 					UE_LOG(LogMaterial, Warning, TEXT("Shader map key recomputed from DDC data: %s"), *InDataKey);
 					UE_LOG(LogMaterial, Warning, TEXT("Shader map key from request: %s"), *DataKey);
-					UE_LOG(LogMaterial, Warning, TEXT("Cached data size %llu, hash %s"), CachedData.GetSize(), *LexToString(CachedDataHash))
+					UE_LOG(LogMaterial, Warning, TEXT("Cached data size %llu"), BytesReceived);
 					checkf(false, TEXT("DDC key constructed from deserialized shadermap does not match request key!"));
 				}
 
@@ -1958,23 +1959,26 @@ TSharedRef<FMaterialShaderMap::FAsyncLoadContext> FMaterialShaderMap::BeginLoadF
 			//   - since the get call is synchronous, this can cause a hitch if there's network latency
 			if (Material->IsPersistent())
 			{
-				FCacheGetValueRequest Request;
+				FCacheGetRequest Request;
 				Request.Name = GetMaterialShaderMapName(Material->GetFullPath(), ShaderMapId, InPlatform);
 				Request.Key = GetMaterialShaderMapKey(Result->DataKey);
 				Result->Material = Material;
 				Result->Platform = InPlatform;
 
-				GetCache().GetValue({Request}, Result->RequestOwner, [Result, bCheckCache](FCacheGetValueResponse&& Response)
-				{
-					if (bCheckCache)
+				GetCache().Get({ Request }, Result->RequestOwner, [Result, bCheckCache](FCacheGetResponse&& Response)
 					{
-						Result->CachedData = Response.Value.GetData().Decompress();
-						Result->CachedDataHash = Response.Value.GetRawHash();
-						// This callback might hold the last reference to Result, which owns RequestOwner, so
-						// we must not cancel in the Owner's destructor; Canceling in a callback will deadlock.
-						Result->RequestOwner.KeepAlive();
-					}
-				});
+						if (bCheckCache)
+						{
+							if (Response.Status == EStatus::Ok)
+							{
+								Result->LoadContext.ReadFromRecord(Response.Record, /* bIsPersisent */ true);
+							}
+
+							// This callback might hold the last reference to Result, which owns RequestOwner, so
+							// we must not cancel in the Owner's destructor; Canceling in a callback will deadlock.
+							Result->RequestOwner.KeepAlive();
+						}
+					});
 			}
 		}
 		INC_FLOAT_STAT_BY(STAT_ShaderCompiling_DDCLoading,(float)MaterialDDCTime);
@@ -1987,28 +1991,24 @@ void FMaterialShaderMap::SaveToDerivedDataCache(const FMaterialShaderParameters&
 {
 	TRACE_CPUPROFILER_EVENT_SCOPE(FMaterialShaderMap::SaveToDerivedDataCache);
 	COOK_STAT(auto Timer = MaterialShaderCookStats::UsageStats.TimeSyncWork());
-	TArray64<uint8> SaveData;
-	FMemoryWriter64 Ar(SaveData, true);
-	Serialize(Ar, FShaderMapBase::FSerializationContext{});
 
-	TRACE_COUNTER_ADD(Shaders_FMaterialShaderMapDDCBytesSent, SaveData.Num());
-	COOK_STAT(Timer.AddMiss(SaveData.Num()));
+	FShaderCacheSaveContext Ctx;
+	Serialize(Ctx);
 
 	const FString DataKey = GetMaterialShaderMapKeyString(ShaderMapId, ShaderParameters, GetShaderPlatform());
 
 	using namespace UE::DerivedData;
-	FCachePutValueRequest Request;
-	Request.Name = GetMaterialShaderMapName(GetMaterialPath(), ShaderMapId, GetShaderPlatform());
-	Request.Key = GetMaterialShaderMapKey(DataKey);
+	FCacheKey Key = GetMaterialShaderMapKey(DataKey);
 
-	UE_LOG(LogMaterial, Verbose, TEXT("Saved shaders for %s to DDC (key hash: %s)"), *Request.Name, *LexToString(Request.Key.Hash));
-	UE_LOG(LogMaterial, VeryVerbose, TEXT("Full DDC data key for %s: %s"), *Request.Name, *DataKey);
+	UE::FSharedString RequestName = GetMaterialShaderMapName(GetMaterialPath(), ShaderMapId, GetShaderPlatform());
+	UE_LOG(LogMaterial, Verbose, TEXT("Saved shaders for %s to DDC (key hash: %s)"), *RequestName, *LexToString(Key.Hash));
+	UE_LOG(LogMaterial, VeryVerbose, TEXT("Full DDC data key for %s: %s"), *RequestName, *DataKey);
 
 	if (!CVarMaterialShaderMapDump->GetString().IsEmpty() && CVarMaterialShaderMapDump->GetString() == GetMaterialPath())
 	{
 		TStringBuilder<256> Path;
 		FPathViews::Append(Path, FPaths::ProjectSavedDir(), TEXT("MaterialShaderMaps"), TEXT(""));
-		Path << Request.Key.Hash << ".txt";
+		Path << Key.Hash << ".txt";
 		FArchive* DumpAr = IFileManager::Get().CreateFileWriter(*Path, FILEWRITE_Silent);
 		if (DumpAr)
 		{
@@ -2017,10 +2017,14 @@ void FMaterialShaderMap::SaveToDerivedDataCache(const FMaterialShaderParameters&
 		}
 	}
 
-	Request.Value = FValue::Compress(MakeSharedBufferFromArray(MoveTemp(SaveData)));
 	FRequestOwner AsyncOwner(EPriority::Normal);
 	FRequestBarrier AsyncBarrier(AsyncOwner);
-	GetCache().PutValue({Request}, AsyncOwner);
+	GetCache().Put({ { RequestName, Ctx.BuildCacheRecord(Key) } }, AsyncOwner);
+
+	uint64 SerializedSize = Ctx.GetSerializedSize();
+	TRACE_COUNTER_ADD(Shaders_FMaterialShaderMapDDCBytesSent, SerializedSize);
+	COOK_STAT(Timer.AddMiss(SerializedSize));
+
 	AsyncOwner.KeepAlive();
 }
 #endif // WITH_EDITOR
@@ -2048,7 +2052,8 @@ void FMaterialShaderMap::SaveForRemoteRecompile(FArchive& Ar, const TMap<FString
 			{
 				uint8 bIsValid = 1;
 				Ar << bIsValid;
-				ShaderMap->Serialize(Ar, FShaderMapBase::FSerializationContext{});
+				FShaderSerializeContext Ctx(Ar);
+				ShaderMap->Serialize(Ctx);
 			}
 			else
 			{
@@ -2100,7 +2105,8 @@ void FMaterialShaderMap::LoadForRemoteRecompile(FArchive& Ar, EShaderPlatform Sh
 				TRefCountPtr<FMaterialShaderMap> ShaderMap = new FMaterialShaderMap();
 
 				// serialize the id and the material shader map
-				ShaderMap->Serialize(Ar, FShaderMapBase::FSerializationContext{});
+				FShaderSerializeContext Ctx(Ar);
+				ShaderMap->Serialize(Ctx);
 
 				LoadedShaderMapsDictionary.Add(ShaderMap->GetShaderMapId(), ShaderMap);
 				MaterialShaderMapData.LoadedShaderMapsIds.Add(ShaderMap->GetShaderMapId());
@@ -3589,14 +3595,15 @@ FMaterialShaderMap* FMaterialShaderMap::GetFinalizedClone() const
 }
 #endif // WITH_EDITOR
 
-bool FMaterialShaderMap::Serialize(FArchive& Ar, const FShaderMapBase::FSerializationContext& Ctx)
+bool FMaterialShaderMap::Serialize(FShaderSerializeContext& Ctx)
 {
 	SCOPED_LOADTIMER(FMaterialShaderMap_Serialize);
 	// Note: This is saved to the DDC, not into packages (except when cooked)
 	// Backwards compatibility therefore will not work based on the version of Ar
 	// Instead, just bump MATERIALSHADERMAP_DERIVEDDATA_VER
-	ShaderMapId.Serialize(Ar, Ctx.bLoadedByCookedMaterial);
-	bool bSerialized = Super::Serialize(Ar, Ctx);
+	FArchive& Ar = Ctx.GetMainArchive();
+	ShaderMapId.Serialize(Ar, Ctx.bLoadingCooked);
+	bool bSerialized = Super::Serialize(Ctx);
 #if STATS
 	// This is an unsavory hack to repair STAT_Shaders_NumShadersLoaded not being calculated right in the superclass because the Content class isn't allowed to have virtual functions atm.
 	// A better way is tracked in UE-127112
diff --git a/Engine/Source/Runtime/Engine/Private/Materials/MaterialShared.cpp b/Engine/Source/Runtime/Engine/Private/Materials/MaterialShared.cpp
--- a/Engine/Source/Runtime/Engine/Private/Materials/MaterialShared.cpp
+++ b/Engine/Source/Runtime/Engine/Private/Materials/MaterialShared.cpp
@@ -1540,7 +1540,8 @@ void FMaterial::SerializeInlineShaderMap(FArchive& Ar, const FName& SerializingA
 
 			if (bValid)
 			{
-				GameThreadShaderMap->Serialize(Ar, FShaderMapBase::FSerializationContext{});
+				FShaderSerializeContext Ctx{ Ar };
+				GameThreadShaderMap->Serialize(Ctx);
 			}
 			else
 			{
@@ -1562,7 +1563,10 @@ void FMaterial::SerializeInlineShaderMap(FArchive& Ar, const FName& SerializingA
 			if (bValid)
 			{
 				TRefCountPtr<FMaterialShaderMap> LoadedShaderMap = new FMaterialShaderMap();
-				if (LoadedShaderMap->Serialize(Ar, FShaderMapBase::FSerializationContext{ bCooked && Ar.IsLoading(), SerializingAsset }))
+				FShaderSerializeContext Ctx(Ar);
+				Ctx.bLoadingCooked = bCooked && Ar.IsLoading();
+				Ctx.SerializingAsset = SerializingAsset;
+				if (LoadedShaderMap->Serialize(Ctx))
 				{
 					GameThreadShaderMap = MoveTemp(LoadedShaderMap);
 					GameThreadShaderMap->GetResource()->SetOwnerName(GetOwnerFName());
@@ -1579,7 +1583,7 @@ void FMaterial::SerializeInlineShaderMap(FArchive& Ar, const FName& SerializingA
 	}
 }
 
-void FMaterial::RegisterInlineShaderMap(bool bLoadedByCookedMaterial)
+void FMaterial::RegisterInlineShaderMap(bool bLoadingCooked)
 {
 	if (GameThreadShaderMap)
 	{
diff --git a/Engine/Source/Runtime/Engine/Private/ShaderCompiler/ShaderCompiler.cpp b/Engine/Source/Runtime/Engine/Private/ShaderCompiler/ShaderCompiler.cpp
--- a/Engine/Source/Runtime/Engine/Private/ShaderCompiler/ShaderCompiler.cpp
+++ b/Engine/Source/Runtime/Engine/Private/ShaderCompiler/ShaderCompiler.cpp
@@ -43,6 +43,7 @@
 #include "Serialization/MemoryReader.h"
 #include "Serialization/NameAsStringProxyArchive.h"
 #include "ShaderCodeLibrary.h"
+#include "ShaderSerialization.h"
 #include "ShaderDiagnostics.h"
 #include "ShaderPlatformCachedIniValue.h"
 #include "StaticBoundShaderState.h"
@@ -5045,13 +5046,13 @@ void CompileGlobalShaderMap(EShaderPlatform Platform, const ITargetPlatform* Tar
 				using namespace UE::DerivedData;
 
 				int32 BufferIndex = 0;
-				TArray<FCacheGetValueRequest> Requests;
+				TArray<FCacheGetRequest> Requests;
 
 				// Submit DDC requests.
 				SlowTask.EnterProgressFrame(ProgressStep, LOCTEXT("SubmitDDCRequests", "Submitting global shader DDC Requests..."));
 				for (const auto& ShaderFilenameDependencies : ShaderMapId.GetShaderFilenameToDependeciesMap())
 				{
-					FCacheGetValueRequest& Request = Requests.AddDefaulted_GetRef();
+					FCacheGetRequest& Request = Requests.AddDefaulted_GetRef();
 					Request.Name = GetGlobalShaderMapName(ShaderMapId, Platform, ShaderFilenameDependencies.Key);
 					Request.Key = GetGlobalShaderMapKey(ShaderMapId, Platform, TargetPlatform, ShaderFilenameDependencies.Value);
 					Request.UserData = uint64(BufferIndex);
@@ -5078,20 +5079,22 @@ void CompileGlobalShaderMap(EShaderPlatform Platform, const ITargetPlatform* Tar
 
 				// Process finished DDC requests.
 				SlowTask.EnterProgressFrame(ProgressStep, LOCTEXT("ProcessDDCRequests", "Processing global shader DDC requests..."));
-				TArray<FValue> GlobalShaderMapBuffers;
-				GlobalShaderMapBuffers.SetNum(Requests.Num());
+				TArray<FShaderCacheLoadContext> GlobalShaderMapLoads;
+				GlobalShaderMapLoads.SetNum(Requests.Num());
 				{
 					COOK_STAT(auto Timer = GlobalShaderCookStats::UsageStats.TimeSyncWork());
 					COOK_STAT(Timer.TrackCyclesOnly());
 					FRequestOwner BlockingOwner(EPriority::Blocking);
-					GetCache().GetValue(Requests, BlockingOwner, [&GlobalShaderMapBuffers, &bTempNoShaderDDC](FCacheGetValueResponse&& Response)
+					GetCache().Get(Requests, BlockingOwner, [&GlobalShaderMapLoads, &bTempNoShaderDDC](FCacheGetResponse&& Response)
 					{
 						if (bTempNoShaderDDC)
 						{
 							return;
 						}
-
-						GlobalShaderMapBuffers[int32(Response.UserData)] = MoveTemp(Response.Value);
+						if (Response.Status == EStatus::Ok)
+						{
+							GlobalShaderMapLoads[int32(Response.UserData)].ReadFromRecord(Response.Record);
+						}
 					});
 					BlockingOwner.Wait();
 				}
@@ -5100,12 +5103,10 @@ void CompileGlobalShaderMap(EShaderPlatform Platform, const ITargetPlatform* Tar
 				for (const auto& ShaderFilenameDependencies : ShaderMapId.GetShaderFilenameToDependeciesMap())
 				{
 					COOK_STAT(auto Timer = GlobalShaderCookStats::UsageStats.TimeSyncWork());
-					if (GlobalShaderMapBuffers[BufferIndex].HasData())
+					if (GlobalShaderMapLoads[BufferIndex].ShaderObjectData)
 					{
-						COOK_STAT(Timer.AddHit(int64(GlobalShaderMapBuffers[BufferIndex].GetRawSize())));
-						const FSharedBuffer CachedData = GlobalShaderMapBuffers[BufferIndex].GetData().Decompress();
-						FMemoryReaderView MemoryReader(CachedData);
-						GGlobalShaderMap[Platform]->AddSection(FGlobalShaderMapSection::CreateFromArchive(MemoryReader));
+						GGlobalShaderMap[Platform]->AddSection(FGlobalShaderMapSection::CreateFromCache(GlobalShaderMapLoads[BufferIndex]));
+						COOK_STAT(Timer.AddHit(int64(GlobalShaderMapLoads[BufferIndex].GetSerializedSize())));
 						DDCHits++;
 					}
 					else
diff --git a/Engine/Source/Runtime/Engine/Private/ShaderCompiler/ShaderCompilerEditor.cpp b/Engine/Source/Runtime/Engine/Private/ShaderCompiler/ShaderCompilerEditor.cpp
--- a/Engine/Source/Runtime/Engine/Private/ShaderCompiler/ShaderCompilerEditor.cpp
+++ b/Engine/Source/Runtime/Engine/Private/ShaderCompiler/ShaderCompilerEditor.cpp
@@ -24,6 +24,7 @@
 #include "ProfilingDebugging/StallDetector.h"
 #include "Serialization/NameAsStringProxyArchive.h"
 #include "ShaderDiagnostics.h"
+#include "ShaderSerialization.h"
 #include "UObject/Linker.h"
 #include "UObject/Package.h"
 #include "UObject/UObjectIterator.h"
@@ -542,8 +543,7 @@ static void SaveGlobalShaderMapToDerivedDataCache(EShaderPlatform Platform)
 	COOK_STAT(auto Timer = GlobalShaderCookStats::UsageStats.TimeSyncWork());
 
 	const ITargetPlatform* TargetPlatform = GGlobalShaderTargetPlatform[Platform];
-	TArray<uint8> SaveData;
-
+	FShaderCacheSaveContext Ctx;
 	FGlobalShaderMapId ShaderMapId(Platform, TargetPlatform);
 	// caller should prevent incomplete shadermaps to be saved
 	FGlobalShaderMap* GlobalSM = GetGlobalShaderMap(Platform);
@@ -554,19 +554,19 @@ static void SaveGlobalShaderMapToDerivedDataCache(EShaderPlatform Platform)
 		{
 			Section->FinalizeContent();
 
-			SaveData.Reset();
-			FMemoryWriter Ar(SaveData, true);
-			Section->Serialize(Ar);
-			COOK_STAT(Timer.AddMiss(SaveData.Num()));
+			// reuse serialize context, internal allocations will be kept so this minimizes heap alloc churn
+			Ctx.Reset();
+
+			Section->Serialize(Ctx);
+			COOK_STAT(Timer.AddMiss(Ctx.GetSerializedSize()));
 
 			using namespace UE::DerivedData;
-			FCachePutValueRequest Request;
-			Request.Name = GetGlobalShaderMapName(ShaderMapId, Platform, ShaderFilenameDependencies.Key);
-			Request.Key = GetGlobalShaderMapKey(ShaderMapId, Platform, TargetPlatform, ShaderFilenameDependencies.Value);
-			Request.Value = FValue::Compress(MakeSharedBufferFromArray(MoveTemp(SaveData)));
+			UE::FSharedString Name = GetGlobalShaderMapName(ShaderMapId, Platform, ShaderFilenameDependencies.Key);
+			FCacheKey Key = GetGlobalShaderMapKey(ShaderMapId, Platform, TargetPlatform, ShaderFilenameDependencies.Value);
+
 			FRequestOwner AsyncOwner(EPriority::Normal);
 			FRequestBarrier AsyncBarrier(AsyncOwner);
-			GetCache().PutValue({ Request }, AsyncOwner);
+			GetCache().Put({ { Name, Ctx.BuildCacheRecord(Key) } }, AsyncOwner);
 			AsyncOwner.KeepAlive();
 		}
 	}
diff --git a/Engine/Source/Runtime/Engine/Private/ShaderCompiler/ShaderCompilerJobCache.cpp b/Engine/Source/Runtime/Engine/Private/ShaderCompiler/ShaderCompilerJobCache.cpp
--- a/Engine/Source/Runtime/Engine/Private/ShaderCompiler/ShaderCompilerJobCache.cpp
+++ b/Engine/Source/Runtime/Engine/Private/ShaderCompiler/ShaderCompilerJobCache.cpp
@@ -17,6 +17,7 @@
 #include "ProfilingDebugging/StallDetector.h"
 #include "Serialization/MemoryReader.h"
 #include "Serialization/MemoryWriter.h"
+#include "ShaderSerialization.h"
 
 
 // Switch to Verbose after initial testing
@@ -139,7 +140,6 @@ namespace
 {
 	/** The FCacheBucket used with the DDC, cached to avoid recreating it for each request */
 	UE::DerivedData::FCacheBucket ShaderJobCacheDDCBucket = UE::DerivedData::FCacheBucket(ANSITEXTVIEW("FShaderJobCacheShaders"), TEXTVIEW("Shader"));
-	UE::DerivedData::FValueId ShaderJobCacheId = UE::DerivedData::FValueId::FromName("FShaderJobCacheShaderID");
 }
 #endif
 
@@ -377,6 +377,9 @@ public:
 	/** Canned output */
 	FSharedBuffer JobOutput;
 
+	/** Separate blobs for shader code */
+	TArray<FSharedBuffer> JobCode;
+
 	/** Path to where the cached debug info is stored. */
 	FString CachedDebugInfoPath;
 
@@ -609,7 +612,6 @@ public:
 
 private:
 	using FJobOutputHash = FBlake3Hash;
-	using FJobCachedOutput = FSharedBuffer;
 	using FStoredOutput = FShaderJobCacheStoredOutput;
 
 	// cannot allow managing this from outside as the caching logic is not exposed
@@ -620,10 +622,22 @@ private:
 		return PrevNumOutstandingJobs - Value;
 	}
 
+	FJobOutputHash ComputeJobHash(const FShaderCacheSerializeContext& SerializeContext)
+	{
+		FBlake3 Hasher;
+		check(SerializeContext);
+		Hasher.Update(SerializeContext.ShaderObjectData.GetData(), SerializeContext.ShaderObjectData.GetSize());
+		for (const FSharedBuffer& CodeBuf : SerializeContext.ShaderCode)
+		{
+			Hasher.Update(CodeBuf.GetData(), CodeBuf.GetSize());
+		}
+		return Hasher.Finalize();
+	}
+
 	void InternalSetPriority(FShaderCommonCompileJob* Job, EShaderCompileJobPriority InPriority);
 
 	/** Looks for or adds an entry for the given hash in the cache.  Returns cached output if it exists, or may initialize DDC request if one has been issued. */
-	FShaderJobCacheRef FindOrAdd(const FShaderCompilerInputHash& Hash, EShaderCompileJobPriority JobPriority, const bool bCheckDDC, TPimplPtr<UE::DerivedData::FRequestOwner>& InoutRequestOwner, FJobCachedOutput*& OutCachedOutput);
+	FShaderCacheLoadContext FindOrAdd(const FShaderCompilerInputHash& Hash, FShaderCommonCompileJob* Job, const bool bCheckDDC);
 
 	/** Find an existing item in the cache. */
 	FShaderJobData* Find(const FShaderCompilerInputHash& Hash);
@@ -635,7 +649,7 @@ private:
 	void RemoveDuplicateJob(FShaderCommonCompileJob* DuplicateJob);
 
 	/** Adds a job output to the cache */
-	void AddJobOutput(FShaderJobData& JobData, const FShaderCommonCompileJob* FinishedJob, const FShaderCompilerInputHash& Hash, const FJobCachedOutput& Contents, int32 InitialHitCount, const bool bAddToDDC);
+	void AddJobOutput(FShaderJobData& JobData, const FShaderCommonCompileJob* FinishedJob, const FShaderCompilerInputHash& Hash, FShaderCacheSaveContext& SaveContext, int32 InitialHitCount, const bool bAddToDDC);
 
 	/** Returns memory used by the cache*/
 	uint64 GetAllocatedMemory() const;
@@ -1365,11 +1379,10 @@ void FShaderJobCache::SubmitJob(FShaderCommonCompileJob* Job)
 		}
 		check(Job->JobIndex != INDEX_NONE);
 
-		FSharedBuffer* ExistingOutput;
-		FShaderJobCacheRef JobCacheRef = FindOrAdd(InputHash, Job->Priority, bCheckDDC, Job->RequestOwner, ExistingOutput);
+		FShaderCacheLoadContext LoadContext = FindOrAdd(InputHash, Job, bCheckDDC);
 
-		// see if there are already cached results for this job
-		if (ExistingOutput)
+		// see if there are already cached results for this job that were returned synchronously by FindOrAdd
+		if (LoadContext)
 		{
 			Unlink(*Job);		// from PendingSubmitJobTaskJobs
 
@@ -1380,17 +1393,15 @@ void FShaderJobCache::SubmitJob(FShaderCommonCompileJob* Job)
 			bJobCacheLocked = false;
 
 			UE_LOG(LogShaderCompilers, UE_SHADERCACHE_LOG_LEVEL, TEXT("There is already a cached job with the ihash %s, processing the new one immediately."), *LexToString(InputHash));
-			FMemoryReaderView MemReader(*ExistingOutput);
-			Job->SerializeOutput(MemReader);
+
+			Job->SerializeOutput(LoadContext);
 
 			// finish the job instantly
 			Job->TimeTaskSubmitJobsStall += ProcessFinishedJob(Job, true);
 		}
 		else
 		{
-			FShaderJobData& JobData = GetShaderJobData(JobCacheRef);
-			Job->JobCacheRef = JobCacheRef;
-
+			FShaderJobData& JobData = GetShaderJobData(Job->JobCacheRef);
 			// see if another job with the same input hash is being worked on
 			if (JobData.JobInFlight)
 			{
@@ -1424,11 +1435,10 @@ void FShaderJobCache::SubmitJob(FShaderCommonCompileJob* Job)
 	}
 	else if (ShaderCompiler::IsJobCacheDebugValidateEnabled())
 	{
-		FSharedBuffer* ExistingOutput;
 		const FShaderCompilerInputHash& InputHash = Job->GetInputHash();
 		const bool bCheckDDC = !(Job->bIsDefaultMaterial || Job->bIsGlobalShader);
 		JobLock.WriteLock();
-		Job->JobCacheRef = FindOrAdd(InputHash, Job->Priority, bCheckDDC, Job->RequestOwner, ExistingOutput);
+		FShaderCacheLoadContext LoadContext = FindOrAdd(InputHash, Job, bCheckDDC);
 		bJobCacheLocked = true;
 	}
 
@@ -1568,11 +1578,8 @@ void FShaderJobCache::AddToCacheAndProcessPending(FShaderCommonCompileJob* Finis
 	ensureMsgf(FinishedJob->bInputHashSet, TEXT("Finished job didn't have input hash set, was shader compiler jobs cache toggled runtime?"));
 
 	const FShaderCompilerInputHash& InputHash = FinishedJob->GetInputHash();
-	TArray<uint8> Output;
-	FMemoryWriter Writer(Output);
-	FinishedJob->SerializeOutput(Writer);
-
-	FSharedBuffer Buffer = MakeSharedBufferFromArray(MoveTemp(Output));
+	FShaderCacheSaveContext SaveContext;
+	FinishedJob->SerializeOutput(SaveContext);
 
 	FShaderJobData& JobData = GetShaderJobData(FinishedJob->JobCacheRef);
 
@@ -1605,7 +1612,7 @@ void FShaderJobCache::AddToCacheAndProcessPending(FShaderCommonCompileJob* Finis
 		{
 			const bool bAddToDDC = GShaderCompilerPerShaderDDCGlobal || !(FinishedJob->bIsDefaultMaterial || FinishedJob->bIsGlobalShader);
 			// we only cache jobs that succeded
-			AddJobOutput(JobData, FinishedJob, InputHash, Buffer, FinishedDuplicateJobs.Num(), bAddToDDC);
+			AddJobOutput(JobData, FinishedJob, InputHash, SaveContext, FinishedDuplicateJobs.Num(), bAddToDDC);
 		}
 
 		// remove ourselves from the jobs in flight
@@ -1626,10 +1633,16 @@ void FShaderJobCache::AddToCacheAndProcessPending(FShaderCommonCompileJob* Finis
 	{
 		UE_LOG(LogShaderCompilers, UE_SHADERCACHE_LOG_LEVEL, TEXT("Processed %d outstanding jobs with the same ihash %s."), FinishedDuplicateJobs.Num(), *LexToString(InputHash));
 
+		check(JobData.HasOutput());
+		FStoredOutput** StoredOutput = Outputs.Find(JobData.OutputHash);
+		check(StoredOutput);
+		// Construct a single load context pointing to the data now stored in the cache
+		FShaderCacheLoadContext LoadContext((*StoredOutput)->JobOutput, (*StoredOutput)->JobCode);
 		for (FShaderCommonCompileJob* DuplicateJob : FinishedDuplicateJobs)
 		{
-			FMemoryReaderView MemReader(Buffer);
-			DuplicateJob->SerializeOutput(MemReader);
+			// reuse the same load context for each duplicate job to avoid reallocating anything
+			LoadContext.Reuse();
+			DuplicateJob->SerializeOutput(LoadContext);
 			checkf(DuplicateJob->bSucceeded == FinishedJob->bSucceeded, TEXT("Different success status for the job with the same ihash"));
 
 			// finish the job instantly
@@ -1731,7 +1744,7 @@ int32 FShaderJobCache::GetPendingJobs(EShaderCompilerWorkerType InWorkerType, ES
 	return NumJobs;
 }
 
-FShaderJobCacheRef FShaderJobCache::FindOrAdd(const FShaderCompilerInputHash& Hash, EShaderCompileJobPriority JobPriority, const bool bCheckDDC, TPimplPtr<UE::DerivedData::FRequestOwner>& InoutRequestOwner, FJobCachedOutput*& OutCachedOutput)
+FShaderCacheLoadContext FShaderJobCache::FindOrAdd(const FShaderCompilerInputHash& Hash, FShaderCommonCompileJob* Job, const bool bCheckDDC)
 {
 	LLM_SCOPE_BYTAG(ShaderCompiler);
 
@@ -1740,13 +1753,10 @@ FShaderJobCacheRef FShaderJobCache::FindOrAdd(const FShaderCompilerInputHash& Ha
 	++TotalSearchAttempts;
 	TRACE_COUNTER_INCREMENT(Shaders_JobCacheSearchAttempts);
 	TRACE_CPUPROFILER_EVENT_SCOPE(FShaderJobCache::Find);
-
-	OutCachedOutput = nullptr;
-
 	uint64 InputHashToJobDataSize = InputHashToJobData.GetAllocatedSize();
 
-	FShaderJobCacheRef JobCacheRef = InputHashToJobData.FindOrAdd(Hash);
-	FShaderJobData& JobData = GetShaderJobData(JobCacheRef);
+	Job->JobCacheRef = InputHashToJobData.FindOrAdd(Hash);
+	FShaderJobData& JobData = GetShaderJobData(Job->JobCacheRef);
 
 	CurrentlyAllocatedMemory += InputHashToJobData.GetAllocatedSize() - InputHashToJobDataSize;
 
@@ -1762,7 +1772,7 @@ FShaderJobCacheRef FShaderJobCache::FindOrAdd(const FShaderCompilerInputHash& Ha
 		// update the output hit count
 		(*CannedOutput)->NumHits++;
 
-		OutCachedOutput = &(*CannedOutput)->JobOutput;
+		return FShaderCacheLoadContext((*CannedOutput)->JobOutput, (*CannedOutput)->JobCode);
 	}
 #if WITH_EDITOR
 	else
@@ -1773,7 +1783,7 @@ FShaderJobCacheRef FShaderJobCache::FindOrAdd(const FShaderCompilerInputHash& Ha
 		// If we didn't find it in memory search the DDC if it's enabled.
 		// Don't search if this isn't the first job with this hash (JobInFlight already set), or there's already a request in flight.
 		const bool bCachePerShaderDDC = IsShaderJobCacheDDCEnabled() && bCheckDDC && !bNoShaderDDC;
-		if (bCachePerShaderDDC && (JobData.JobInFlight == nullptr) && !InoutRequestOwner)
+		if (bCachePerShaderDDC && (JobData.JobInFlight == nullptr) && !Job->RequestOwner)
 		{
 			TRACE_COUNTER_INCREMENT(Shaders_JobCacheDDCRequests);
 
@@ -1791,15 +1801,15 @@ FShaderJobCacheRef FShaderJobCache::FindOrAdd(const FShaderCompilerInputHash& Ha
 				}
 				else
 				{
-					switch (JobPriority)
+					switch (Job->Priority)
 					{
 					case EShaderCompileJobPriority::Low:		DerivedDataPriority = UE::DerivedData::EPriority::Low;		break;
 					case EShaderCompileJobPriority::Normal:		DerivedDataPriority = UE::DerivedData::EPriority::Normal;	break;
 					default:									DerivedDataPriority = UE::DerivedData::EPriority::Highest;	break;
 					}
 				}
-				InoutRequestOwner = MakePimpl<UE::DerivedData::FRequestOwner>(DerivedDataPriority);
-				RequestOwner = InoutRequestOwner.Get();
+				Job->RequestOwner = MakePimpl<UE::DerivedData::FRequestOwner>(DerivedDataPriority);
+				RequestOwner = Job->RequestOwner.Get();
 			}
 			else
 			{
@@ -1814,15 +1824,17 @@ FShaderJobCacheRef FShaderJobCache::FindOrAdd(const FShaderCompilerInputHash& Ha
 			Request.Key.Hash = Hash;
 			Request.Policy = IsShaderJobCacheDDCRemotePolicyEnabled() ? UE::DerivedData::ECachePolicy::Default : UE::DerivedData::ECachePolicy::Local;
 
-			// Optionally read the cached output back to the main thread
-			FJobCachedOutput** OutCachedOutputPtr = &OutCachedOutput;
 			bool bCompletedSynchronously = false;
 			bool* bCompletedSynchronouslyPtr = &bCompletedSynchronously;
+			
+			// Optionally read the cached output back to the main thread (this is only ever set if bCompletedSynchronously is true)
+			FStoredOutput* StoredOutput = nullptr;
+			FStoredOutput** StoredOutputPtr = &StoredOutput;
 
 			UE::DerivedData::GetCache().Get(
 				{ Request },
 				*RequestOwner,
-				[this, JobDataPtr = &JobData, OutCachedOutputPtr, DerivedDataPriority, bCompletedSynchronouslyPtr](UE::DerivedData::FCacheGetResponse&& Response)
+				[this, JobDataPtr = &JobData, DerivedDataPriority, StoredOutputPtr, bCompletedSynchronouslyPtr](UE::DerivedData::FCacheGetResponse&& Response)
 				{
 					if (GShaderCompilerDebugStallDDCQuery > 0)
 					{
@@ -1843,10 +1855,11 @@ FShaderJobCacheRef FShaderJobCache::FindOrAdd(const FShaderCompilerInputHash& Ha
 					{
 						// Retrieve the shared buffer containing the job output and compute the associated output hash for the result retrieved from DDC
 						// If an existing duplicate of this buffer is already registered in the Outputs map, this copy will be freed at end of scope
-						FSharedBuffer JobOutput = Response.Record.GetValue(ShaderJobCacheId).GetData().Decompress();
-						FJobOutputHash OutputHash = FBlake3::HashBuffer(JobOutput.GetData(), JobOutput.GetSize());
+						FShaderCacheLoadContext LoadContext;
+						LoadContext.ReadFromRecord(Response.Record);
+						FJobOutputHash OutputHash = ComputeJobHash(LoadContext);
 
-						TRACE_COUNTER_ADD(Shaders_JobCacheDDCBytesReceived, JobOutput.GetSize());
+						TRACE_COUNTER_ADD(Shaders_JobCacheDDCBytesReceived, LoadContext.GetSerializedSize());
 						TRACE_COUNTER_INCREMENT(Shaders_JobCacheDDCHits);
 
 						// If we are running the cache logic async (not blocking in the main thread), we need a lock before writing to the job cache.
@@ -1887,8 +1900,10 @@ FShaderJobCacheRef FShaderJobCache::FindOrAdd(const FShaderCompilerInputHash& Ha
 						if (StoredOutput == nullptr)
 						{
 							// Create a new entry to store in the FShaderJobCache if one doesn't already exist for this output hash
+							check(LoadContext); // sanity check that the load context was populated properly
 							StoredOutput = new FStoredOutput();
-							StoredOutput->JobOutput = JobOutput;
+							StoredOutput->JobOutput = LoadContext.ShaderObjectData;
+							StoredOutput->JobCode = MoveTemp(LoadContext.OwnedShaderCode);
 							Outputs.Add(OutputHash, StoredOutput);
 							CurrentlyAllocatedMemory += StoredOutput->GetAllocatedSize();
 						}
@@ -1934,11 +1949,14 @@ FShaderJobCacheRef FShaderJobCache::FindOrAdd(const FShaderCompilerInputHash& Ha
 							// Need to release the lock before calling ProcessFinishedJobs
 							JobLock.WriteUnlock();
 
+							// Re-initialize load context pointing to the stored code array in the cache and reuse for each duplicate job needing population
+							LoadContext = FShaderCacheLoadContext(StoredOutput->JobOutput, StoredOutput->JobCode);
+
 							// Call ProcessFinishedJob on main job and duplicates
 							for (FShaderCommonCompileJob* FinishedJob : FinishedJobs)
 							{
-								FMemoryReaderView MemReader(StoredOutput->JobOutput);
-								FinishedJob->SerializeOutput(MemReader);
+								LoadContext.Reuse();
+								FinishedJob->SerializeOutput(LoadContext);
 								ProcessFinishedJob(FinishedJob, true);
 							}
 
@@ -1950,7 +1968,7 @@ FShaderJobCacheRef FShaderJobCache::FindOrAdd(const FShaderCompilerInputHash& Ha
 						else
 						{
 							// Send results back to the main thread when running synchronous
-							*OutCachedOutputPtr = &StoredOutput->JobOutput;
+							*StoredOutputPtr = StoredOutput;
 						}
 					}
 					else
@@ -2001,13 +2019,17 @@ FShaderJobCacheRef FShaderJobCache::FindOrAdd(const FShaderCompilerInputHash& Ha
 				// in which case we can delete the TPimplPtr request owner by setting it to null.  This tells the main thread
 				// there is no async DDC request in flight, and it should handle adding the pending job to the queue, since the
 				// DDC request callback won't be handling that.
-				InoutRequestOwner = nullptr;
+				Job->RequestOwner = nullptr;
+
+				if (StoredOutput)
+				{
+					return FShaderCacheLoadContext(StoredOutput->JobOutput, StoredOutput->JobCode);
+				}
 			}
 		}
 	}
 #endif
-
-	return JobCacheRef;
+	return FShaderCacheLoadContext();
 }
 
 FShaderJobData* FShaderJobCache::Find(const FShaderCompilerInputHash& Hash)
@@ -2071,7 +2093,7 @@ FShaderJobCache::~FShaderJobCache()
 	}
 }
 
-void FShaderJobCache::AddJobOutput(FShaderJobData& JobData, const FShaderCommonCompileJob* FinishedJob, const FShaderCompilerInputHash& Hash, const FJobCachedOutput& Contents, int32 InitialHitCount, const bool bAddToDDC)
+void FShaderJobCache::AddJobOutput(FShaderJobData& JobData, const FShaderCommonCompileJob* FinishedJob, const FShaderCompilerInputHash& Hash, FShaderCacheSaveContext& SaveContext, int32 InitialHitCount, const bool bAddToDDC)
 {
 	TRACE_CPUPROFILER_EVENT_SCOPE(FShaderJobCache::Add);
 
@@ -2085,7 +2107,11 @@ void FShaderJobCache::AddJobOutput(FShaderJobData& JobData, const FShaderCommonC
 		return;
 	}
 
-	FJobOutputHash OutputHash = FBlake3::HashBuffer(Contents.GetData(), Contents.GetSize());
+	// Explicitly finalize the serialization to generate the job struct FSharedBuffer since we store
+	// this in the in-memory cache prior to pushing to DDC (subsequent Finalize call in BuildCacheRecord 
+	// will be a noop)
+	SaveContext.Finalize();
+	FJobOutputHash OutputHash = ComputeJobHash(SaveContext);
 
 	if (JobData.HasOutput() && ShaderCompiler::IsJobCacheDebugValidateEnabled())
 	{
@@ -2179,9 +2205,11 @@ void FShaderJobCache::AddJobOutput(FShaderJobData& JobData, const FShaderCommonC
 		{
 			const uint64 OutputsOriginalSize = Outputs.GetAllocatedSize();
 
+			check(SaveContext);
 			FStoredOutput* NewStoredOutput = new FStoredOutput();
 			NewStoredOutput->NumHits = InitialHitCount;
-			NewStoredOutput->JobOutput = Contents;
+			NewStoredOutput->JobOutput = SaveContext.ShaderObjectData;
+			NewStoredOutput->JobCode = MoveTemp(SaveContext.OwnedShaderCode);
 			NewStoredOutput->CachedDebugInfoPath = InputDebugInfoPath;
 			NewStoredOutput->AddRef();
 			Outputs.Add(OutputHash, NewStoredOutput);
@@ -2202,8 +2230,14 @@ void FShaderJobCache::AddJobOutput(FShaderJobData& JobData, const FShaderCommonC
 			// Write new allocated cache file
 			if (!InputDebugInfoPath.IsEmpty())
 			{
-				const FString CacheFilename = FString::Printf(TEXT("%s/%s.bytecode"), *InputDebugInfoPath, *InputSourceFilename);
-				FFileHelper::SaveArrayToFile(TArrayView<const uint8>((const uint8*)Contents.GetData(), Contents.GetSize()), *CacheFilename);
+				const FString CacheOutputFilename = FString::Printf(TEXT("%s/%s.joboutput"), *InputDebugInfoPath, *InputSourceFilename);
+				FFileHelper::SaveArrayToFile(TArrayView<const uint8>((const uint8*)SaveContext.ShaderObjectData.GetData(), SaveContext.ShaderObjectData.GetSize()), *CacheOutputFilename);
+				for (int32 CodeIndex = 0; CodeIndex < SaveContext.ShaderCode.Num(); ++CodeIndex)
+				{
+					const FString CacheCodeFilename = FString::Printf(TEXT("%s/%s_%d.bytecode"), *InputDebugInfoPath, *InputSourceFilename);
+					const FSharedBuffer& JobCode = SaveContext.ShaderCode[CodeIndex];
+					FFileHelper::SaveArrayToFile(TArrayView<const uint8>((const uint8*)JobCode.GetData(), JobCode.GetSize()), *CacheCodeFilename);
+				}
 			}
 		}
 
@@ -2235,21 +2269,21 @@ void FShaderJobCache::AddJobOutput(FShaderJobData& JobData, const FShaderCommonC
 		UE::DerivedData::FCacheKey Key;
 		Key.Bucket = ShaderJobCacheDDCBucket;
 		Key.Hash = Hash;
-		UE::DerivedData::FCacheRecordBuilder RecordBuilder(Key);
-
-		RecordBuilder.AddValue(ShaderJobCacheId, FSharedBuffer::MakeView(Contents));
-
-		TRACE_COUNTER_ADD(Shaders_JobCacheDDCBytesSent, Contents.GetSize());
 
 		UE::DerivedData::FRequestOwner RequestOwner(UE::DerivedData::EPriority::Normal);
 		UE::DerivedData::FRequestBarrier RequestBarrier(RequestOwner);
 		RequestOwner.KeepAlive();
 		UE::DerivedData::GetCache().Put(
-			{ {{TEXT("FShaderJobCache")}, RecordBuilder.Build(), IsShaderJobCacheDDCRemotePolicyEnabled() ? UE::DerivedData::ECachePolicy::Default : UE::DerivedData::ECachePolicy::Local } },
+			{ {{TEXT("FShaderJobCache")}, SaveContext.BuildCacheRecord(Key), IsShaderJobCacheDDCRemotePolicyEnabled() ? UE::DerivedData::ECachePolicy::Default : UE::DerivedData::ECachePolicy::Local } },
 			RequestOwner
 		);
+
+
+		TRACE_COUNTER_ADD(Shaders_JobCacheDDCBytesSent, SaveContext.GetSerializedSize());
 	}
 #endif
+
+
 }
 
 /** Returns memory used by the cache*/
diff --git a/Engine/Source/Runtime/Engine/Public/MaterialShared.h b/Engine/Source/Runtime/Engine/Public/MaterialShared.h
--- a/Engine/Source/Runtime/Engine/Public/MaterialShared.h
+++ b/Engine/Source/Runtime/Engine/Public/MaterialShared.h
@@ -1183,7 +1183,7 @@ public:
 	ENGINE_API void SetShaderDependencies(const TArray<FShaderType*>& ShaderTypes, const TArray<const FShaderPipelineType*>& ShaderPipelineTypes, const TArray<FVertexFactoryType*>& VFTypes, EShaderPlatform ShaderPlatform);
 #endif
 
-	void Serialize(FArchive& Ar, bool bLoadedByCookedMaterial);
+	void Serialize(FArchive& Ar, bool bLoadingCooked);
 
 	bool IsCookedId() const
 	{
@@ -1368,8 +1368,6 @@ private:
 		inline const FHashedName& operator()(const FMeshMaterialShaderMap* InShaderMap) { return InShaderMap->GetVertexFactoryTypeName(); }
 	};
 
-	//void Serialize(FArchive& Ar, bool bInlineShaderResources, bool bLoadedByCookedMaterial);
-
 	ENGINE_API const FMeshMaterialShaderMap* GetMeshShaderMap(const FHashedName& VertexFactoryTypeName) const;
 	ENGINE_API FMeshMaterialShaderMap* AcquireMeshShaderMap(const FHashedName& VertexFactoryTypeName);
 
@@ -1538,15 +1536,17 @@ public:
 	ENGINE_API void AddRef();
 	ENGINE_API void Release();
 
-	/** Serializes the shader map. */
-	UE_DEPRECATED(5.5, "Inlining of shader code into shadermaps is no longer supported. If you rely on this please reach out to the UE rendering team.")
-	// note: FSerializationContext param in the below overload can be defaulted once this is removed, to avoid a bunch of explicit instantiations of FSerializationContext in various callsites
-	bool Serialize(FArchive& Ar, bool bInlineShaderResources = true, bool bLoadedByCookedMaterial = false, bool bInlineShaderCode = false, const FName& SerializingAsset = NAME_None)
+	UE_DEPRECATED(5.5, "Please use overload accepting an FShaderSerializeContext. Note that inlining of shader code into runtime shadermaps is no longer supported. If you rely on this please reach out to the UE rendering team.")
+	bool Serialize(FArchive& Ar, bool bInlineShaderResources = true, bool bLoadingCooked = false, bool bInlineShaderCode = false, const FName& SerializingAsset = NAME_None)
 	{
-		return Serialize(Ar, { bLoadedByCookedMaterial, SerializingAsset });
+		FShaderSerializeContext Ctx(Ar);
+		Ctx.bLoadingCooked = bLoadingCooked;
+		Ctx.SerializingAsset = SerializingAsset;
+		return Serialize(Ctx);
 	}
 
-	bool Serialize(FArchive& Ar, const FShaderMapBase::FSerializationContext& Ctx);
+	/** Serializes the shader map. */
+	bool Serialize(FShaderSerializeContext& Ctx);
 
 #if WITH_EDITOR
 	/** Saves this shader map to the derived data cache. */
@@ -2097,7 +2097,7 @@ public:
 	void SerializeInlineShaderMap(FArchive& Ar, const FName& SerializingAsset = NAME_None);
 
 	/** Serializes the shader map inline in this material, including any shader dependencies. */
-	void RegisterInlineShaderMap(bool bLoadedByCookedMaterial);
+	void RegisterInlineShaderMap(bool bLoadingCooked);
 
 	/** Releases this material's shader map.  Must only be called on materials not exposed to the rendering thread! */
 	void ReleaseShaderMap();
diff --git a/Engine/Source/Runtime/RenderCore/Internal/ShaderSerialization.h b/Engine/Source/Runtime/RenderCore/Internal/ShaderSerialization.h
new file mode 100644
--- /dev/null
+++ b/Engine/Source/Runtime/RenderCore/Internal/ShaderSerialization.h
@@ -0,0 +1,151 @@
+// Copyright Epic Games, Inc. All Rights Reserved.
+
+#pragma once
+
+#include "Memory/SharedBuffer.h"
+#include "Templates/Function.h"
+#include "Templates/UniquePtr.h"
+#include "UObject/NameTypes.h" 
+
+class FMemoryReaderView;
+class FMemoryWriter64;
+struct FShaderCodeResource;
+
+#if WITH_EDITOR
+namespace UE::DerivedData
+{
+	class FCacheRecord;
+	struct FCacheKey;
+}
+#endif
+
+/* Context object used for storing state and serialization parameterization related to shader/shadermap serialization. */
+struct FShaderSerializeContext
+{
+	/* Default constructor, if called it's expected the archive pointer will be set by derived implementations. */
+	FShaderSerializeContext() = default;
+
+	/* Constructor which accepts an FArchive reference, used for simple serialization cases */
+	FShaderSerializeContext(FArchive& InArchive) : Ar(&InArchive) {}
+
+	/* Archive pointer which should be used for serializing the object, possibly excluding shader code if the SerializeCodeFunc is set */
+	FArchive* Ar = nullptr;
+
+	/* Flag indicating whether this serialization is a cooked load; used to change serialization behaviour for cooked data vs cached data */
+	bool bLoadingCooked = false;
+
+	/* FName of the asset which triggered the serialization, this is used only for diagnostic messages */
+	FName SerializingAsset = NAME_None;
+
+	/* Optional function which can be used to serialize shader code separately from the main object; if this is set serialization of bytecode 
+	 * will be excluded from the main object (i.e. not using the Ar member archive reference) and is the responsibility of this function.
+	 */
+	TFunction<void(FShaderCodeResource&, int32)> SerializeCodeFunc{};
+
+	/* Optional function which can be used to reserve space for a number of code entries as a minor optimization. */
+	TFunction<void(int32)> ReserveCodeFunc{};
+
+	/* Helper function to retrieve an FArchive reference; exists for convenience and validation purposes. */
+	FArchive& GetMainArchive()
+	{
+		check(Ar);
+		return *Ar;
+	}
+};
+
+
+/* FShaderSerializeContext implementation used for saving to or loading from caches (either DDC or in-memory job cache) 
+ * Note that this is just a base class of the below and should not be used directly.
+ */
+struct FShaderCacheSerializeContext : public FShaderSerializeContext
+{
+	FShaderCacheSerializeContext() : FShaderSerializeContext() 
+	{
+		ShaderCode = OwnedShaderCode;
+	}
+
+	/* Buffer which stores the main object data for a cache entry - i.e. a shadermap or job structure */
+	FSharedBuffer ShaderObjectData;
+
+	/* View on array of buffers which store the bytecode objects for the object, one per shader/stage.
+	 * Note that this may or may not point to the OwnedShaderCode array below, depending on usage.
+	 */
+	TArrayView<FSharedBuffer> ShaderCode;
+
+	/* Array of buffers actually owned by this context object; it is valid for this to be empty in cases
+	 * where the array of ShaderCode buffers is stored externally.
+	 */
+	TArray<FSharedBuffer> OwnedShaderCode;
+
+	/* Get the total serialized size of data for this context; note that this will return 0 if called prior to the FSharedBuffers
+	 * being set (this is done in the derived implementations, see below).
+	 */
+	int64 GetSerializedSize() const
+	{
+		int64 Size = 0u;
+		if (ShaderObjectData)
+		{
+			Size += ShaderObjectData.GetSize();
+
+			for (FSharedBuffer& CodeBuf : ShaderCode)
+			{
+				Size += CodeBuf.GetSize();
+			}
+		}
+		return Size;
+	}
+
+	/* Boolean conversion operator which checks for validity of the data referenced by this context. */
+	explicit operator bool() const { return ShaderObjectData && !ShaderCode.IsEmpty(); }
+};
+
+/* Implementation of FShaderCacheSerializeContext used for saving data to caches. */
+struct FShaderCacheSaveContext : public FShaderCacheSerializeContext
+{
+	/* Default constructor; sets up base class TFunctions to save bytecode independently of object structure, as
+	 * well as base class FArchive pointing to the owned memory writer.
+	 */
+	RENDERCORE_API FShaderCacheSaveContext();
+
+	/* Converts the raw serialized object data into the ShaderObjectData FSharedBuffer.
+	 * Note that this is called by BuildCacheRecord as well, calls subsequent to the first will have no effect.
+	 */
+	RENDERCORE_API void Finalize();
+
+#if WITH_EDITOR
+	/* Helper function to generate a DDC record from the data serialized using this context. Serialization must
+	 * be executed prior to calling this function.
+	 */
+	RENDERCORE_API UE::DerivedData::FCacheRecord BuildCacheRecord(const UE::DerivedData::FCacheKey& Key);
+#endif
+
+	/* Call to reset internal state, this allows re-using internal allocations to store other objects (as an optimization). */
+	RENDERCORE_API void Reset();
+
+	/* Data that will be written to when serializing */
+	TArray64<uint8> ShaderObjectRawData;
+
+	/* FArchive that is passed to the base class, which writes to ShaderObjectRawData */
+	TUniquePtr<FMemoryWriter64> Writer;
+};
+
+/* Implementation of FShaderCacheSerializeContext used for loading data from caches. */
+struct FShaderCacheLoadContext : public FShaderCacheSerializeContext
+{
+	/*	Default constructor, use when array of code buffers will be allocated via ReadFromRecord */
+	RENDERCORE_API FShaderCacheLoadContext();
+	/* Constructor which references buffers (and an array of code buffers) owned elsewhere, does not allocate the OwnedShaderCode array */
+	RENDERCORE_API FShaderCacheLoadContext(FSharedBuffer ShaderObjectData, TArrayView<FSharedBuffer> CodeBuffers);
+
+	/* Call to reset reader to start position so the same load context can be used to populate multiple objects. */
+	RENDERCORE_API void Reuse();
+
+#if WITH_EDITOR
+	/* Helper function to populate the internal state (FSharedBuffers defined on FShaderCacheSerializeContext) from a DDC record. */
+	RENDERCORE_API void ReadFromRecord(const UE::DerivedData::FCacheRecord& CacheRecord, bool bIsPersistent = false);
+#endif
+
+	/* FArchive passed to the base class; this will be used to read data from the FShaderCacheSerializeContext::ShaderObjectData shared buffer */
+	TUniquePtr<FMemoryReaderView> Reader;
+};
+
diff --git a/Engine/Source/Runtime/RenderCore/Private/GlobalShader.cpp b/Engine/Source/Runtime/RenderCore/Private/GlobalShader.cpp
--- a/Engine/Source/Runtime/RenderCore/Private/GlobalShader.cpp
+++ b/Engine/Source/Runtime/RenderCore/Private/GlobalShader.cpp
@@ -6,14 +6,15 @@
 
 #include "GlobalShader.h"
 
+#include "Containers/StaticBitArray.h"
 #include "Interfaces/ITargetPlatform.h"
+#include "Interfaces/ITargetPlatformManagerModule.h"
+#include "Misc/ConfigCacheIni.h"
+#include "Misc/CoreMisc.h"
 #include "Misc/Paths.h"
 #include "Serialization/MemoryWriter.h"
-#include "Misc/CoreMisc.h"
-#include "Misc/ConfigCacheIni.h"
-#include "Interfaces/ITargetPlatformManagerModule.h"
-#include "Containers/StaticBitArray.h"
 #include "ShaderCodeLibrary.h"
+#include "ShaderSerialization.h"
 
 /** The global shader map. */
 FGlobalShaderMap* GGlobalShaderMap[SP_NumPlatforms] = {};
@@ -538,17 +539,31 @@ FGlobalShaderMap* GetGlobalShaderMap(EShaderPlatform Platform)
 FGlobalShaderMapSection* FGlobalShaderMapSection::CreateFromArchive(FArchive& Ar)
 {
 	FGlobalShaderMapSection* Section = new FGlobalShaderMapSection();
-	if (Section->Serialize(Ar))
+	FShaderSerializeContext Ctx(Ar);
+	if (Section->Serialize(Ctx))
+	{
+		return Section;
+	}
+	delete Section;
+	return nullptr;
+}
+
+#if WITH_EDITOR
+FGlobalShaderMapSection* FGlobalShaderMapSection::CreateFromCache(FShaderCacheLoadContext& Ctx)
+{
+	FGlobalShaderMapSection* Section = new FGlobalShaderMapSection();
+	if (Section->Serialize(Ctx))
 	{
 		return Section;
 	}
 	delete Section;
 	return nullptr;
 }
+#endif
 
-bool FGlobalShaderMapSection::Serialize(FArchive& Ar)
+bool FGlobalShaderMapSection::Serialize(FShaderSerializeContext& Ctx)
 {
-	return Super::Serialize(Ar, FShaderMapBase::FSerializationContext());
+	return Super::Serialize(Ctx);
 }
 
 TShaderRef<FShader> FGlobalShaderMapSection::GetShader(FShaderType* ShaderType, int32 PermutationId) const
@@ -766,7 +781,8 @@ void FGlobalShaderMap::SaveToGlobalArchive(FArchive& Ar)
 
 	for (const auto& It : SectionMap)
 	{
-		It.Value->Serialize(Ar);
+		FShaderSerializeContext Ctx(Ar);
+		It.Value->Serialize(Ctx);
 	}
 }
 
diff --git a/Engine/Source/Runtime/RenderCore/Private/ShaderCodeLibrary.cpp b/Engine/Source/Runtime/RenderCore/Private/ShaderCodeLibrary.cpp
--- a/Engine/Source/Runtime/RenderCore/Private/ShaderCodeLibrary.cpp
+++ b/Engine/Source/Runtime/RenderCore/Private/ShaderCodeLibrary.cpp
@@ -1417,7 +1417,7 @@ struct FEditorShaderCodeArchive
 
 		if (SerializedShaders.FindOrAddShaderMap(Code->ResourceHash, ShaderMapIndex, &AssociatedAssets))
 		{
-			const int32 NumShaders = Code->ShaderEntries.Num();
+			const int32 NumShaders = Code->ShaderCodeResources.Num();
 			FShaderMapEntry& ShaderMapEntry = SerializedShaders.ShaderMapEntries[ShaderMapIndex];
 			ShaderMapEntry.NumShaders = NumShaders;
 			ShaderMapEntry.ShaderIndicesOffset = SerializedShaders.ShaderIndices.AddZeroed(NumShaders);
@@ -1427,25 +1427,25 @@ struct FEditorShaderCodeArchive
 				int32 ShaderIndex = INDEX_NONE;
 				if (SerializedShaders.FindOrAddShader(Code->ShaderHashes[i], ShaderIndex))
 				{
-					const FShaderMapResourceCode::FShaderEntry& SourceShaderEntry = Code->ShaderEntries[i];
+					const FShaderCodeResource& SourceShaderResource = Code->ShaderCodeResources[i];
 					FShaderCodeEntry& SerializedShaderEntry = SerializedShaders.ShaderEntries[ShaderIndex];
-					SerializedShaderEntry.Frequency = SourceShaderEntry.Frequency;
-					SerializedShaderEntry.Size = SourceShaderEntry.Code.Num();
-					SerializedShaderEntry.UncompressedSize = SourceShaderEntry.UncompressedSize;
-					check(!SourceShaderEntry.Code.IsEmpty());
-					ShaderCode.Add(SourceShaderEntry.Code);
+					SerializedShaderEntry.Frequency = SourceShaderResource.Frequency;
+					SerializedShaderEntry.Size = SourceShaderResource.Code.Num();
+					SerializedShaderEntry.UncompressedSize = SourceShaderResource.UncompressedSize;
+					check(!SourceShaderResource.Code.IsEmpty());
+					ShaderCode.Add(SourceShaderResource.Code);
 					check(ShaderCode.Num() == SerializedShaders.ShaderEntries.Num());
 
 					CodeStats.NumUniqueShaders++;
-					CodeStats.ShadersUniqueSize += SourceShaderEntry.Code.Num();
+					CodeStats.ShadersUniqueSize += SourceShaderResource.Code.Num();
 				}
-				CodeStats.ShadersSize += Code->ShaderEntries[i].Code.Num();
+				CodeStats.ShadersSize += Code->ShaderCodeResources[i].Code.Num();
 				SerializedShaders.ShaderIndices[ShaderMapEntry.ShaderIndicesOffset + i] = ShaderIndex;
 			}
 
 			// for total shaders, only count shaders when we're adding a new shadermap. AddShaderCode() for the same shadermap can be called several times during
 			// the cook because of serialization path being reused for other purposes than actual saving, so counting them every time artificially inflates number of shaders.
-			CodeStats.NumShaders += Code->ShaderEntries.Num();
+			CodeStats.NumShaders += Code->ShaderCodeResources.Num();
 			CodeStats.NumShaderMaps++;
 		}
 		// always mark the shadermap dirty, because it might have gotten new asset associations
diff --git a/Engine/Source/Runtime/RenderCore/Private/ShaderCore.cpp b/Engine/Source/Runtime/RenderCore/Private/ShaderCore.cpp
--- a/Engine/Source/Runtime/RenderCore/Private/ShaderCore.cpp
+++ b/Engine/Source/Runtime/RenderCore/Private/ShaderCore.cpp
@@ -32,6 +32,7 @@
 #include "Serialization/MemoryWriter.h"
 #include "Serialization/ShaderKeyGenerator.h"
 #include "Shader.h"
+#include "ShaderSerialization.h"
 #include "ShaderCompilerCore.h"
 #include "ShaderCompilerDefinitions.h"
 #include "ShaderCompilerJobTypes.h"
@@ -3999,40 +4000,36 @@ FShaderCompilerInputHash FShaderCompileJob::GetInputHash()
 	return InputHash;
 }
 
-void FShaderCompileJob::SerializeOutput(FArchive& Ar)
+void FShaderCompileJob::SerializeOutput(FShaderCacheSerializeContext& Ctx, int32 CodeIndex)
 {
-	double ActualCompileTime = 0.0;
-	// Save the preprocess time as set in the job regardless of whether saving or loading - if loading from the cache and the preprocessed job
-	// cache is enabled this job will have already run its own preprocessing and we want to track/aggregate this time properly.
-	double ActualPreprocessTime = Output.PreprocessTime;
-	if (Ar.IsSaving())
+	Output.bSerializingForCache = true;
+	FArchive& Ar = Ctx.GetMainArchive();
+	bool bIsSaving = Ar.IsSaving();
+	bool bIsLoading = Ar.IsLoading();
+
+	Ar << Output;
+
+	FShaderCodeResource CodeResource;
+	if (bIsSaving)
 	{
-		// Clear preprocess time and compile time when storing a job in the cache. This reduces storage requirements since these objects are
-		// deduplicated based on a hash (and otherwise duplicate jobs will still differ in these values).
-		ActualCompileTime = Output.CompileTime;
-		Output.CompileTime = 0.0;
-		Output.PreprocessTime = 0.0;
+		CodeResource = Output.ConvertCodeToResource();
 	}
 
-	Ar << Output;
+	Ctx.SerializeCodeFunc(CodeResource, CodeIndex);
+
+	// we intentionally re-set the internal ShaderCode even when saving; GetCodeResource moves the code array into the
+	// FShaderCodeResource's internal array and this moves it back (preventing an unnecessary temporary copy of the code)
+	Output.SetCodeFromResource(MoveTemp(CodeResource));
+	
 	// output hash is now serialized as part of the output, as the shader code is compressed in SCWs
 	checkf(!Output.bSucceeded || Output.OutputHash != FSHAHash(), TEXT("Successful compile job does not have an OutputHash generated."));
 	checkf(Output.Target == Input.Target, TEXT("Output FShaderTarget does not match the input struct; incorrect results associated with job?"));
 
-	if (Ar.IsLoading())
+	if (bIsLoading)
 	{
 		bFinalized = true;
 		bSucceeded = Output.bSucceeded;
 	}
-	else
-	{
-		// Restore the compile time for this job if we're saving to the cache.
-		// Jobs that will be deserialized from the cache will have a compile time of 0.0
-		Output.CompileTime = ActualCompileTime;
-	}
-
-	// Unconditionally restore the preprocess time for this job after saving to or loading from the cache.
-	Output.PreprocessTime = ActualPreprocessTime;
 }
 
 void FShaderCompileJob::OnComplete()
@@ -4191,16 +4188,22 @@ FShaderPipelineCompileJob::FShaderPipelineCompileJob(uint32 InHash, uint32 InId,
 	}
 }
 
-void FShaderPipelineCompileJob::SerializeOutput(FArchive& Ar)
+void FShaderPipelineCompileJob::SerializeOutput(FShaderCacheSerializeContext& Ctx)
 {
 	bool bAllStagesSucceeded = true;
+
+	if (Ctx.ReserveCodeFunc)
+	{
+		Ctx.ReserveCodeFunc(StageJobs.Num());
+	}
+
 	for (int32 Index = 0, Num = StageJobs.Num(); Index < Num; ++Index)
 	{
-		StageJobs[Index]->SerializeOutput(Ar);
+		StageJobs[Index]->SerializeOutput(Ctx, Index);
 		bAllStagesSucceeded = bAllStagesSucceeded && StageJobs[Index]->bSucceeded;
 	}
 
-	if (Ar.IsLoading())
+	if (Ctx.GetMainArchive().IsLoading())
 	{
 		bFinalized = true;
 		bSucceeded = bAllStagesSucceeded;
diff --git a/Engine/Source/Runtime/RenderCore/Private/ShaderMap.cpp b/Engine/Source/Runtime/RenderCore/Private/ShaderMap.cpp
--- a/Engine/Source/Runtime/RenderCore/Private/ShaderMap.cpp
+++ b/Engine/Source/Runtime/RenderCore/Private/ShaderMap.cpp
@@ -235,8 +235,9 @@ void FShaderMapBase::UnfreezeContent()
 
 #define CHECK_SHADERMAP_DEPENDENCIES (WITH_EDITOR || !(UE_BUILD_SHIPPING || UE_BUILD_TEST))
 
-bool FShaderMapBase::Serialize(FArchive& Ar, const FSerializationContext& Ctx)
+bool FShaderMapBase::Serialize(FShaderSerializeContext& Ctx)
 {
+	FArchive& Ar = Ctx.GetMainArchive();
 	LLM_SCOPE(ELLMTag::Shaders);
 	if (Ar.IsSaving())
 	{
@@ -288,10 +289,10 @@ bool FShaderMapBase::Serialize(FArchive& Ar, const FSerializationContext& Ctx)
 			Ar << ResourceHash;
 			FShaderLibraryCooker::AddShaderCode(ShaderPlatform, Code, GetAssociatedAssets());
 		}
-		else
+		else 
 #endif // WITH_EDITOR
 		{
-			Code->Serialize(Ar, Ctx.bLoadedByCookedMaterial);
+			Code->Serialize(Ctx);
 		}
 	}
 	else
@@ -331,7 +332,7 @@ bool FShaderMapBase::Serialize(FArchive& Ar, const FSerializationContext& Ctx)
 		else
 		{
 			Code = new FShaderMapResourceCode();
-			Code->Serialize(Ar, Ctx.bLoadedByCookedMaterial);
+			Code->Serialize(Ctx);
 			Resource = new FShaderMapResource_InlineCode(ShaderPlatform, Code);
 		}
 
diff --git a/Engine/Source/Runtime/RenderCore/Private/ShaderResource.cpp b/Engine/Source/Runtime/RenderCore/Private/ShaderResource.cpp
--- a/Engine/Source/Runtime/RenderCore/Private/ShaderResource.cpp
+++ b/Engine/Source/Runtime/RenderCore/Private/ShaderResource.cpp
@@ -5,22 +5,22 @@
 =============================================================================*/
 
 #include "Shader.h"
+#include "Compression/OodleDataCompression.h"
+#include "DataDrivenShaderPlatformInfo.h"
 #include "Misc/Compression.h"
 #include "Misc/CoreMisc.h"
 #include "Misc/StringBuilder.h"
 #include "Interfaces/ITargetPlatform.h"
 #include "Interfaces/ITargetPlatformManagerModule.h"
 #include "Interfaces/IShaderFormat.h"
-#include "RHI.h"
-#include "ShaderCore.h"
+#include "Misc/MemStack.h"
 #include "Misc/ScopeLock.h"
 #include "RenderingThread.h"
-#include "UObject/RenderingObjectVersion.h"
-#include "Misc/MemStack.h"
-#include "ShaderCompilerCore.h"
-#include "Compression/OodleDataCompression.h"
+#include "RHI.h"
 #include "RHIResources.h"	// Access to FRHIRayTracingShader::RayTracingPayloadType requires this
-#include "DataDrivenShaderPlatformInfo.h"
+#include "ShaderCompilerCore.h"
+#include "ShaderCore.h"
+#include "UObject/RenderingObjectVersion.h"
 
 #if WITH_EDITORONLY_DATA
 #include "Interfaces/IShaderFormat.h"
@@ -177,7 +177,7 @@ static void ApplyResourceStats(FShaderMapResourceCode& Resource)
 {
 #if STATS
 	INC_DWORD_STAT_BY(STAT_Shaders_ShaderResourceMemory, Resource.GetSizeBytes());
-	for (const FShaderMapResourceCode::FShaderEntry& Shader : Resource.ShaderEntries)
+	for (const FShaderCodeResource& Shader : Resource.ShaderCodeResources)
 	{
 		INC_DWORD_STAT_BY_FName(GetMemoryStatType(Shader.Frequency).GetName(), Shader.Code.Num());
 	}
@@ -188,7 +188,7 @@ static void RemoveResourceStats(FShaderMapResourceCode& Resource)
 {
 #if STATS
 	DEC_DWORD_STAT_BY(STAT_Shaders_ShaderResourceMemory, Resource.GetSizeBytes());
-	for (const FShaderMapResourceCode::FShaderEntry& Shader : Resource.ShaderEntries)
+	for (const FShaderCodeResource& Shader : Resource.ShaderCodeResources)
 	{
 		DEC_DWORD_STAT_BY_FName(GetMemoryStatType(Shader.Frequency).GetName(), Shader.Code.Num());
 	}
@@ -199,17 +199,19 @@ FShaderMapResourceCode::FShaderMapResourceCode(const FShaderMapResourceCode& Oth
 {
 	ResourceHash = Other.ResourceHash;
 	ShaderHashes = Other.ShaderHashes;
-	ShaderEntries = Other.ShaderEntries;
+	ShaderCodeResources = Other.ShaderCodeResources;
 
 #if WITH_EDITORONLY_DATA
 	ShaderEditorOnlyDataEntries = Other.ShaderEditorOnlyDataEntries;
 #endif // WITH_EDITORONLY_DATA
 }
 
+PRAGMA_DISABLE_DEPRECATION_WARNINGS
 FShaderMapResourceCode::~FShaderMapResourceCode()
 {
 	RemoveResourceStats(*this);
 }
+PRAGMA_ENABLE_DEPRECATION_WARNINGS
 
 void FShaderMapResourceCode::Finalize()
 {
@@ -226,8 +228,8 @@ void FShaderMapResourceCode::Finalize()
 
 uint32 FShaderMapResourceCode::GetSizeBytes() const
 {
-	uint64 Size = sizeof(*this) + ShaderHashes.GetAllocatedSize() + ShaderEntries.GetAllocatedSize();
-	for (const FShaderEntry& Entry : ShaderEntries)
+	uint64 Size = sizeof(*this) + ShaderHashes.GetAllocatedSize() + ShaderCodeResources.GetAllocatedSize();
+	for (const FShaderCodeResource& Entry : ShaderCodeResources)
 	{
 		Size += Entry.Code.GetAllocatedSize();
 	}
@@ -245,7 +247,6 @@ void FShaderMapResourceCode::AddShaderCompilerOutput(const FShaderCompilerOutput
 	TRACE_CPUPROFILER_EVENT_SCOPE(FShaderMapResourceCode::AddShaderCode);
 
 	const FSHAHash& InHash = Output.OutputHash;
-	const FShaderCode& InCode = Output.ShaderCode;
 	const int32 Index = Algo::LowerBound(ShaderHashes, InHash);
 	if (Index >= ShaderHashes.Num() || ShaderHashes[Index] != InHash)
 	{
@@ -256,66 +257,9 @@ void FShaderMapResourceCode::AddShaderCompilerOutput(const FShaderCompilerOutput
 		AddEditorOnlyData(Index, DebugName, Output.PlatformDebugData, Output.Errors, Output.ShaderStatistics, DebugInfo);
 #endif
 
-		FShaderEntry& Entry = ShaderEntries.InsertDefaulted_GetRef(Index);
-		Entry.Frequency = Output.Target.GetFrequency();
-		const TArray<uint8>& ShaderCode = InCode.GetReadAccess();
-
-		FName ShaderCompressionFormat = GetShaderCompressionFormat();
-		if (ShaderCompressionFormat != NAME_None)
-		{
-			Entry.UncompressedSize = InCode.GetUncompressedSize();
-
-			// we trust that SCWs also obeyed by the same CVar, so we expect a compressed shader code at this point
-			// However, if we see an uncompressed shader, it perhaps means that SCW tried to compress it, but the result was worse than uncompressed. 
-			// Because of that we special-case NAME_None here
-			if (ShaderCompressionFormat != InCode.GetCompressionFormat())
-			{
-				if (InCode.GetCompressionFormat() != NAME_None)
-				{
-					UE_LOG(LogShaders, Fatal, TEXT("Shader %s is expected to be compressed with %s, but it is compressed with %s instead."),
-						*InHash.ToString(),
-						*ShaderCompressionFormat.ToString(),
-						*InCode.GetCompressionFormat().ToString()
-						);
-					// unreachable
-					return;
-				}
-				
-				// assume uncompressed due to worse ratio than the compression
-				Entry.UncompressedSize = ShaderCode.Num();
-				UE_LOG(LogShaders, Verbose, TEXT("Shader %s is expected to be compressed with %s, but it arrived uncompressed (size=%d). Assuming compressing made it longer and storing uncompressed."),
-					*InHash.ToString(),
-					*ShaderCompressionFormat.ToString(),
-					ShaderCode.Num()
-				);
-			}
-			else if (ShaderCompressionFormat == NAME_Oodle)
-			{
-				// check if Oodle-specific settings match
-				FOodleDataCompression::ECompressor OodleCompressor;
-				FOodleDataCompression::ECompressionLevel OodleLevel;
-				GetShaderCompressionOodleSettings(OodleCompressor, OodleLevel);
-
-				if (InCode.GetOodleCompressor() != OodleCompressor || InCode.GetOodleLevel() != OodleLevel)
-				{
-					UE_LOG(LogShaders, Fatal, TEXT("Shader %s is expected to be compressed with Oodle compressor %d level %d, but it is compressed with compressor %d level %d instead."),
-						*InHash.ToString(),
-						static_cast<int32>(OodleCompressor),
-						static_cast<int32>(OodleLevel),
-						static_cast<int32>(InCode.GetOodleCompressor()),
-						static_cast<int32>(InCode.GetOodleLevel())
-					);
-					// unreachable
-					return;
-				}
-			}
-		}
-		else
-		{
-			Entry.UncompressedSize = ShaderCode.Num();
-		}
-
-		Entry.Code = ShaderCode;
+		// This moves the internal array in the Output's ShaderCode object into the resource format; assuming that after
+		// the code is moved to the shadermap that it's no longer required in the job output struct.
+		ShaderCodeResources.Insert(Output.ConvertCodeToResource(), Index);
 	}
 #if WITH_EDITORONLY_DATA
 	else
@@ -410,20 +354,41 @@ void FShaderMapResourceCode::ToString(FStringBuilderBase& OutString) const
 	OutString.Appendf(TEXT("Shaders: Num=%d\n"), ShaderHashes.Num());
 	for (int32 i = 0; i < ShaderHashes.Num(); ++i)
 	{
-		const FShaderEntry& Entry = ShaderEntries[i];
+		const FShaderCodeResource& Res = ShaderCodeResources[i];
 		OutString.Appendf(TEXT("    [%d]: { Hash: %s, Freq: %s, Size: %d, UncompressedSize: %d }\n"),
-			i, *ShaderHashes[i].ToString(), GetShaderFrequencyString(Entry.Frequency), Entry.Code.Num(), Entry.UncompressedSize);
+			i, *ShaderHashes[i].ToString(), GetShaderFrequencyString(Res.Frequency), Res.Code.Num(), Res.UncompressedSize);
 	}
 }
 
-void FShaderMapResourceCode::Serialize(FArchive& Ar, bool bLoadedByCookedMaterial)
+void FShaderMapResourceCode::Serialize(FShaderSerializeContext& Ctx)
 {
+	FArchive& Ar = Ctx.GetMainArchive();
 	Ar << ResourceHash;
 	Ar << ShaderHashes;
-	Ar << ShaderEntries;
-	check(ShaderEntries.Num() == ShaderHashes.Num());
+	if (!Ctx.SerializeCodeFunc)
+	{
+		Ar << ShaderCodeResources;
+	}
+	else
+	{
+		if (Ar.IsLoading())
+		{
+			ShaderCodeResources.SetNum(ShaderHashes.Num());
+		}
+		
+		if (Ctx.ReserveCodeFunc)
+		{
+			Ctx.ReserveCodeFunc(ShaderCodeResources.Num());
+		}
+
+		for (int32 CodeIndex = 0; CodeIndex < ShaderCodeResources.Num(); ++CodeIndex)
+		{
+			Ctx.SerializeCodeFunc(ShaderCodeResources[CodeIndex], CodeIndex);
+		}
+	}
+	check(ShaderCodeResources.Num() == ShaderHashes.Num());
 #if WITH_EDITORONLY_DATA
-	const bool bSerializeEditorOnlyData = !bLoadedByCookedMaterial && (!Ar.IsCooking() || Ar.CookingTarget()->HasEditorOnlyData());
+	const bool bSerializeEditorOnlyData = !Ctx.bLoadingCooked && (!Ar.IsCooking() || Ar.CookingTarget()->HasEditorOnlyData());
 	if (bSerializeEditorOnlyData)
 	{
 		Ar << ShaderEditorOnlyDataEntries;
@@ -689,21 +654,21 @@ FRHIShader* FShaderMapResource_InlineCode::CreateRHIShaderOrCrash(int32 ShaderIn
 	}
 
 	FMemStackBase& MemStack = FMemStack::Get();
-	const FShaderMapResourceCode::FShaderEntry& ShaderEntry = Code->ShaderEntries[ShaderIndex];
-	const uint8* ShaderCode = ShaderEntry.Code.GetData();
+	const FShaderCodeResource& ShaderCodeResource = Code->ShaderCodeResources[ShaderIndex];
+	const uint8* ShaderCode = ShaderCodeResource.Code.GetData();
 
 	FMemMark Mark(MemStack);
-	if (ShaderEntry.Code.Num() != ShaderEntry.UncompressedSize)
+	if (ShaderCodeResource.Code.Num() != ShaderCodeResource.UncompressedSize)
 	{
-		void* UncompressedCode = MemStack.Alloc(ShaderEntry.UncompressedSize, 16);
-		bool bSucceed = FCompression::UncompressMemory(GetShaderCompressionFormat(), UncompressedCode, ShaderEntry.UncompressedSize, ShaderCode, ShaderEntry.Code.Num());
+		void* UncompressedCode = MemStack.Alloc(ShaderCodeResource.UncompressedSize, 16);
+		bool bSucceed = FCompression::UncompressMemory(GetShaderCompressionFormat(), UncompressedCode, ShaderCodeResource.UncompressedSize, ShaderCode, ShaderCodeResource.Code.Num());
 		check(bSucceed);
 		ShaderCode = (uint8*)UncompressedCode;
 	}
 
-	const auto ShaderCodeView = MakeArrayView(ShaderCode, ShaderEntry.UncompressedSize);
+	const auto ShaderCodeView = MakeArrayView(ShaderCode, ShaderCodeResource.UncompressedSize);
 	const FSHAHash& ShaderHash = Code->ShaderHashes[ShaderIndex];
-	const EShaderFrequency Frequency = ShaderEntry.Frequency;
+	const EShaderFrequency Frequency = ShaderCodeResource.Frequency;
 
 	TRefCountPtr<FRHIShader> RHIShader;
 	switch (Frequency)
diff --git a/Engine/Source/Runtime/RenderCore/Private/ShaderSerialization.cpp b/Engine/Source/Runtime/RenderCore/Private/ShaderSerialization.cpp
new file mode 100644
--- /dev/null
+++ b/Engine/Source/Runtime/RenderCore/Private/ShaderSerialization.cpp
@@ -0,0 +1,117 @@
+// Copyright Epic Games, Inc. All Rights Reserved.
+
+#include "ShaderSerialization.h"
+
+#include "DerivedDataCacheKey.h"
+#include "DerivedDataCacheRecord.h"
+#include "DerivedDataValue.h"
+#include "Serialization/CompactBinaryWriter.h"
+#include "Serialization/MemoryReader.h"
+#include "Serialization/MemoryWriter.h"
+#include "ShaderCore.h"
+
+#if WITH_EDITOR
+FShaderCacheSaveContext::FShaderCacheSaveContext()
+{
+	SerializeCodeFunc = [this](FShaderCodeResource& Res, int32 Index)
+		{
+			TArray64<uint8> CodeData;
+			FMemoryWriter64 Ar(CodeData, true);
+			Ar << Res;
+			OwnedShaderCode.Add(MakeSharedBufferFromArray(MoveTemp(CodeData)));
+			// reset the array view any time an entry is added; we do this instead of calling Resize in the reserve delegate
+			// and setting it there since not all code paths (i.e. single job cache records) call reserve
+			ShaderCode = OwnedShaderCode;
+		};
+
+	ReserveCodeFunc = [this](int32 Num)
+		{
+			OwnedShaderCode.Reserve(Num);
+		};
+	Reset();
+}
+
+const UE::DerivedData::FValueId ShaderObjectDataValue = UE::DerivedData::FValueId::FromName(TEXT("ShaderObjectData"));
+const UE::DerivedData::FValueId ShaderCodeDataValue = UE::DerivedData::FValueId::FromName(TEXT("ShaderCodeData"));
+const FAnsiStringView CodeCountMetaField = ANSITEXTVIEW("CodeCount");
+
+void FShaderCacheSaveContext::Reset()
+{
+	ShaderObjectData.Reset();
+	OwnedShaderCode.Reset();
+	Writer = MakeUnique<FMemoryWriter64>(ShaderObjectRawData);
+	Ar = Writer.Get();
+}
+
+void FShaderCacheSaveContext::Finalize()
+{
+	if (!ShaderObjectData)
+	{
+		ShaderObjectData = MakeSharedBufferFromArray(MoveTemp(ShaderObjectRawData));
+	}
+}
+
+UE::DerivedData::FCacheRecord FShaderCacheSaveContext::BuildCacheRecord(const UE::DerivedData::FCacheKey& Key)
+{
+	Finalize();
+
+	UE::DerivedData::FCacheRecordBuilder RecordBuilder(Key);
+	RecordBuilder.AddValue(ShaderObjectDataValue, UE::DerivedData::FValue::Compress(ShaderObjectData));
+	int32 CodeIndex = 0;
+	for (FSharedBuffer& CodeBuf : ShaderCode)
+	{
+		RecordBuilder.AddValue(ShaderCodeDataValue.MakeIndexed(CodeIndex++), UE::DerivedData::FValue::Compress(CodeBuf));
+	}
+
+	TCbWriter<16> MetaWriter;
+	MetaWriter.BeginObject();
+	MetaWriter.AddInteger(CodeCountMetaField, ShaderCode.Num());
+	MetaWriter.EndObject();
+
+	RecordBuilder.SetMeta(MetaWriter.Save().AsObject());
+
+	return RecordBuilder.Build();
+}
+
+FShaderCacheLoadContext::FShaderCacheLoadContext()
+{
+	SerializeCodeFunc = [this](FShaderCodeResource& Res, int32 Index)
+		{
+			FMemoryReaderView Ar(ShaderCode[Index]);
+			Ar << Res;
+		};
+}
+
+FShaderCacheLoadContext::FShaderCacheLoadContext(FSharedBuffer InShaderObjectData, TArrayView<FSharedBuffer> InCodeBuffers) : FShaderCacheLoadContext()
+{
+	ShaderObjectData = InShaderObjectData;
+	ShaderCode = InCodeBuffers;
+
+	Reader = MakeUnique<FMemoryReaderView>(ShaderObjectData);
+	Ar = Reader.Get();
+}
+
+void FShaderCacheLoadContext::Reuse()
+{
+	Reader->Seek(0u);
+}
+
+void FShaderCacheLoadContext::ReadFromRecord(const UE::DerivedData::FCacheRecord& Record, bool bIsPersistent)
+{
+	ShaderObjectData = Record.GetValue(ShaderObjectDataValue).GetData().Decompress();
+	
+	// Must initialize a memory reader (and the base class archive pointer) after reading the base shadermap data buffer
+	// from the DDC record
+	Reader = MakeUnique<FMemoryReaderView>(ShaderObjectData, bIsPersistent);
+	Ar = Reader.Get();
+	
+	int32 CodeCount = Record.GetMeta()[CodeCountMetaField].AsInt32();
+	OwnedShaderCode.Reserve(CodeCount);
+	for (int32 CodeIndex = 0; CodeIndex < CodeCount; ++CodeIndex)
+	{
+		OwnedShaderCode.Add(Record.GetValue(ShaderCodeDataValue.MakeIndexed(CodeIndex)).GetData().Decompress());
+	}
+	ShaderCode = OwnedShaderCode;
+}
+#endif
+
diff --git a/Engine/Source/Runtime/RenderCore/Public/GlobalShader.h b/Engine/Source/Runtime/RenderCore/Public/GlobalShader.h
--- a/Engine/Source/Runtime/RenderCore/Public/GlobalShader.h
+++ b/Engine/Source/Runtime/RenderCore/Public/GlobalShader.h
@@ -33,6 +33,9 @@ class FShaderUniformBufferParameter;
 class FVertexFactoryType;
 class ITargetPlatform;
 struct FShaderCompilerEnvironment;
+#if WITH_EDITOR
+struct FShaderCacheLoadContext;
+#endif
 
 /** Used to identify the global shader map in compile queues. */
 extern RENDERCORE_API const int32 GlobalShaderMapId;
@@ -180,7 +183,17 @@ class FGlobalShaderMapSection : public TShaderMap<FGlobalShaderMapContent, FShad
 public:
 	RENDERCORE_API static FGlobalShaderMapSection* CreateFromArchive(FArchive& Ar);
 
-	RENDERCORE_API bool Serialize(FArchive& Ar);
+#if WITH_EDITOR
+	RENDERCORE_API static FGlobalShaderMapSection* CreateFromCache(FShaderCacheLoadContext& Ctx);
+#endif
+
+	UE_DEPRECATED(5.5, "Use overload accepting a FShaderSerializeContext instead")
+	RENDERCORE_API bool Serialize(FArchive& Ar)
+	{
+		return false;
+	}
+
+	RENDERCORE_API bool Serialize(FShaderSerializeContext& Ctx);
 	RENDERCORE_API virtual void GetShaderList(TMap<FHashedName, TShaderRef<FShader>>& OutShaders) const override;
 	RENDERCORE_API virtual void GetShaderPipelineList(TArray<FShaderPipelineRef>& OutShaderPipelines) const override;
 private:
diff --git a/Engine/Source/Runtime/RenderCore/Public/Shader.h b/Engine/Source/Runtime/RenderCore/Public/Shader.h
--- a/Engine/Source/Runtime/RenderCore/Public/Shader.h
+++ b/Engine/Source/Runtime/RenderCore/Public/Shader.h
@@ -45,6 +45,7 @@
 #include "ShaderParameterMetadata.h"
 #include "ShaderParameters.h"
 #include "ShaderPermutation.h"
+#include "ShaderSerialization.h"
 #include "Templates/RefCounting.h"
 #include "Templates/TypeHash.h"
 #include "Templates/UniquePtr.h"
@@ -479,7 +480,7 @@ private:
 class FShaderMapResourceCode : public FThreadSafeRefCountedObject
 {
 public:
-	struct FShaderEntry
+	struct UE_DEPRECATED(5.5, "Use FShaderCodeResource") FShaderEntry
 	{
 		TArray<uint8> Code;
 		int32 UncompressedSize;
@@ -526,7 +527,15 @@ public:
 
 	RENDERCORE_API void Finalize();
 
-	RENDERCORE_API void Serialize(FArchive& Ar, bool bLoadedByCookedMaterial);
+	UE_DEPRECATED(5.5, "Use overload accepting an FShaderSerializeContext object")
+		void Serialize(FArchive& Ar, bool bLoadingCooked)
+	{
+		FShaderSerializeContext Ctx(Ar);
+		Ctx.bLoadingCooked = bLoadingCooked;
+		Serialize(Ctx);
+	}
+
+	RENDERCORE_API void Serialize(FShaderSerializeContext& Ctx);
 #if WITH_EDITORONLY_DATA
 	RENDERCORE_API void NotifyShadersCompiled(FName FormatName);
 #endif // WITH_EDITORONLY_DATA
@@ -549,7 +558,13 @@ public:
 	/** A hash describing the total contents of *this. Constructed from the contents of ShaderHashes during Finalize. */
 	FSHAHash ResourceHash;
 	TArray<FSHAHash> ShaderHashes;
+
+	UE_DEPRECATED(5.5, "Use ShaderCodeResources") 
+	PRAGMA_DISABLE_DEPRECATION_WARNINGS
 	TArray<FShaderEntry> ShaderEntries;
+	PRAGMA_ENABLE_DEPRECATION_WARNINGS
+
+	TArray<FShaderCodeResource> ShaderCodeResources;
 #if WITH_EDITORONLY_DATA
 	// Optional array of editor-only data indexed in the same order as ShaderEntries (sorted by the shader hash)
 	// Empty in the cases where the editor-only data is not serialized.
@@ -566,7 +581,7 @@ class FShaderMapResource_InlineCode : public FShaderMapResource
 {
 public:
 	FShaderMapResource_InlineCode(EShaderPlatform InPlatform, FShaderMapResourceCode* InCode)
-		: FShaderMapResource(InPlatform, InCode->ShaderEntries.Num())
+		: FShaderMapResource(InPlatform, InCode->ShaderCodeResources.Num())
 		, Code(InCode)
 	{}
 
@@ -1617,7 +1632,7 @@ struct FShaderCompiledShaderInitializerType
 	const FShaderType* Type;
 	const FShaderType::FParameters* Parameters;
 	FShaderTarget Target;
-	const TArray<uint8>& Code;
+	TArrayView<const uint8> Code;
 	const FShaderParameterMap& ParameterMap;
 	const FSHAHash& OutputHash;
 	FSHAHash MaterialShaderMapHash;
@@ -2392,18 +2407,16 @@ public:
 
 	RENDERCORE_API void FinalizeContent();
 	RENDERCORE_API void UnfreezeContent();
-	UE_DEPRECATED(5.5, "Inlining of shader code into shadermaps is no longer supported. If you rely on this please reach out to the UE rendering team.")
-	RENDERCORE_API bool Serialize(FArchive& Ar, bool bInlineShaderResources, bool bLoadedByCookedMaterial, bool bInlineShaderCode = false, const FName& SerializingAsset = NAME_None)
+	UE_DEPRECATED(5.5, "Please use overload accepting an FShaderSerializeContext. Note that inlining of shader code into runtime shadermaps is no longer supported. If you rely on this please reach out to the UE rendering team.")
+	RENDERCORE_API bool Serialize(FArchive& Ar, bool bInlineShaderResources, bool bLoadingCooked, bool bInlineShaderCode = false, const FName& SerializingAsset = NAME_None)
 	{
-		return Serialize(Ar, { bLoadedByCookedMaterial, SerializingAsset });
+		FShaderSerializeContext Ctx(Ar);
+		Ctx.bLoadingCooked = bLoadingCooked;
+		Ctx.SerializingAsset = SerializingAsset;
+		return Serialize(Ctx);
 	}
 
-	struct FSerializationContext
-	{
-		bool bLoadedByCookedMaterial = false;
-		FName SerializingAsset = NAME_None;
-	};
-	RENDERCORE_API bool Serialize(FArchive& Ar, const FSerializationContext& Ctx);
+	RENDERCORE_API bool Serialize(FShaderSerializeContext& Ctx);
 
 	EShaderPermutationFlags GetPermutationFlags() const
 	{
diff --git a/Engine/Source/Runtime/RenderCore/Public/ShaderCompilerCore.h b/Engine/Source/Runtime/RenderCore/Public/ShaderCompilerCore.h
--- a/Engine/Source/Runtime/RenderCore/Public/ShaderCompilerCore.h
+++ b/Engine/Source/Runtime/RenderCore/Public/ShaderCompilerCore.h
@@ -33,7 +33,7 @@ typedef TSharedPtr<TArray<ANSICHAR>, ESPMode::ThreadSafe> FShaderSharedAnsiStrin
 // this is for the protocol, not the data, bump if FShaderCompilerInput/FShaderPreprocessOutput serialization, SerializeWorkerInput or ProcessInputFromArchive changes.
 inline const int32 ShaderCompileWorkerInputVersion = 28;
 // this is for the protocol, not the data, bump if FShaderCompilerOutput or WriteToOutputArchive changes.
-inline const int32 ShaderCompileWorkerOutputVersion = 23;
+inline const int32 ShaderCompileWorkerOutputVersion = 24;
 // this is for the protocol, not the data.
 inline const int32 ShaderCompileWorkerSingleJobHeader = 'S';
 // this is for the protocol, not the data.
@@ -41,7 +41,7 @@ inline const int32 ShaderCompileWorkerPipelineJobHeader = 'P';
 
 // Modify this to invalidate _just_ the cache/DDC entries for individual shaders (will not cause shadermaps to rebuild if they are not otherwise out-of-date).
 // This should be bumped for changes to the FShaderCompilerOutput data structure (in addition to ShaderCompileWorkerOutputVersion)
-inline static const FGuid UE_SHADER_CACHE_VERSION = FGuid("DB7CA032-63BC-4D95-8F65-64C6B39EF882");
+inline static const FGuid UE_SHADER_CACHE_VERSION = FGuid("55EB8969-BA75-4CBD-B7A1-84194D140CBD");
 
 namespace UE::ShaderCompiler
 {
@@ -566,11 +566,13 @@ struct FShaderCompilerOutput
 	,	bSucceeded(false)
 	,	bSupportsQueryingUsedAttributes(false)
 	,	bSerializeModifiedSource(false)
+	,	bSerializingForCache(false)
 	{
 	}
 
 	FShaderParameterMap ParameterMap;
 	TArray<FShaderCompilerError> Errors;
+	UE_DEPRECATED(5.5, "PragmaDirectives are now only accessible via the FShaderPreprocessOutput struct")
 	TArray<FString> PragmaDirectives;
 	FShaderTarget Target;
 	FShaderCode ShaderCode;
@@ -583,12 +585,13 @@ struct FShaderCompilerOutput
 	bool bSucceeded;
 	bool bSupportsQueryingUsedAttributes;
 	bool bSerializeModifiedSource;
+	bool bSerializingForCache;
 	TArray<FString> UsedAttributes;
 
 	PRAGMA_DISABLE_DEPRECATION_WARNINGS
 	// Explicitly-defaulted copy/move ctors & assignment operators are needed temporarily due to 
-	// deprecation of OptionalFinalShaderSource field. These can be removed once the deprecation
-	// window for said field ends.
+	// deprecation of OptionalFinalShaderSource and PragmaDirectives fields. These can be removed 
+	// once the deprecation window for said fields ends.
 	FShaderCompilerOutput(FShaderCompilerOutput&&) = default;
 	FShaderCompilerOutput(const FShaderCompilerOutput&) = default;
 	FShaderCompilerOutput& operator=(FShaderCompilerOutput&&) = default;
@@ -642,14 +645,35 @@ struct FShaderCompilerOutput
 		Stat.Flags = Flags;
 	}
 
+	FShaderCodeResource ConvertCodeToResource() const
+	{
+		return ShaderCode.ConvertToResource(Target.GetFrequency(), OutputHash);
+	}
+
+	void SetCodeFromResource(FShaderCodeResource&& Resource)
+	{
+		return ShaderCode.SetFromResource(MoveTemp(Resource));
+	}
+
 	// Bump ShaderCompileWorkerOutputVersion if FShaderCompilerOutput changes
 	friend FArchive& operator<<(FArchive& Ar, FShaderCompilerOutput& Output)
 	{
 		// Note: this serialize is used to pass between UE and the shader compile worker, recompile both when modifying
-		Ar << Output.ParameterMap << Output.Errors << Output.Target << Output.ShaderCode << Output.OutputHash << Output.NumInstructions << Output.NumTextureSamplers << Output.bSucceeded;
-		Ar << Output.bSupportsQueryingUsedAttributes << Output.UsedAttributes;
-		Ar << Output.CompileTime;
-		Ar << Output.PreprocessTime;
+		Ar << Output.ParameterMap;
+		Ar << Output.Errors;
+		Ar << Output.Target;
+		Ar << Output.bSerializingForCache;
+		if (!Output.bSerializingForCache)
+		{
+			// skip serializing these fields when saving to cache/DDC; only needed when reading back results from workers
+			Ar << Output.ShaderCode;
+			Ar << Output.ValidateInputHash;
+			Ar << Output.CompileTime;
+		}
+		Ar << Output.OutputHash;
+		Ar << Output.NumInstructions;
+		Ar << Output.NumTextureSamplers;
+		Ar << Output.bSucceeded;
 		Ar << Output.bSerializeModifiedSource;
 		if (Output.bSerializeModifiedSource)
 		{
@@ -658,7 +682,11 @@ struct FShaderCompilerOutput
 		}
 		Ar << Output.PlatformDebugData;
 		Ar << Output.ShaderStatistics;
-		Ar << Output.ValidateInputHash;
+
+		// note: intentionally never serializing the following fields:
+		// - PreprocessTime - it is always set in the cooker since we no longer run preprocessing in SCW
+		// - bSupportsQueryingUsedAttributes - only used when compiling pipelines by subsequent stage compile steps, these are always executed in order in a single SCW job invocation
+		// - UsedAttributes - as above
 
 		return Ar;
 	}
diff --git a/Engine/Source/Runtime/RenderCore/Public/ShaderCompilerJobTypes.h b/Engine/Source/Runtime/RenderCore/Public/ShaderCompilerJobTypes.h
--- a/Engine/Source/Runtime/RenderCore/Public/ShaderCompilerJobTypes.h
+++ b/Engine/Source/Runtime/RenderCore/Public/ShaderCompilerJobTypes.h
@@ -13,6 +13,7 @@
 class FShaderCommonCompileJob;
 class FShaderCompileJob;
 class FShaderPipelineCompileJob;
+struct FShaderCacheSerializeContext;
 
 namespace UE::DerivedData { class FRequestOwner; }
 
@@ -149,9 +150,12 @@ public:
 	/** Returns hash of all inputs for this job (needed for caching). */
 	virtual FShaderCompilerInputHash GetInputHash() { return FShaderCompilerInputHash(); }
 
-	/** Serializes (and deserializes) the output for caching purposes. */
+	UE_DEPRECATED(5.5, "Use overload accepting an FShaderCacheSerializeContext struct")
 	virtual void SerializeOutput(FArchive& Ar) {}
 
+	/** Serializes (and deserializes) the output for caching purposes. */
+	virtual void SerializeOutput(FShaderCacheSerializeContext& Ctx) {}
+
 	FShaderCompileJob* GetSingleShaderJob();
 	const FShaderCompileJob* GetSingleShaderJob() const;
 	FShaderPipelineCompileJob* GetShaderPipelineJob();
@@ -253,7 +257,17 @@ public:
 	TMap<const FVertexFactoryType*, TArray<const FShaderPipelineType*>> SharingPipelines;
 
 	virtual RENDERCORE_API FShaderCompilerInputHash GetInputHash() override;
-	virtual RENDERCORE_API void SerializeOutput(FArchive& Ar) override;
+
+	UE_DEPRECATED(5.5, "Use overload accepting an FShaderCacheSerializeContext")
+	virtual RENDERCORE_API void SerializeOutput(FArchive& Ar) override
+	{
+	}
+
+	RENDERCORE_API void SerializeOutput(FShaderCacheSerializeContext& Ctx, int32 CodeIndex);
+	virtual void SerializeOutput(FShaderCacheSerializeContext& Ctx)
+	{
+		SerializeOutput(Ctx, 0);
+	}
 
 	virtual RENDERCORE_API void OnComplete() override;
 	
@@ -314,7 +328,14 @@ public:
 	TArray<TRefCountPtr<FShaderCompileJob>> StageJobs;
 
 	virtual RENDERCORE_API FShaderCompilerInputHash GetInputHash() override;
-	virtual RENDERCORE_API void SerializeOutput(FArchive& Ar) override;
+
+	UE_DEPRECATED(5.5, "Use overload accepting an FShaderCommonCompileJob::FSerializationContext")
+	virtual RENDERCORE_API void SerializeOutput(FArchive& Ar) override
+	{
+	}
+
+	virtual RENDERCORE_API void SerializeOutput(FShaderCacheSerializeContext& Ctx) override;
+
 	virtual RENDERCORE_API void OnComplete() override;
 	virtual RENDERCORE_API void AppendDebugName(FStringBuilderBase& OutName) const override;
 
diff --git a/Engine/Source/Runtime/RenderCore/Public/ShaderCore.h b/Engine/Source/Runtime/RenderCore/Public/ShaderCore.h
--- a/Engine/Source/Runtime/RenderCore/Public/ShaderCore.h
+++ b/Engine/Source/Runtime/RenderCore/Public/ShaderCore.h
@@ -1087,6 +1087,22 @@ public:
 	}
 };
 
+struct FShaderCodeResource
+{
+	TArray<uint8> Code;
+	int32 UncompressedSize; // full size of code array before compression
+	int32 ShaderCodeSize;	// uncompressed size excluding optional data
+	EShaderFrequency Frequency;
+
+	friend FArchive& operator<<(FArchive& Ar, FShaderCodeResource& Resource)
+	{
+		uint8 Freq = Resource.Frequency;
+		Ar << Resource.Code << Resource.UncompressedSize << Resource.ShaderCodeSize << Freq;
+		Resource.Frequency = (EShaderFrequency)Freq;
+		return Ar;
+	}
+};
+
 class FShaderCode
 {
 	// -1 if ShaderData was finalized
@@ -1225,6 +1241,89 @@ public:
 		AddOptionalData(Key, (uint8*)InString, Size);
 	}
 
+	// Converts code to FShaderCodeResource format. Note that this is destructive (internal code array will be moved into the resource).
+	FShaderCodeResource ConvertToResource(EShaderFrequency Frequency, FSHAHash OutputHash) const
+	{
+		check(OptionalDataSize == -1); // shader code must be finalized before converting to FShaderCodeResource
+		FShaderCodeResource Resource;
+		// need to persist this in the resource since it will not be readable from the code array after compression
+		Resource.ShaderCodeSize = GetShaderCodeSize();
+		Resource.Frequency = Frequency;
+
+		// Validate that compression settings used for this ShaderCode by the compilation process match what is expected
+		FName ShaderCompressionFormat = GetShaderCompressionFormat();
+		if (ShaderCompressionFormat != NAME_None)
+		{
+			Resource.UncompressedSize = GetUncompressedSize();
+
+			// we trust that SCWs also obeyed by the same CVar, so we expect a compressed shader code at this point
+			// However, if we see an uncompressed shader, it perhaps means that SCW tried to compress it, but the result was worse than uncompressed. 
+			// Because of that we special-case NAME_None here
+			if (ShaderCompressionFormat != GetCompressionFormat())
+			{
+				if (GetCompressionFormat() != NAME_None)
+				{
+					UE_LOG(LogShaders, Fatal, TEXT("Shader %s is expected to be compressed with %s, but it is compressed with %s instead."),
+						*OutputHash.ToString(),
+						*ShaderCompressionFormat.ToString(),
+						*GetCompressionFormat().ToString()
+					);
+					// unreachable
+					return Resource;
+				}
+
+				// assume uncompressed due to worse ratio than the compression
+				Resource.UncompressedSize = ShaderCodeWithOptionalData.Num();
+				UE_LOG(LogShaders, Verbose, TEXT("Shader %s is expected to be compressed with %s, but it arrived uncompressed (size=%d). Assuming compressing made it longer and storing uncompressed."),
+					*OutputHash.ToString(),
+					*ShaderCompressionFormat.ToString(),
+					ShaderCodeWithOptionalData.Num()
+				);
+			}
+			else if (ShaderCompressionFormat == NAME_Oodle)
+			{
+				// check if Oodle-specific settings match
+				FOodleDataCompression::ECompressor OodleCompressorSetting;
+				FOodleDataCompression::ECompressionLevel OodleLevelSetting;
+				GetShaderCompressionOodleSettings(OodleCompressorSetting, OodleLevelSetting);
+
+				if (GetOodleCompressor() != OodleCompressorSetting || GetOodleLevel() != OodleLevelSetting)
+				{
+					UE_LOG(LogShaders, Fatal, TEXT("Shader %s is expected to be compressed with Oodle compressor %d level %d, but it is compressed with compressor %d level %d instead."),
+						*OutputHash.ToString(),
+						static_cast<int32>(OodleCompressorSetting),
+						static_cast<int32>(OodleLevelSetting),
+						static_cast<int32>(GetOodleCompressor()),
+						static_cast<int32>(GetOodleLevel())
+					);
+					// unreachable
+					return Resource;
+				}
+			}
+		}
+		else
+		{
+			Resource.UncompressedSize = ShaderCodeWithOptionalData.Num();
+		}
+
+		Resource.Code = MoveTemp(ShaderCodeWithOptionalData);
+		return Resource;
+	}
+
+	void SetFromResource(FShaderCodeResource&& Resource)
+	{
+		// Set to the state of a finalized ShaderCode object
+		UncompressedSize = Resource.UncompressedSize;
+		ShaderCodeSize = Resource.ShaderCodeSize;
+		OptionalDataSize = -1;
+
+		// already validated that compression settings matched when serializing the resource, so we can just initialize them to the known-correct values
+		CompressionFormat = GetShaderCompressionFormat();
+		GetShaderCompressionOodleSettings(OodleCompressor, OodleLevel);
+
+		ShaderCodeWithOptionalData = MoveTemp(Resource.Code);
+	}
+
 	friend RENDERCORE_API FArchive& operator<<(FArchive& Ar, FShaderCode& Output);
 };
 
diff --git a/Engine/Source/Runtime/RenderCore/RenderCore.Build.cs b/Engine/Source/Runtime/RenderCore/RenderCore.Build.cs
--- a/Engine/Source/Runtime/RenderCore/RenderCore.Build.cs
+++ b/Engine/Source/Runtime/RenderCore/RenderCore.Build.cs
@@ -39,11 +39,18 @@ public class RenderCore : ModuleRules
 		}
 
 		PrivateDependencyModuleNames.AddRange(new string[] { "Core", "Projects", "ApplicationCore", "TraceLog", "CookOnTheFly" });
-
-        PrivateIncludePathModuleNames.AddRange(new string[] { "DerivedDataCache" });
 		
 		PublicIncludePathModuleNames.AddRange(new string[] { "RHI" });
 
+		if (Target.bBuildEditor == true)
+		{
+			PrivateDependencyModuleNames.Add("DerivedDataCache");
+		}
+		else
+		{
+			PrivateIncludePathModuleNames.Add("DerivedDataCache");
+		}
+
 		// Added in Dev-VT, still needed?
 		PrivateIncludePathModuleNames.AddRange(new string[] { "TargetPlatform" });
 
