[StaticMesh] added PreStaticMeshCompilation callback

diff --git a/Engine/Source/Runtime/Engine/Classes/Components/ActorStaticMeshComponentInterface.h b/Engine/Source/Runtime/Engine/Classes/Components/ActorStaticMeshComponentInterface.h
--- a/Engine/Source/Runtime/Engine/Classes/Components/ActorStaticMeshComponentInterface.h
+++ b/Engine/Source/Runtime/Engine/Classes/Components/ActorStaticMeshComponentInterface.h
@@ -9,6 +9,7 @@ class FActorStaticMeshComponentInterface : public IStaticMeshComponent
 public:
 #if WITH_EDITOR
 	virtual void OnMeshRebuild(bool bRenderDataChanged) override;
+	virtual void PreStaticMeshCompilation() override;
 	virtual void PostStaticMeshCompilation() override;
 #endif
 	virtual UStaticMesh* GetStaticMesh() const override;
diff --git a/Engine/Source/Runtime/Engine/Classes/Components/ComponentInterfaces.h b/Engine/Source/Runtime/Engine/Classes/Components/ComponentInterfaces.h
--- a/Engine/Source/Runtime/Engine/Classes/Components/ComponentInterfaces.h
+++ b/Engine/Source/Runtime/Engine/Classes/Components/ComponentInterfaces.h
@@ -114,6 +114,7 @@ class IStaticMeshComponent
 public:
 #if WITH_EDITOR
 	virtual void OnMeshRebuild(bool bRenderDataChanged) = 0;
+	virtual void PreStaticMeshCompilation() = 0;
 	virtual void PostStaticMeshCompilation() = 0;
 #endif
 	virtual UStaticMesh* GetStaticMesh() const = 0;
diff --git a/Engine/Source/Runtime/Engine/Classes/Components/StaticMeshComponent.h b/Engine/Source/Runtime/Engine/Classes/Components/StaticMeshComponent.h
--- a/Engine/Source/Runtime/Engine/Classes/Components/StaticMeshComponent.h
+++ b/Engine/Source/Runtime/Engine/Classes/Components/StaticMeshComponent.h
@@ -613,6 +613,7 @@ public:
 	ENGINE_API virtual bool IsShown(const FEngineShowFlags& ShowFlags) const override;
 #if WITH_EDITOR
 	ENGINE_API void OnMeshRebuild(bool bRenderDataChanged);
+	ENGINE_API virtual void PreStaticMeshCompilation();
 	ENGINE_API virtual void PostStaticMeshCompilation();
 	ENGINE_API virtual bool ComponentIsTouchingSelectionBox(const FBox& InSelBBox, const bool bConsiderOnlyBSP, const bool bMustEncompassEntireComponent) const override;
 	ENGINE_API virtual bool ComponentIsTouchingSelectionFrustum(const FConvexVolume& InSelBBox, const bool bConsiderOnlyBSP, const bool bMustEncompassEntireComponent) const override;
diff --git a/Engine/Source/Runtime/Engine/Private/Components/StaticMeshComponent.cpp b/Engine/Source/Runtime/Engine/Private/Components/StaticMeshComponent.cpp
--- a/Engine/Source/Runtime/Engine/Private/Components/StaticMeshComponent.cpp
+++ b/Engine/Source/Runtime/Engine/Private/Components/StaticMeshComponent.cpp
@@ -3339,6 +3339,11 @@ void UStaticMeshComponent::UpdateBounds()
 	Super::UpdateBounds();
 }
 
+void UStaticMeshComponent::PreStaticMeshCompilation()
+{
+	FNavigationSystem::UnregisterComponent(*this);
+}
+
 void UStaticMeshComponent::PostStaticMeshCompilation()
 {
 	// Flag indicates that the component is currently registering and all the following actions
@@ -3986,6 +3991,11 @@ void FActorStaticMeshComponentInterface::OnMeshRebuild(bool bRenderDataChanged)
 	UStaticMeshComponent::GetStaticMeshComponent(this)->OnMeshRebuild(bRenderDataChanged);
 }
 
+void FActorStaticMeshComponentInterface::PreStaticMeshCompilation()
+{
+	UStaticMeshComponent::GetStaticMeshComponent(this)->PreStaticMeshCompilation();
+}
+
 void FActorStaticMeshComponentInterface::PostStaticMeshCompilation()
 {
 	UStaticMeshComponent::GetStaticMeshComponent(this)->PostStaticMeshCompilation();
diff --git a/Engine/Source/Runtime/Engine/Private/StaticMeshBuild.cpp b/Engine/Source/Runtime/Engine/Private/StaticMeshBuild.cpp
--- a/Engine/Source/Runtime/Engine/Private/StaticMeshBuild.cpp
+++ b/Engine/Source/Runtime/Engine/Private/StaticMeshBuild.cpp
@@ -198,6 +198,8 @@ void UStaticMesh::BatchBuild(const TArray<UStaticMesh*>& InStaticMeshes, const F
 
 				for (IStaticMeshComponent* Component : ObjectCacheScope.GetContext().GetStaticMeshComponents(StaticMesh))
 				{
+					Component->PreStaticMeshCompilation();
+
 					IPrimitiveComponent* PrimitiveComponent = Component->GetPrimitiveComponentInterface();
 
 					// Detach all instances of those static meshes from the scene.
