Adding min distance cvar for stop oldest sound concurrency flavors

[RN] Cvar allowing a min-distance delta for distance-based concurrency resolution

diff --git a/Engine/Source/Runtime/Engine/Private/SoundConcurrency.cpp b/Engine/Source/Runtime/Engine/Private/SoundConcurrency.cpp
--- a/Engine/Source/Runtime/Engine/Private/SoundConcurrency.cpp
+++ b/Engine/Source/Runtime/Engine/Private/SoundConcurrency.cpp
@@ -16,6 +16,14 @@ FAutoConsoleVariableRef CVarConcurrencyMinVolumeScale(
 	ECVF_Default);
 
 
+static float ConcurrencyMinDistanceDeltaCVar = 0.0f;
+FAutoConsoleVariableRef CVarConcurrencyMinDistanceDelta(
+	TEXT("au.Concurrency.MinDistanceDelta"),
+	ConcurrencyMinDistanceDeltaCVar,
+	TEXT("Minimum distance delta a sound needs to be to determine if it's further away from another sound.\n"),
+	ECVF_Default);
+
+
 namespace SoundConcurrency
 {
 	void SetSoundDataTarget(const FActiveSound& ActiveSound, FConcurrencySoundData& SoundData, float InTargetVolume, float InLerpTime)
@@ -653,8 +661,12 @@ FActiveSound* FSoundConcurrencyManager::GetEvictableSoundStopFarthest(const FAct
 {
 	const EMaxConcurrentResolutionRule::Type Rule = ConcurrencyGroup.GetSettings().ResolutionRule;
 
+	// Get the mininum distance used to calculate distance of a sound
+	float MinDistanceSquared = FMath::Max(ConcurrencyMinDistanceDeltaCVar, 0.0f);
+	MinDistanceSquared *= MinDistanceSquared;
+
 	check(AudioDevice);
-	float DistanceToStopSoundSq;
+	float DistanceToStopSoundSq = 0.0f;
 	const bool bAllowAttenuationOverrides = true;
 	int32 ClosestListenerIndex = AudioDevice->FindClosestListenerIndex(NewActiveSound.Transform.GetTranslation(), DistanceToStopSoundSq, bAllowAttenuationOverrides);
 
@@ -667,25 +679,30 @@ FActiveSound* FSoundConcurrencyManager::GetEvictableSoundStopFarthest(const FAct
 		float DistanceToActiveSoundSq;
 		ClosestListenerIndex = AudioDevice->FindClosestListenerIndex(ActiveSound->Transform.GetTranslation(), DistanceToActiveSoundSq, bAllowAttenuationOverrides);
 
+
 		// Stop Farthest
-		if (DistanceToActiveSoundSq > DistanceToStopSoundSq)
+		if (DistanceToActiveSoundSq > (DistanceToStopSoundSq + MinDistanceSquared))
 		{
-			DistanceToStopSoundSq = DistanceToActiveSoundSq;
-			EvictableSound = ActiveSound;
-			continue;
+				DistanceToStopSoundSq = DistanceToActiveSoundSq;
+				EvictableSound = ActiveSound;
+				continue;
 		}
 
-		// Stop Farthest than Oldest
-		if (Rule == EMaxConcurrentResolutionRule::StopFarthestThenOldest && DistanceToActiveSoundSq == DistanceToStopSoundSq)
+		// Stop Farthest then Oldest
+		if (Rule == EMaxConcurrentResolutionRule::StopFarthestThenOldest)
 		{
-			const bool bIsOlderThanChosen = EvictableSound == nullptr || ActiveSound->PlaybackTime > EvictableSound->PlaybackTime;
-			if (bIsOlderThanChosen)
+			const bool bIsSameDistance = FMath::Abs(DistanceToActiveSoundSq - DistanceToStopSoundSq) <= MinDistanceSquared;
+			if (bIsSameDistance)
 			{
-				// Don't evict if attempting to re-trigger an older sound than that which is currently playing
-				if (!bIsRetriggering || ActiveSound->PlaybackTime > NewActiveSound.PlaybackTime)
+				const bool bIsOlderThanChosen = EvictableSound == nullptr || ActiveSound->PlaybackTime > EvictableSound->PlaybackTime;
+				if (bIsOlderThanChosen)
 				{
-					DistanceToStopSoundSq = DistanceToActiveSoundSq;
-					EvictableSound = ActiveSound;
+					// Don't evict if attempting to re-trigger an older sound than that which is currently playing
+					if (!bIsRetriggering || ActiveSound->PlaybackTime > NewActiveSound.PlaybackTime)
+					{
+						DistanceToStopSoundSq = DistanceToActiveSoundSq;
+						EvictableSound = ActiveSound;
+					}
 				}
 			}
 		}
