Expose mechanism for codecs to opt-out of raw data sanitizing during animation compression

diff --git a/Engine/Plugins/Animation/ACLPlugin/Source/ACLPlugin/Classes/AnimBoneCompressionCodec_ACLBase.h b/Engine/Plugins/Animation/ACLPlugin/Source/ACLPlugin/Classes/AnimBoneCompressionCodec_ACLBase.h
--- a/Engine/Plugins/Animation/ACLPlugin/Source/ACLPlugin/Classes/AnimBoneCompressionCodec_ACLBase.h
+++ b/Engine/Plugins/Animation/ACLPlugin/Source/ACLPlugin/Classes/AnimBoneCompressionCodec_ACLBase.h
@@ -72,6 +72,7 @@ class UAnimBoneCompressionCodec_ACLBase : public UAnimBoneCompressionCodec
 	float ErrorThreshold;
 
 	// UAnimBoneCompressionCodec implementation
+	virtual bool IsHighFidelity(const FCompressibleAnimData& CompressibleAnimData) const override;
 	virtual bool Compress(const FCompressibleAnimData& CompressibleAnimData, FCompressibleAnimDataResult& OutResult) override;
 	virtual void PopulateDDCKey(const UE::Anim::Compression::FAnimDDCKeyArgs& KeyArgs, FArchive& Ar) override;
 	virtual int64 EstimateCompressionMemoryUsage(const UAnimSequence& AnimSequence) const override;
diff --git a/Engine/Plugins/Animation/ACLPlugin/Source/ACLPlugin/Private/AnimBoneCompressionCodec_ACLBase.cpp b/Engine/Plugins/Animation/ACLPlugin/Source/ACLPlugin/Private/AnimBoneCompressionCodec_ACLBase.cpp
--- a/Engine/Plugins/Animation/ACLPlugin/Source/ACLPlugin/Private/AnimBoneCompressionCodec_ACLBase.cpp
+++ b/Engine/Plugins/Animation/ACLPlugin/Source/ACLPlugin/Private/AnimBoneCompressionCodec_ACLBase.cpp
@@ -76,6 +76,12 @@ UAnimBoneCompressionCodec_ACLBase::UAnimBoneCompressionCodec_ACLBase(const FObje
 }
 
 #if WITH_EDITORONLY_DATA
+bool UAnimBoneCompressionCodec_ACLBase::IsHighFidelity(const FCompressibleAnimData& CompressibleAnimData) const
+{
+	// ACL handles raw data sanitizing internally
+	return true;
+}
+
 static void AppendMaxVertexDistances(USkeletalMesh* OptimizationTarget, TMap<FName, float>& BoneMaxVertexDistanceMap)
 {
 #if (ENGINE_MAJOR_VERSION == 4 && ENGINE_MINOR_VERSION >= 27) || ENGINE_MAJOR_VERSION >= 5
@@ -439,7 +445,7 @@ void UAnimBoneCompressionCodec_ACLBase::PopulateDDCKey(const UE::Anim::Compressi
 {
 	Super::PopulateDDCKey(KeyArgs, Ar);
 
-	uint32 ForceRebuildVersion = 19;
+	uint32 ForceRebuildVersion = 20;
 
 	Ar << ForceRebuildVersion << DefaultVirtualVertexDistance << SafeVirtualVertexDistance << ErrorThreshold;
 	Ar << CompressionLevel;
diff --git a/Engine/Source/Runtime/Engine/Classes/Animation/AnimBoneCompressionCodec.h b/Engine/Source/Runtime/Engine/Classes/Animation/AnimBoneCompressionCodec.h
--- a/Engine/Source/Runtime/Engine/Classes/Animation/AnimBoneCompressionCodec.h
+++ b/Engine/Source/Runtime/Engine/Classes/Animation/AnimBoneCompressionCodec.h
@@ -37,6 +37,9 @@ class UAnimBoneCompressionCodec : public UObject
 	/** Returns whether or not we can use this codec to compress. */
 	virtual bool IsCodecValid() const { return true; }
 
+	/** Returns whether or not this codec is high fidelity. High fidelity codecs requires the original raw data without any sanitizing by the engine. */
+	ENGINE_API virtual bool IsHighFidelity(const FCompressibleAnimData& CompressibleAnimData) const { return false; }
+
 	/** Compresses the curve data from an animation sequence. */
 	ENGINE_API virtual bool Compress(const FCompressibleAnimData& CompressibleAnimData, FCompressibleAnimDataResult& OutResult) PURE_VIRTUAL(UAnimCurveCompressionCodec::Compress, return false;);
 
diff --git a/Engine/Source/Runtime/Engine/Classes/Animation/AnimBoneCompressionSettings.h b/Engine/Source/Runtime/Engine/Classes/Animation/AnimBoneCompressionSettings.h
--- a/Engine/Source/Runtime/Engine/Classes/Animation/AnimBoneCompressionSettings.h
+++ b/Engine/Source/Runtime/Engine/Classes/Animation/AnimBoneCompressionSettings.h
@@ -49,6 +49,9 @@ class UAnimBoneCompressionSettings : public UObject
 	/** Returns whether or not we can use these settings to compress. */
 	ENGINE_API bool AreSettingsValid() const;
 
+	/** Returns whether or not a codec within is high fidelity. @see UAnimBoneCompressionCodec::IsHighFidelity */
+	ENGINE_API bool IsHighFidelity(const FCompressibleAnimData& CompressibleAnimData) const;
+
 	/*
 	 * Compresses the animation bones inside the supplied sequence.
 	 * The resultant compressed data is applied to the OutCompressedData structure.
diff --git a/Engine/Source/Runtime/Engine/Classes/Animation/AnimCompress_LeastDestructive.h b/Engine/Source/Runtime/Engine/Classes/Animation/AnimCompress_LeastDestructive.h
--- a/Engine/Source/Runtime/Engine/Classes/Animation/AnimCompress_LeastDestructive.h
+++ b/Engine/Source/Runtime/Engine/Classes/Animation/AnimCompress_LeastDestructive.h
@@ -17,4 +17,9 @@ UCLASS()
 class UAnimCompress_LeastDestructive : public UAnimCompress_BitwiseCompressOnly
 {
 	GENERATED_UCLASS_BODY()
+
+#if WITH_EDITORONLY_DATA
+	// Least destructive is high fidelity
+	virtual bool IsHighFidelity(const FCompressibleAnimData& CompressibleAnimData) const override { return true; }
+#endif
 };
diff --git a/Engine/Source/Runtime/Engine/Classes/Animation/AnimSequenceHelpers.h b/Engine/Source/Runtime/Engine/Classes/Animation/AnimSequenceHelpers.h
--- a/Engine/Source/Runtime/Engine/Classes/Animation/AnimSequenceHelpers.h
+++ b/Engine/Source/Runtime/Engine/Classes/Animation/AnimSequenceHelpers.h
@@ -229,7 +229,7 @@ namespace Anim {
 		* @return Whether or not the operation was succesful
 		*/
 		UE_DEPRECATED(5.4, "CompressAnimationDataTracks has been deprecated use other signature with skeleton and bone mappings")
-		ENGINE_API bool CompressAnimationDataTracks(TArray<FRawAnimSequenceTrack>& RawAnimationData, int32 NumberOfKeys, FName ErrorName, float MaxPosDiff = 0.0001f, float MaxAngleDiff = 0.0003f);
+		ENGINE_API bool CompressAnimationDataTracks(TArray<FRawAnimSequenceTrack>& RawAnimationData, int32 NumberOfKeys, FName ErrorName, float MaxPosDiff = 0.0001f, float MaxAngleDiff = 0.0003f, float MaxScaleDiff = 0.0001f);
 
 		/**
 		* Compress the provided tracks using two methods.
@@ -244,7 +244,7 @@ namespace Anim {
 		*
 		* @return Whether or not the operation was succesful
 		*/
-		ENGINE_API bool CompressAnimationDataTracks(const USkeleton* Skeleton, const TArray<FTrackToSkeletonMap>& TrackToSkeleton,TArray<FRawAnimSequenceTrack>& RawAnimationData, int32 NumberOfKeys, FName ErrorName, float MaxPosDiff = 0.0001f, float MaxAngleDiff = 0.0003f);
+		ENGINE_API bool CompressAnimationDataTracks(const USkeleton* Skeleton, const TArray<FTrackToSkeletonMap>& TrackToSkeleton,TArray<FRawAnimSequenceTrack>& RawAnimationData, int32 NumberOfKeys, FName ErrorName, float MaxPosDiff = 0.0001f, float MaxAngleDiff = 0.0003f, float MaxScaleDiff = 0.0001f);
 
 		/**
 		* Compress the provided track by checking for uniform key data according to the provided error metrics
@@ -257,8 +257,23 @@ namespace Anim {
 		*
 		* @return Whether or not the operation was succesful
 		*/
+		UE_DEPRECATED(5.5, "CompressRawAnimSequenceTrack requires the MaxScaleDiff argument")
 		ENGINE_API bool CompressRawAnimSequenceTrack(FRawAnimSequenceTrack& RawTrack, int32 NumberOfKeys, FName ErrorName, float MaxPosDiff, float MaxAngleDiff);
 
+		/**
+		* Compress the provided track by checking for uniform key data according to the provided error metrics
+		*
+		* @param	RawTrack			Individual track to apply the reduction to
+		* @param	NumberOfKeys		Expected number of keys for each component within the tracks
+		* @param	ErrorName			Identifier for outer callsite when failing to reduce the data
+		* @param	MaxPosDiff			Maximum positional delta when determining whether or not two positional keys are identical
+		* @param	MaxAngleDiff		Maximum rotational delta when determining whether or not two rotational keys are identical
+		* @param	MaxScaleDiff		Maximum scale delta when determining whether or not two scale keys are identical
+		*
+		* @return Whether or not the operation was succesful
+		*/
+		ENGINE_API bool CompressRawAnimSequenceTrack(FRawAnimSequenceTrack& RawTrack, int32 NumberOfKeys, FName ErrorName, float MaxPosDiff, float MaxAngleDiff, float MaxScaleDiff);
+
 		/**
 		* Sanitize the provided track by snapping small scale values to 0, and normalizing any rotational keys
 		*
diff --git a/Engine/Source/Runtime/Engine/Private/Animation/AnimBoneCompressionSettings.cpp b/Engine/Source/Runtime/Engine/Private/Animation/AnimBoneCompressionSettings.cpp
--- a/Engine/Source/Runtime/Engine/Private/Animation/AnimBoneCompressionSettings.cpp
+++ b/Engine/Source/Runtime/Engine/Private/Animation/AnimBoneCompressionSettings.cpp
@@ -77,6 +77,20 @@ bool UAnimBoneCompressionSettings::AreSettingsValid() const
 	return NumValidCodecs != 0;
 }
 
+bool UAnimBoneCompressionSettings::IsHighFidelity(const FCompressibleAnimData& CompressibleAnimData) const
+{
+	for (const UAnimBoneCompressionCodec* Codec : Codecs)
+	{
+		if (Codec != nullptr && Codec->IsHighFidelity(CompressibleAnimData))
+		{
+			// At least one codec is high fidelity
+			return true;
+		}
+	}
+
+	return false;
+}
+
 struct FAnimBoneCompressionContext
 {
 	/** The animation sequence we are compressing. */
diff --git a/Engine/Source/Runtime/Engine/Private/Animation/AnimCompressionTypes.cpp b/Engine/Source/Runtime/Engine/Private/Animation/AnimCompressionTypes.cpp
--- a/Engine/Source/Runtime/Engine/Private/Animation/AnimCompressionTypes.cpp
+++ b/Engine/Source/Runtime/Engine/Private/Animation/AnimCompressionTypes.cpp
@@ -859,12 +859,21 @@ void FCompressibleAnimData::FetchData(const ITargetPlatform* InPlatform)
 	{
 		return InLHS.GetName().LexicalLess(InRHS.GetName());
 	});
+
+	// High fidelity codecs wish to see the original raw data where possible
+	const bool bIsHighFidelity = BoneCompressionSettings->IsHighFidelity(*this);
 	
 	// Apply any key reduction if possible
 	if (RawAnimationData.Num())
 	{ 
-		UE::Anim::Compression::CompressAnimationDataTracks(Skeleton, TrackToSkeletonMapTable, RawAnimationData, NumberOfKeys, AnimSequence->GetFName(), -1.f, -1.f);
-		UE::Anim::Compression::CompressAnimationDataTracks(Skeleton, TrackToSkeletonMapTable, RawAnimationData, NumberOfKeys, AnimSequence->GetFName());
+		// Fixup broken data
+		UE::Anim::Compression::CompressAnimationDataTracks(Skeleton, TrackToSkeletonMapTable, RawAnimationData, NumberOfKeys, AnimSequence->GetFName(), -1.f, -1.f, -1.f);
+
+		if (!bIsHighFidelity)
+		{
+			// Low fidelity codecs need some help, sanitize the raw data
+			UE::Anim::Compression::CompressAnimationDataTracks(Skeleton, TrackToSkeletonMapTable, RawAnimationData, NumberOfKeys, AnimSequence->GetFName());
+		}
 	}
 
 	auto IsKeyArrayValidForRemoval = [](const auto& Keys, const auto& IdentityValue) -> bool
@@ -912,7 +921,9 @@ void FCompressibleAnimData::FetchData(const ITargetPlatform* InPlatform)
 			const int32 BoneIndex = RefSkeleton.FindBoneIndex(OriginalTrackNames[TrackIndex]);
 
 			const bool bValidBoneIndex = BoneIndex != INDEX_NONE;
-			const bool bValidAdditiveTrack = !IsRawTrackZeroAdditive(Track);
+
+			// Low fidelity codecs need some help, sanitize the raw data
+			const bool bValidAdditiveTrack = bIsHighFidelity || !IsRawTrackZeroAdditive(Track);
 
 			// Only include track if it contains valid (additive) data and its name corresponds to a bone on the skeleton
 			if ((!bIsAdditiveAnimation || bValidAdditiveTrack) && bValidBoneIndex)
diff --git a/Engine/Source/Runtime/Engine/Private/Animation/AnimSequenceHelpers.cpp b/Engine/Source/Runtime/Engine/Private/Animation/AnimSequenceHelpers.cpp
--- a/Engine/Source/Runtime/Engine/Private/Animation/AnimSequenceHelpers.cpp
+++ b/Engine/Source/Runtime/Engine/Private/Animation/AnimSequenceHelpers.cpp
@@ -464,7 +464,7 @@ bool CopyNotifies(const UAnimSequenceBase* SourceAnimSeq, UAnimSequenceBase* Des
 
 #endif // WITH_EDITOR
 
-bool Compression::CompressAnimationDataTracks(TArray<FRawAnimSequenceTrack>& RawAnimationData, int32 NumberOfKeys, FName ErrorName, float MaxPosDiff /*= 0.0001f*/, float MaxAngleDiff /*= 0.0003f*/)
+bool Compression::CompressAnimationDataTracks(TArray<FRawAnimSequenceTrack>& RawAnimationData, int32 NumberOfKeys, FName ErrorName, float MaxPosDiff /*= 0.0001f*/, float MaxAngleDiff /*= 0.0003f*/, float MaxScaleDiff)
 {
 	bool bRemovedKeys = false;
 
@@ -474,14 +474,14 @@ bool Compression::CompressAnimationDataTracks(TArray<FRawAnimSequenceTrack>& Raw
 		// This removes trivial keys, and this has to happen before the removing tracks
 		for (int32 TrackIndex = 0; TrackIndex < RawAnimationData.Num(); TrackIndex++)
 		{
-			bRemovedKeys |= CompressRawAnimSequenceTrack(RawAnimationData[TrackIndex], NumberOfKeys, ErrorName, MaxPosDiff, MaxAngleDiff);
+			bRemovedKeys |= CompressRawAnimSequenceTrack(RawAnimationData[TrackIndex], NumberOfKeys, ErrorName, MaxPosDiff, MaxAngleDiff, MaxScaleDiff);
 		}
 	}
 #endif
 	return bRemovedKeys;
 }
 
-bool Compression::CompressAnimationDataTracks(const USkeleton* Skeleton, const TArray<FTrackToSkeletonMap>& TrackToSkeleton, TArray<FRawAnimSequenceTrack>& RawAnimationData, int32 NumberOfKeys, FName ErrorName, float MaxPosDiff /*= 0.0001f*/, float MaxAngleDiff /*= 0.0003f*/)
+bool Compression::CompressAnimationDataTracks(const USkeleton* Skeleton, const TArray<FTrackToSkeletonMap>& TrackToSkeleton, TArray<FRawAnimSequenceTrack>& RawAnimationData, int32 NumberOfKeys, FName ErrorName, float MaxPosDiff, float MaxAngleDiff, float MaxScaleDiff)
 {
 	bool bRemovedKeys = false;
 
@@ -491,15 +491,19 @@ bool Compression::CompressAnimationDataTracks(const USkeleton* Skeleton, const T
 		// This removes trivial keys, and this has to happen before the removing tracks
 		for (int32 TrackIndex = 0; TrackIndex < RawAnimationData.Num(); TrackIndex++)
 		{
-			bRemovedKeys |= CompressRawAnimSequenceTrack(RawAnimationData[TrackIndex], NumberOfKeys, ErrorName, MaxPosDiff, MaxAngleDiff);
+			bRemovedKeys |= CompressRawAnimSequenceTrack(RawAnimationData[TrackIndex], NumberOfKeys, ErrorName, MaxPosDiff, MaxAngleDiff, MaxScaleDiff);
 		}
 	}
 #endif
 	return bRemovedKeys;
 }
 
-
 bool Compression::CompressRawAnimSequenceTrack(FRawAnimSequenceTrack& RawTrack, int32 NumberOfKeys, FName ErrorName, float MaxPosDiff, float MaxAngleDiff)
+{
+	return CompressRawAnimSequenceTrack(RawTrack, NumberOfKeys, ErrorName, MaxPosDiff, MaxAngleDiff, 0.0001f);
+}
+
+bool Compression::CompressRawAnimSequenceTrack(FRawAnimSequenceTrack& RawTrack, int32 NumberOfKeys, FName ErrorName, float MaxPosDiff, float MaxAngleDiff, float MaxScaleDiff)
 {
 	bool bRemovedKeys = false;
 
@@ -583,8 +587,6 @@ bool Compression::CompressRawAnimSequenceTrack(FRawAnimSequenceTrack& RawTrack,
 		}
 	}
 
-	float MaxScaleDiff = 0.0001f;
-
 	// Check variation of Scaleition keys
 	if ((RawTrack.ScaleKeys.Num() > 1) && (MaxScaleDiff >= 0.0f))
 	{
