[UBA] Add breadcrumbs to UBA trace files.

This change adds a new field to the ProcessStartInfo struct in UBA trace files for a description of the task, called "breadcrumbs" which allow to diagnose and backtrace the origin of a UBA task.
This is also hooked up to the distributed shader compilation (added as new field in FTaskCommandData) and a descriptive string is generated for each shader compile job batch.
UbaVisualizer is updated accordingly to integrate this new data when hovering over a task in the timeline.
This feature can be enabled via the new CVar "r.ShaderCompiler.DistributedJobDescriptionLevel" with values 0 (Disabled), 1 (Basic information), 2 (Adds shader format to each job). It's disabled by default.

#rnx

diff --git a/Engine/Plugins/UbaController/Source/Private/UbaJobProcessor.cpp b/Engine/Plugins/UbaController/Source/Private/UbaJobProcessor.cpp
--- a/Engine/Plugins/UbaController/Source/Private/UbaJobProcessor.cpp
+++ b/Engine/Plugins/UbaController/Source/Private/UbaJobProcessor.cpp
@@ -261,6 +261,7 @@ void FUbaJobProcessor::RunTaskWithUba(FTask* Task)
 	FStringToUbaStringConversion UbaParametersStr(*Parameters);
 	FStringToUbaStringConversion UbaInputFileNameStr(*InputFileName);
 	FStringToUbaStringConversion UbaWorkingDirStr(*AppDir);
+	FStringToUbaStringConversion UbaTaskDescription(*Data.Description);
 
 	uba::ProcessStartInfo ProcessInfo;
 	ProcessInfo.application = UbaCommandStr.Get();
@@ -268,6 +269,7 @@ void FUbaJobProcessor::RunTaskWithUba(FTask* Task)
 	ProcessInfo.description = UbaInputFileNameStr.Get();
 	ProcessInfo.workingDir = UbaWorkingDirStr.Get();
 	ProcessInfo.writeOutputFilesOnFail = true;
+	ProcessInfo.breadcrumbs = UbaTaskDescription.Get();
 
 	if (UbaJobProcessorOptions::bProcessLogEnabled)
 	{
diff --git a/Engine/Source/Developer/DistributedBuildInterface/Public/DistributedBuildControllerInterface.h b/Engine/Source/Developer/DistributedBuildInterface/Public/DistributedBuildControllerInterface.h
--- a/Engine/Source/Developer/DistributedBuildInterface/Public/DistributedBuildControllerInterface.h
+++ b/Engine/Source/Developer/DistributedBuildInterface/Public/DistributedBuildControllerInterface.h
@@ -25,6 +25,7 @@ struct FTaskCommandData
 	FString InputFileName;
 	FString OutputFileName;
 	FString ExtraCommandArgs;
+	FString Description; // Optional string describing the task. Shows up in UBA trace files for each job.
 	uint32 DispatcherPID = 0;
 	TArray<FString> Dependencies;
 };
diff --git a/Engine/Source/Programs/UnrealBuildAccelerator/Common/Private/UbaSession.cpp b/Engine/Source/Programs/UnrealBuildAccelerator/Common/Private/UbaSession.cpp
--- a/Engine/Source/Programs/UnrealBuildAccelerator/Common/Private/UbaSession.cpp
+++ b/Engine/Source/Programs/UnrealBuildAccelerator/Common/Private/UbaSession.cpp
@@ -1509,7 +1509,7 @@ namespace uba
 			u32 exitCode = process.GetExitCode();
 			Vector<ProcessLogLine> emptyLines;
 			auto& logLines = (exitCode != 0 || m_detailedTrace) ? process.m_logLines : emptyLines;
-			m_trace.ProcessExited(id, exitCode, writer.GetData(), writer.GetPosition(), logLines);
+			m_trace.ProcessExited(id, exitCode, writer.GetData(), writer.GetPosition(), logLines, process.GetStartInfo().breadcrumbs);
 			SCOPED_WRITE_LOCK(m_processStatsLock, lock);
 			m_processStats.Add(process.m_processStats);
 			m_stats.Add(process.m_sessionStats);
diff --git a/Engine/Source/Programs/UnrealBuildAccelerator/Common/Private/UbaSessionServer.cpp b/Engine/Source/Programs/UnrealBuildAccelerator/Common/Private/UbaSessionServer.cpp
--- a/Engine/Source/Programs/UnrealBuildAccelerator/Common/Private/UbaSessionServer.cpp
+++ b/Engine/Source/Programs/UnrealBuildAccelerator/Common/Private/UbaSessionServer.cpp
@@ -476,7 +476,7 @@ namespace uba
 		StackBinaryWriter<1024> statsWriter;
 		ProcessStats processStats;
 		processStats.Write(statsWriter);
-		m_trace.ProcessExited(id, exitCode, statsWriter.GetData(), statsWriter.GetPosition(), Vector<ProcessLogLine>());
+		m_trace.ProcessExited(id, exitCode, statsWriter.GetData(), statsWriter.GetPosition(), Vector<ProcessLogLine>(), TC(""));
 	}
 
 	void SessionServer::UpdateProgress(u32 processesTotal, u32 processesDone, u32 errorCount)
@@ -1239,7 +1239,7 @@ namespace uba
 		u32 id = process.m_processId;
 		Vector<ProcessLogLine> emptyLines;
 		auto& logLines = (exitCode != 0 || m_detailedTrace) ? process.m_logLines : emptyLines;
-		m_trace.ProcessExited(id, exitCode, reader.GetPositionData(), reader.GetLeft(), logLines);
+		m_trace.ProcessExited(id, exitCode, reader.GetPositionData(), reader.GetLeft(), logLines, process.GetStartInfo().breadcrumbs);
 
 		ProcessStats processStats;
 		processStats.Read(reader, ~0u);
@@ -1781,7 +1781,7 @@ namespace uba
 			SessionStats().Write(writer);
 			StorageStats().Write(writer);
 			KernelStats().Write(writer);
-			m_trace.ProcessExited(process->m_processId, process->m_exitCode, writer.GetData(), writer.GetPosition(), Vector<ProcessLogLine>());
+			m_trace.ProcessExited(process->m_processId, process->m_exitCode, writer.GetData(), writer.GetPosition(), Vector<ProcessLogLine>(), process->GetStartInfo().breadcrumbs);
 
 			m_logger.Warning(TC("Cancelling remote active processes has not been tested. Notify devs"));
 
diff --git a/Engine/Source/Programs/UnrealBuildAccelerator/Common/Private/UbaTrace.cpp b/Engine/Source/Programs/UnrealBuildAccelerator/Common/Private/UbaTrace.cpp
--- a/Engine/Source/Programs/UnrealBuildAccelerator/Common/Private/UbaTrace.cpp
+++ b/Engine/Source/Programs/UnrealBuildAccelerator/Common/Private/UbaTrace.cpp
@@ -264,12 +264,13 @@ namespace uba
 		writer.WriteBytes(data, dataSize);
 	}
 
-	void Trace::ProcessExited(u32 processId, u32 exitCode, const u8* data, u64 dataSize, const Vector<ProcessLogLine>& logLines)
+	void Trace::ProcessExited(u32 processId, u32 exitCode, const u8* data, u64 dataSize, const Vector<ProcessLogLine>& logLines, const tchar* breadcrumbs)
 	{
 		BEGIN_TRACE_ENTRY(ProcessExited);
 		writer.WriteU32(processId);
 		writer.WriteU32(exitCode);
 		writer.WriteBytes(data, dataSize);
+		writer.WriteString(breadcrumbs);
 		u32 lineCounter = 0;
 		for (auto& line : logLines)
 		{
diff --git a/Engine/Source/Programs/UnrealBuildAccelerator/Common/Public/UbaProcessStartInfo.h b/Engine/Source/Programs/UnrealBuildAccelerator/Common/Public/UbaProcessStartInfo.h
--- a/Engine/Source/Programs/UnrealBuildAccelerator/Common/Public/UbaProcessStartInfo.h
+++ b/Engine/Source/Programs/UnrealBuildAccelerator/Common/Public/UbaProcessStartInfo.h
@@ -27,6 +27,7 @@ namespace uba
 		bool trackInputs = false;				// Track all files read. Can read result in ProcessHandle.GetTrackedInputs()
 		bool useCustomAllocator = true;			// Disable detouring of allocator inside processes. If Session.disableCustomAllocator is false this will be overridden
 		bool writeOutputFilesOnFail = false;	// If set to true, output files will be written to disk regardless if process succeeds or not
+		const tchar* breadcrumbs = TC("");		// If not empty, write additional information to the UBA trace file
 		bool startSuspended = false;			// Start process suspended.. a bit internal atm and only supported on windows
 
 		using LogLineCallback = void(void* userData, const tchar* line, u32 length, LogEntryType type);
diff --git a/Engine/Source/Programs/UnrealBuildAccelerator/Common/Public/UbaProcessStartInfoHolder.h b/Engine/Source/Programs/UnrealBuildAccelerator/Common/Public/UbaProcessStartInfoHolder.h
--- a/Engine/Source/Programs/UnrealBuildAccelerator/Common/Public/UbaProcessStartInfoHolder.h
+++ b/Engine/Source/Programs/UnrealBuildAccelerator/Common/Public/UbaProcessStartInfoHolder.h
@@ -48,6 +48,9 @@ namespace uba
 
 			logFileStr = si.logFile;
 			logFile = logFileStr.c_str();
+
+			breadcrumbsStr = si.breadcrumbs;
+			breadcrumbs = breadcrumbsStr.c_str();
 		}
 
 		void Write(BinaryWriter& writer)
@@ -57,6 +60,7 @@ namespace uba
 			writer.WriteString(argumentsStr);
 			writer.WriteString(workingDirStr);
 			writer.WriteString(logFileStr);
+			// ignore breadcrumbs here
 			writer.WriteU32(*(u32*)&weight);
 			writer.WriteBool(trackInputs);
 			writer.WriteBool(writeOutputFilesOnFail);
@@ -70,6 +74,7 @@ namespace uba
 			argumentsStr = reader.ReadString();
 			workingDirStr = reader.ReadString();
 			logFileStr = reader.ReadString();
+			// ignore breadcrumbs here
 
 			Replace(applicationStr.data(), '/', PathSeparator); // TODO: Is this needed?
 
@@ -106,6 +111,8 @@ namespace uba
 			description = descriptionStr.c_str();
 			logFileStr = o.logFileStr;
 			logFile = logFileStr.c_str();
+			breadcrumbsStr = o.breadcrumbsStr;
+			breadcrumbs = breadcrumbsStr.c_str();
 
 			weight = o.weight;
 		}
@@ -117,6 +124,7 @@ namespace uba
 		TString argumentsStr;
 		TString workingDirStr;
 		TString logFileStr;
+		TString breadcrumbsStr;
 		float weight = 1.0f;
 	};
 }
\ No newline at end of file
diff --git a/Engine/Source/Programs/UnrealBuildAccelerator/Common/Public/UbaTrace.h b/Engine/Source/Programs/UnrealBuildAccelerator/Common/Public/UbaTrace.h
--- a/Engine/Source/Programs/UnrealBuildAccelerator/Common/Public/UbaTrace.h
+++ b/Engine/Source/Programs/UnrealBuildAccelerator/Common/Public/UbaTrace.h
@@ -68,7 +68,7 @@ namespace uba
 	using Color = u32;
 	inline Color toColor(u8 r, u8 g, u8 b) { return (r << 16) + (g << 8) + b; }
 
-	static constexpr u32 TraceVersion = 33;
+	static constexpr u32 TraceVersion = 34;
 	static constexpr u32 TraceReadCompatibilityVersion = 6;
 
 	class Trace : public WorkTracker
@@ -87,7 +87,7 @@ namespace uba
 		void SessionDisconnect(u32 sessionId);
 		void ProcessAdded(u32 sessionId, u32 processId, const tchar* description);
 		void ProcessEnvironmentUpdated(u32 processId, const tchar* reason, const u8* data, u64 dataSize);
-		void ProcessExited(u32 processId, u32 exitCode, const u8* data, u64 dataSize, const Vector<ProcessLogLine>& logLines);
+		void ProcessExited(u32 processId, u32 exitCode, const u8* data, u64 dataSize, const Vector<ProcessLogLine>& logLines, const tchar* breadcrumbs);
 		void ProcessReturned(u32 processId, const StringView& reason);
 		void ProxyCreated(u32 clientId, const tchar* proxyName);
 		void ProxyUsed(u32 clientId, const tchar* proxyName);
diff --git a/Engine/Source/Programs/UnrealBuildAccelerator/Visualizer/Private/UbaTraceReader.cpp b/Engine/Source/Programs/UnrealBuildAccelerator/Visualizer/Private/UbaTraceReader.cpp
--- a/Engine/Source/Programs/UnrealBuildAccelerator/Visualizer/Private/UbaTraceReader.cpp
+++ b/Engine/Source/Programs/UnrealBuildAccelerator/Visualizer/Private/UbaTraceReader.cpp
@@ -601,6 +601,11 @@ namespace uba
 			process.stats.resize(dataEnd - dataStart);
 			memcpy(process.stats.data(), dataStart, dataEnd - dataStart);
 
+			if (out.version >= 34)
+			{
+				process.breadcrumbs = reader.ReadString();
+			}
+
 			process.createFilesTime = processStats.createFile.time;
 			process.writeFilesTime = Max(processStats.writeFiles.time, processStats.sendFiles.time);
 
diff --git a/Engine/Source/Programs/UnrealBuildAccelerator/Visualizer/Private/UbaTraceReader.h b/Engine/Source/Programs/UnrealBuildAccelerator/Visualizer/Private/UbaTraceReader.h
--- a/Engine/Source/Programs/UnrealBuildAccelerator/Visualizer/Private/UbaTraceReader.h
+++ b/Engine/Source/Programs/UnrealBuildAccelerator/Visualizer/Private/UbaTraceReader.h
@@ -21,6 +21,7 @@ namespace uba
 			u64 stop = 0;
 			TString description;
 			TString returnedReason;
+			TString breadcrumbs;
 			HBITMAP bitmap = 0;
 			u32 bitmapOffset = 0;
 			bool bitmapDirty = true;
diff --git a/Engine/Source/Programs/UnrealBuildAccelerator/Visualizer/Private/UbaVisualizer.cpp b/Engine/Source/Programs/UnrealBuildAccelerator/Visualizer/Private/UbaVisualizer.cpp
--- a/Engine/Source/Programs/UnrealBuildAccelerator/Visualizer/Private/UbaVisualizer.cpp
+++ b/Engine/Source/Programs/UnrealBuildAccelerator/Visualizer/Private/UbaVisualizer.cpp
@@ -1767,6 +1767,37 @@ namespace uba
 			if (hasExited && process.exitCode != 0)
 				logger.Info(L"  ExitCode: %7u", process.exitCode);
 
+			const auto& breadcrumbs = process.breadcrumbs;
+			if (!breadcrumbs.empty())
+			{
+				constexpr TString::size_type maxLineLen = 37;
+				logger.Info(L"");
+				logger.Info(L"  ------------ Breadcrumbs ------------");
+				for (TString::size_type lineStart = 0, lineEnd = 0; lineEnd < breadcrumbs.size(); lineStart = lineEnd + 1)
+				{
+					// Log each individual line
+					lineEnd = breadcrumbs.find(L'\n', lineStart);
+					TString line = (lineEnd == TString::npos ? breadcrumbs.substr(lineStart) : breadcrumbs.substr(lineStart, lineEnd - lineStart));
+
+					// Break each line down into smaller section if they are longer than the maximum allowed length
+					if (line.size() > maxLineLen)
+					{
+						for (TString::size_type sectionStart = 0, sectionEnd = 0; sectionStart < line.size(); sectionStart = sectionEnd)
+						{
+							const TString::size_type maxSectionLen = sectionStart == 0 ? maxLineLen : maxLineLen - 2;
+							sectionEnd = std::min<TString::size_type>(sectionEnd + maxSectionLen, line.size());
+							TString section = (sectionStart == 0 ? L"  " : L"    ") + line.substr(sectionStart, sectionEnd - sectionStart);
+							logger.Info(section.c_str());
+						}
+					}
+					else
+					{
+						line = L"  " + line;
+						logger.Info(line.c_str());
+					}
+				}
+			}
+
 			if (process.stop != ~u64(0) && !process.stats.empty())
 			{
 				BinaryReader reader(process.stats.data(), 0, process.stats.size());
@@ -2978,7 +3009,7 @@ namespace uba
 			{
 				m_trace.ReadFile(m_traceView, m_fileName.data, m_replay != 0);
 				m_traceView.finished = m_replay == 0;
-				PostNewTitle(GetTitlePrefix(title).Append(m_fileName));
+				PostNewTitle(GetTitlePrefix(title).Appendf(L"%s (v%u)", m_fileName.data, m_traceView.version));
 			}
 			else
 			{
diff --git a/Engine/Source/Runtime/Engine/Private/ShaderCompiler/ShaderCompilerDistributed.cpp b/Engine/Source/Runtime/Engine/Private/ShaderCompiler/ShaderCompilerDistributed.cpp
--- a/Engine/Source/Runtime/Engine/Private/ShaderCompiler/ShaderCompilerDistributed.cpp
+++ b/Engine/Source/Runtime/Engine/Private/ShaderCompiler/ShaderCompilerDistributed.cpp
@@ -30,6 +30,17 @@ namespace DistributedShaderCompilerVariables
 		TEXT("Maximum number of seconds we expect to pass between getting distributed controller complete a task (this is used to detect problems with the distribution controllers).")
 	);
 
+	static int32 GDistributedJobDescriptionLevel = 0;
+	static FAutoConsoleVariableRef CVarDistributedJobDescriptionLevel(
+		TEXT("r.ShaderCompiler.DistributedJobDescriptionLevel"),
+		GDistributedJobDescriptionLevel,
+		TEXT("Sets the level of descriptive details for each distributed job batch. The following modes are supported:\n")
+		TEXT(" Mode 0: Disabled.\n")
+		TEXT(" Mode 1: Basic information of the first 20 compile jobs per batch.\n")
+		TEXT(" Mode 2: Additional information of the shader format per compile job.\n")
+		TEXT("This will show up in the UBA trace files. By default 0.")
+	);
+
 }
 
 bool FShaderCompileDistributedThreadRunnable_Interface::IsSupported()
@@ -68,6 +79,59 @@ FShaderCompileDistributedThreadRunnable_Interface::~FShaderCompileDistributedThr
 {
 }
 
+static FString BuildCompactTaskDescriptionForJob(const FShaderCommonCompileJob& Job)
+{
+	FString JobDescription;
+	if (const FShaderCompileJob* SingleJob = Job.GetSingleShaderJob())
+	{
+		JobDescription = SingleJob->Input.DebugGroupName;
+		if (DistributedShaderCompilerVariables::GDistributedJobDescriptionLevel >= 2)
+		{
+			JobDescription += FString::Printf(TEXT("(%s)"), *SingleJob->Input.ShaderFormat.ToString());
+		}
+	}
+	else if (const FShaderPipelineCompileJob* PipelineJob = Job.GetShaderPipelineJob())
+	{
+		JobDescription = TEXT("Stages:");
+		for (int32 StageJobIndex = 0; StageJobIndex < PipelineJob->StageJobs.Num(); ++StageJobIndex)
+		{
+			if (StageJobIndex > 0)
+			{
+				JobDescription += TEXT(",");
+			}
+			JobDescription += BuildCompactTaskDescriptionForJob(*PipelineJob->StageJobs[StageJobIndex]);
+		}
+	}
+	return JobDescription;
+}
+
+// Builds a compact description of the shader compile task that will show up in UBA trace files for instance.
+// It shall contain a brief summary of the shaders being compiled to diagnose issues with overly long remote jobs.
+static FString BuildCompactTaskDescription(const TArray<FShaderCommonCompileJobPtr>& JobsToSerialize)
+{
+	FString Description;
+
+	if (!JobsToSerialize.IsEmpty())
+	{
+		constexpr int32 MaxNumJobsInDescription = 20;
+		const int32 NumJobsInDescription = JobsToSerialize.Num() > MaxNumJobsInDescription ? MaxNumJobsInDescription - 1 : JobsToSerialize.Num();
+		for (int32 JobIndex = 0; JobIndex < NumJobsInDescription; ++JobIndex)
+		{
+			if (JobIndex > 0)
+			{
+				Description += TEXT("\n");
+			}
+			Description += BuildCompactTaskDescriptionForJob(*JobsToSerialize[JobIndex]);
+		}
+		if (JobsToSerialize.Num() > NumJobsInDescription)
+		{
+			Description += FString::Printf(TEXT("%d more shaders ...\n"), JobsToSerialize.Num() - NumJobsInDescription);
+		}
+	}
+
+	return Description;
+}
+
 void FShaderCompileDistributedThreadRunnable_Interface::DispatchShaderCompileJobsBatch(TArray<FShaderCommonCompileJobPtr>& JobsToSerialize)
 {
 	const FString BaseFilePath = CachedController.CreateUniqueFilePath();
@@ -93,6 +157,11 @@ void FShaderCompileDistributedThreadRunnable_Interface::DispatchShaderCompileJob
 	TaskCommandData.OutputFileName = OutputFilePath;
 	TaskCommandData.ExtraCommandArgs = FString::Printf(TEXT("%s%s"), *FCommandLine::GetSubprocessCommandline(), GIsBuildMachine ? TEXT(" -buildmachine") : TEXT(""));
 	TaskCommandData.Dependencies = GetDependencyFilesForJobs(JobsToSerialize);
+
+	if (DistributedShaderCompilerVariables::GDistributedJobDescriptionLevel > 0)
+	{
+		TaskCommandData.Description = BuildCompactTaskDescription(JobsToSerialize);
+	}
 	
 	DispatchedTasks.Add(
 		new FDistributedShaderCompilerTask(
