Remove bShouldModulationRoutingBeUpdated bool from AudioComponent call + site fixes
#rnx

diff --git a/Engine/Source/Runtime/Engine/Classes/Components/AudioComponent.h b/Engine/Source/Runtime/Engine/Classes/Components/AudioComponent.h
--- a/Engine/Source/Runtime/Engine/Classes/Components/AudioComponent.h
+++ b/Engine/Source/Runtime/Engine/Classes/Components/AudioComponent.h
@@ -759,7 +759,7 @@ public:
 	* @param bShouldModulationRoutingBeUpdated Whether we should get the ActiveSound to update the modulator on its next update call
 	*/
 	UFUNCTION(BlueprintCallable, Category = "Audio|Components|Audio", DisplayName = "Add Modulation Routing")
-	ENGINE_API void AddModulationRouting(const TSet<USoundModulatorBase*>& Modulators, const EModulationDestination Destination, const bool bShouldModulationRoutingBeUpdated = true);
+	ENGINE_API void AddModulationRouting(const TSet<USoundModulatorBase*>& Modulators, const EModulationDestination Destination);
 
 	/**
 	* Removes the given set of Modulators to the modulators currently set on the Audio Component. This performs a logical subtraction of the sets, so modulators that are not in the given set will stay.
diff --git a/Engine/Source/Runtime/Engine/Private/Components/AudioComponent.cpp b/Engine/Source/Runtime/Engine/Private/Components/AudioComponent.cpp
--- a/Engine/Source/Runtime/Engine/Private/Components/AudioComponent.cpp
+++ b/Engine/Source/Runtime/Engine/Private/Components/AudioComponent.cpp
@@ -1913,7 +1913,7 @@ void UAudioComponent::SetModulationRouting(const TSet<USoundModulatorBase*>& Mod
 
 }
 
-void UAudioComponent::AddModulationRouting(const TSet<USoundModulatorBase*>& Modulators, const EModulationDestination Destination, const bool bShouldModulationRoutingBeUpdated)
+void UAudioComponent::AddModulationRouting(const TSet<USoundModulatorBase*>& Modulators, const EModulationDestination Destination)
 {
 	FAudioDevice* AudioDevice = GetAudioDevice();
 	if (!AudioDevice)
