[World Partition Data Layers]
- Replaced an overly complex comment that was showing as tooltip in the world outliner with a simpler, meaningful one.
- Fixed in-editor data layers filtering logic to respect the streaming generation logic: attached actors will inherit their parent data layer.
- During streaming generation, resolve parent actor descriptor views before resolving anything else, since it can be used to resolve data layers from the parents, for instance.

diff --git a/Engine/Source/Editor/DataLayerEditor/Private/DataLayer/DataLayerEditorSubsystem.cpp b/Engine/Source/Editor/DataLayerEditor/Private/DataLayer/DataLayerEditorSubsystem.cpp
--- a/Engine/Source/Editor/DataLayerEditor/Private/DataLayer/DataLayerEditorSubsystem.cpp
+++ b/Engine/Source/Editor/DataLayerEditor/Private/DataLayer/DataLayerEditorSubsystem.cpp
@@ -1585,15 +1585,19 @@ bool UDataLayerEditorSubsystem::PassDataLayersFilter(UWorld* World, const FWorld
 		// If Actor is loaded and dirty, use a newly resolved DataLayerInstanceNames array
 		auto GetLatestDataLayerInstanceNames = [DataLayerManager, ActorHandle]() -> FDataLayerInstanceNames
 		{
-			if(AActor* Actor = ActorHandle->GetActor(false); Actor && Actor->GetPackage()->IsDirty())
+			FWorldPartitionHandle TopActorHandle(ActorHandle);
+			while (TopActorHandle->GetParentActor().IsValid())
 			{
-				TUniquePtr<FWorldPartitionActorDesc> NewActorDesc = Actor->CreateActorDesc();
-				return FDataLayerUtils::ResolveDataLayerInstanceNames(DataLayerManager, NewActorDesc.Get());
+				TopActorHandle = FWorldPartitionHandle(TopActorHandle->GetContainerInstance(), TopActorHandle->GetParentActor());
 			}
-			else
+
+			if(AActor* Actor = TopActorHandle->GetActor(false); Actor && Actor->GetPackage()->IsDirty())
 			{
-				return ActorHandle->GetDataLayerInstanceNames();
+				TUniquePtr<FWorldPartitionActorDesc> NewActorDesc = Actor->CreateActorDesc();
+				return FDataLayerUtils::ResolveDataLayerInstanceNames(DataLayerManager, NewActorDesc.Get());
 			}
+
+			return TopActorHandle->GetDataLayerInstanceNames();
 		};
 
 		if (IsRunningCookCommandlet())
diff --git a/Engine/Source/Runtime/Engine/Classes/GameFramework/Actor.h b/Engine/Source/Runtime/Engine/Classes/GameFramework/Actor.h
--- a/Engine/Source/Runtime/Engine/Classes/GameFramework/Actor.h
+++ b/Engine/Source/Runtime/Engine/Classes/GameFramework/Actor.h
@@ -1030,8 +1030,7 @@ protected:
 	UPROPERTY(VisibleAnywhere, AdvancedDisplay, Category = DataLayers)
 	TArray<FActorDataLayer> DataLayers;
 
-	// There is currently an issue where if we allow property override of DataLayerAssets and it contains Private datalayers
-	// then it will always serialize a diff since those are outered to the instanced level and will get remapped differently between the Override instance and Archetype instance we are comparing against
+	/** DataLayers assets the actor belongs to.*/
 	UPROPERTY(EditAnywhere, AdvancedDisplay, Category = DataLayers, meta=(DisableLevelInstancePropertyOverride))
 	TArray<TSoftObjectPtr<UDataLayerAsset>> DataLayerAssets;
 
diff --git a/Engine/Source/Runtime/Engine/Private/WorldPartition/WorldPartitionStreamingGeneration.cpp b/Engine/Source/Runtime/Engine/Private/WorldPartition/WorldPartitionStreamingGeneration.cpp
--- a/Engine/Source/Runtime/Engine/Private/WorldPartition/WorldPartitionStreamingGeneration.cpp
+++ b/Engine/Source/Runtime/Engine/Private/WorldPartition/WorldPartitionStreamingGeneration.cpp
@@ -1076,11 +1076,11 @@ class FWorldPartitionStreamingGenerator
 	{
 		auto ResolveActorDescView = [this, &ContainerCollectionInstanceDescriptor](FStreamingGenerationActorDescView& ActorDescView)
 		{
+			ResolveParentView(ActorDescView, ContainerCollectionInstanceDescriptor.ActorDescViewMap.Get(), ContainerCollectionInstanceDescriptor.EditorOnlyActorDescSet, ContainerCollectionInstanceDescriptor.EditorOnlyParentActorTransforms);
 			ResolveRuntimeSpatiallyLoaded(ActorDescView);
 			ResolveRuntimeGrid(ActorDescView);
 			ResolveRuntimeDataLayers(ActorDescView, ContainerCollectionInstanceDescriptor.DataLayerResolvers);
-			ResolveHLODLayer(ActorDescView, WorldPartitionContext ? FSoftObjectPath(WorldPartitionContext->GetDefaultHLODLayer()) : FSoftObjectPath());
-			ResolveParentView(ActorDescView, ContainerCollectionInstanceDescriptor.ActorDescViewMap.Get(), ContainerCollectionInstanceDescriptor.EditorOnlyActorDescSet, ContainerCollectionInstanceDescriptor.EditorOnlyParentActorTransforms);
+			ResolveHLODLayer(ActorDescView, WorldPartitionContext ? FSoftObjectPath(WorldPartitionContext->GetDefaultHLODLayer()) : FSoftObjectPath());			
 		};
 
 		ContainerCollectionInstanceDescriptor.ActorDescViewMap->ForEachActorDescView([this, &ResolveActorDescView](FStreamingGenerationActorDescView& ActorDescView)
