Fix UDIM bug caused by incorrect assumption that BlockDataOffsets[0] == 0
bump DDC key for affected textures

diff --git a/Engine/Source/Runtime/Engine/Classes/Engine/Texture.h b/Engine/Source/Runtime/Engine/Classes/Engine/Texture.h
--- a/Engine/Source/Runtime/Engine/Classes/Engine/Texture.h
+++ b/Engine/Source/Runtime/Engine/Classes/Engine/Texture.h
@@ -309,6 +309,9 @@ struct FTextureSource
 	/** Computes the size of a single mip in bytes.  Same size as the image in GetMipImageInfo.
 	See also: CalcMipOffset. */
 	ENGINE_API int64 CalcMipSize(int32 BlockIndex, int32 LayerIndex, int32 MipIndex) const;
+	
+	/** Retrieve the size and offset for a source mip. The size includes all slices. */
+	ENGINE_API int64 CalcMipOffset(int32 BlockIndex, int32 LayerIndex, int32 MipIndex) const;
 
 	/** Computes the number of bytes per-pixel. */
 	ENGINE_API int64 GetBytesPerPixel(int32 LayerIndex = 0) const;
@@ -647,10 +650,7 @@ private:
 	DoUEDeltaTransform( DoUEDeltaTransform(Buffer,true), false) == Buffer
 	*/
 	FSharedBuffer DoUEDeltaTransform(FSharedBuffer InBuffer,bool bForward) const;
-
-	/** Retrieve the size and offset for a source mip. The size includes all slices. */
-	int64 CalcMipOffset(int32 BlockIndex, int32 LayerIndex, int32 MipIndex) const;
-	
+		
 	/* total size of source data in bytes */
 	int64 CalcTotalSize() const;
 	/* size of block in bytes, over all layers */
diff --git a/Engine/Source/Runtime/Engine/Private/Texture.cpp b/Engine/Source/Runtime/Engine/Private/Texture.cpp
--- a/Engine/Source/Runtime/Engine/Private/Texture.cpp
+++ b/Engine/Source/Runtime/Engine/Private/Texture.cpp
@@ -1092,7 +1092,9 @@ void UTexture::Serialize(FArchive& Ar)
 
 	if (Ar.IsLoading())
 	{
-		// Could potentially guard this with a new custom version, but overhead of just checking on every load should be very small
+		// EnsureBlocksAreSorted here should do nothing, because it only sets them if they were not saved
+		//	if they were saved, they are not changed
+		// this mainly acts to create a BlockDataOffsets array for legacy non-UDIM textures that had none
 		Source.EnsureBlocksAreSorted();
 
 		if (Ar.UEVer() < VER_UE4_TEXTURE_LEGACY_GAMMA)
@@ -3499,10 +3501,10 @@ int64 FTextureSource::CalcLayerSize(const FTextureSourceBlock& Block, int32 Laye
 
 int64 FTextureSource::CalcMipOffset(int32 BlockIndex, int32 LayerIndex, int32 OffsetToMipIndex) const
 {
-	if ( BlockIndex == 0 && LayerIndex == 0 && OffsetToMipIndex == 0 )
+	if ( LayerIndex == 0 && OffsetToMipIndex == 0 )
 	{
 		// early out common case
-		return 0;
+		return BlockDataOffsets[BlockIndex];
 	}
 
 	/*************
@@ -3519,6 +3521,8 @@ int64 FTextureSource::CalcMipOffset(int32 BlockIndex, int32 LayerIndex, int32 Of
 
 	note these are the mips in the *source*, not the number of mips generated
 
+	note: BlockDataOffsets[] are not sorted, and BlockDataOffsets[0] == 0 is not guaranteed
+
 	*************/
 
 	check( BlockIndex < GetNumBlocks() );
@@ -4924,7 +4928,7 @@ void FTextureSource::InitBlockedImpl(const ETextureSourceFormat* InLayerFormats,
 		LayerFormat[i] = InLayerFormats[i];
 	}
 
-	EnsureBlocksAreSorted();
+	EnsureBlocksAreSorted(); // this resizes and fills out the BlockDataOffsets
 
 	CheckTextureIsUnlocked(TEXT("InitBlockedImpl"));
 }
@@ -4946,11 +4950,16 @@ inline bool operator<(const FSortedTextureSourceBlock& Lhs, const FSortedTexture
 
 bool FTextureSource::EnsureBlocksAreSorted()
 {
+	// confusingly, EnsureBlocksAreSorted does not sort the offsets if they already exist
+	// it populates the BlockDataOffsets in sorted order if it does not exist
+	//	  sort order is by XY, not offset
+	// also for non-UDIM it ensure you have a 1-entry array with [0]==0
+
 	// BlockDataOffsets is of size NumBlocks, even when NumBlocks==1
-	// and BlockDataOffsets[0] == 0
 	const int32 NumBlocks = GetNumBlocks();
 	if (BlockDataOffsets.Num() == NumBlocks)
 	{
+		// do nothing if BlockDataOffsets is already set up
 		return false;
 	}
 
@@ -4969,6 +4978,7 @@ bool FTextureSource::EnsureBlocksAreSorted()
 			GetBlock(BlockIndex, SortedBlock.Block);
 			SortedBlock.SourceBlockIndex = BlockIndex;
 			SortedBlock.DataOffset = CurrentDataOffset;
+			// note: Sort is by XY position, NOT by DataOffset
 			SortedBlock.SortKey = SortedBlock.Block.BlockY * SizeInBlocks.X + SortedBlock.Block.BlockX;
 			CurrentDataOffset += CalcBlockSize(SortedBlock.Block);
 		}
diff --git a/Engine/Source/Runtime/Engine/Private/TextureDerivedData.cpp b/Engine/Source/Runtime/Engine/Private/TextureDerivedData.cpp
--- a/Engine/Source/Runtime/Engine/Private/TextureDerivedData.cpp
+++ b/Engine/Source/Runtime/Engine/Private/TextureDerivedData.cpp
@@ -497,6 +497,13 @@ void GetTextureDerivedDataKeySuffix(const UTexture& Texture, const FTextureBuild
 		KeyBuilder.Appendf(TEXT("VT%s_"), TEXTURE_VT_DERIVEDDATA_VER);
 	}
 
+	if ( Texture.Source.GetNumBlocks() > 1 && Texture.Source.CalcMipOffset(0,0,0) != 0 )
+	{
+		// bug introduced in CL 32770500 4/5/2024 , incorrectly assumed CalcMipOffset(0,0,0) == 0
+		// fix 09/10/2024
+		KeyBuilder.Appendf(TEXT("UDIMOffsetBug_"));
+	}
+
 #if PLATFORM_CPU_ARM_FAMILY
 	// Separate out arm keys as x64 and arm64 clang do not generate the same data for a given
 	// input. Add the arm specifically so that a) we avoid rebuilding the current DDC and
