#UE Better debug name info in BodyInstances

diff --git a/Engine/Source/Runtime/Engine/Private/PhysicsEngine/BodyInstance.cpp b/Engine/Source/Runtime/Engine/Private/PhysicsEngine/BodyInstance.cpp
--- a/Engine/Source/Runtime/Engine/Private/PhysicsEngine/BodyInstance.cpp
+++ b/Engine/Source/Runtime/Engine/Private/PhysicsEngine/BodyInstance.cpp
@@ -1115,7 +1115,7 @@ TSharedPtr<TArray<ANSICHAR>> GetDebugDebugName(const UPrimitiveComponent* Primit
 
 	if (PrimitiveComp)
 	{
-		DebugName += FString::Printf(TEXT("%s %s "), *AActor::GetDebugName(PrimitiveComp->GetOwner()), *PrimitiveComp->GetName());
+		DebugName += FString::Printf(TEXT("%s %s "), *AActor::GetDebugName(PrimitiveComp->GetOwner()), *PrimitiveComp->GetPathName());
 	}
 
 	if (BodySetup->BoneName != NAME_None)
