Clamped skeletal mesh component global scale animation rate minimum to 0, in order to avoid other systems issues wuth global negative values

diff --git a/Engine/Source/Runtime/Engine/Classes/Components/SkeletalMeshComponent.h b/Engine/Source/Runtime/Engine/Classes/Components/SkeletalMeshComponent.h
--- a/Engine/Source/Runtime/Engine/Classes/Components/SkeletalMeshComponent.h
+++ b/Engine/Source/Runtime/Engine/Classes/Components/SkeletalMeshComponent.h
@@ -563,7 +563,7 @@ protected:
 
 public:
 	/** Used to scale speed of all animations on this skeletal mesh. */
-	UPROPERTY(EditAnywhere, AdvancedDisplay, BlueprintReadWrite, Category=Animation, meta=(EditCondition = bEnableAnimation))
+	UPROPERTY(EditAnywhere, AdvancedDisplay, BlueprintReadWrite, Category=Animation, meta=(EditCondition = bEnableAnimation, ClampMin = 0.f))
 	float GlobalAnimRateScale;
 	
 	/** If we are running physics, should we update non-simulated bones based on the animation bone positions. */
