Add PeriodicWorldPosition to the TransformPosition material expression.
PeriodicWorldPosition is similar to absolute world space, but the world origin is moved to the center of the tile the camera is in.
This offers better precision and scalability than absolute world position, and allows for worldspace tiling which is not possible in a camera-relative space.

#tests Tested with M_ProcGrid for precision at origin, at 2^20 and 1 trillion.

#virtualized

diff --git a/Engine/Shaders/Private/MaterialTemplate.ush b/Engine/Shaders/Private/MaterialTemplate.ush
--- a/Engine/Shaders/Private/MaterialTemplate.ush
+++ b/Engine/Shaders/Private/MaterialTemplate.ush
@@ -988,6 +988,23 @@ FWSVector3 GetPrevParticleWorldPosition(FMaterialPixelParameters Parameters)
 	return Parameters.LWCData.PrevParticleWorldPosition;
 }
 
+// Periodic World Position
+
+float3 GetPeriodicWorldOrigin(float3 Scale)
+{
+	// Following three lines can be replaced with `ScaledPos = DFDivide(ResolvedView.PreViewTranslation, Scale)` but this is faster and reasonably precise.
+	float3 InvScale = rcp(Scale);
+	FDFVector3 ScaledPos = DFTwoProduct(ResolvedView.PreViewTranslation.High, InvScale);
+	ScaledPos.Low += ResolvedView.PreViewTranslation.Low * InvScale;
+	
+	return DFFracDemote(ScaledPos) * Scale;
+}
+
+float3 GetPeriodicWorldOrigin_Pow2(float3 Scale)
+{
+	return DFFmodByPow2Demote(ResolvedView.PreViewTranslation, Scale);
+}
+
 // Local position
 
 float3 GetPositionPrimitiveSpace(FMaterialPixelParameters Parameters)
diff --git a/Engine/Source/Runtime/Engine/Private/Materials/HLSLMaterialDerivativeAutogen.cpp b/Engine/Source/Runtime/Engine/Private/Materials/HLSLMaterialDerivativeAutogen.cpp
--- a/Engine/Source/Runtime/Engine/Private/Materials/HLSLMaterialDerivativeAutogen.cpp
+++ b/Engine/Source/Runtime/Engine/Private/Materials/HLSLMaterialDerivativeAutogen.cpp
@@ -819,19 +819,6 @@ int32 FMaterialDerivativeAutogen::GenerateExpressionFunc1(FHLSLMaterialTranslato
 	return Ret;
 }
 
-bool FMaterialDerivativeAutogen::IsConstFloatOfPow2Expression(FHLSLMaterialTranslator& Translator, int32 ExpCode)
-{
-	FMaterialUniformExpression* ExpressionB = Translator.GetParameterUniformExpression(ExpCode);
-	FLinearColor ValueB;
-	bool bIsPow2 = false;
-	if (ExpressionB && Translator.GetConstParameterValue(ExpressionB, ValueB))
-	{
-		auto IsFloatPowerOfTwo = [](float Value) { return ((*reinterpret_cast<int*>(&Value)) & 0x007FFFFF) == 0; }; // zero mantisse
-		bIsPow2 = IsFloatPowerOfTwo(ValueB.R) && IsFloatPowerOfTwo(ValueB.G) && IsFloatPowerOfTwo(ValueB.B) && IsFloatPowerOfTwo(ValueB.A);
-	}
-	return bIsPow2;
-}
-
 int32 FMaterialDerivativeAutogen::GenerateExpressionFunc2(FHLSLMaterialTranslator& Translator, EFunc2 Op, int32 LhsCode, int32 RhsCode)
 {
 	if (LhsCode == INDEX_NONE || RhsCode == INDEX_NONE)
@@ -953,7 +940,7 @@ int32 FMaterialDerivativeAutogen::GenerateExpressionFunc2(FHLSLMaterialTranslato
 		case EFunc2::Mul:
 			if (bIsLWC)
 			{
-				if (IsConstFloatOfPow2Expression(Translator, RhsCode))
+				if (Translator.IsConstFloatOfPow2Expression(RhsCode))
 				{
 					//Translator.AddLWCFuncUsage(ELWCFunctionKind::MultiplyVectorVector);
 					DstToken = TEXT("WSMultiplyByPow2(") + LhsToken + TEXT(", ") + RhsToken + TEXT(")");
@@ -972,7 +959,7 @@ int32 FMaterialDerivativeAutogen::GenerateExpressionFunc2(FHLSLMaterialTranslato
 		case EFunc2::Div:
 			if (bIsLWC)
 			{
-				if (IsConstFloatOfPow2Expression(Translator, RhsCode))
+				if (Translator.IsConstFloatOfPow2Expression(RhsCode))
 				{
 					//Translator.AddLWCFuncUsage(ELWCFunctionKind::Divide);
 					DstToken = TEXT("WSDivideByPow2(") + LhsToken + TEXT(", ") + RhsToken + TEXT(")");
diff --git a/Engine/Source/Runtime/Engine/Private/Materials/HLSLMaterialDerivativeAutogen.h b/Engine/Source/Runtime/Engine/Private/Materials/HLSLMaterialDerivativeAutogen.h
--- a/Engine/Source/Runtime/Engine/Private/Materials/HLSLMaterialDerivativeAutogen.h
+++ b/Engine/Source/Runtime/Engine/Private/Materials/HLSLMaterialDerivativeAutogen.h
@@ -178,8 +178,6 @@ private:
 	FString CoerceValueRaw(FHLSLMaterialTranslator& Translator, const FString& Token, const FDerivInfo& SrcInfo, EDerivativeType DstType);
 	FString CoerceValueDeriv(const FString& Token, const FDerivInfo& SrcInfo, EDerivativeType DstType);
 
-	bool IsConstFloatOfPow2Expression(FHLSLMaterialTranslator& Translator, int32 ExpressionCode);
-
 	void EnableGeneratedDepencencies();
 
 	// State to keep track of which derivative functions have been used and need to be generated.
diff --git a/Engine/Source/Runtime/Engine/Private/Materials/HLSLMaterialTranslator.cpp b/Engine/Source/Runtime/Engine/Private/Materials/HLSLMaterialTranslator.cpp
--- a/Engine/Source/Runtime/Engine/Private/Materials/HLSLMaterialTranslator.cpp
+++ b/Engine/Source/Runtime/Engine/Private/Materials/HLSLMaterialTranslator.cpp
@@ -11078,6 +11078,91 @@ int32 FHLSLMaterialTranslator::TransformPosition(EMaterialCommonBasis SourceCoor
 	return TransformBase(SourceCoordBasis, DestCoordBasis, A, 1);
 }
 
+int32 FHLSLMaterialTranslator::CalculatePeriodicWorldPositionOrigin(int TileScaleIdx)
+{
+	// float3 PeriodicWorldOrigin = GetPeriodicWorldOrigin(Scale);
+	FString PeriodicWorldOrigin = CreateSymbolName(TEXT("PeriodicWorldOrigin"));
+	FString FuncName = IsConstFloatOfPow2Expression(TileScaleIdx) ? TEXT("GetPeriodicWorldOrigin_Pow2") : TEXT("GetPeriodicWorldOrigin");
+	FString Code = FString::Format(TEXT("	float3 {0} = {1}({2});{3}"), { PeriodicWorldOrigin, FuncName, *GetParameterCode(TileScaleIdx), HLSL_LINE_TERMINATOR });
+
+	AddInlinedCodeChunk(EMaterialValueType::MCT_VoidStatement, *Code);
+	return AddCodeChunk(MCT_Float3, *PeriodicWorldOrigin);
+}
+
+int32 FHLSLMaterialTranslator::TransformFromPeriodicWorldPosition(EMaterialCommonBasis DestCoordBasis, int TileScaleIdx, int32 PeriodicWorldPosition)
+{
+	const EDerivativeStatus InputDerivStatus = GetDerivativeStatus(PeriodicWorldPosition);
+
+	int32 PeriodicWorldOrigin = CalculatePeriodicWorldPositionOrigin(TileScaleIdx);
+	FString TranslatedWorldPositionCode = FString::Format(TEXT("(({0}) + ({1}))"), { GetParameterCode(PeriodicWorldPosition), GetParameterCode(PeriodicWorldOrigin) });
+
+	int32 Result;
+	if (IsAnalyticDerivEnabled() && IsDerivativeValid(InputDerivStatus))
+	{
+		if (InputDerivStatus == EDerivativeStatus::Valid)
+		{
+			FString PositionDerivCode = GetParameterCodeDeriv(PeriodicWorldPosition, CompiledPDV_Analytic);
+
+			FString AnalyticCode = DerivativeAutogen.ConstructDeriv(
+				*TranslatedWorldPositionCode, 
+				*(PositionDerivCode + TEXT(".Ddx")), 
+				*(PositionDerivCode + TEXT(".Ddy")), 
+				GetDerivType(MCT_Float3));
+			Result = AddCodeChunkInnerDeriv(*TranslatedWorldPositionCode, *AnalyticCode, MCT_Float3, false, EDerivativeStatus::Valid);
+		}
+		else
+		{
+			Result = AddCodeChunkZeroDeriv(MCT_Float3, *TranslatedWorldPositionCode);
+		}
+	}
+	else
+	{
+		Result = AddCodeChunk(MCT_Float3, *TranslatedWorldPositionCode);
+	}
+
+	return TransformBase(EMaterialCommonBasis::MCB_TranslatedWorld, DestCoordBasis, Result, 1);
+}
+
+int32 FHLSLMaterialTranslator::TransformToPeriodicWorldPosition(EMaterialCommonBasis SourceCoordBasis, int TileScaleIdx, int32 Input)
+{
+	int32 TranslatedWorldPos = TransformBase(SourceCoordBasis, EMaterialCommonBasis::MCB_TranslatedWorld, Input, 1);
+	if (TranslatedWorldPos == INDEX_NONE) // Transform failed
+	{
+		return INDEX_NONE;
+	}
+
+	const EDerivativeStatus InputDerivStatus = GetDerivativeStatus(TranslatedWorldPos);
+
+	int32 PeriodicWorldOrigin = CalculatePeriodicWorldPositionOrigin(TileScaleIdx);
+	FString PeriodicWorldPositionCode = FString::Format(TEXT("(({0}) - ({1}))"), { GetParameterCode(TranslatedWorldPos), GetParameterCode(PeriodicWorldOrigin) });
+
+	int32 Result;
+	if (IsAnalyticDerivEnabled() && IsDerivativeValid(InputDerivStatus))
+	{
+		if (InputDerivStatus == EDerivativeStatus::Valid)
+		{
+			FString TranslatedWorldPosDerivCode = GetParameterCodeDeriv(TranslatedWorldPos, CompiledPDV_Analytic);
+
+			FString AnalyticCode = DerivativeAutogen.ConstructDeriv(
+				*PeriodicWorldPositionCode, 
+				*(TranslatedWorldPosDerivCode + TEXT(".Ddx")), 
+				*(TranslatedWorldPosDerivCode + TEXT(".Ddy")), 
+				GetDerivType(MCT_Float3));
+			Result = AddCodeChunkInnerDeriv(*PeriodicWorldPositionCode, *AnalyticCode, MCT_Float3, false, EDerivativeStatus::Valid);
+		}
+		else
+		{
+			Result = AddCodeChunkZeroDeriv(MCT_Float3, *PeriodicWorldPositionCode);
+		}
+	}
+	else
+	{
+		Result = AddCodeChunk(MCT_Float3, *PeriodicWorldPositionCode);
+	}
+
+	return Result;
+}
+
 int32 FHLSLMaterialTranslator::TransformNormalFromRequestedBasisToWorld(int32 NormalCodeChunk)
 {
 	if (IsTangentSpaceNormal())
@@ -11207,6 +11292,19 @@ int32 FHLSLMaterialTranslator::GenericSwitch(const TCHAR* SwitchExpressionText,
 	}
 }
 
+bool FHLSLMaterialTranslator::IsConstFloatOfPow2Expression(int32 ExpCode)
+{
+	FMaterialUniformExpression* ExpressionB = GetParameterUniformExpression(ExpCode);
+	FLinearColor ValueB;
+	bool bIsPow2 = false;
+	if (ExpressionB && GetConstParameterValue(ExpressionB, ValueB))
+	{
+		auto IsFloatPowerOfTwo = [](float Value) { return ((*reinterpret_cast<int*>(&Value)) & 0x007FFFFF) == 0; }; // zero mantisse
+		bIsPow2 = IsFloatPowerOfTwo(ValueB.R) && IsFloatPowerOfTwo(ValueB.G) && IsFloatPowerOfTwo(ValueB.B) && IsFloatPowerOfTwo(ValueB.A);
+	}
+	return bIsPow2;
+}
+
 int32 FHLSLMaterialTranslator::GIReplace(int32 Direct, int32 StaticIndirect, int32 DynamicIndirect)
 { 
 	return GenericSwitch(TEXT("GetGIReplaceState()"), DynamicIndirect, Direct);
diff --git a/Engine/Source/Runtime/Engine/Private/Materials/HLSLMaterialTranslator.h b/Engine/Source/Runtime/Engine/Private/Materials/HLSLMaterialTranslator.h
--- a/Engine/Source/Runtime/Engine/Private/Materials/HLSLMaterialTranslator.h
+++ b/Engine/Source/Runtime/Engine/Private/Materials/HLSLMaterialTranslator.h
@@ -848,6 +848,8 @@ protected:
 
 	int32 GenericSwitch(const TCHAR* Function, int32 IfTrue, int32 IfFalse);
 
+	bool IsConstFloatOfPow2Expression(int32 ExpressionCode);
+
 	FString SubstrateGetCastParameterCode(int32 Index, EMaterialValueType DestType);
 	FString SubstrateGetCastParameterCodeWithDeriv(int32 Index, EMaterialValueType DestType);
 
@@ -1158,6 +1160,9 @@ protected:
 	
 	virtual int32 TransformVector(EMaterialCommonBasis SourceCoordBasis, EMaterialCommonBasis DestCoordBasis, int32 A) override;
 	virtual int32 TransformPosition(EMaterialCommonBasis SourceCoordBasis, EMaterialCommonBasis DestCoordBasis, int32 A) override;
+	int32 CalculatePeriodicWorldPositionOrigin(int TileScaleIdx);
+	virtual int32 TransformFromPeriodicWorldPosition(EMaterialCommonBasis DestCoordBasis, int TileScaleIdx, int32 A) override;
+	virtual int32 TransformToPeriodicWorldPosition(EMaterialCommonBasis SourceCoordBasis, int TileScaleIdx, int32 A) override;
 	virtual int32 TransformNormalFromRequestedBasisToWorld(int32 NormalCodeChunk) override;
 	virtual int32 DynamicParameter(FLinearColor& DefaultValue, uint32 ParameterIndex = 0) override;
 	virtual int32 LightmapUVs() override;
diff --git a/Engine/Source/Runtime/Engine/Private/Materials/MaterialExpressionHLSL.cpp b/Engine/Source/Runtime/Engine/Private/Materials/MaterialExpressionHLSL.cpp
--- a/Engine/Source/Runtime/Engine/Private/Materials/MaterialExpressionHLSL.cpp
+++ b/Engine/Source/Runtime/Engine/Private/Materials/MaterialExpressionHLSL.cpp
@@ -3804,6 +3804,27 @@ bool UMaterialExpressionTransform::GenerateHLSLExpression(FMaterialHLSLGenerator
 	return true;
 }
 
+static const UE::HLSLTree::FExpression* TransformFromPeriodicWorldPosition(UE::HLSLTree::FTree& Tree, EMaterialCommonBasis DestCoordBasis, const UE::HLSLTree::FExpression* TileScaleExpression, const UE::HLSLTree::FExpression* PeriodicWorldPosition)
+{
+	using namespace UE::HLSLTree;
+	using namespace UE::HLSLTree::Material;
+
+	const FExpression* PeriodicWorldOrigin = Tree.NewExpression<FExpressionPeriodicWorldOrigin>(TileScaleExpression);
+	const FExpression* Result = Tree.NewAdd(PeriodicWorldPosition, PeriodicWorldOrigin);
+
+	return TransformBase(Tree, EMaterialCommonBasis::MCB_TranslatedWorld, DestCoordBasis, Result, true);
+}
+
+static const UE::HLSLTree::FExpression* TransformToPeriodicWorldPosition(UE::HLSLTree::FTree& Tree, EMaterialCommonBasis SourceCoordBasis, const UE::HLSLTree::FExpression* TileScaleExpression, const UE::HLSLTree::FExpression* Input)
+{
+	using namespace UE::HLSLTree;
+	using namespace UE::HLSLTree::Material;
+
+	const FExpression* TranslatedWorldPos = TransformBase(Tree, SourceCoordBasis, EMaterialCommonBasis::MCB_TranslatedWorld, Input, true);
+	const FExpression* PeriodicWorldOrigin = Tree.NewExpression<FExpressionPeriodicWorldOrigin>(TileScaleExpression);
+
+	return Tree.NewSub(TranslatedWorldPos, PeriodicWorldOrigin);
+}
 
 bool UMaterialExpressionTransformPosition::GenerateHLSLExpression(FMaterialHLSLGenerator& Generator, UE::HLSLTree::FScope& Scope, int32 OutputIndex, UE::HLSLTree::FExpression const*& OutExpression) const
 {
@@ -3812,6 +3833,7 @@ bool UMaterialExpressionTransformPosition::GenerateHLSLExpression(FMaterialHLSLG
 	static const EMaterialCommonBasis kTable[TRANSFORMPOSSOURCE_MAX] = {
 		MCB_Local,						// TRANSFORMPOSSOURCE_Local
 		MCB_World,						// TRANSFORMPOSSOURCE_World
+		MCB_MAX,						// No match for TRANSFORMPOSSOURCE_PeriodicWorld
 		MCB_TranslatedWorld,			// TRANSFORMPOSSOURCE_TranslatedWorld
 		MCB_View,						// TRANSFORMPOSSOURCE_View
 		MCB_Camera,						// TRANSFORMPOSSOURCE_Camera
@@ -3825,8 +3847,36 @@ bool UMaterialExpressionTransformPosition::GenerateHLSLExpression(FMaterialHLSLG
 		return false;
 	}
 
-	OutExpression = TransformBase(Generator.GetTree(), kTable[TransformSourceType], kTable[TransformType], ExpressionInput, true);
-	return true;
+	if (TransformSourceType == TRANSFORMPOSSOURCE_PeriodicWorld || TransformType == TRANSFORMPOSSOURCE_PeriodicWorld)
+	{
+		const FExpression* PeriodicWorldTileSizeExpression = PeriodicWorldTileSize.AcquireHLSLExpressionOrConstant(Generator, Scope, ConstPeriodicWorldTileSize);
+
+		if (TransformSourceType == TRANSFORMPOSSOURCE_PeriodicWorld)
+		{
+			if (TransformType == TRANSFORMPOSSOURCE_PeriodicWorld)
+			{
+				OutExpression = ExpressionInput;
+				return true;
+			}
+			OutExpression = TransformFromPeriodicWorldPosition(Generator.GetTree(), kTable[TransformType], PeriodicWorldTileSizeExpression, ExpressionInput);
+			return true;
+		}
+		else
+		{
+			if (TransformSourceType == TRANSFORMPOSSOURCE_PeriodicWorld)
+			{
+				OutExpression = ExpressionInput;
+				return true;
+			}
+			OutExpression = TransformToPeriodicWorldPosition(Generator.GetTree(), kTable[TransformSourceType], PeriodicWorldTileSizeExpression, ExpressionInput);
+			return true;
+		}
+	}
+	else
+	{
+		OutExpression = TransformBase(Generator.GetTree(), kTable[TransformSourceType], kTable[TransformType], ExpressionInput, true);
+		return true;
+	}
 }
 
 bool UMaterialExpressionIf::GenerateHLSLExpression(FMaterialHLSLGenerator& Generator, UE::HLSLTree::FScope& Scope, int32 OutputIndex, UE::HLSLTree::FExpression const*& OutExpression) const
diff --git a/Engine/Source/Runtime/Engine/Private/Materials/MaterialExpressions.cpp b/Engine/Source/Runtime/Engine/Private/Materials/MaterialExpressions.cpp
--- a/Engine/Source/Runtime/Engine/Private/Materials/MaterialExpressions.cpp
+++ b/Engine/Source/Runtime/Engine/Private/Materials/MaterialExpressions.cpp
@@ -13552,6 +13552,7 @@ static EMaterialCommonBasis GetMaterialCommonBasis(EMaterialPositionTransformSou
 	static const EMaterialCommonBasis ConversionTable[TRANSFORMPOSSOURCE_MAX] = {
 		MCB_Local,						// TRANSFORMPOSSOURCE_Local
 		MCB_World,						// TRANSFORMPOSSOURCE_World
+		MCB_MAX,						// No match for TRANSFORMPOSSOURCE_PeriodicWorld
 		MCB_TranslatedWorld,			// TRANSFORMPOSSOURCE_TranslatedWorld
 		MCB_View,						// TRANSFORMPOSSOURCE_View
 		MCB_Camera,						// TRANSFORMPOSSOURCE_Camera
@@ -13564,21 +13565,52 @@ static EMaterialCommonBasis GetMaterialCommonBasis(EMaterialPositionTransformSou
 #if WITH_EDITOR
 int32 UMaterialExpressionTransformPosition::Compile(class FMaterialCompiler* Compiler, int32 OutputIndex)
 {
-	int32 Result=INDEX_NONE;
-	
-	if( !Input.GetTracedInput().Expression )
+	if(!Input.GetTracedInput().Expression)
 	{
-		Result = Compiler->Errorf(TEXT("Missing Transform Position input vector"));
+		return Compiler->Errorf(TEXT("Missing Transform Position input vector"));
+	}
+
+	int32 VecInputIdx = Input.Compile(Compiler);
+
+	// Periodic World Position (see TRANSFORMPOSSOURCE_PeriodicWorld) takes a different path as the transform depends on a user-specified tile scale
+	if (TransformSourceType == TRANSFORMPOSSOURCE_PeriodicWorld || TransformType == TRANSFORMPOSSOURCE_PeriodicWorld)
+	{
+		int32 PeriodicWorldTileSizeIdx;
+		if (PeriodicWorldTileSize.GetTracedInput().Expression)
+		{
+			PeriodicWorldTileSizeIdx = PeriodicWorldTileSize.Compile(Compiler);
+		}
+		else
+		{
+			PeriodicWorldTileSizeIdx = Compiler->Constant(ConstPeriodicWorldTileSize);
+		}
+
+		if (TransformSourceType == TRANSFORMPOSSOURCE_PeriodicWorld)
+		{
+			if (TransformType == TRANSFORMPOSSOURCE_PeriodicWorld)
+			{
+				return VecInputIdx;
+			}
+			const auto TransformDestBasis = GetMaterialCommonBasis(TransformType);
+			return Compiler->TransformFromPeriodicWorldPosition(TransformDestBasis, PeriodicWorldTileSizeIdx, VecInputIdx);
+		}
+		else
+		{
+			if (TransformSourceType == TRANSFORMPOSSOURCE_PeriodicWorld)
+			{
+				return VecInputIdx;
+			}
+			const auto TransformSourceBasis = GetMaterialCommonBasis(TransformSourceType);
+			return Compiler->TransformToPeriodicWorldPosition(TransformSourceBasis, PeriodicWorldTileSizeIdx, VecInputIdx);
+		}
 	}
 	else
 	{
-		int32 VecInputIdx = Input.Compile(Compiler);
 		const auto TransformSourceBasis = GetMaterialCommonBasis(TransformSourceType);
 		const auto TransformDestBasis = GetMaterialCommonBasis(TransformType);
-		Result = Compiler->TransformPosition(TransformSourceBasis, TransformDestBasis, VecInputIdx);
-	}
 
-	return Result;
+		return Compiler->TransformPosition(TransformSourceBasis, TransformDestBasis, VecInputIdx);
+	}
 }
 
 void UMaterialExpressionTransformPosition::GetCaption(TArray<FString>& OutCaptions) const
@@ -13598,6 +13630,60 @@ void UMaterialExpressionTransformPosition::GetCaption(TArray<FString>& OutCaptio
 	
 	OutCaptions.Add(TEXT("TransformPosition"));
 }
+
+TArrayView<FExpressionInput*> UMaterialExpressionTransformPosition::GetInputsView()
+{
+	CachedInputs.Empty();
+	uint32 InputIndex = 0;
+	while (FExpressionInput* Ptr = GetInput(InputIndex++))
+	{
+		CachedInputs.Add(Ptr);
+	}
+	return CachedInputs;
+}
+
+// this define is only used for the following function
+#define IF_INPUT_RETURN(Item) if(!InputIndex) return &Item; --InputIndex
+FExpressionInput* UMaterialExpressionTransformPosition::GetInput(int32 InputIndex)
+{
+	IF_INPUT_RETURN(Input);
+
+	if (bUsesPeriodicWorldPosition)
+	{
+		IF_INPUT_RETURN(PeriodicWorldTileSize);
+	}
+
+	return nullptr;
+}
+#undef IF_INPUT_RETURN
+
+#define IF_INPUT_RETURN(Name) if(!InputIndex) return Name; --InputIndex
+FName UMaterialExpressionTransformPosition::GetInputName(int32 InputIndex) const
+{
+	IF_INPUT_RETURN(TEXT(""));
+	IF_INPUT_RETURN(TEXT("Periodic World Tile Size"));
+
+	return TEXT("");
+}
+#undef IF_INPUT_RETURN
+
+void UMaterialExpressionTransformPosition::PostEditChangeProperty(struct FPropertyChangedEvent& PropertyChangedEvent)
+{
+	Super::PostEditChangeProperty(PropertyChangedEvent);
+
+	if (PropertyChangedEvent.GetMemberPropertyName() == GET_MEMBER_NAME_CHECKED(ThisClass, TransformSourceType) ||
+		PropertyChangedEvent.GetMemberPropertyName() == GET_MEMBER_NAME_CHECKED(ThisClass, TransformType))
+	{
+		bUsesPeriodicWorldPosition = TransformSourceType == TRANSFORMPOSSOURCE_PeriodicWorld || TransformType == TRANSFORMPOSSOURCE_PeriodicWorld;
+		if (GraphNode)
+		{
+			GraphNode->ReconstructNode();
+		}
+	}
+
+	// Need to update expression properties before super call (which triggers recompile)
+	Super::PostEditChangeProperty(PropertyChangedEvent);
+}
 #endif // WITH_EDITOR
 
 ///////////////////////////////////////////////////////////////////////////////
diff --git a/Engine/Source/Runtime/Engine/Private/Materials/MaterialHLSLTree.cpp b/Engine/Source/Runtime/Engine/Private/Materials/MaterialHLSLTree.cpp
--- a/Engine/Source/Runtime/Engine/Private/Materials/MaterialHLSLTree.cpp
+++ b/Engine/Source/Runtime/Engine/Private/Materials/MaterialHLSLTree.cpp
@@ -3574,6 +3574,41 @@ void FExpressionDefaultSubsurfaceColor::EmitValuePreshader(FEmitContext& Context
 	OutResult.Preshader.WriteOpcode(Shader::EPreshaderOpcode::Constant).Write(Value);
 }
 
+bool FExpressionPeriodicWorldOrigin::PrepareValue(FEmitContext& Context, FEmitScope& Scope, const FRequestedType& RequestedType, FPrepareValueResult& OutResult) const
+{
+	return OutResult.SetType(Context, RequestedType, EExpressionEvaluation::Shader, Shader::EValueType::Float3);
+}
+
+static const bool IsConstFloatOfPow2Expression(FEmitContext& Context, FEmitScope& Scope, const UE::HLSLTree::FExpression* Expression)
+{
+	bool bIsConstPow2 = false;
+
+	const EExpressionEvaluation Evaluation = Context.GetPreparedType(Expression, Shader::EValueType::Float3).GetEvaluation(Scope);
+	if (Evaluation == EExpressionEvaluation::Constant || Evaluation == EExpressionEvaluation::ConstantZero)
+	{
+		UE::Shader::FFloatValue Value = Expression->GetValueConstant(Context, Scope, Shader::EValueType::Float3).AsFloat();
+		auto IsFloatPowerOfTwo = [](float Value) { return ((*reinterpret_cast<int*>(&Value)) & 0x007FFFFF) == 0; }; // zero mantisse
+		bIsConstPow2 = IsFloatPowerOfTwo(Value[0]) && IsFloatPowerOfTwo(Value[1]) && IsFloatPowerOfTwo(Value[2]) && IsFloatPowerOfTwo(Value[3]);
+	}
+
+	return bIsConstPow2;
+}
+
+void FExpressionPeriodicWorldOrigin::EmitValueShader(FEmitContext& Context, FEmitScope& Scope, const FRequestedType& RequestedType, FEmitValueShaderResult& OutResult) const
+{
+	FEmitShaderExpression* EmitScaleExpression = Scale->GetValueShader(Context, Scope, Shader::EValueType::Float3);
+	bool bIsScalePow2 = IsConstFloatOfPow2Expression(Context, Scope, Scale);
+	const TCHAR* FuncName = bIsScalePow2 ? TEXT("GetPeriodicWorldOrigin_Pow2") : TEXT("GetPeriodicWorldOrigin");
+	OutResult.Code = Context.EmitExpression(Scope, Shader::EValueType::Float3, TEXT("%s(%s)"), FuncName, EmitScaleExpression);
+}
+
+void FExpressionPeriodicWorldOrigin::ComputeAnalyticDerivatives(FTree& Tree, FExpressionDerivatives& OutResult) const
+{
+	const Shader::FValue ZeroValue(Shader::EValueType::Float3);
+	OutResult.ExpressionDdx = Tree.NewConstant(ZeroValue);
+	OutResult.ExpressionDdy = OutResult.ExpressionDdx;
+}
+
 int32 FEmitData::FindInterpolatorIndex(const FExpression* Expression) const
 {
 	for (int32 Index = 0; Index < VertexInterpolators.Num(); ++Index)
diff --git a/Engine/Source/Runtime/Engine/Public/MaterialCompiler.h b/Engine/Source/Runtime/Engine/Public/MaterialCompiler.h
--- a/Engine/Source/Runtime/Engine/Public/MaterialCompiler.h
+++ b/Engine/Source/Runtime/Engine/Public/MaterialCompiler.h
@@ -496,6 +496,8 @@ public:
 	virtual int32 AppendVector(int32 A,int32 B) = 0;
 	virtual int32 TransformVector(EMaterialCommonBasis SourceCoordBasis, EMaterialCommonBasis DestCoordBasis, int32 A) = 0;
 	virtual int32 TransformPosition(EMaterialCommonBasis SourceCoordBasis, EMaterialCommonBasis DestCoordBasis, int32 A) = 0;
+	virtual int32 TransformFromPeriodicWorldPosition(EMaterialCommonBasis DestCoordBasis, int TileScaleIdx, int32 A) = 0;
+	virtual int32 TransformToPeriodicWorldPosition(EMaterialCommonBasis SourceCoordBasis, int TileScaleIdx, int32 A) = 0;
 	virtual int32 TransformNormalFromRequestedBasisToWorld(int32 NormalCodeChunk) = 0;
 
 	virtual int32 DynamicParameter(FLinearColor& DefaultValue, uint32 ParameterIndex = 0) = 0;
@@ -1048,6 +1050,14 @@ public:
 	{
 		return Compiler->TransformPosition(SourceCoordBasis, DestCoordBasis, A);
 	}
+	virtual int32 TransformFromPeriodicWorldPosition(EMaterialCommonBasis DestCoordBasis, int TileScaleIdx, int32 A) override
+	{
+		return Compiler->TransformFromPeriodicWorldPosition(DestCoordBasis, TileScaleIdx, A);
+	}
+	virtual int32 TransformToPeriodicWorldPosition(EMaterialCommonBasis SourceCoordBasis, int TileScaleIdx, int32 A) override
+	{
+		return Compiler->TransformToPeriodicWorldPosition(SourceCoordBasis, TileScaleIdx, A);
+	}
 	virtual int32 TransformNormalFromRequestedBasisToWorld(int32 NormalCodeChunk) override { return Compiler->TransformNormalFromRequestedBasisToWorld(NormalCodeChunk); }
 
 	virtual int32 DynamicParameter(FLinearColor& DefaultValue, uint32 ParameterIndex = 0) override { return Compiler->DynamicParameter(DefaultValue, ParameterIndex); }
diff --git a/Engine/Source/Runtime/Engine/Public/MaterialHLSLTree.h b/Engine/Source/Runtime/Engine/Public/MaterialHLSLTree.h
--- a/Engine/Source/Runtime/Engine/Public/MaterialHLSLTree.h
+++ b/Engine/Source/Runtime/Engine/Public/MaterialHLSLTree.h
@@ -1244,6 +1244,20 @@ public:
 	virtual void EmitValuePreshader(FEmitContext& Context, FEmitScope& Scope, const FRequestedType& RequestedType, FEmitValuePreshaderResult& OutResult) const override;
 };
 
+class FExpressionPeriodicWorldOrigin : public FExpression
+{
+public:
+	FExpressionPeriodicWorldOrigin(const FExpression* InScale = nullptr)
+	: Scale(InScale)
+	{}
+
+	const FExpression* Scale;
+
+	virtual bool PrepareValue(FEmitContext& Context, FEmitScope& Scope, const FRequestedType& RequestedType, FPrepareValueResult& OutResult) const override;
+	virtual void EmitValueShader(FEmitContext& Context, FEmitScope& Scope, const FRequestedType& RequestedType, FEmitValueShaderResult& OutResult) const override;
+	virtual void ComputeAnalyticDerivatives(FTree& Tree, FExpressionDerivatives& OutResult) const override;
+};
+
 struct FVertexInterpolator
 {
 	FVertexInterpolator() = default;
diff --git a/Engine/Source/Runtime/Engine/Public/Materials/MaterialExpressionTransformPosition.h b/Engine/Source/Runtime/Engine/Public/Materials/MaterialExpressionTransformPosition.h
--- a/Engine/Source/Runtime/Engine/Public/Materials/MaterialExpressionTransformPosition.h
+++ b/Engine/Source/Runtime/Engine/Public/Materials/MaterialExpressionTransformPosition.h
@@ -17,8 +17,17 @@ enum EMaterialPositionTransformSource : int
 	
 	/** Absolute world space */
 	TRANSFORMPOSSOURCE_World UMETA(DisplayName="Absolute World Space"),
+
+	/**
+	  Like absolute world space, but the world origin is moved to the center of the tile the camera is in.
+	  Logically similar to `fmod(CameraAbsoluteWorldPosition, TileSize) + CameraRelativeWorldPosition`.
+	  This offers better precision and scalability than absolute world position.
+	  Suitable as a position input for functions that tile based on world position, e.g. frac(Position / TileSize).
+	  Works best when the tile size is a power of two.
+	*/
+	TRANSFORMPOSSOURCE_PeriodicWorld UMETA(DisplayName="Periodic World Space"),
 	
-	/** Camera relative world space */
+	/** Translated world space, i.e. world space rotation and scale but with a position relative to the camera */
 	TRANSFORMPOSSOURCE_TranslatedWorld  UMETA(DisplayName="Camera Relative World Space"),
 
 	/** View space (differs from camera space in the shadow passes) */
@@ -53,15 +62,30 @@ class UMaterialExpressionTransformPosition : public UMaterialExpression
 	UPROPERTY(EditAnywhere, Category=MaterialExpressionTransformPosition, meta=(DisplayName = "Destination", ShowAsInputPin = "Advanced"))
 	TEnumAsByte<enum EMaterialPositionTransformSource> TransformType;
 
+	/** scale of the tiles used in Periodic World Space */
+	UPROPERTY(meta=(DisplayName = "Periodic World Tile Size", ToolTip = "Distance the camera can move before the world origin is moved", RequiredInput = "false"))
+	FExpressionInput PeriodicWorldTileSize;
+
+	/** only used if PeriodicWorldTileSize is not hooked up */
+	UPROPERTY(EditAnywhere, Category=MaterialExpressionTransformPosition, meta=(DisplayName = "Periodic World Tile Size", RequiredInput="false", ClampMin=0.0001, UIMax=524288.0, OverridingInputProperty = "PeriodicWorldTileSize", EditCondition=bUsesPeriodicWorldPosition, EditConditionHides))
+	float ConstPeriodicWorldTileSize = 32.0;
+
+private:
+	UPROPERTY()
+	bool bUsesPeriodicWorldPosition = false;
+
+public:
 
 	//~ Begin UMaterialExpression Interface
 #if WITH_EDITOR
 	virtual bool GenerateHLSLExpression(FMaterialHLSLGenerator& Generator, UE::HLSLTree::FScope& Scope, int32 OutputIndex, UE::HLSLTree::FExpression const*& OutExpression) const override;
 	virtual int32 Compile(class FMaterialCompiler* Compiler, int32 OutputIndex) override;
 	virtual void GetCaption(TArray<FString>& OutCaptions) const override;
+	virtual TArrayView<FExpressionInput*> GetInputsView() override;
+	virtual FExpressionInput* GetInput(int32 InputIndex) override;
+	virtual FName GetInputName(int32 InputIndex) const override;
+	virtual void PostEditChangeProperty(struct FPropertyChangedEvent& PropertyChangedEvent) override;
 #endif
 	//~ End UMaterialExpression Interface
 };
 
-
-
