[shaders] fix processing duplicate jobs in the case the job failed (and so was not added to the cache)

diff --git a/Engine/Source/Runtime/Engine/Private/ShaderCompiler/ShaderCompilerJobCache.cpp b/Engine/Source/Runtime/Engine/Private/ShaderCompiler/ShaderCompilerJobCache.cpp
--- a/Engine/Source/Runtime/Engine/Private/ShaderCompiler/ShaderCompilerJobCache.cpp
+++ b/Engine/Source/Runtime/Engine/Private/ShaderCompiler/ShaderCompilerJobCache.cpp
@@ -1580,6 +1580,9 @@ void FShaderJobCache::AddToCacheAndProcessPending(FShaderCommonCompileJob* Finis
 	const FShaderCompilerInputHash& InputHash = FinishedJob->GetInputHash();
 	FShaderCacheSaveContext SaveContext;
 	FinishedJob->SerializeOutput(SaveContext);
+	// Explicitly finalize the serialization to generate the job struct FSharedBuffer since it's needed below in the case
+	// we need to process/populate any duplicate job results
+	SaveContext.Finalize();
 
 	FShaderJobData& JobData = GetShaderJobData(FinishedJob->JobCacheRef);
 
@@ -1633,11 +1636,9 @@ void FShaderJobCache::AddToCacheAndProcessPending(FShaderCommonCompileJob* Finis
 	{
 		UE_LOG(LogShaderCompilers, UE_SHADERCACHE_LOG_LEVEL, TEXT("Processed %d outstanding jobs with the same ihash %s."), FinishedDuplicateJobs.Num(), *LexToString(InputHash));
 
-		check(JobData.HasOutput());
-		FStoredOutput** StoredOutput = Outputs.Find(JobData.OutputHash);
-		check(StoredOutput);
-		// Construct a single load context pointing to the data now stored in the cache
-		FShaderCacheLoadContext LoadContext((*StoredOutput)->JobOutput, (*StoredOutput)->JobCode);
+		check(SaveContext);
+		// Construct a single load context pointing to the data in the save context used above
+		FShaderCacheLoadContext LoadContext(SaveContext.ShaderObjectData, SaveContext.ShaderCode);
 		for (FShaderCommonCompileJob* DuplicateJob : FinishedDuplicateJobs)
 		{
 			// reuse the same load context for each duplicate job to avoid reallocating anything
@@ -2107,10 +2108,6 @@ void FShaderJobCache::AddJobOutput(FShaderJobData& JobData, const FShaderCommonC
 		return;
 	}
 
-	// Explicitly finalize the serialization to generate the job struct FSharedBuffer since we store
-	// this in the in-memory cache prior to pushing to DDC (subsequent Finalize call in BuildCacheRecord 
-	// will be a noop)
-	SaveContext.Finalize();
 	FJobOutputHash OutputHash = ComputeJobHash(SaveContext);
 
 	if (JobData.HasOutput() && ShaderCompiler::IsJobCacheDebugValidateEnabled())
@@ -2210,6 +2207,8 @@ void FShaderJobCache::AddJobOutput(FShaderJobData& JobData, const FShaderCommonC
 			NewStoredOutput->NumHits = InitialHitCount;
 			NewStoredOutput->JobOutput = SaveContext.ShaderObjectData;
 			NewStoredOutput->JobCode = MoveTemp(SaveContext.OwnedShaderCode);
+			// reset ShaderCode view in SaveContext to point to the cached array, in case its data is reused to populate duplicate jobs
+			SaveContext.ShaderCode = NewStoredOutput->JobCode;
 			NewStoredOutput->CachedDebugInfoPath = InputDebugInfoPath;
 			NewStoredOutput->AddRef();
 			Outputs.Add(OutputHash, NewStoredOutput);
