Fix PVS static analysis warning introduced in CL 36190594.

diff --git a/Engine/Source/Runtime/Engine/Private/Materials/MaterialExpressionHLSL.cpp b/Engine/Source/Runtime/Engine/Private/Materials/MaterialExpressionHLSL.cpp
--- a/Engine/Source/Runtime/Engine/Private/Materials/MaterialExpressionHLSL.cpp
+++ b/Engine/Source/Runtime/Engine/Private/Materials/MaterialExpressionHLSL.cpp
@@ -3851,23 +3851,18 @@ bool UMaterialExpressionTransformPosition::GenerateHLSLExpression(FMaterialHLSLG
 	{
 		const FExpression* PeriodicWorldTileSizeExpression = PeriodicWorldTileSize.AcquireHLSLExpressionOrConstant(Generator, Scope, ConstPeriodicWorldTileSize);
 
+		if (TransformSourceType == TRANSFORMPOSSOURCE_PeriodicWorld && TransformType == TRANSFORMPOSSOURCE_PeriodicWorld)
+		{
+			OutExpression = ExpressionInput;
+			return true;
+		}
 		if (TransformSourceType == TRANSFORMPOSSOURCE_PeriodicWorld)
 		{
-			if (TransformType == TRANSFORMPOSSOURCE_PeriodicWorld)
-			{
-				OutExpression = ExpressionInput;
-				return true;
-			}
 			OutExpression = TransformFromPeriodicWorldPosition(Generator.GetTree(), kTable[TransformType], PeriodicWorldTileSizeExpression, ExpressionInput);
 			return true;
 		}
 		else
 		{
-			if (TransformSourceType == TRANSFORMPOSSOURCE_PeriodicWorld)
-			{
-				OutExpression = ExpressionInput;
-				return true;
-			}
 			OutExpression = TransformToPeriodicWorldPosition(Generator.GetTree(), kTable[TransformSourceType], PeriodicWorldTileSizeExpression, ExpressionInput);
 			return true;
 		}
diff --git a/Engine/Source/Runtime/Engine/Private/Materials/MaterialExpressions.cpp b/Engine/Source/Runtime/Engine/Private/Materials/MaterialExpressions.cpp
--- a/Engine/Source/Runtime/Engine/Private/Materials/MaterialExpressions.cpp
+++ b/Engine/Source/Runtime/Engine/Private/Materials/MaterialExpressions.cpp
@@ -13585,21 +13585,17 @@ int32 UMaterialExpressionTransformPosition::Compile(class FMaterialCompiler* Com
 			PeriodicWorldTileSizeIdx = Compiler->Constant(ConstPeriodicWorldTileSize);
 		}
 
+		if (TransformSourceType == TRANSFORMPOSSOURCE_PeriodicWorld && TransformType == TRANSFORMPOSSOURCE_PeriodicWorld)
+		{
+			return VecInputIdx;
+		}
 		if (TransformSourceType == TRANSFORMPOSSOURCE_PeriodicWorld)
 		{
-			if (TransformType == TRANSFORMPOSSOURCE_PeriodicWorld)
-			{
-				return VecInputIdx;
-			}
 			const auto TransformDestBasis = GetMaterialCommonBasis(TransformType);
 			return Compiler->TransformFromPeriodicWorldPosition(TransformDestBasis, PeriodicWorldTileSizeIdx, VecInputIdx);
 		}
 		else
 		{
-			if (TransformSourceType == TRANSFORMPOSSOURCE_PeriodicWorld)
-			{
-				return VecInputIdx;
-			}
 			const auto TransformSourceBasis = GetMaterialCommonBasis(TransformSourceType);
 			return Compiler->TransformToPeriodicWorldPosition(TransformSourceBasis, PeriodicWorldTileSizeIdx, VecInputIdx);
 		}
