Typo fix

diff --git a/Engine/Source/Runtime/Engine/Public/Animation/AnimTypes.h b/Engine/Source/Runtime/Engine/Public/Animation/AnimTypes.h
--- a/Engine/Source/Runtime/Engine/Public/Animation/AnimTypes.h
+++ b/Engine/Source/Runtime/Engine/Public/Animation/AnimTypes.h
@@ -329,7 +329,7 @@ struct FAnimNotifyEvent : public FAnimLinkableElement
 	UPROPERTY(EditAnywhere, BlueprintReadWrite, Category=AnimNotifyEvent)
 	TEnumAsByte<EMontageNotifyTickType::Type> MontageTickType;
 
-	/** Defines the chance of of this notify triggering, 0 = No Chance, 1 = Always triggers */
+	/** Defines the chance of this notify triggering, 0 = No Chance, 1 = Always triggers */
 	UPROPERTY(EditAnywhere, BlueprintReadWrite, Category = AnimNotifyTriggerSettings, meta = (ClampMin = "0.0", ClampMax = "1.0"))
 	float NotifyTriggerChance;
 
