Interchange: fix deadlock when a selected skeletal mesh in the world is re-import.

#rnx

diff --git a/Engine/Source/Runtime/Engine/Private/Components/SkinnedMeshComponent.cpp b/Engine/Source/Runtime/Engine/Private/Components/SkinnedMeshComponent.cpp
--- a/Engine/Source/Runtime/Engine/Private/Components/SkinnedMeshComponent.cpp
+++ b/Engine/Source/Runtime/Engine/Private/Components/SkinnedMeshComponent.cpp
@@ -1531,16 +1531,19 @@ void USkinnedMeshComponent::UpdateFollowerComponent()
 // this has to be skeletalmesh material. You can't have more than what SkeletalMesh materials have
 int32 USkinnedMeshComponent::GetNumMaterials() const
 {
-	if (GetSkinnedAsset())
+	if (GetSkinnedAsset() && !GetSkinnedAsset()->IsCompiling())
 	{
 		return GetSkinnedAsset()->GetMaterials().Num();
 	}
-
 	return 0;
 }
 
 void USkinnedMeshComponent::GetPrimitiveStats(FPrimitiveStats& PrimitiveStats) const
 {
+	if (!IsRegistered() || IsCompiling())
+	{
+		return;
+	}
 	const TIndirectArray<FSkeletalMeshLODRenderData>& LODRenderData = GetSkeletalMeshRenderData()->LODRenderData;
 	const TArray<FSkeletalMaterial>& Materials = GetSkinnedAsset()->GetMaterials();
 	PrimitiveStats.LODStats.Reserve(LODRenderData.Num());
@@ -1572,7 +1575,7 @@ UMaterialInterface* USkinnedMeshComponent::GetMaterial(int32 MaterialIndex) cons
 	{
 		return OverrideMaterials[MaterialIndex];
 	}
-	else if (GetSkinnedAsset() && GetSkinnedAsset()->GetMaterials().IsValidIndex(MaterialIndex) && GetSkinnedAsset()->GetMaterials()[MaterialIndex].MaterialInterface)
+	else if (GetSkinnedAsset() && !GetSkinnedAsset()->IsCompiling() && GetSkinnedAsset()->GetMaterials().IsValidIndex(MaterialIndex) && GetSkinnedAsset()->GetMaterials()[MaterialIndex].MaterialInterface)
 	{
 		return GetSkinnedAsset()->GetMaterials()[MaterialIndex].MaterialInterface;
 	}
@@ -1582,7 +1585,7 @@ UMaterialInterface* USkinnedMeshComponent::GetMaterial(int32 MaterialIndex) cons
 
 int32 USkinnedMeshComponent::GetMaterialIndex(FName MaterialSlotName) const
 {
-	if (GetSkinnedAsset() != nullptr)
+	if (GetSkinnedAsset() && !GetSkinnedAsset()->IsCompiling())
 	{
 		const TArray<FSkeletalMaterial>& SkeletalMeshMaterials = GetSkinnedAsset()->GetMaterials();
 		for (int32 MaterialIndex = 0; MaterialIndex < SkeletalMeshMaterials.Num(); ++MaterialIndex)
@@ -1600,7 +1603,7 @@ int32 USkinnedMeshComponent::GetMaterialIndex(FName MaterialSlotName) const
 TArray<FName> USkinnedMeshComponent::GetMaterialSlotNames() const
 {
 	TArray<FName> MaterialNames;
-	if (GetSkinnedAsset() != nullptr)
+	if (GetSkinnedAsset() && !GetSkinnedAsset()->IsCompiling())
 	{
 		const TArray<FSkeletalMaterial>& SkeletalMeshMaterials = GetSkinnedAsset()->GetMaterials();
 		for (int32 MaterialIndex = 0; MaterialIndex < SkeletalMeshMaterials.Num(); ++MaterialIndex)
