[shaders] minor cleanup to cache serialization context APIs

diff --git a/Engine/Source/Runtime/Engine/Private/Materials/MaterialShader.cpp b/Engine/Source/Runtime/Engine/Private/Materials/MaterialShader.cpp
--- a/Engine/Source/Runtime/Engine/Private/Materials/MaterialShader.cpp
+++ b/Engine/Source/Runtime/Engine/Private/Materials/MaterialShader.cpp
@@ -1858,7 +1858,7 @@ TSharedRef<FMaterialShaderMap::FAsyncLoadContext> FMaterialShaderMap::BeginLoadF
 			TRACE_CPUPROFILER_EVENT_SCOPE(FMaterialShaderMap::FinishLoadFromDerivedDataCache);
 			COOK_STAT(auto Timer = MaterialShaderCookStats::UsageStats.TimeSyncWork());
 
-			if (LoadContext)
+			if (LoadContext.HasData())
 			{
 				TRACE_COUNTER_INCREMENT(Shaders_FMaterialShaderMapDDCHits);
 
diff --git a/Engine/Source/Runtime/Engine/Private/ShaderCompiler/ShaderCompilerJobCache.cpp b/Engine/Source/Runtime/Engine/Private/ShaderCompiler/ShaderCompilerJobCache.cpp
--- a/Engine/Source/Runtime/Engine/Private/ShaderCompiler/ShaderCompilerJobCache.cpp
+++ b/Engine/Source/Runtime/Engine/Private/ShaderCompiler/ShaderCompilerJobCache.cpp
@@ -625,7 +625,7 @@ private:
 	FJobOutputHash ComputeJobHash(const FShaderCacheSerializeContext& SerializeContext)
 	{
 		FBlake3 Hasher;
-		check(SerializeContext);
+		check(SerializeContext.HasData());
 		Hasher.Update(SerializeContext.ShaderObjectData.GetData(), SerializeContext.ShaderObjectData.GetSize());
 		for (const FSharedBuffer& CodeBuf : SerializeContext.ShaderCode)
 		{
@@ -1382,7 +1382,7 @@ void FShaderJobCache::SubmitJob(FShaderCommonCompileJob* Job)
 		FShaderCacheLoadContext LoadContext = FindOrAdd(InputHash, Job, bCheckDDC);
 
 		// see if there are already cached results for this job that were returned synchronously by FindOrAdd
-		if (LoadContext)
+		if (LoadContext.HasData())
 		{
 			Unlink(*Job);		// from PendingSubmitJobTaskJobs
 
@@ -1636,7 +1636,7 @@ void FShaderJobCache::AddToCacheAndProcessPending(FShaderCommonCompileJob* Finis
 	{
 		UE_LOG(LogShaderCompilers, UE_SHADERCACHE_LOG_LEVEL, TEXT("Processed %d outstanding jobs with the same ihash %s."), FinishedDuplicateJobs.Num(), *LexToString(InputHash));
 
-		check(SaveContext);
+		check(SaveContext.HasData());
 		// Construct a single load context pointing to the data in the save context used above
 		FShaderCacheLoadContext LoadContext(SaveContext.ShaderObjectData, SaveContext.ShaderCode);
 		for (FShaderCommonCompileJob* DuplicateJob : FinishedDuplicateJobs)
@@ -1901,10 +1901,10 @@ FShaderCacheLoadContext FShaderJobCache::FindOrAdd(const FShaderCompilerInputHas
 						if (StoredOutput == nullptr)
 						{
 							// Create a new entry to store in the FShaderJobCache if one doesn't already exist for this output hash
-							check(LoadContext); // sanity check that the load context was populated properly
+							check(LoadContext.HasData()); // sanity check that the load context was populated properly
 							StoredOutput = new FStoredOutput();
 							StoredOutput->JobOutput = LoadContext.ShaderObjectData;
-							StoredOutput->JobCode = MoveTemp(LoadContext.OwnedShaderCode);
+							LoadContext.MoveCode(StoredOutput->JobCode);
 							Outputs.Add(OutputHash, StoredOutput);
 							CurrentlyAllocatedMemory += StoredOutput->GetAllocatedSize();
 						}
@@ -2202,13 +2202,11 @@ void FShaderJobCache::AddJobOutput(FShaderJobData& JobData, const FShaderCommonC
 		{
 			const uint64 OutputsOriginalSize = Outputs.GetAllocatedSize();
 
-			check(SaveContext);
+			check(SaveContext.HasData());
 			FStoredOutput* NewStoredOutput = new FStoredOutput();
 			NewStoredOutput->NumHits = InitialHitCount;
 			NewStoredOutput->JobOutput = SaveContext.ShaderObjectData;
-			NewStoredOutput->JobCode = MoveTemp(SaveContext.OwnedShaderCode);
-			// reset ShaderCode view in SaveContext to point to the cached array, in case its data is reused to populate duplicate jobs
-			SaveContext.ShaderCode = NewStoredOutput->JobCode;
+			SaveContext.MoveCode(NewStoredOutput->JobCode);
 			NewStoredOutput->CachedDebugInfoPath = InputDebugInfoPath;
 			NewStoredOutput->AddRef();
 			Outputs.Add(OutputHash, NewStoredOutput);
diff --git a/Engine/Source/Runtime/RenderCore/Internal/ShaderSerialization.h b/Engine/Source/Runtime/RenderCore/Internal/ShaderSerialization.h
--- a/Engine/Source/Runtime/RenderCore/Internal/ShaderSerialization.h
+++ b/Engine/Source/Runtime/RenderCore/Internal/ShaderSerialization.h
@@ -95,8 +95,15 @@ struct FShaderCacheSerializeContext : public FShaderSerializeContext
 		return Size;
 	}
 
-	/* Boolean conversion operator which checks for validity of the data referenced by this context. */
-	explicit operator bool() const { return ShaderObjectData && !ShaderCode.IsEmpty(); }
+	/* Populates the given code array (transfering ownership) and resets the internal view to point to the new owning array's data */
+	void MoveCode(TArray<FSharedBuffer>& TargetCode)
+	{
+		TargetCode = MoveTemp(OwnedShaderCode);
+		ShaderCode = TargetCode;
+	}
+
+	/* Returns true if there is valid serialized data referenced by this context. */
+	bool HasData() const { return ShaderObjectData && !ShaderCode.IsEmpty(); }
 };
 
 /* Implementation of FShaderCacheSerializeContext used for saving data to caches. */
