SplineComponent: Added support via new blueprint exposed method to enable modification of spline points within blueprints and outside of construction scripts / direct manipulation in the level editor.

diff --git a/Engine/Source/Runtime/Engine/Classes/Components/SplineComponent.h b/Engine/Source/Runtime/Engine/Classes/Components/SplineComponent.h
--- a/Engine/Source/Runtime/Engine/Classes/Components/SplineComponent.h
+++ b/Engine/Source/Runtime/Engine/Classes/Components/SplineComponent.h
@@ -371,6 +371,10 @@ public:
 	UFUNCTION(BlueprintCallable, Category = Spline)
 	ENGINE_API virtual void UpdateSpline();
 
+	/** Set the spline to be edited outside of the construction script */
+	UFUNCTION(BlueprintCallable, Category = Spline)
+	ENGINE_API virtual void SetOverrideConstructionScript(bool InOverride);
+
 	/** Get location along spline at the provided input key value */
 	UFUNCTION(BlueprintCallable, Category = Spline)
 	ENGINE_API FVector GetLocationAtSplineInputKey(float InKey, ESplineCoordinateSpace::Type CoordinateSpace) const;
diff --git a/Engine/Source/Runtime/Engine/Private/Components/SplineComponent.cpp b/Engine/Source/Runtime/Engine/Private/Components/SplineComponent.cpp
--- a/Engine/Source/Runtime/Engine/Private/Components/SplineComponent.cpp
+++ b/Engine/Source/Runtime/Engine/Private/Components/SplineComponent.cpp
@@ -283,6 +283,12 @@ void USplineComponent::UpdateSpline()
 #endif
 }
 
+void USplineComponent::SetOverrideConstructionScript(bool bInOverride)
+{
+	bSplineHasBeenEdited = bInOverride;
+}
+
+
 float FSplineCurves::GetSegmentLength(const int32 Index, const float Param, bool bClosedLoop, const FVector& Scale3D) const
 {
 	const int32 NumPoints = Position.Points.Num();
