Don't run vertex color override fixup on static mesh components if the base static mesh is cooked. This doesn't work, and causes the overrides to be discarded.

diff --git a/Engine/Source/Runtime/Engine/Private/Components/StaticMeshComponent.cpp b/Engine/Source/Runtime/Engine/Private/Components/StaticMeshComponent.cpp
--- a/Engine/Source/Runtime/Engine/Private/Components/StaticMeshComponent.cpp
+++ b/Engine/Source/Runtime/Engine/Private/Components/StaticMeshComponent.cpp
@@ -1294,6 +1294,15 @@ bool UStaticMeshComponent::RequiresOverrideVertexColorsFixup()
 	{
 		return false;
 	}
+	
+	// Don't attempt tofixup vertex colours if the base mesh is cooked. We don't have the source data in order
+	// to do it correctly, and we'll end up discarding the old data and lose the overrides altogether.
+	// Serialization during cooking should discard the overrides if the stream is incompatible with the static
+	// mesh, so we shouldn't end up with mismatching streams at runtime.
+	if (Mesh->GetPackage() && Mesh->GetPackage()->HasAllPackagesFlags(PKG_Cooked))
+	{
+		return false;
+	}
 
 	if ( !Mesh->GetRenderData())
 	{
