Disable r.SupportDepthOnlyIndexBuffers and r.SupportReversedIndexBuffers by default now. These settings generally wasted 100s of mb in Nanite heavy content, and are completely useless when rendering via Nanite. Projects or specific platforms can still re-enable them if desirable.

diff --git a/Engine/Source/Runtime/Engine/Private/StaticMesh.cpp b/Engine/Source/Runtime/Engine/Private/StaticMesh.cpp
--- a/Engine/Source/Runtime/Engine/Private/StaticMesh.cpp
+++ b/Engine/Source/Runtime/Engine/Private/StaticMesh.cpp
@@ -132,13 +132,13 @@ static TAutoConsoleVariable<int32> CVarStaticMeshKeepMobileMinLODSettingOnDeskto
 
 static TAutoConsoleVariable<int32> CVarSupportDepthOnlyIndexBuffers(
 	TEXT("r.SupportDepthOnlyIndexBuffers"),
-	1,
+	0,
 	TEXT("Enables depth-only index buffers. Saves a little time at the expense of doubling the size of index buffers."),
 	ECVF_ReadOnly | ECVF_RenderThreadSafe);
 
 static TAutoConsoleVariable<int32> CVarSupportReversedIndexBuffers(
 	TEXT("r.SupportReversedIndexBuffers"),
-	1,
+	0,
 	TEXT("Enables reversed index buffers. Saves a little time at the expense of doubling the size of index buffers."),
 	ECVF_ReadOnly | ECVF_RenderThreadSafe);
 
