Flush Rendering and RHI before GC in trim memory in case there are some UStreamableRenderAssets with pending RenderAssetUpdate RHI tasks (under Engine.FlushRTAndRHIBeforeGCInTrimMemory cvar).

#rnx

diff --git a/Engine/Source/Runtime/Engine/Private/UnrealEngine.cpp b/Engine/Source/Runtime/Engine/Private/UnrealEngine.cpp
--- a/Engine/Source/Runtime/Engine/Private/UnrealEngine.cpp
+++ b/Engine/Source/Runtime/Engine/Private/UnrealEngine.cpp
@@ -490,6 +490,15 @@ static FAutoConsoleVariableRef GDelayTrimMemoryDuringMapLoadModeCVar(
 	ECVF_Default
 );
 
+// Flush Rendering and RHI before GC in trim memory
+int32 GFlushRTAndRHIBeforeGCInTrimMemory = 1;
+static FAutoConsoleVariableRef GFlushRTAndRHIBeforeGCInTrimMemoryCVar(
+	TEXT("Engine.FlushRTAndRHIBeforeGCInTrimMemory"),
+	GFlushRTAndRHIBeforeGCInTrimMemory,
+	TEXT("Flush Rendering and RHI before GC in trim memory in case there are some UStreamableRenderAssets with pending RenderAssetUpdate RHI tasks (default: 1)"),
+	ECVF_Default
+);
+
 static int32 GVerifyLoadMapWorldCleanup_Severity = 2;
 static FAutoConsoleVariableRef CVarVerifyLoadMapWorldCleanup_Severity(
 #if UE_BUILD_SHIPPING
@@ -15944,6 +15953,17 @@ bool UEngine::LoadMap( FWorldContext& WorldContext, FURL URL, class UPendingNetG
 
 void UEngine::TrimMemory()
 {
+	if (GFlushRTAndRHIBeforeGCInTrimMemory > 0)
+	{
+		// Flush render and RHI commands in case there are some UStreamableRenderAssets with pending RenderAssetUpdate RHI tasks
+		ENQUEUE_RENDER_COMMAND(FlushCommand)(
+			[](FRHICommandList& RHICmdList)
+		{
+			GRHICommandList.GetImmediateCommandList().ImmediateFlush(EImmediateFlushType::FlushRHIThread);
+		});
+		FlushRenderingCommands();
+	}
+
 	// Clean up the previous level out of memory.
 	CollectGarbage(GARBAGE_COLLECTION_KEEPFLAGS, true);
 
