- Add additional option to how we track particle system memory to the CSV stats
- We have just the component, Component + Asset Estimate & Component + Asset Count Bytes
- The default is not Component + Asset Estimate

#rnx
#tests CSV captures and processing to confirm charts are updated

diff --git a/Engine/Source/Runtime/Engine/Classes/Particles/ParticlePerfStats.h b/Engine/Source/Runtime/Engine/Classes/Particles/ParticlePerfStats.h
--- a/Engine/Source/Runtime/Engine/Classes/Particles/ParticlePerfStats.h
+++ b/Engine/Source/Runtime/Engine/Classes/Particles/ParticlePerfStats.h
@@ -285,6 +285,8 @@ struct FParticlePerfStats
 	FName CSVStat_Culled = NAME_None;
 	FName CSVStat_MemoryKB = NAME_None;
 
+	TOptional<uint64> CSVMemoryKB_Asset;
+
 	void PopulateStatNames(const FName InName);
 	void ResetStatNames();
 #endif
diff --git a/Engine/Source/Runtime/Engine/Private/Particles/ParticlePerfStats.cpp b/Engine/Source/Runtime/Engine/Private/Particles/ParticlePerfStats.cpp
--- a/Engine/Source/Runtime/Engine/Private/Particles/ParticlePerfStats.cpp
+++ b/Engine/Source/Runtime/Engine/Private/Particles/ParticlePerfStats.cpp
@@ -1160,14 +1160,15 @@ static FAutoConsoleVariable CVarWriteDetailedCSVStats(
 	ECVF_Default | ECVF_RenderThreadSafe
 );
 
-int32 GFXDetailedCSVMemorySMode = 1;
+int32 GFXDetailedCSVMemoryMode = 2;
 static FAutoConsoleVariableRef CVarFXDetailedCSVMemoryStats(
 	TEXT("fx.DetailedCSVStats.MemoryMode"),
-	GFXDetailedCSVMemorySMode,
+	GFXDetailedCSVMemoryMode,
 	TEXT("Gathers approximate memory information depending on the mode.\n")
 	TEXT("0 = Disabled (default).\n")
-	TEXT("1 = Minimal information (small performance impact).\n")
-	TEXT("2 = Full information (large performance impact)."),
+	TEXT("1 = Component Approximate Size (low perf impact).\n")
+	TEXT("2 = Asset Size + Component Approximate Size (low perf impact).\n")
+	TEXT("3 = Asset Size + Component Resource Size (high performance impact)."),
 	ECVF_Default
 );
 
@@ -1227,9 +1228,36 @@ bool FParticlePerfStatsListener_CSVProfiler::Tick()
 	{
 #if WITH_PER_SYSTEM_PARTICLE_PERF_STATS
 		// Very slow, this gives coverage over everything but is bad for performance
-		if (GFXDetailedCSVMemorySMode != 0)
+		if (GFXDetailedCSVMemoryMode != 0)
 		{
 			TMap<UFXSystemAsset*, uint64> MemoryUsage;
+
+			const bool bIncludeAssetSize = GFXDetailedCSVMemoryMode > 1;
+			const bool bComponentUseResourceSize = GFXDetailedCSVMemoryMode == 3;
+
+			// Gather per loaded asset cost
+			if (bIncludeAssetSize)
+			{
+				for (TObjectIterator<UFXSystemAsset> It; It; ++It)
+				{
+					UFXSystemAsset* FXAsset = *It;
+					if (!IsValid(FXAsset) || FXAsset->IsUnreachable() || FXAsset->HasAnyFlags(EObjectFlags::RF_ClassDefaultObject))
+					{
+						continue;
+					}
+
+					FParticlePerfStats* Stats = FParticlePerfStatsManager::GetSystemPerfStats(FXAsset);
+					if (!Stats->CSVMemoryKB_Asset.IsSet())
+					{
+						FResourceSizeEx ResSize = FResourceSizeEx(EResourceSizeMode::EstimatedTotal);
+						FXAsset->GetResourceSizeEx(ResSize);
+						Stats->CSVMemoryKB_Asset = ResSize.GetTotalMemoryBytes();
+					}
+					MemoryUsage.FindOrAdd(FXAsset) = Stats->CSVMemoryKB_Asset.GetValue();
+				}
+			}
+
+			// Gather per component cost
 			for (TObjectIterator<UFXSystemComponent> It; It; ++It)
 			{
 				UFXSystemComponent* FXComponent = *It;
@@ -1244,15 +1272,10 @@ bool FParticlePerfStatsListener_CSVProfiler::Tick()
 				{
 					continue;
 				}
-				const bool bFullMemory = GFXDetailedCSVMemorySMode == 2;
 
 				uint64& MemoryBytes = MemoryUsage.FindOrAdd(FXAsset);
-				if (bFullMemory)
+				if (bComponentUseResourceSize)
 				{
-					if (MemoryBytes == 0)
-					{
-						MemoryBytes += FArchiveCountMem(FXAsset).GetMax();
-					}
 					MemoryBytes += FArchiveCountMem(FXComponent).GetMax();
 
 					FResourceSizeEx CompResSize = FResourceSizeEx(EResourceSizeMode::EstimatedTotal);
