Fix attempt for a TSAN warning

#tests loaded into a level in the editor

diff --git a/Engine/Source/Runtime/Engine/Private/MeshCardRepresentation.cpp b/Engine/Source/Runtime/Engine/Private/MeshCardRepresentation.cpp
--- a/Engine/Source/Runtime/Engine/Private/MeshCardRepresentation.cpp
+++ b/Engine/Source/Runtime/Engine/Private/MeshCardRepresentation.cpp
@@ -411,7 +411,7 @@ void FCardRepresentationAsyncQueue::AddTask(FAsyncCardRepresentationTask* Task)
 
 	if (!MeshUtilities)
 	{
-		MeshUtilities = &FModuleManager::Get().LoadModuleChecked<IMeshUtilities>(TEXT("MeshUtilities"));
+		FPlatformAtomics::InterlockedCompareExchangePointer((void**)&MeshUtilities, &FModuleManager::Get().LoadModuleChecked<IMeshUtilities>(TEXT("MeshUtilities")), nullptr);
 	}
 
 	const bool bUseAsyncBuild = GUseAsyncCardRepresentationBuildQueue || !IsInGameThread();
