#UE Preload clothing assets so that they will have data in time for when PostLoad happen in USkeletalMesh

diff --git a/Engine/Source/Runtime/Engine/Private/SkeletalMesh.cpp b/Engine/Source/Runtime/Engine/Private/SkeletalMesh.cpp
--- a/Engine/Source/Runtime/Engine/Private/SkeletalMesh.cpp
+++ b/Engine/Source/Runtime/Engine/Private/SkeletalMesh.cpp
@@ -2035,6 +2035,20 @@ void USkeletalMesh::Serialize( FArchive& Ar )
 			ThisLODInfo.BuildSettings.bUseBackwardsCompatibleF16TruncUVs = true;
 		}
 	}
+
+#if WITH_EDITOR
+	// Preload MeshClothingAssets because we call ConditionalPostLoad on them in our PostLoad. The PostLoad of these assets requires the data to actually have been loaded already
+	if (Ar.IsLoading())
+	{
+		for (UClothingAssetBase* MeshClothingAsset : GetMeshClothingAssets())
+		{
+			if (MeshClothingAsset)
+			{
+				Ar.Preload(MeshClothingAsset);
+			}
+		}
+	}
+#endif //WITH_EDITOR
 }
 
 #if WITH_EDITORONLY_DATA
