Initial work on new animation bank asset type

diff --git a/Engine/Plugins/Editor/EngineAssetDefinitions/Source/Private/Animation/AssetDefinition_AnimBank.cpp b/Engine/Plugins/Editor/EngineAssetDefinitions/Source/Private/Animation/AssetDefinition_AnimBank.cpp
new file mode 100644
--- /dev/null
+++ b/Engine/Plugins/Editor/EngineAssetDefinitions/Source/Private/Animation/AssetDefinition_AnimBank.cpp
@@ -0,0 +1,53 @@
+// Copyright Epic Games, Inc. All Rights Reserved.
+
+#include "Animation/AssetDefinition_AnimBank.h"
+
+#include "ContentBrowserMenuContexts.h"
+#include "Misc/PackageName.h"
+#include "ToolMenus.h"
+#include "Misc/MessageDialog.h"
+#include "ContentBrowserModule.h"
+#include "Framework/Application/SlateApplication.h"
+#include "SBlueprintDiff.h"
+#include "SSkeletonWidget.h"
+#include "Styling/SlateIconFinder.h"
+#include "Kismet2/KismetEditorUtilities.h"
+#include "IAnimationBlueprintEditorModule.h"
+#include "IAssetTools.h"
+#include "Algo/AllOf.h"
+#include "Algo/NoneOf.h"
+#include "ToolMenu.h"
+#include "ToolMenuEntry.h"
+#include "ToolMenuSection.h"
+#include "SPrimaryButton.h"
+#include "SSimpleButton.h"
+#include "BlueprintEditorSettings.h"
+#include "Logging/MessageLog.h"
+
+#define LOCTEXT_NAMESPACE "AssetTypeActions"
+
+TSharedPtr<SWidget> UAssetDefinition_AnimBank::GetThumbnailOverlay(const FAssetData& AssetData) const
+{
+	const FSlateBrush* Icon = FSlateIconFinder::FindIconBrushForClass(UAnimBank::StaticClass());
+
+	return SNew(SBorder)
+		.BorderImage(FAppStyle::GetNoBrush())
+		.Visibility(EVisibility::HitTestInvisible)
+		.Padding(FMargin(0.0f, 0.0f, 0.0f, 3.0f))
+		.HAlign(HAlign_Right)
+		.VAlign(VAlign_Bottom)
+	[
+		SNew(SImage)
+		.Image(Icon)
+	];
+}
+
+// Menu Extensions
+//--------------------------------------------------------------------
+
+namespace MenuExtension_AnimBank
+{
+
+}
+
+#undef LOCTEXT_NAMESPACE
diff --git a/Engine/Plugins/Editor/EngineAssetDefinitions/Source/Private/Animation/AssetDefinition_AnimBank.h b/Engine/Plugins/Editor/EngineAssetDefinitions/Source/Private/Animation/AssetDefinition_AnimBank.h
new file mode 100644
--- /dev/null
+++ b/Engine/Plugins/Editor/EngineAssetDefinitions/Source/Private/Animation/AssetDefinition_AnimBank.h
@@ -0,0 +1,34 @@
+// Copyright Epic Games, Inc. All Rights Reserved.
+
+#pragma once
+
+#include "Animation/AnimBank.h"
+#include "Animation/AssetDefinition_AnimationAsset.h"
+
+#include "AssetDefinition_AnimBank.generated.h"
+
+class UAssetDefinition_AnimationAsset;
+
+UCLASS()
+class UAssetDefinition_AnimBank : public UAssetDefinition_AnimationAsset
+{
+	GENERATED_BODY()
+
+public:
+	virtual FText GetAssetDisplayName() const override
+	{
+		return NSLOCTEXT("AssetTypeActions", "AssetTypeActions_AnimBank", "Animation Bank");
+	}
+	
+	virtual TSoftClassPtr<UObject> GetAssetClass() const override
+	{
+		return UAnimBank::StaticClass();
+	}
+
+	virtual FLinearColor GetAssetColor() const override
+	{
+		return FLinearColor(FColor(237, 100, 36));
+	}
+
+	virtual TSharedPtr<SWidget> GetThumbnailOverlay(const FAssetData& AssetData) const override;
+};
diff --git a/Engine/Source/Editor/UnrealEd/Classes/Factories/AnimBankFactory.h b/Engine/Source/Editor/UnrealEd/Classes/Factories/AnimBankFactory.h
new file mode 100644
--- /dev/null
+++ b/Engine/Source/Editor/UnrealEd/Classes/Factories/AnimBankFactory.h
@@ -0,0 +1,36 @@
+// Copyright Epic Games, Inc. All Rights Reserved.
+
+#pragma once
+
+#include "CoreMinimal.h"
+#include "UObject/ObjectMacros.h"
+#include "Factories/Factory.h"
+#include "AnimBankFactory.generated.h"
+
+struct FAssetData;
+class SWindow;
+
+UCLASS(HideCategories=Object, BlueprintType, MinimalAPI)
+class UAnimBankFactory : public UFactory
+{
+	GENERATED_UCLASS_BODY()
+
+	UPROPERTY(EditAnywhere, BlueprintReadWrite, Category = Animation)
+	TObjectPtr<class USkeleton> TargetSkeleton;
+
+	/** The preview mesh to use with this animation bank */
+	UPROPERTY(EditAnywhere, BlueprintReadWrite, Category = Animation)
+	TObjectPtr<class USkeletalMesh> PreviewSkeletalMesh;
+
+	//~ Begin UFactory Interface
+	UNREALED_API virtual bool ConfigureProperties() override;
+	UNREALED_API virtual UObject* FactoryCreateNew(UClass* Class,UObject* InParent, FName Name, EObjectFlags Flags, UObject* Context, FFeedbackContext* Warn) override;
+	virtual bool ShouldShowInNewMenu() const override { return true; } 
+	//~ End UFactory Interface	
+
+private:
+	UNREALED_API void OnTargetSkeletonSelected(const FAssetData& SelectedAsset);
+
+private:
+	TSharedPtr<SWindow> PickerWindow;
+};
diff --git a/Engine/Source/Editor/UnrealEd/Private/Factories/AnimBankFactory.cpp b/Engine/Source/Editor/UnrealEd/Private/Factories/AnimBankFactory.cpp
new file mode 100644
--- /dev/null
+++ b/Engine/Source/Editor/UnrealEd/Private/Factories/AnimBankFactory.cpp
@@ -0,0 +1,86 @@
+// Copyright Epic Games, Inc. All Rights Reserved.
+
+#include "Factories/AnimBankFactory.h"
+#include "Modules/ModuleManager.h"
+#include "Widgets/DeclarativeSyntaxSupport.h"
+#include "Widgets/SWindow.h"
+#include "Widgets/Layout/SBorder.h"
+#include "Styling/AppStyle.h"
+#include "Animation/Skeleton.h"
+#include "Animation/AnimBank.h"
+#include "Editor.h"
+
+#include "IContentBrowserSingleton.h"
+#include "ContentBrowserModule.h"
+
+#define LOCTEXT_NAMESPACE "AnimBankFactory"
+
+UAnimBankFactory::UAnimBankFactory(const FObjectInitializer& ObjectInitializer)
+: Super(ObjectInitializer)
+{
+	bCreateNew = true;
+	SupportedClass = UAnimBank::StaticClass();
+}
+
+bool UAnimBankFactory::ConfigureProperties()
+{
+	// Null the skeleton so we can check for selection later
+	TargetSkeleton = nullptr;
+
+	// Load the content browser module to display an asset picker
+	FContentBrowserModule& ContentBrowserModule = FModuleManager::LoadModuleChecked<FContentBrowserModule>("ContentBrowser");
+
+	FAssetPickerConfig AssetPickerConfig;
+	
+	/** The asset picker will only show skeletons */
+	AssetPickerConfig.Filter.ClassPaths.Add(USkeleton::StaticClass()->GetClassPathName());
+	AssetPickerConfig.Filter.bRecursiveClasses = true;
+
+	/** The delegate that fires when an asset was selected */
+	AssetPickerConfig.OnAssetSelected = FOnAssetSelected::CreateUObject(this, &UAnimBankFactory::OnTargetSkeletonSelected);
+
+	/** The default view mode should be a list view */
+	AssetPickerConfig.InitialAssetViewType = EAssetViewType::List;
+
+	PickerWindow = SNew(SWindow)
+		.Title(LOCTEXT("CreateAnimBankOptions", "Pick Skeleton"))
+		.ClientSize(FVector2D(500, 600))
+		.SupportsMinimize(false) .SupportsMaximize(false)
+	[
+		SNew(SBorder)
+		.BorderImage( FAppStyle::GetBrush("Menu.Background") )
+		[
+			ContentBrowserModule.Get().CreateAssetPicker(AssetPickerConfig)
+		]
+	];
+
+	GEditor->EditorAddModalWindow(PickerWindow.ToSharedRef());
+	PickerWindow.Reset();
+
+	return TargetSkeleton != nullptr;
+}
+
+UObject* UAnimBankFactory::FactoryCreateNew(UClass* Class, UObject* InParent, FName Name, EObjectFlags Flags, UObject* Context, FFeedbackContext* Warn)
+{
+	UAnimBank* AnimBank = NewObject<UAnimBank>(InParent, Class, Name, Flags);
+	
+	if (TargetSkeleton)
+	{
+		AnimBank->SetSkeleton(TargetSkeleton);
+	}
+
+	if (PreviewSkeletalMesh)
+	{
+		AnimBank->SetPreviewMesh(PreviewSkeletalMesh);
+	}
+
+	return AnimBank;
+}
+
+void UAnimBankFactory::OnTargetSkeletonSelected(const FAssetData& SelectedAsset)
+{
+	TargetSkeleton = Cast<USkeleton>(SelectedAsset.GetAsset());
+	PickerWindow->RequestDestroyWindow();
+}
+
+#undef LOCTEXT_NAMESPACE
diff --git a/Engine/Source/Runtime/Engine/Classes/Animation/AnimBank.h b/Engine/Source/Runtime/Engine/Classes/Animation/AnimBank.h
new file mode 100644
--- /dev/null
+++ b/Engine/Source/Runtime/Engine/Classes/Animation/AnimBank.h
@@ -0,0 +1,64 @@
+// Copyright Epic Games, Inc. All Rights Reserved.
+
+#pragma once
+
+#include "CoreMinimal.h"
+#include "AnimationAsset.h"
+#include "AnimBank.generated.h"
+
+USTRUCT(BlueprintType)
+struct FAnimBankItem
+{
+	GENERATED_USTRUCT_BODY()
+
+	UPROPERTY(EditAnywhere, BlueprintReadWrite, Category = Animation)
+	TObjectPtr<class UAnimSequence> Sequence;
+
+	UPROPERTY(EditAnywhere, BlueprintReadWrite, Category = Animation, meta = (DisplayName = "Looping"))
+	uint32 bLooping : 1;
+
+	UPROPERTY(EditAnywhere, BlueprintReadWrite, Category = Animation, meta = (DisplayName = "AutoStart"))
+	uint32 bAutoStart : 1;
+
+	UPROPERTY(EditAnywhere, BlueprintReadWrite, Category = Animation, meta = (DisplayName = "Position"))
+	float Position;
+
+	UPROPERTY(EditAnywhere, BlueprintReadWrite, Category = Animation, meta = (DisplayName = "PlayRate"))
+	float PlayRate;
+
+	FAnimBankItem()
+	{
+		Sequence	= nullptr;
+		PlayRate	= 1.0f;
+		bLooping 	= true;
+		bAutoStart	= true;
+		Position	= 0.0f;
+	}
+
+	ENGINE_API void ValidatePosition();
+};
+
+UCLASS(BlueprintType, hidecategories=Object, editinlinenew, MinimalAPI)
+class UAnimBank : public UAnimationAsset
+{
+	GENERATED_UCLASS_BODY()
+
+	UPROPERTY(EditAnywhere, BlueprintReadWrite, Category = Sequences, meta = (ShowOnlyInnerProperties))
+	TArray<struct FAnimBankItem> Items;
+
+public:
+	//~ Begin UObject Interface
+	ENGINE_API virtual void Serialize(FArchive& Ar) override;
+	ENGINE_API virtual void PostLoad() override;
+	virtual void PreSave(FObjectPreSaveContext ObjectSaveContext) override;
+	//~ End UObject Interface
+
+#if WITH_EDITOR
+	ENGINE_API virtual void PostEditChangeProperty(FPropertyChangedEvent& PropertyChangedEvent) override;
+
+	//~ Begin AnimationAsset interface
+	virtual bool GetAllAnimationSequencesReferred(TArray<UAnimationAsset*>& AnimationAssets, bool bRecursive = true) override;
+	virtual void ReplaceReferredAnimations(const TMap<UAnimationAsset*, UAnimationAsset*>& ReplacementMap) override;
+	//~ End AnimationAsset interface
+#endif
+};
diff --git a/Engine/Source/Runtime/Engine/Private/Animation/AnimBank.cpp b/Engine/Source/Runtime/Engine/Private/Animation/AnimBank.cpp
new file mode 100644
--- /dev/null
+++ b/Engine/Source/Runtime/Engine/Private/Animation/AnimBank.cpp
@@ -0,0 +1,111 @@
+// Copyright Epic Games, Inc. All Rights Reserved.
+
+#include "Animation/AnimBank.h"
+#include "UObject/FrameworkObjectVersion.h"
+#include "UObject/UE5MainStreamObjectVersion.h"
+#include "UObject/ObjectSaveContext.h"
+#include "UObject/UObjectThreadContext.h"
+#include "EngineUtils.h"
+#include "AnimationUtils.h"
+#include "Logging/MessageLog.h"
+#include "DerivedDataCacheInterface.h"
+#include "Interfaces/ITargetPlatform.h"
+#include "Interfaces/ITargetPlatformManagerModule.h"
+
+#include UE_INLINE_GENERATED_CPP_BY_NAME(AnimBank)
+
+#define LOCTEXT_NAMESPACE "AnimBank"
+
+CSV_DECLARE_CATEGORY_MODULE_EXTERN(ENGINE_API, Animation);
+
+UAnimBank::UAnimBank(const FObjectInitializer& ObjectInitializer)
+: Super(ObjectInitializer)
+{
+}
+
+void UAnimBank::Serialize(FArchive& Ar)
+{
+	LLM_SCOPE(ELLMTag::Animation);
+	
+	Ar.UsingCustomVersion(FFrameworkObjectVersion::GUID);
+	Ar.UsingCustomVersion(FUE5MainStreamObjectVersion::GUID);
+
+	Super::Serialize(Ar);
+}
+
+void UAnimBank::PostLoad()
+{
+	Super::PostLoad();
+}
+
+void UAnimBank::PreSave(FObjectPreSaveContext ObjectSaveContext)
+{
+	Super::PreSave(ObjectSaveContext);
+}
+
+#if WITH_EDITOR
+
+void UAnimBank::PostEditChangeProperty(FPropertyChangedEvent& PropertyChangedEvent)
+{
+	Super::PostEditChangeProperty(PropertyChangedEvent);
+
+	if (!IsTemplate() && !FUObjectThreadContext::Get().IsRoutingPostLoad)
+	{
+		MarkPackageDirty();
+	}
+
+	if (PropertyChangedEvent.Property)
+	{
+		if (PropertyChangedEvent.Property->GetFName() == GET_MEMBER_NAME_CHECKED(UAnimBank, Items))
+		{
+			// ...
+		}
+	}
+}
+
+bool UAnimBank::GetAllAnimationSequencesReferred(TArray<UAnimationAsset*>& AnimationAssets, bool bRecursive)
+{
+	Super::GetAllAnimationSequencesReferred(AnimationAssets, bRecursive);
+
+	for (FAnimBankItem& BankItem : Items)
+	{
+		if (BankItem.Sequence && !AnimationAssets.Contains(BankItem.Sequence))
+		{
+			BankItem.Sequence->HandleAnimReferenceCollection(AnimationAssets, bRecursive);
+		}
+	}
+
+	return AnimationAssets.Num() > 0;
+}
+
+template<class AssetType>
+inline void HandleAnimBankReferenceReplacement(AssetType*& OriginalAsset, const TMap<UAnimationAsset*, UAnimationAsset*>& ReplacementMap)
+{
+	AssetType* CacheOriginalAsset = OriginalAsset;
+	OriginalAsset = nullptr;
+
+	if (UAnimationAsset* const* ReplacementAsset = ReplacementMap.Find(CacheOriginalAsset))
+	{
+		OriginalAsset = Cast<AssetType>(*ReplacementAsset);
+	}
+}
+
+template<class AssetType>
+inline void HandleAnimBankReferenceReplacement(TObjectPtr<AssetType>& OriginalAsset, const TMap<UAnimationAsset*, UAnimationAsset*>& ReplacementMap)
+{
+	HandleAnimBankReferenceReplacement(static_cast<AssetType*&>(OriginalAsset), ReplacementMap);
+}
+
+void UAnimBank::ReplaceReferredAnimations(const TMap<UAnimationAsset*, UAnimationAsset*>& ReplacementMap)
+{
+	Super::ReplaceReferredAnimations(ReplacementMap);
+
+	for (FAnimBankItem& BankItem : Items)
+	{
+		HandleAnimBankReferenceReplacement(BankItem.Sequence, ReplacementMap);
+	}
+}
+
+#endif
+
+#undef LOCTEXT_NAMESPACE 
\ No newline at end of file
