[Engine]
* Made PhysicsEngine methods const

diff --git a/Engine/Source/Runtime/Engine/Classes/PhysicsEngine/PhysicsAsset.h b/Engine/Source/Runtime/Engine/Classes/PhysicsEngine/PhysicsAsset.h
--- a/Engine/Source/Runtime/Engine/Classes/PhysicsEngine/PhysicsAsset.h
+++ b/Engine/Source/Runtime/Engine/Classes/PhysicsEngine/PhysicsAsset.h
@@ -321,10 +321,10 @@ public:
 #endif
 
 	// @todo document
-	void GetCollisionMesh(int32 ViewIndex, FMeshElementCollector& Collector, const FReferenceSkeleton& RefSkeleton, const TArray<FTransform>& SpaceBases, const FTransform& LocalToWorld, const FVector& Scale3D);
+	void GetCollisionMesh(int32 ViewIndex, FMeshElementCollector& Collector, const FReferenceSkeleton& RefSkeleton, const TArray<FTransform>& SpaceBases, const FTransform& LocalToWorld, const FVector& Scale3D) const;
 
 	// @todo document
-	void DrawConstraints(int32 ViewIndex, FMeshElementCollector& Collector, const FReferenceSkeleton& RefSkeleton, const TArray<FTransform>& SpaceBases, const FTransform& LocalToWorld, float Scale);
+	void DrawConstraints(int32 ViewIndex, FMeshElementCollector& Collector, const FReferenceSkeleton& RefSkeleton, const TArray<FTransform>& SpaceBases, const FTransform& LocalToWorld, float Scale) const;
 
 	void GetUsedMaterials(TArray<UMaterialInterface*>& Materials);
 
diff --git a/Engine/Source/Runtime/Engine/Private/PhysicsEngine/PhysDrawing.cpp b/Engine/Source/Runtime/Engine/Private/PhysicsEngine/PhysDrawing.cpp
--- a/Engine/Source/Runtime/Engine/Private/PhysicsEngine/PhysDrawing.cpp
+++ b/Engine/Source/Runtime/Engine/Private/PhysicsEngine/PhysDrawing.cpp
@@ -1371,7 +1371,7 @@ FTransform GetSkelBoneTransform(int32 BoneIndex, const TArray<FTransform>& Space
 	}
 }
 
-void UPhysicsAsset::GetCollisionMesh(int32 ViewIndex, FMeshElementCollector& Collector, const FReferenceSkeleton& RefSkeleton, const TArray<FTransform>& SpaceBases, const FTransform& LocalToWorld, const FVector& Scale3D)
+void UPhysicsAsset::GetCollisionMesh(int32 ViewIndex, FMeshElementCollector& Collector, const FReferenceSkeleton& RefSkeleton, const TArray<FTransform>& SpaceBases, const FTransform& LocalToWorld, const FVector& Scale3D) const
 {
 	for( int32 i=0; i<SkeletalBodySetups.Num(); i++)
 	{
@@ -1398,7 +1398,7 @@ void UPhysicsAsset::GetUsedMaterials(TArray<UMaterialInterface*>& Materials)
 	}
 }
 
-void UPhysicsAsset::DrawConstraints(int32 ViewIndex, FMeshElementCollector& Collector, const FReferenceSkeleton& RefSkeleton, const TArray<FTransform>& SpaceBases, const FTransform& LocalToWorld, float Scale)
+void UPhysicsAsset::DrawConstraints(int32 ViewIndex, FMeshElementCollector& Collector, const FReferenceSkeleton& RefSkeleton, const TArray<FTransform>& SpaceBases, const FTransform& LocalToWorld, float Scale) const
 {
 	for (int32 i = 0; i < ConstraintSetup.Num(); i++)
 	{
