[World Partition] Disable code meant to only run in PIE that made sure to reload a new instance of actor levels when streaming actors from non-OFPA levels.
This is required in PIE to handle the case where a cell streams actors from a non-OFPA level, then unload the cell and reload it before a GC; in such case, it shouldn't find the old level because actors were moved into the first cell.
This was causing massive cook time degradation because we were always loading a new instance of the level for each cell, triggering expensive BP recompilation of the level script.

diff --git a/Engine/Source/Runtime/Engine/Private/WorldPartition/WorldPartitionLevelHelper.cpp b/Engine/Source/Runtime/Engine/Private/WorldPartition/WorldPartitionLevelHelper.cpp
--- a/Engine/Source/Runtime/Engine/Private/WorldPartition/WorldPartitionLevelHelper.cpp
+++ b/Engine/Source/Runtime/Engine/Private/WorldPartition/WorldPartitionLevelHelper.cpp
@@ -605,9 +605,12 @@ bool FWorldPartitionLevelHelper::LoadActorsInternal(FLoadActorsParams&& InParams
 		const uint64 ContextObjectPathNameHash = CityHash64(TCHAR_TO_ANSI(*ContextObjectPathName), ContextObjectPathName.Len());
 		PackageNameBuilder.Appendf(TEXT("_%llx"), ContextObjectPathNameHash);
 
-		// Distinguish between loading the same package after a reload between GCs (only for PIE)
-		static uint32 ContextObjectUniqueID = 0;
-		PackageNameBuilder.Appendf(TEXT("_%llx"), ContextObjectUniqueID++);
+		if (!IsRunningCommandlet())
+		{
+			// Distinguish between loading the same package after a reload between GCs (only for PIE)
+			static uint32 ContextObjectUniqueID = 0;
+			PackageNameBuilder.Appendf(TEXT("_%llx"), ContextObjectUniqueID++);
+		}
 
 		return PackageNameBuilder.ToString();
 	};
