Ensure imported materials are correctly auto-converted with Substrate.

diff --git a/Engine/Source/Runtime/Engine/Private/Materials/Material.cpp b/Engine/Source/Runtime/Engine/Private/Materials/Material.cpp
--- a/Engine/Source/Runtime/Engine/Private/Materials/Material.cpp
+++ b/Engine/Source/Runtime/Engine/Private/Materials/Material.cpp
@@ -4862,6 +4862,13 @@ void UMaterial::PostEditChangeProperty(FPropertyChangedEvent& PropertyChangedEve
 	return PostEditChangePropertyInternal(PropertyChangedEvent, EPostEditChangeEffectOnShaders::Default);
 }
 
+static bool IsMaterialAlreadyConvertedToSubstrate(const UMaterial* InMaterial, const UMaterialEditorOnlyData* InEditorOnly)
+{
+	return 
+		InEditorOnly->FrontMaterial.IsConnected() || 
+		InMaterial->HasAnyExpressionsInMaterialAndFunctionsOfType<UMaterialExpressionSubstrateBSDF>();
+}
+
 void UMaterial::PostEditChangePropertyInternal(FPropertyChangedEvent& PropertyChangedEvent, const EPostEditChangeEffectOnShaders EffectOnShaders)
 {
 	// PreEditChange is not enforced to be called before PostEditChange.
@@ -4878,6 +4885,15 @@ void UMaterial::PostEditChangePropertyInternal(FPropertyChangedEvent& PropertyCh
 
 	const UMaterialEditorOnlyData* EditorOnly = GetEditorOnlyData();
 
+	// Apply Substrate material conversion if needed.
+	// For imported materials, the conversion won't happen during the PostLoad() phase. 
+	// Usually material importers create material by code, and then calls PostEditChange(), which calls PostEditChangeProperty()
+	// This late conversion ensures that imported materials are converted properly with Substrate.
+	if (Substrate::IsSubstrateEnabled() && !IsMaterialAlreadyConvertedToSubstrate(this, EditorOnly))
+	{
+		ConvertMaterialToSubstrateMaterial();
+	}
+
 	// Check for distortion in material 
 	bUsesDistortion = RefractionMethod != RM_None;
 	bRootNodeOverridesDefaultDistortion = bUsesDistortion ? IsRefractionPinPluggedIn(EditorOnly) : false;
