Always set Lighting ShowFlag for VMI_VisualizeVirtualShadowMap to prevent missing VSM debug overlay in some cases.

#rnx

diff --git a/Engine/Source/Runtime/Engine/Private/ShowFlags.cpp b/Engine/Source/Runtime/Engine/Private/ShowFlags.cpp
--- a/Engine/Source/Runtime/Engine/Private/ShowFlags.cpp
+++ b/Engine/Source/Runtime/Engine/Private/ShowFlags.cpp
@@ -556,7 +556,8 @@ void EngineShowFlagOverride(EShowFlagInitMode ShowFlagInitMode, EViewModeIndex V
 		if( ViewModeIndex == VMI_Lit ||
 			ViewModeIndex == VMI_Lit_Wireframe ||
 			ViewModeIndex == VMI_LightingOnly ||
-			ViewModeIndex == VMI_LitLightmapDensity)
+			ViewModeIndex == VMI_LitLightmapDensity ||
+			ViewModeIndex == VMI_VisualizeVirtualShadowMap )
 		{
 			EngineShowFlags.SetLighting(true);
 		}
