[World Partition] Disable code meant to only run in PIE on non-OFPA levels when loading external actors.
This is required in PIE to handle the case where a cell streams actors from a non-OFPA level, then unload the cell and reload it before a GC; in such case, it shouldn't find the old level because actors were moved into the first cell. This will be fixed in an upcoming change by removing this logic and duplicating actors from their non-OFPA level instead of moving them into cells.
This was causing PIE degradation because we were always loading a new instance of the level for each external, triggering expensive BP recompilation of the level script.

diff --git a/Engine/Source/Runtime/Engine/Private/WorldPartition/WorldPartitionLevelHelper.cpp b/Engine/Source/Runtime/Engine/Private/WorldPartition/WorldPartitionLevelHelper.cpp
--- a/Engine/Source/Runtime/Engine/Private/WorldPartition/WorldPartitionLevelHelper.cpp
+++ b/Engine/Source/Runtime/Engine/Private/WorldPartition/WorldPartitionLevelHelper.cpp
@@ -592,7 +592,7 @@ bool FWorldPartitionLevelHelper::LoadActorsInternal(FLoadActorsParams&& InParams
 
 	// Generate a unique name to load a level instance embedded actor if there are multiple instances of this level instance and possibly across 
 	// multiple instances of the WP world:
-	auto GetContainerPackage = [](const FActorContainerID& InContainerID, const FString& InPackageName, const UObject* InContextObject) -> FName
+	auto GetContainerPackage = [](const FActorContainerID& InContainerID, const FString& InPackageName, const UObject* InContextObject, bool bUniquePackage) -> FName
 	{
 		TStringBuilder<512> PackageNameBuilder;
 		
@@ -605,7 +605,7 @@ bool FWorldPartitionLevelHelper::LoadActorsInternal(FLoadActorsParams&& InParams
 		const uint64 ContextObjectPathNameHash = CityHash64(TCHAR_TO_ANSI(*ContextObjectPathName), ContextObjectPathName.Len());
 		PackageNameBuilder.Appendf(TEXT("_%llx"), ContextObjectPathNameHash);
 
-		if (!IsRunningCommandlet())
+		if (!IsRunningCommandlet() && bUniquePackage)
 		{
 			// Distinguish between loading the same package after a reload between GCs (only for PIE)
 			static uint32 ContextObjectUniqueID = 0;
@@ -621,7 +621,7 @@ bool FWorldPartitionLevelHelper::LoadActorsInternal(FLoadActorsParams&& InParams
 
 		// Add main container context
 		LinkerInstancingContexts.Add(FActorContainerID::GetMainContainerID(), MoveTemp(InParams.InstancingContext));
-			
+
 		for (FWorldPartitionRuntimeCellObjectMapping& PackageObjectMapping : InParams.ActorPackages)
 		{
 			FLinkerInstancingContext* Context = LinkerInstancingContexts.Find(PackageObjectMapping.ContainerID);
@@ -639,7 +639,9 @@ bool FWorldPartitionLevelHelper::LoadActorsInternal(FLoadActorsParams&& InParams
 				NewContext.AddTag(ULevel::DontLoadExternalObjectsTag);
 				NewContext.AddTag(ULevel::DontLoadExternalFoldersTag);
 
-				const FName ContainerPackageInstanceName(GetContainerPackage(PackageObjectMapping.ContainerID, PackageObjectMapping.ContainerPackage.ToString(), InParams.OuterWorld));
+				// We only want unique packages for non-OFPA actors, @todo_ow: remove this and duplicate actors from non-OFPA levels instead of renaming.
+				const bool bUniquePackage = !PackageObjectMapping.Package.ToString().Contains(FPackagePath::GetExternalActorsFolderName());
+				const FName ContainerPackageInstanceName(GetContainerPackage(PackageObjectMapping.ContainerID, PackageObjectMapping.ContainerPackage.ToString(), InParams.OuterWorld, bUniquePackage));
 				NewContext.AddPackageMapping(PackageObjectMapping.ContainerPackage, ContainerPackageInstanceName);
 				Context = &NewContext;
 			}
