[NetPhysics] Trigger resimulation from desync in rotation, linear velocity and angular velocity, alongside the already implemented positional desync.
- Rename project setting bEnablePhysicsResimulation to bEnablePhysicsHistoryCapture
- Rename project setting ResimulationErrorThreshold to ResimulationErrorPositionThreshold and move to Project Settings -> Physics -> Replication -> Physics Prediction -> Resimulation Replication Settings
- Implement project settings for enabling and threshold value for triggering resimulation from XRVW, also implement API to get those settings on the physics thread.
- Implemented per actor overrides for new settings via NetworkPhysicsSettingsComponent
- Logic to compare thresholds and trigger reism in PhysicsReplication.cpp and RewindData.cpp
- Stopped caching target states from the server in rewind data if we have not yet setup the physics tick offset over the network

diff --git a/Engine/Config/BaseEngine.ini b/Engine/Config/BaseEngine.ini
--- a/Engine/Config/BaseEngine.ini
+++ b/Engine/Config/BaseEngine.ini
@@ -1308,6 +1308,9 @@ DefaultWorldPartitionHLODLayer=/Engine/Maps/Templates/HLODs/HLODLayer_Instanced.
 +FunctionRedirects=(OldName="KismetMathLibrary.MinimumAreaRectangle", NewName="MinAreaRectangle")
 +PropertyRedirects=(OldName="KismetMathLibrary.MinAreaRectangle.InVerts", NewName="InPoints")
 
++PropertyRedirects=(OldName="NetworkPhysicsSettingsResimulation.bOverrideResimulationErrorThreshold", NewName="bOverrideResimulationErrorPositionThreshold")
++PropertyRedirects=(OldName="NetworkPhysicsSettingsResimulation.ResimulationErrorThreshold", NewName="ResimulationErrorPositionThreshold")
+
 ; Moved UObject-based camera shakes and animations to EngineCameras
 +ClassRedirects=(OldName="/Script/GameplayCameras.CompositeCameraShakePattern", NewName="/Script/EngineCameras.CompositeCameraShakePattern")
 +ClassRedirects=(OldName="/Script/GameplayCameras.DefaultCameraShakeBase", NewName="/Script/EngineCameras.DefaultCameraShakeBase")
diff --git a/Engine/Source/Runtime/Engine/Classes/PhysicsEngine/PhysicsSettings.h b/Engine/Source/Runtime/Engine/Classes/PhysicsEngine/PhysicsSettings.h
--- a/Engine/Source/Runtime/Engine/Classes/PhysicsEngine/PhysicsSettings.h
+++ b/Engine/Source/Runtime/Engine/Classes/PhysicsEngine/PhysicsSettings.h
@@ -85,37 +85,95 @@ namespace ESettingsDOF
 	};
 }
 
+/** Default settings for physics replication using EPhysicsReplicationMode::Resimulation */
+USTRUCT()
+struct FPhysicsReplicationResimulationSettings
+{
+	GENERATED_BODY();
+
+	/** Enable positional error threshold to trigger resimulation */
+	UPROPERTY(config, EditDefaultsOnly, Category = "Replication", Meta = (InlineEditConditionToggle))
+	bool bEnableResimulationErrorPositionThreshold;
+
+	/** Distance in centimeters before a state discrepancy triggers a resimulation */
+	UPROPERTY(EditAnywhere, Category = "Replication", meta = (editcondition = "bEnableResimulationErrorPositionThreshold"))
+	float ResimulationErrorPositionThreshold;
+
+	/** Enable rotational error threshold to trigger resimulation */
+	UPROPERTY(config, EditDefaultsOnly, Category = "Replication", Meta = (InlineEditConditionToggle))
+	bool bEnableResimulationErrorRotationThreshold;
+
+	/** Rotation difference in degrees before a state discrepancy triggers a resimulation */
+	UPROPERTY(EditAnywhere, Category = "Replication", meta = (editcondition = "bEnableResimulationErrorRotationThreshold"))
+	float ResimulationErrorRotationThreshold;
+
+	/** Enable linear velocity error threshold to trigger resimulation */
+	UPROPERTY(config, EditDefaultsOnly, Category = "Replication", Meta = (InlineEditConditionToggle))
+	bool bEnableResimulationErrorLinearVelocityThreshold;
+
+	/** Velocity difference in centimeters / second before a state discrepancy triggers a resimulation */
+	UPROPERTY(EditAnywhere, Category = "Replication", meta = (editcondition = "bEnableResimulationErrorLinearVelocityThreshold"))
+	float ResimulationErrorLinearVelocityThreshold;
+
+	/** Enable angular velocity error threshold to trigger resimulation */
+	UPROPERTY(config, EditDefaultsOnly, Category = "Replication", Meta = (InlineEditConditionToggle))
+	bool bEnableResimulationErrorAngularVelocityThreshold;
+
+	/** Velocity degrees / second before a state discrepancy triggers a resimulation */
+	UPROPERTY(EditAnywhere, Category = "Replication", meta = (editcondition = "bEnableResimulationErrorAngularVelocityThreshold"))
+	float ResimulationErrorAngularVelocityThreshold;
+
+	FPhysicsReplicationResimulationSettings()
+		: bEnableResimulationErrorPositionThreshold(true)
+		, ResimulationErrorPositionThreshold(10.0f)
+		, bEnableResimulationErrorRotationThreshold(true)
+		, ResimulationErrorRotationThreshold(4.0f)
+		, bEnableResimulationErrorLinearVelocityThreshold(false)
+		, ResimulationErrorLinearVelocityThreshold(5.0f)
+		, bEnableResimulationErrorAngularVelocityThreshold(false)
+		, ResimulationErrorAngularVelocityThreshold(2.0f)
+	{ }
+};
+
+
 /** Physics Prediction Settings */
 USTRUCT()
 struct FPhysicsPredictionSettings
 {
 	GENERATED_BODY();
 
+	/** bEnablePhysicsResimulation has been renamed, please use bEnablePhysicsHistoryCapture*/
+	UPROPERTY(config)
+	bool bEnablePhysicsResimulation_DEPRECATED;
+
+	/** ResimulationErrorThreshold has been renamed and moved, please use FPhysicsReplicationResimulationSettings.ResimulationErrorPositionThreshold*/
+	UPROPERTY(config)
+	float ResimulationErrorThreshold_DEPRECATED;
+
 	/** Enable networked physics prediction (experimental)
-	* Note: If an AActor::PhysicsReplicationMode is set to use Resimulation this will allow physics to cache history which is required by resimulation replication.
-	* Note: This can also affect how physics is solved even when not using resimulation. 
-	* Note: Physics Prediction only works when Physics -> Framerate -> Tick Physics Async is enabled */
+	* This syncs the physics tick number between client and server and keeps it in sync via time dilation performed on the client, see APlayerController::GetPhysicsTimestamp().
+	* If an AActor::PhysicsReplicationMode is set to use Resimulation this will also enable RewindData to cache physics history on the client which is required by resimulation replication.
+	* IMPORTANT: Physics Prediction needs Physics -> Framerate -> Tick Physics Async enabled to function as intended. */
 	UPROPERTY(EditAnywhere, Category = "Replication")
 	bool bEnablePhysicsPrediction;
 
-	/** Forces the PlayerController to sync inputs as used in Physics Prediction.
-	* Only enable this if actively using a custom solution that needs this enabled for resimulation.
-	* This is automatically enabled when using the recommended NetworkPhysicsComponent on a pawn to handle Rewind / Resimulation. */
-	UPROPERTY(EditAnywhere, Category = "Replication", meta = (editcondition = "bEnablePhysicsPrediction"))
-	bool bEnablePhysicsResimulation;
-
-	/** Distance in centimeters before a state discrepancy triggers a resimulation */
-	UPROPERTY(EditAnywhere, Category = "Replication", meta = (editcondition = "bEnablePhysicsPrediction"))
-	float ResimulationErrorThreshold;
+	/** Enables FRewindData to cache physics history
+	* Note: This is not recommended for networked physics unless developing a custom resimulation solution since this starts caching physics on both client and server, 
+	 instead only enable bEnablePhysicsPrediction which will automatically enable FRewindData caching on the client if needed by the chosen replication mode. */
+	UPROPERTY(EditAnywhere, Category = "Replication")
+	bool bEnablePhysicsHistoryCapture;
 
 	/** Amount of RTT (Round Trip Time) latency for the prediction to support in milliseconds. */
-	UPROPERTY(EditAnywhere, Category = "Replication", meta = (editcondition = "bEnablePhysicsPrediction"))
+	UPROPERTY(EditAnywhere, Category = "Replication")
 	float MaxSupportedLatencyPrediction;
 
+	/** Default settings for physics replication using EPhysicsReplicationMode::Resimulation. */
+	UPROPERTY(EditAnywhere, Category = "Replication", meta = (DisplayName = "Default Resimulation Replication Settings (EPhysicsReplicationMode::Resimulation)"))
+	FPhysicsReplicationResimulationSettings ResimulationSettings;
+
 	FPhysicsPredictionSettings()
 		: bEnablePhysicsPrediction(false)
-		, bEnablePhysicsResimulation(false)
-		, ResimulationErrorThreshold(10.0)
+		, bEnablePhysicsHistoryCapture(false)
 		, MaxSupportedLatencyPrediction(1000)
 	{ }
 };
@@ -150,8 +208,8 @@ class UPhysicsSettings : public UPhysicsSettingsCore
 	UPROPERTY(Config, EditAnywhere, Category = Replication, meta = (DisplayName = "Physics Prediction (Experimental)"))
 	FPhysicsPredictionSettings PhysicsPrediction;
 
-	/** Error correction data for replicating simulated physics (rigid bodies) */
-	UPROPERTY(config, EditAnywhere, Category = Replication)
+	/** Default settings for physics replication using EPhysicsReplicationMode::Default */
+	UPROPERTY(config, EditAnywhere, Category = Replication, meta = (DisplayName = "Default Physics Replication Settings (EPhysicsReplicationMode::Default)"))
 	FRigidBodyErrorCorrection PhysicErrorCorrection;
 
 	UPROPERTY(config)
diff --git a/Engine/Source/Runtime/Engine/Private/ActorReplication.cpp b/Engine/Source/Runtime/Engine/Private/ActorReplication.cpp
--- a/Engine/Source/Runtime/Engine/Private/ActorReplication.cpp
+++ b/Engine/Source/Runtime/Engine/Private/ActorReplication.cpp
@@ -1198,5 +1198,5 @@ EPhysicsReplicationMode AActor::GetPhysicsReplicationMode()
 
 float AActor::GetResimulationThreshold() const
 {
-	return UPhysicsSettings::Get()->PhysicsPrediction.ResimulationErrorThreshold;
+	return UPhysicsSettings::Get()->PhysicsPrediction.ResimulationSettings.ResimulationErrorPositionThreshold;
 }
diff --git a/Engine/Source/Runtime/Engine/Private/PhysicsEngine/NetworkPhysicsComponent.cpp b/Engine/Source/Runtime/Engine/Private/PhysicsEngine/NetworkPhysicsComponent.cpp
--- a/Engine/Source/Runtime/Engine/Private/PhysicsEngine/NetworkPhysicsComponent.cpp
+++ b/Engine/Source/Runtime/Engine/Private/PhysicsEngine/NetworkPhysicsComponent.cpp
@@ -382,7 +382,7 @@ void UNetworkPhysicsSystem::OnWorldPostInit(UWorld* World, const UWorld::Initial
 		return;
 	}
 
-	if (UPhysicsSettings::Get()->PhysicsPrediction.bEnablePhysicsPrediction)
+	if (UPhysicsSettings::Get()->PhysicsPrediction.bEnablePhysicsPrediction || UPhysicsSettings::Get()->PhysicsPrediction.bEnablePhysicsHistoryCapture)
 	{
 		if (FPhysScene* PhysScene = World->GetPhysicsScene())
 		{
@@ -393,9 +393,8 @@ void UNetworkPhysicsSystem::OnWorldPostInit(UWorld* World, const UWorld::Initial
 					Solver->SetRewindCallback(MakeUnique<FNetworkPhysicsCallback>(World));
 				}
 
-				if (UPhysicsSettings::Get()->PhysicsPrediction.bEnablePhysicsResimulation)
+				if (UPhysicsSettings::Get()->PhysicsPrediction.bEnablePhysicsHistoryCapture)
 				{
-					// Enable RewindData from having bEnablePhysicsResimulation set
 					if (Solver->GetRewindData() == nullptr)
 					{
 						const int32 NumFrames = UPhysicsSettings::Get()->GetPhysicsHistoryCount();
diff --git a/Engine/Source/Runtime/Engine/Private/PhysicsEngine/PhysLevel.cpp b/Engine/Source/Runtime/Engine/Private/PhysicsEngine/PhysLevel.cpp
--- a/Engine/Source/Runtime/Engine/Private/PhysicsEngine/PhysLevel.cpp
+++ b/Engine/Source/Runtime/Engine/Private/PhysicsEngine/PhysLevel.cpp
@@ -49,9 +49,50 @@ public:
 		return GetSettings()->PhysicsPrediction.bEnablePhysicsPrediction;
 	}
 
+	UE_DEPRECATED(5.5, "GetResimulationErrorThreshold has been renamed, please use GetResimulationErrorPositionThreshold.")
 	virtual float GetResimulationErrorThreshold() const override
 	{
-		return GetSettings()->PhysicsPrediction.ResimulationErrorThreshold;
+		return GetResimulationErrorPositionThreshold();
+	}
+
+	virtual bool GetResimulationErrorPositionThresholdEnabled() const
+	{
+		return GetSettings()->PhysicsPrediction.ResimulationSettings.bEnableResimulationErrorPositionThreshold;
+	}
+
+	virtual float GetResimulationErrorPositionThreshold() const
+	{
+		return GetSettings()->PhysicsPrediction.ResimulationSettings.ResimulationErrorPositionThreshold;
+	}
+
+	virtual bool GetResimulationErrorRotationThresholdEnabled() const
+	{
+		return GetSettings()->PhysicsPrediction.ResimulationSettings.bEnableResimulationErrorRotationThreshold;
+	}
+
+	virtual float GetResimulationErrorRotationThreshold() const
+	{
+		return GetSettings()->PhysicsPrediction.ResimulationSettings.ResimulationErrorRotationThreshold;
+	}
+
+	virtual bool GetResimulationErrorLinearVelocityThresholdEnabled() const
+	{
+		return GetSettings()->PhysicsPrediction.ResimulationSettings.bEnableResimulationErrorLinearVelocityThreshold;
+	}
+
+	virtual float GetResimulationErrorLinearVelocityThreshold() const
+	{
+		return GetSettings()->PhysicsPrediction.ResimulationSettings.ResimulationErrorLinearVelocityThreshold;
+	}
+
+	virtual bool GetResimulationErrorAngularVelocityThresholdEnabled() const
+	{
+		return GetSettings()->PhysicsPrediction.ResimulationSettings.bEnableResimulationErrorAngularVelocityThreshold;
+	}
+
+	virtual float GetResimulationErrorAngularVelocityThreshold() const
+	{
+		return GetSettings()->PhysicsPrediction.ResimulationSettings.ResimulationErrorAngularVelocityThreshold;
 	}
 
 	virtual int32 GetPhysicsHistoryCount() const override
diff --git a/Engine/Source/Runtime/Engine/Private/PhysicsEngine/PhysicsReplication.cpp b/Engine/Source/Runtime/Engine/Private/PhysicsEngine/PhysicsReplication.cpp
--- a/Engine/Source/Runtime/Engine/Private/PhysicsEngine/PhysicsReplication.cpp
+++ b/Engine/Source/Runtime/Engine/Private/PhysicsEngine/PhysicsReplication.cpp
@@ -386,6 +386,7 @@ void FPhysicsReplication::OnTick(float DeltaSeconds, TMap<TWeakObjectPtr<UPrimit
 	}
 
 	int32 LocalFrameOffset = 0; // LocalFrame = ServerFrame + LocalFrameOffset;
+	bool LocalFrameOffsetAssigned = false;
 
 	if (UPhysicsSettings::Get()->PhysicsPrediction.bEnablePhysicsPrediction)
 	{
@@ -395,6 +396,7 @@ void FPhysicsReplication::OnTick(float DeltaSeconds, TMap<TWeakObjectPtr<UPrimit
 			{
 				if (APlayerController* PlayerController = World->GetFirstPlayerController())
 				{
+					LocalFrameOffsetAssigned = PlayerController->GetNetworkPhysicsTickOffsetAssigned();
 					LocalFrameOffset = PlayerController->GetNetworkPhysicsTickOffset();
 				}
 			}
@@ -481,8 +483,12 @@ void FPhysicsReplication::OnTick(float DeltaSeconds, TMap<TWeakObjectPtr<UPrimit
 		AsyncInputData.Proxy = nullptr;
 		AsyncInputData.RepMode = PhysicsTarget.ReplicationMode;
 		AsyncInputData.ServerFrame = PhysicsTarget.ServerFrame;
-		AsyncInputData.FrameOffset = LocalFrameOffset;
 		AsyncInputData.LatencyOneWay = PingSecondsOneWay;
+		
+		if (LocalFrameOffsetAssigned)
+		{
+			AsyncInputData.FrameOffset = LocalFrameOffset;
+		}
 
 		AsyncInput->InputData.Add(AsyncInputData);
 	}
@@ -983,6 +989,12 @@ void FPhysicsReplicationAsync::UpdateRewindDataTarget(const FPhysicsRepAsyncInpu
 		return;
 	}
 
+	// If there is no FrameOffset set then we have not synced up physics ticks with the server yet so don't cache this data
+	if (Input.FrameOffset.IsSet() == false)
+	{
+		return;
+	}
+
 	Chaos::FPBDRigidsSolver* RigidsSolver = static_cast<Chaos::FPBDRigidsSolver*>(GetSolver());
 	if (RigidsSolver == nullptr)
 	{
@@ -999,7 +1011,7 @@ void FPhysicsReplicationAsync::UpdateRewindDataTarget(const FPhysicsRepAsyncInpu
 	if (Chaos::FGeometryParticleHandle* Handle = Interface.GetParticle(Input.PhysicsObject))
 	{
 		// Cache all target states inside RewindData
-		const int32 LocalFrame = Input.ServerFrame - Input.FrameOffset;
+		const int32 LocalFrame = Input.ServerFrame - *Input.FrameOffset;
 		RewindData->SetTargetStateAtFrame(*Handle, LocalFrame, Chaos::FFrameAndPhase::EParticleHistoryPhase::PostPushData,
 			Input.TargetState.Position, Input.TargetState.Quaternion,
 			Input.TargetState.LinVel, FMath::DegreesToRadians(Input.TargetState.AngVel), (Input.TargetState.Flags & ERigidBodyFlags::Sleeping));
@@ -1092,7 +1104,7 @@ void FPhysicsReplicationAsync::UpdateAsyncTarget(const FPhysicsRepAsyncInputData
 		Target->ReceiveFrame = CurrentFrame;
 		Target->TargetState = Input.TargetState;
 		Target->RepMode = Input.RepMode;
-		Target->FrameOffset = Input.FrameOffset;
+		Target->FrameOffset = Input.FrameOffset.IsSet() ? *Input.FrameOffset : 0;
 		Target->TickCount = 0;
 		Target->AccumulatedSleepSeconds = 0.0f;
 
@@ -2006,20 +2018,51 @@ bool FPhysicsReplicationAsync::ResimulationReplication(Chaos::FPBDRigidParticleH
 	bool bClearTarget = true;
 
 	static constexpr Chaos::FFrameAndPhase::EParticleHistoryPhase RewindPhase = Chaos::FFrameAndPhase::EParticleHistoryPhase::PostPushData;
-
-	const float ResimErrorThreshold = SettingsCurrent.ResimulationSettings.GetResimulationErrorThreshold(Chaos::FPhysicsSolverBase::ResimulationErrorThreshold());
+	
+	// Get state from locally cached history for frame corresponding to received data
 	const Chaos::FGeometryParticleState PastState = RewindData->GetPastStateAtFrame(*Handle, LocalFrame, RewindPhase);
 
-	const FVector ErrorOffset = (Target.TargetState.Position - PastState.GetX());
-	const float ErrorDistance = ErrorOffset.Size();
-	const bool ShouldTriggerResim = ErrorDistance >= ResimErrorThreshold;
+	// Check which comparisons to perform to trigger resimulation from
+	const bool bCompareX = Chaos::FPhysicsSolverBase::GetResimulationErrorPositionThresholdEnabled() || SettingsCurrent.ResimulationSettings.bOverrideResimulationErrorPositionThreshold;
+	const bool bCompareR = Chaos::FPhysicsSolverBase::GetResimulationErrorRotationThresholdEnabled() || SettingsCurrent.ResimulationSettings.bOverrideResimulationErrorRotationThreshold;
+	const bool bCompareV = Chaos::FPhysicsSolverBase::GetResimulationErrorLinearVelocityThresholdEnabled() || SettingsCurrent.ResimulationSettings.bOverrideResimulationErrorLinearVelocityThreshold;
+	const bool bCompareW = Chaos::FPhysicsSolverBase::GetResimulationErrorAngularVelocityThresholdEnabled() || SettingsCurrent.ResimulationSettings.bOverrideResimulationErrorAngularVelocityThreshold;
+	bool bShouldTriggerResim = false;
+
+	// Check for positional discrepancy in Distance between client and server
+	if (bCompareX)
+	{
+		const float ResimPositionErrorThreshold = SettingsCurrent.ResimulationSettings.GetResimulationErrorPositionThreshold(Chaos::FPhysicsSolverBase::GetResimulationErrorPositionThreshold());
+		bShouldTriggerResim = Chaos::FRewindData::CheckVectorThreshold(Target.TargetState.Position, PastState.GetX(), ResimPositionErrorThreshold);
+	}
+
+	// Check for linear velocity discrepancy in Distance / s between client and server
+	if (!bShouldTriggerResim && bCompareV)
+	{
+		const float ResimLinVelocityErrorThreshold = SettingsCurrent.ResimulationSettings.GetResimulationErrorLinearVelocityThreshold(Chaos::FPhysicsSolverBase::GetResimulationErrorLinearVelocityThreshold());
+		bShouldTriggerResim = Chaos::FRewindData::CheckVectorThreshold(Target.TargetState.LinVel, PastState.GetV(), ResimLinVelocityErrorThreshold);
+	}
+
+	// Check for angular velocity discrepancy in Degrees / s between client and server
+	if (!bShouldTriggerResim && bCompareW)
+	{
+		const float ResimAngVelocityErrorThreshold = SettingsCurrent.ResimulationSettings.GetResimulationErrorAngularVelocityThreshold(Chaos::FPhysicsSolverBase::GetResimulationErrorAngularVelocityThreshold());
+		bShouldTriggerResim = Chaos::FRewindData::CheckVectorThreshold(FMath::DegreesToRadians(Target.TargetState.AngVel), PastState.GetW(), ResimAngVelocityErrorThreshold);
+	}
+
+	// Check for rotational discrepancy in Degrees between client and server
+	if (!bShouldTriggerResim && bCompareR)
+	{
+		const float ResimRotationErrorThreshold = SettingsCurrent.ResimulationSettings.GetResimulationErrorRotationThreshold(Chaos::FPhysicsSolverBase::GetResimulationErrorRotationThreshold());
+		bShouldTriggerResim = Chaos::FRewindData::CheckQuaternionThreshold(Target.TargetState.Quaternion, PastState.GetR(), ResimRotationErrorThreshold);
+	}
 
 #if !(UE_BUILD_SHIPPING || UE_BUILD_TEST)
 
 	if (Chaos::FPhysicsSolverBase::CanDebugNetworkPhysicsPrediction())
 	{
 		UE_LOG(LogTemp, Log, TEXT("Apply Rigid body state at local frame %d with offset = %d"), LocalFrame, Target.FrameOffset);
-		UE_LOG(LogTemp, Log, TEXT("Particle Position Error = %f | Should Trigger Resim = %s | Server Frame = %d | Client Frame = %d"), ErrorDistance, (ShouldTriggerResim ? TEXT("True") : TEXT("False")), Target.ServerFrame, LocalFrame);
+		UE_LOG(LogTemp, Log, TEXT("Should Trigger Resim = %s | Server Frame = %d | Client Frame = %d"), (bShouldTriggerResim ? TEXT("True") : TEXT("False")), Target.ServerFrame, LocalFrame);
 		UE_LOG(LogTemp, Log, TEXT("Particle Target Position = %s | Current Position = %s"), *Target.TargetState.Position.ToString(), *PastState.GetX().ToString());
 		UE_LOG(LogTemp, Log, TEXT("Particle Target Velocity = %s | Current Velocity = %s"), *Target.TargetState.LinVel.ToString(), *PastState.GetV().ToString());
 		UE_LOG(LogTemp, Log, TEXT("Particle Target Quaternion = %s | Current Quaternion = %s"), *Target.TargetState.Quaternion.ToString(), *PastState.GetR().ToString());
@@ -2029,7 +2072,7 @@ bool FPhysicsReplicationAsync::ResimulationReplication(Chaos::FPBDRigidParticleH
 	if (PhysicsReplicationCVars::ResimulationCVars::bDrawDebug)
 	{ 
 		static constexpr float BoxSize = 5.0f;
-		const float ColorLerp = ShouldTriggerResim ? 1.0f : 0.0f;
+		const float ColorLerp = bShouldTriggerResim ? 1.0f : 0.0f;
 		const FColor DebugColor = FLinearColor::LerpUsingHSV(FLinearColor::Green, FLinearColor::Red, ColorLerp).ToFColor(false);
 
 		Chaos::FDebugDrawQueue::GetInstance().DrawDebugBox(Target.TargetState.Position, FVector(BoxSize, BoxSize, BoxSize), Target.TargetState.Quaternion, FColor::Orange, true, CharacterMovementCVars::NetCorrectionLifetime, 0, 1.0f);
@@ -2045,7 +2088,7 @@ bool FPhysicsReplicationAsync::ResimulationReplication(Chaos::FPBDRigidParticleH
 		RigidsSolver->GetEvolution()->SetParticleObjectState(Handle, Chaos::EObjectStateType::Dynamic);
 	}
 
-	if (ShouldTriggerResim && Target.TickCount == 0 && LocalFrame > RewindData->GetBlockedResimFrame())
+	if (bShouldTriggerResim && Target.TickCount == 0 && LocalFrame > RewindData->GetBlockedResimFrame())
 	{
 		// Trigger resimulation
 		RigidsSolver->GetEvolution()->GetIslandManager().SetParticleResimFrame(Handle, LocalFrame);
@@ -2060,6 +2103,8 @@ bool FPhysicsReplicationAsync::ResimulationReplication(Chaos::FPBDRigidParticleH
 
 		if (Target.TickCount <= NumPredictedFrames && NumPredictedFrames > 0)
 		{
+			const FVector ErrorOffset = (Target.TargetState.Position - PastState.GetX());
+
 			// Positional Correction
 			const float CorrectionAmountX = SettingsCurrent.ResimulationSettings.GetPosStabilityMultiplier() / NumPredictedFrames;
 			const FVector PosDiffCorrection = ErrorOffset * CorrectionAmountX; // Same result as (ErrorOffset / NumPredictedFrames) * PosStabilityMultiplier
diff --git a/Engine/Source/Runtime/Engine/Private/PhysicsEngine/PhysicsSettings.cpp b/Engine/Source/Runtime/Engine/Private/PhysicsEngine/PhysicsSettings.cpp
--- a/Engine/Source/Runtime/Engine/Private/PhysicsEngine/PhysicsSettings.cpp
+++ b/Engine/Source/Runtime/Engine/Private/PhysicsEngine/PhysicsSettings.cpp
@@ -39,6 +39,17 @@ void UPhysicsSettings::PostInitProperties()
 	LoadSurfaceType();
 #endif
 
+	if (PhysicsPrediction.bEnablePhysicsResimulation_DEPRECATED)
+	{
+		PhysicsPrediction.bEnablePhysicsHistoryCapture = true;
+	}
+
+	if (PhysicsPrediction.ResimulationErrorThreshold_DEPRECATED > 0.0f)
+	{
+		PhysicsPrediction.ResimulationSettings.bEnableResimulationErrorPositionThreshold = true;
+		PhysicsPrediction.ResimulationSettings.ResimulationErrorPositionThreshold = PhysicsPrediction.ResimulationErrorThreshold_DEPRECATED;
+	}
+
 	if (LockedAxis_DEPRECATED == static_cast<ESettingsLockedAxis::Type>(-1))
 	{
 		LockedAxis_DEPRECATED = ESettingsLockedAxis::Invalid;
diff --git a/Engine/Source/Runtime/Engine/Public/Physics/NetworkPhysicsSettingsComponent.h b/Engine/Source/Runtime/Engine/Public/Physics/NetworkPhysicsSettingsComponent.h
--- a/Engine/Source/Runtime/Engine/Public/Physics/NetworkPhysicsSettingsComponent.h
+++ b/Engine/Source/Runtime/Engine/Public/Physics/NetworkPhysicsSettingsComponent.h
@@ -245,12 +245,44 @@ struct FNetworkPhysicsSettingsResimulation
 {
 	GENERATED_BODY()
 
+	/** Deprecated UE 5.5 - bOverrideResimulationErrorThreshold has been renamed, please use bOverrideResimulationErrorPositionThreshold*/
+	UPROPERTY(config)
+	uint32 bOverrideResimulationErrorThreshold_DEPRECATED;
+
+	/** Deprecated UE 5.5 - ResimulationErrorThreshold has been renamed, please use ResimulationErrorPositionThreshold */
+	UPROPERTY(config)
+	uint32 ResimulationErrorThreshold_DEPRECATED;
+
+	UE_DEPRECATED(5.5, "GetResimulationErrorThreshold has been renamed, please use GetResimulationErrorPositionThreshold.")
+	uint32 GetResimulationErrorThreshold(uint32 DefaultValue) { return GetResimulationErrorPositionThreshold(DefaultValue); }
+
+	UPROPERTY(config, EditDefaultsOnly, Category = "Overrides", Meta = (InlineEditConditionToggle))
+	uint32 bOverrideResimulationErrorPositionThreshold : 1;
+	// Overrides Project Settings -> Physics -> Replication -> Physics Prediction -> Resimulation Error Position Threshold -- Distance that the object is allowed to desync from the server before triggering a resimulation, within this threshold runtime correction can be performed if RuntimeCorrectionEnabled is true.
+	UPROPERTY(config, EditDefaultsOnly, Category = "Overrides", Meta = (EditCondition = "bOverrideResimulationErrorPositionThreshold"))
+	float ResimulationErrorPositionThreshold = 10.0f;
+	float GetResimulationErrorPositionThreshold(uint32 DefaultValue) { return bOverrideResimulationErrorPositionThreshold ? ResimulationErrorPositionThreshold : DefaultValue; }
+
+	UPROPERTY(config, EditDefaultsOnly, Category = "Overrides", Meta = (InlineEditConditionToggle))
+	uint32 bOverrideResimulationErrorRotationThreshold : 1;
+	// Overrides Project Settings -> Physics -> Replication -> Physics Prediction -> Resimulation Error Rotation Threshold -- Rotation difference in degrees that the object is allowed to desync from the server before triggering a resimulation, within this threshold runtime correction can be performed if RuntimeCorrectionEnabled is true.
+	UPROPERTY(config, EditDefaultsOnly, Category = "Overrides", Meta = (EditCondition = "bOverrideResimulationErrorRotationThreshold"))
+	float ResimulationErrorRotationThreshold = 4.0f;
+	float GetResimulationErrorRotationThreshold(uint32 DefaultValue) { return bOverrideResimulationErrorRotationThreshold ? ResimulationErrorRotationThreshold : DefaultValue; }
+
+	UPROPERTY(config, EditDefaultsOnly, Category = "Overrides", Meta = (InlineEditConditionToggle))
+	uint32 bOverrideResimulationErrorLinearVelocityThreshold : 1;
+	// Overrides Project Settings -> Physics -> Replication -> Physics Prediction -> Resimulation Error Linear Velocity Threshold -- Velocity difference in centimeters / second that the object is allowed to desync from the server before triggering a resimulation, within this threshold runtime correction can be performed if RuntimeCorrectionEnabled is true.
+	UPROPERTY(config, EditDefaultsOnly, Category = "Overrides", Meta = (EditCondition = "bOverrideResimulationErrorLinearVelocityThreshold"))
+	float ResimulationErrorLinearVelocityThreshold = 5.0f;
+	float GetResimulationErrorLinearVelocityThreshold(uint32 DefaultValue) { return bOverrideResimulationErrorLinearVelocityThreshold ? ResimulationErrorLinearVelocityThreshold : DefaultValue; }
+
 	UPROPERTY(config, EditDefaultsOnly, Category = "Overrides", Meta = (InlineEditConditionToggle))
-	uint32 bOverrideResimulationErrorThreshold : 1;
-	// Overrides Project Settings -> Physics -> Replication -> Physics Prediction -> Resimulation Error Threshold -- Distance that the object is allowed to desync from the server before triggering a resimulation, within this distance runtime correction can be performed if RuntimeCorrectionEnabled is true.
-	UPROPERTY(config, EditDefaultsOnly, Category = "Overrides", Meta = (EditCondition = "bOverrideResimulationErrorThreshold"))
-	uint32 ResimulationErrorThreshold = 10;
-	uint32 GetResimulationErrorThreshold(uint32 DefaultValue) { return bOverrideResimulationErrorThreshold ? ResimulationErrorThreshold : DefaultValue; }
+	uint32 bOverrideResimulationErrorAngularVelocityThreshold : 1;
+	// Overrides Project Settings -> Physics -> Replication -> Physics Prediction -> Resimulation Error Angular Velocity Threshold -- Degrees / second that the object is allowed to desync from the server before triggering a resimulation, within this threshold runtime correction can be performed if RuntimeCorrectionEnabled is true.
+	UPROPERTY(config, EditDefaultsOnly, Category = "Overrides", Meta = (EditCondition = "bOverrideResimulationErrorAngularVelocityThreshold"))
+	float ResimulationErrorAngularVelocityThreshold = 2.0f;
+	float GetResimulationErrorAngularVelocityThreshold(uint32 DefaultValue) { return bOverrideResimulationErrorAngularVelocityThreshold ? ResimulationErrorAngularVelocityThreshold : DefaultValue; }
 
 	UPROPERTY(config, EditDefaultsOnly, Category = "Overrides", Meta = (InlineEditConditionToggle))
 	uint32 bOverrideRuntimeCorrectionEnabled : 1;
diff --git a/Engine/Source/Runtime/Engine/Public/PhysicsReplication.h b/Engine/Source/Runtime/Engine/Public/PhysicsReplication.h
--- a/Engine/Source/Runtime/Engine/Public/PhysicsReplication.h
+++ b/Engine/Source/Runtime/Engine/Public/PhysicsReplication.h
@@ -65,7 +65,7 @@ struct FPhysicsRepAsyncInputData
 	TOptional<FPhysicsRepErrorCorrectionData> ErrorCorrection;
 	EPhysicsReplicationMode RepMode;
 	int32 ServerFrame;
-	int32 FrameOffset;
+	TOptional<int32> FrameOffset;
 	float LatencyOneWay;
 
 	FPhysicsRepAsyncInputData(Chaos::FConstPhysicsObjectHandle POHandle)
diff --git a/Engine/Source/Runtime/Experimental/Chaos/Private/RewindData.cpp b/Engine/Source/Runtime/Experimental/Chaos/Private/RewindData.cpp
--- a/Engine/Source/Runtime/Experimental/Chaos/Private/RewindData.cpp
+++ b/Engine/Source/Runtime/Experimental/Chaos/Private/RewindData.cpp
@@ -269,27 +269,57 @@ const int32 FRewindData::CompareTargetsToLastFrame()
 		return RewindFrame;
 	}
 
+	// TODO: Take per actor settings into consideration via NetworkPhysicsSettingsComponent
+	const bool bCompareX = Chaos::FPhysicsSolverBase::GetResimulationErrorPositionThresholdEnabled();
+	const bool bCompareR = Chaos::FPhysicsSolverBase::GetResimulationErrorRotationThresholdEnabled();
+	const bool bCompareV = Chaos::FPhysicsSolverBase::GetResimulationErrorLinearVelocityThresholdEnabled();
+	const bool bCompareW = Chaos::FPhysicsSolverBase::GetResimulationErrorAngularVelocityThresholdEnabled();
+
+	bool ShouldTriggerResim = false;
+
 	// Iterate over targets that exist for current frame
 	for (FDirtyParticleInfo& DirtyParticleInfo : DirtyParticles)
 	{
-		// TODO: Only iterate targets that are not predicted/interpolated
+		// TODO: Only iterate source target states, i.e. states that are not predicted/interpolated to fill in gaps
 		FGeometryParticleHandle* PTParticle = DirtyParticleInfo.GetObjectPtr();
 		if (PTParticle)
 		{
 			FGeometryParticleStateBase& History = DirtyParticleInfo.GetHistory();
-			if (!History.TargetPositions.IsEmpty())
+			if ((bCompareX || bCompareR) && !History.TargetPositions.IsEmpty())
 			{
 				// Compare with particle for this frame and mark resim if needed from CurrentFrame()
 				if (const FParticlePositionRotation* TargetState = History.TargetPositions.Read(FrameAndPhase, PropertiesPool))
 				{
-					if (const FParticlePositionRotation* XRState = History.ParticlePositionRotation.Read(FrameAndPhase, PropertiesPool))
+					if (const FParticlePositionRotation* PastState = History.ParticlePositionRotation.Read(FrameAndPhase, PropertiesPool))
 					{
-						const FVec3 ErrorOffset = TargetState->GetX() - XRState->GetX();
-						const bool ShouldTriggerResim = ErrorOffset.Size() >= FPhysicsSolverBase::ResimulationErrorThreshold();
+						if (bCompareX)
+						{
+							ShouldTriggerResim |= FRewindData::CheckVectorThreshold(TargetState->GetX(), PastState->GetX(), FPhysicsSolverBase::GetResimulationErrorPositionThreshold()); // TODO: Take per actor settings into consideration via NetworkPhysicsSettingsComponent
+						}
 
-						if (ShouldTriggerResim)
+						if (bCompareR)
 						{
-							RewindFrame = CurrentFrame() - 1;
+							ShouldTriggerResim |= FRewindData::CheckQuaternionThreshold(TargetState->GetR(), PastState->GetR(), FPhysicsSolverBase::GetResimulationErrorRotationThreshold()); // TODO: Take per actor settings into consideration via NetworkPhysicsSettingsComponent
+						}
+					}
+				}
+			}
+
+			if (!ShouldTriggerResim && (bCompareV || bCompareW) && !History.TargetVelocities.IsEmpty())
+			{
+				// Compare with particle for this frame and mark resim if needed from CurrentFrame()
+				if (const FParticleVelocities* TargetState = History.TargetVelocities.Read(FrameAndPhase, PropertiesPool))
+				{
+					if (const FParticleVelocities* PastState = History.Velocities.Read(FrameAndPhase, PropertiesPool))
+					{
+						if (bCompareV)
+						{
+							ShouldTriggerResim |= FRewindData::CheckVectorThreshold(TargetState->GetV(), PastState->GetV(), FPhysicsSolverBase::GetResimulationErrorLinearVelocityThreshold()); // TODO: Take per actor settings into consideration via NetworkPhysicsSettingsComponent
+						}
+
+						if (bCompareW)
+						{
+							ShouldTriggerResim |= FRewindData::CheckVectorThreshold(TargetState->GetW(), PastState->GetW(), FPhysicsSolverBase::GetResimulationErrorAngularVelocityThreshold()); // TODO: Take per actor settings into consideration via NetworkPhysicsSettingsComponent
 						}
 					}
 				}
@@ -297,9 +327,36 @@ const int32 FRewindData::CompareTargetsToLastFrame()
 		}
 	}
 
+	if (ShouldTriggerResim)
+	{
+		RewindFrame = FrameAndPhase.Frame;
+	}
+
 	return RewindFrame;
 }
 
+bool FRewindData::CheckVectorThreshold(FVec3 A, FVec3 B, float Threshold)
+{
+	const FVector Delta = A - B;
+	return Delta.Size() >= Threshold;
+}
+
+bool FRewindData::CheckQuaternionThreshold(FQuat A, FQuat B, float ThresholdDegrees)
+{
+	// Get the rotational delta between A and B 
+	const FQuat RotDelta = A * B.Inverse();
+
+	// Convert delta to angle and axis
+	float Angle;
+	FVector Axis;
+	RotDelta.ToAxisAndAngle(Axis, Angle);
+	Angle = FMath::RadiansToDegrees(FMath::UnwindRadians(Angle));
+	Angle = FMath::Abs(Angle);
+
+	return Angle >= ThresholdDegrees;
+}
+
+
 CHAOS_API bool bResimAllowRewindToResimulatedFrames = false;
 FAutoConsoleVariableRef CVarResimAllowRewindToResimulatedFrames(TEXT("p.Resim.AllowRewindToResimulatedFrames"), bResimAllowRewindToResimulatedFrames, TEXT("Allow rewinding back to a frame that was previously part of a resimulation. If a resimulation is performed between frame 100-110, allow a new resim from 105-115 if needed, else next resim will be able to start from frame 111."));
 
diff --git a/Engine/Source/Runtime/Experimental/Chaos/Public/Chaos/Framework/PhysicsSolverBase.h b/Engine/Source/Runtime/Experimental/Chaos/Public/Chaos/Framework/PhysicsSolverBase.h
--- a/Engine/Source/Runtime/Experimental/Chaos/Public/Chaos/Framework/PhysicsSolverBase.h
+++ b/Engine/Source/Runtime/Experimental/Chaos/Public/Chaos/Framework/PhysicsSolverBase.h
@@ -745,11 +745,59 @@ namespace Chaos
 			return PhysicsHistoryCount;
 		}
 
+		UE_DEPRECATED(5.5, "ResimulationErrorThreshold has been renamed, please use GetResimulationErrorPositionThreshold.")
 		static float ResimulationErrorThreshold()
 		{
-			const float ResimulationErrorThreshold = FChaosSolversModule::GetModule()->GetSettingsProvider().GetResimulationErrorThreshold();
+			const float ResimulationErrorThreshold = FChaosSolversModule::GetModule()->GetSettingsProvider().GetResimulationErrorPositionThreshold();
 			return ResimulationErrorThreshold;
+		}
+
+		static bool GetResimulationErrorPositionThresholdEnabled()
+		{
+			const bool ResimulationErrorPositionThresholdEnabled = FChaosSolversModule::GetModule()->GetSettingsProvider().GetResimulationErrorPositionThresholdEnabled();
+			return ResimulationErrorPositionThresholdEnabled;
+		}
+
+		static float GetResimulationErrorPositionThreshold()
+		{
+			const float ResimulationErrorPositionThreshold = FChaosSolversModule::GetModule()->GetSettingsProvider().GetResimulationErrorPositionThreshold();
+			return ResimulationErrorPositionThreshold;
+		}
+
+		static bool GetResimulationErrorRotationThresholdEnabled()
+		{
+			const bool ResimulationErrorRotationThresholdEnabled = FChaosSolversModule::GetModule()->GetSettingsProvider().GetResimulationErrorRotationThresholdEnabled();
+			return ResimulationErrorRotationThresholdEnabled;
+		}
 
+		static float GetResimulationErrorRotationThreshold()
+		{
+			const float ResimulationErrorRotationThreshold = FChaosSolversModule::GetModule()->GetSettingsProvider().GetResimulationErrorRotationThreshold();
+			return ResimulationErrorRotationThreshold;
+		}
+
+		static bool GetResimulationErrorLinearVelocityThresholdEnabled()
+		{
+			const bool ResimulationErrorLinearVelocityThresholdEnabled = FChaosSolversModule::GetModule()->GetSettingsProvider().GetResimulationErrorLinearVelocityThresholdEnabled();
+			return ResimulationErrorLinearVelocityThresholdEnabled;
+		}
+
+		static float GetResimulationErrorLinearVelocityThreshold()
+		{
+			const float ResimulationErrorLinearVelocityThreshold = FChaosSolversModule::GetModule()->GetSettingsProvider().GetResimulationErrorLinearVelocityThreshold();
+			return ResimulationErrorLinearVelocityThreshold;
+		}
+
+		static bool GetResimulationErrorAngularVelocityThresholdEnabled()
+		{
+			const bool ResimulationErrorAngularVelocityThresholdEnabled = FChaosSolversModule::GetModule()->GetSettingsProvider().GetResimulationErrorAngularVelocityThresholdEnabled();
+			return ResimulationErrorAngularVelocityThresholdEnabled;
+		}
+
+		static float GetResimulationErrorAngularVelocityThreshold()
+		{
+			const float ResimulationErrorAngularVelocityThreshold = FChaosSolversModule::GetModule()->GetSettingsProvider().GetResimulationErrorAngularVelocityThreshold();
+			return ResimulationErrorAngularVelocityThreshold;
 		}
 
 		/** Return the interpolation lerp in case the resim is off */
diff --git a/Engine/Source/Runtime/Experimental/Chaos/Public/ChaosSolversModule.h b/Engine/Source/Runtime/Experimental/Chaos/Public/ChaosSolversModule.h
--- a/Engine/Source/Runtime/Experimental/Chaos/Public/ChaosSolversModule.h
+++ b/Engine/Source/Runtime/Experimental/Chaos/Public/ChaosSolversModule.h
@@ -41,7 +41,48 @@ public:
 		return false;
 	}
 
+	UE_DEPRECATED(5.5, "GetResimulationErrorThreshold has been renamed, please use GetResimulationErrorPositionThreshold.")
 	virtual float GetResimulationErrorThreshold() const
+	{
+		return GetResimulationErrorPositionThreshold();
+	}
+
+	virtual bool GetResimulationErrorPositionThresholdEnabled() const
+	{
+		return false;
+	}
+
+	virtual float GetResimulationErrorPositionThreshold() const
+	{
+		return 0;
+	}
+
+	virtual bool GetResimulationErrorRotationThresholdEnabled() const
+	{
+		return false;
+	}
+
+	virtual float GetResimulationErrorRotationThreshold() const
+	{
+		return 0;
+	}
+
+	virtual bool GetResimulationErrorLinearVelocityThresholdEnabled() const
+	{
+		return false;
+	}
+
+	virtual float GetResimulationErrorLinearVelocityThreshold() const
+	{
+		return 0;
+	}
+
+	virtual bool GetResimulationErrorAngularVelocityThresholdEnabled() const
+	{
+		return false;
+	}
+
+	virtual float GetResimulationErrorAngularVelocityThreshold() const
 	{
 		return 0;
 	}
diff --git a/Engine/Source/Runtime/Experimental/Chaos/Public/RewindData.h b/Engine/Source/Runtime/Experimental/Chaos/Public/RewindData.h
--- a/Engine/Source/Runtime/Experimental/Chaos/Public/RewindData.h
+++ b/Engine/Source/Runtime/Experimental/Chaos/Public/RewindData.h
@@ -1724,6 +1724,9 @@ public:
 	* NOTE: This only happens when the client has desynced and is behind the server, so we receive server states for frames not yet simulated */
 	const int32 CHAOS_API CompareTargetsToLastFrame();
 
+	static bool CHAOS_API CheckVectorThreshold(FVec3 A, FVec3 B, float Threshold);
+	static bool CHAOS_API CheckQuaternionThreshold(FQuat A, FQuat B, float ThresholdDegrees);
+
 private:
 	friend class FPBDRigidsSolver;
 
