Expose runtime cell transformer

diff --git a/Engine/Source/Runtime/Engine/Public/WorldPartition/WorldPartitionRuntimeCellTransformer.h b/Engine/Source/Runtime/Engine/Public/WorldPartition/WorldPartitionRuntimeCellTransformer.h
--- a/Engine/Source/Runtime/Engine/Public/WorldPartition/WorldPartitionRuntimeCellTransformer.h
+++ b/Engine/Source/Runtime/Engine/Public/WorldPartition/WorldPartitionRuntimeCellTransformer.h
@@ -5,7 +5,7 @@
 #include "WorldPartitionRuntimeCellTransformer.generated.h"
 
 UCLASS(Abstract)
-class UWorldPartitionRuntimeCellTransformer : public UObject
+class ENGINE_API UWorldPartitionRuntimeCellTransformer : public UObject
 {
 	GENERATED_UCLASS_BODY()
 
