[Shaders] Wait until SCW process termination is completed when suspending workers.

This hardens the rare scenario in which the SCW process is terminated (which is async on Windows), then the output file is discarded, but the terminated process still writes the output file after it was supposed to be discarded.

#rnx

diff --git a/Engine/Source/Runtime/Engine/Private/ShaderCompiler/ShaderCompilerThreadRunnable.cpp b/Engine/Source/Runtime/Engine/Private/ShaderCompiler/ShaderCompilerThreadRunnable.cpp
--- a/Engine/Source/Runtime/Engine/Private/ShaderCompiler/ShaderCompilerThreadRunnable.cpp
+++ b/Engine/Source/Runtime/Engine/Private/ShaderCompiler/ShaderCompilerThreadRunnable.cpp
@@ -113,11 +113,18 @@ struct FShaderCompileWorkerInfo
 		TerminateWorkerProcess();
 	}
 
-	void TerminateWorkerProcess()
+	void TerminateWorkerProcess(bool bAsynchronous = false)
 	{
 		if (WorkerProcess.IsValid())
 		{
 			FPlatformProcess::TerminateProc(WorkerProcess);
+			if (!bAsynchronous)
+			{
+				while (FPlatformProcess::IsProcRunning(WorkerProcess))
+				{
+					FPlatformProcess::Sleep(0.01f);
+				}
+			}
 			FPlatformProcess::CloseProc(WorkerProcess);
 			WorkerProcess = FProcHandle();
 		}
