Iris:
* Added NetObjectFactory API
* Users can now create their own net object factories to handle discovery and instantiation of any replicated object types.
* An exemple of using this factory is to create your own replicated root objects that are not derived from the Actor class.

Iris ActorBridge:
* Created new Actor and SubObject net object factories.
* Deleted existing code handling discovery and instantiation of actors from the Bridge

ReplicationSystemTest:
* Converted Test bridge to use the new net object factory API

diff --git a/Engine/Plugins/Runtime/ReplicationSystemTestPlugin/Source/Private/ReplicationSystemTestPlugin.cpp b/Engine/Plugins/Runtime/ReplicationSystemTestPlugin/Source/Private/ReplicationSystemTestPlugin.cpp
--- a/Engine/Plugins/Runtime/ReplicationSystemTestPlugin/Source/Private/ReplicationSystemTestPlugin.cpp
+++ b/Engine/Plugins/Runtime/ReplicationSystemTestPlugin/Source/Private/ReplicationSystemTestPlugin.cpp
@@ -3,6 +3,9 @@
 #include "ReplicationSystemTestPlugin/ReplicationSystemTestPlugin.h"
 #include "Iris/ReplicationState/PropertyNetSerializerInfoRegistry.h"
 
+#include "Iris/ReplicationSystem/NetObjectFactoryRegistry.h"
+#include "Tests/ReplicationSystem/ReplicatedTestObjectFactory.h"
+
 class FReplicationSystemTestPlugin: public IReplicationSystemTestPlugin
 {
 	/** IModuleInterface implementation */
@@ -14,11 +17,14 @@ IMPLEMENT_MODULE( FReplicationSystemTestPlugin, ReplicationSystemTestPlugin )
 
 void FReplicationSystemTestPlugin::StartupModule()
 {
+	UE::Net::FNetObjectFactoryRegistry::RegisterFactory(UReplicatedTestObjectFactory::StaticClass(), UReplicatedTestObjectFactory::GetFactoryName());
 }
 
 void FReplicationSystemTestPlugin::ShutdownModule()
 {
 	// This function may be called during shutdown to clean up your module.  For modules that support dynamic reloading,
 	// we call this function before unloading the module.
+
+	UE::Net::FNetObjectFactoryRegistry::UnregisterFactory(UReplicatedTestObjectFactory::GetFactoryName());
 }
 
diff --git a/Engine/Plugins/Runtime/ReplicationSystemTestPlugin/Source/Private/Tests/ReplicationSystem/ReplicatedTestObject.cpp b/Engine/Plugins/Runtime/ReplicationSystemTestPlugin/Source/Private/Tests/ReplicationSystem/ReplicatedTestObject.cpp
--- a/Engine/Plugins/Runtime/ReplicationSystemTestPlugin/Source/Private/Tests/ReplicationSystem/ReplicatedTestObject.cpp
+++ b/Engine/Plugins/Runtime/ReplicationSystemTestPlugin/Source/Private/Tests/ReplicationSystem/ReplicatedTestObject.cpp
@@ -1,6 +1,7 @@
 // Copyright Epic Games, Inc. All Rights Reserved.
 
 #include "ReplicatedTestObject.h"
+#include "ReplicatedTestObjectFactory.h"
 #include "NetworkAutomationTest.h"
 #include "NetworkAutomationTestMacros.h"
 #include "Iris/ReplicationState/PropertyReplicationState.h"
@@ -33,6 +34,15 @@ UReplicatedTestObjectBridge::UReplicatedTestObjectBridge()
 : UObjectReplicationBridge()
 , CreatedObjectsOnNode(nullptr)
 {
+	
+}
+
+void UReplicatedTestObjectBridge::Initialize(UReplicationSystem* InReplicationSystem)
+{
+	Super::Initialize(InReplicationSystem);
+
+	ReplicatedObjectFactoryId = UE::Net::FNetObjectFactoryRegistry::GetFactoryIdFromName(UReplicatedTestObjectFactory::GetFactoryName());
+	check(ReplicatedObjectFactoryId != UE::Net::InvalidNetObjectFactoryId);
 }
 
 const UE::Net::FReplicationInstanceProtocol* UReplicatedTestObjectBridge::GetReplicationInstanceProtocol(FNetRefHandle Handle) const
@@ -51,7 +61,7 @@ UE::Net::FNetRefHandle UReplicatedTestObjectBridge::BeginReplication(UReplicated
 {
 	// Create NetRefHandle for the registered fragments
 	FRootObjectReplicationParams Params;
-	FNetRefHandle Handle = StartReplicatingRootObject(Instance, Params);
+	FNetRefHandle Handle = StartReplicatingRootObject(Instance, Params, ReplicatedObjectFactoryId);
 
 	// This is optional but typically we want to cache at least the NetRefHandle in the game instance to avoid doing map lookups to find it
 	if (Handle.IsValid())
@@ -65,7 +75,7 @@ UE::Net::FNetRefHandle UReplicatedTestObjectBridge::BeginReplication(UReplicated
 UE::Net::FNetRefHandle UReplicatedTestObjectBridge::BeginReplication(UReplicatedTestObject* Instance, const UObjectReplicationBridge::FRootObjectReplicationParams& Params)
 {
 	// Create NetRefHandle for the registered fragments
-	FNetRefHandle Handle = Super::StartReplicatingRootObject(Instance, Params);
+	FNetRefHandle Handle = Super::StartReplicatingRootObject(Instance, Params, ReplicatedObjectFactoryId);
 
 	// This is optional but typically we want to cache at least the NetRefHandle in the game instance to avoid doing map lookups to find it
 	if (Handle.IsValid())
@@ -81,7 +91,8 @@ UE::Net::FNetRefHandle UReplicatedTestObjectBridge::BeginReplication(FNetRefHand
 	check(OwnerHandle.IsValid());
 
 	// Create NetRefHandle for the registered fragments
-	FNetRefHandle Handle = Super::StartReplicatingSubObject(OwnerHandle, SubObjectInstance, InsertRelativeToSubObjectHandle, InsertionOrder);
+	const FSubObjectReplicationParams Params { .RootObjectHandle = OwnerHandle, .InsertRelativeToSubObjectHandle = InsertRelativeToSubObjectHandle, .InsertionOrder = InsertionOrder };
+	FNetRefHandle Handle = Super::StartReplicatingSubObject(SubObjectInstance, Params, ReplicatedObjectFactoryId);
 
 	if (Handle.IsValid())
 	{
@@ -98,145 +109,6 @@ void UReplicatedTestObjectBridge::EndReplication(UReplicatedTestObject* Instance
 	StopReplicatingNetObject(Instance, Flags);
 }
 
-bool UReplicatedTestObjectBridge::WriteCreationHeader(UE::Net::FNetSerializationContext& Context, FNetRefHandle Handle)
-{
-	TUniquePtr<FCreationHeader> Header(GetCreationHeader(Handle));
-	return WriteCreationHeader(Context, Header.Get());
-}
-
-TUniquePtr<UObjectReplicationBridge::FCreationHeader> UReplicatedTestObjectBridge::GetCreationHeader(FNetRefHandle Handle)
-{
-	const UObject* Object = GetReplicatedObject(Handle);
-
-	if (!ensure(Object))
-	{
-		return nullptr;
-	}
-
-	TUniquePtr<FReplicationTestObjectCreationHeader> Header(new FReplicationTestObjectCreationHeader);
-
-	Header->NumComponentsToSpawn = 0U;
-	Header->NumIrisComponentsToSpawn = 0U;
-	Header->NumDynamicComponentsToSpawn = 0U;
-	Header->NumConnectionFilteredComponentsToSpawn = 0U;
-	Header->NumObjectReferenceComponentsToSpawn = 0U;
-	Header->bForceFailCreateRemoteInstance = false;
-
-	if (const UReplicatedTestObject* ReplicatedTestObject = Cast<UReplicatedTestObject>(Object))
-	{
-		Header->bForceFailCreateRemoteInstance = ReplicatedTestObject->bForceFailToInstantiateOnRemote;
-	}
-	if (const UTestReplicatedIrisObject* TestReplicatedIrisObject = Cast<UTestReplicatedIrisObject>(Object))
-	{
-		Header->NumComponentsToSpawn = IntCastChecked<uint16>(TestReplicatedIrisObject->Components.Num());
-		Header->NumIrisComponentsToSpawn = IntCastChecked<uint16>(TestReplicatedIrisObject->IrisComponents.Num());
-		Header->NumDynamicComponentsToSpawn = IntCastChecked<uint16>(TestReplicatedIrisObject->DynamicStateComponents.Num());
-		Header->NumConnectionFilteredComponentsToSpawn = IntCastChecked<uint16>(TestReplicatedIrisObject->ConnectionFilteredComponents.Num());
-		Header->NumObjectReferenceComponentsToSpawn = IntCastChecked<uint16>(TestReplicatedIrisObject->ObjectReferenceComponents.Num());
-	}
-
-	UObject* Archetype = Object->GetArchetype();
-	if (!Archetype)
-	{
-		check(Archetype);
-		return nullptr;
-	}
-
-	Header->ArchetypeName = Archetype->GetPathName();
-
-	return Header;
-}
-
-bool UReplicatedTestObjectBridge::WriteCreationHeader(UE::Net::FNetSerializationContext& Context, const UObjectReplicationBridge::FCreationHeader* InHeader)
-{
-	if (!InHeader)
-	{
-		return false;
-	}
-
-	const FReplicationTestObjectCreationHeader* Header = static_cast<const FReplicationTestObjectCreationHeader*>(InHeader);
-
-	UE::Net::FNetBitStreamWriter& Writer = *Context.GetBitStreamWriter();
-
-	WriteString(&Writer, Header->ArchetypeName);
-	Writer.WriteBits(Header->NumComponentsToSpawn, 16);
-	Writer.WriteBits(Header->NumIrisComponentsToSpawn, 16);
-	Writer.WriteBits(Header->NumDynamicComponentsToSpawn, 16);
-	Writer.WriteBits(Header->NumConnectionFilteredComponentsToSpawn, 16);
-	Writer.WriteBits(Header->NumObjectReferenceComponentsToSpawn, 16);
-	Writer.WriteBool(Header->bForceFailCreateRemoteInstance);
-
-	return !Writer.IsOverflown();
-}
-
-TUniquePtr<UObjectReplicationBridge::FCreationHeader> UReplicatedTestObjectBridge::ReadCreationHeader(UE::Net::FNetSerializationContext& Context)
-{
-	UE::Net::FNetBitStreamReader& Reader = *Context.GetBitStreamReader();
-	TUniquePtr<FReplicationTestObjectCreationHeader> Header(new FReplicationTestObjectCreationHeader);
-
-	ReadString(&Reader, Header->ArchetypeName);
-
-	Header->NumComponentsToSpawn = Reader.ReadBits(16);
-	Header->NumIrisComponentsToSpawn = Reader.ReadBits(16);
-	Header->NumDynamicComponentsToSpawn = Reader.ReadBits(16);
-	Header->NumConnectionFilteredComponentsToSpawn = Reader.ReadBits(16);
-	Header->NumObjectReferenceComponentsToSpawn = Reader.ReadBits(16);
-	Header->bForceFailCreateRemoteInstance = Reader.ReadBool();
-
-	if (Reader.IsOverflown())
-	{
-		return nullptr;
-	}
-
-	return Header;
-}
-
-FObjectReplicationBridgeInstantiateResult UReplicatedTestObjectBridge::BeginInstantiateFromRemote(FNetRefHandle RootObjectOfSubObject, const UE::Net::FNetObjectResolveContext& ResolveContext, const FCreationHeader* InHeader)
-{
-	const FReplicationTestObjectCreationHeader* Header = static_cast<const FReplicationTestObjectCreationHeader*>(InHeader);
-
-	// Force fail to create this remote instance
-	if (Header->bForceFailCreateRemoteInstance)
-	{
-		return FObjectReplicationBridgeInstantiateResult();
-	}
-
-	UObject* ArcheType = StaticFindObject(UObject::StaticClass(), nullptr, *Header->ArchetypeName, false);
-
-	check(ArcheType);
-
-	FStaticConstructObjectParameters ConstructObjectParameters(ArcheType->GetClass());
-	UObject* CreatedObject = StaticConstructObject_Internal(ConstructObjectParameters);
-
-	if (UReplicatedTestObject* BaseTestObject = Cast<UReplicatedTestObject>(CreatedObject))
-	{
-		BaseTestObject->bIsSubObject = RootObjectOfSubObject.IsValid();
-	}
-	
-	if (UTestReplicatedIrisObject* CreatedTestObject = Cast<UTestReplicatedIrisObject>(CreatedObject))
-	{
-		UTestReplicatedIrisObject::FComponents Components;
-		Components.PropertyComponentCount = Header->NumComponentsToSpawn;
-		Components.IrisComponentCount = Header->NumIrisComponentsToSpawn;
-		Components.DynamicStateComponentCount = Header->NumDynamicComponentsToSpawn;
-		Components.ConnectionFilteredComponentCount = Header->NumConnectionFilteredComponentsToSpawn;
-		Components.ObjectReferenceComponentCount = Header->NumObjectReferenceComponentsToSpawn;
-
-		CreatedTestObject->AddComponents(Components);
-	}
-
-	// Store the object so that we can find detached/torn off instances from tests
-	if (CreatedObjectsOnNode)
-	{
-		CreatedObjectsOnNode->Add(TStrongObjectPtr<UObject>(CreatedObject));
-	}
-	
-	FObjectReplicationBridgeInstantiateResult InstantiateResult;
-	InstantiateResult.Object = CreatedObject;
-	InstantiateResult.Flags |= EReplicationBridgeCreateNetRefHandleResultFlags::AllowDestroyInstanceFromRemote;
-	return InstantiateResult;
-}
-
 void UReplicatedTestObjectBridge::EndInstantiateFromRemote(FNetRefHandle Handle)
 {
 	UReplicatedTestObject* Instance = CastChecked<UReplicatedTestObject>(GetReplicatedObject(Handle));
diff --git a/Engine/Plugins/Runtime/ReplicationSystemTestPlugin/Source/Private/Tests/ReplicationSystem/ReplicatedTestObject.h b/Engine/Plugins/Runtime/ReplicationSystemTestPlugin/Source/Private/Tests/ReplicationSystem/ReplicatedTestObject.h
--- a/Engine/Plugins/Runtime/ReplicationSystemTestPlugin/Source/Private/Tests/ReplicationSystem/ReplicatedTestObject.h
+++ b/Engine/Plugins/Runtime/ReplicationSystemTestPlugin/Source/Private/Tests/ReplicationSystem/ReplicatedTestObject.h
@@ -533,6 +533,8 @@ class UReplicatedTestObjectBridge : public UObjectReplicationBridge
 
 public:
 	UReplicatedTestObjectBridge();
+
+	virtual void Initialize(UReplicationSystem* InReplicationSystem) override;
 	
 	void SetCreatedObjectsOnNode(TArray<TStrongObjectPtr<UObject>>* InCreatedObjectsOnNode) { CreatedObjectsOnNode = InCreatedObjectsOnNode; }
 
@@ -573,38 +575,22 @@ public:
 		bool bSuppressCreateInstanceFailedEnsure;
 	};
 
-protected:
+public:
 
-	friend FSupressCreateInstanceFailedEnsureScope;
+	TArray<TStrongObjectPtr<UObject>>* CreatedObjectsOnNode;
 
-	// Type specifics for serializing creation data this will most likely be made into a separate interface to support different types of header data for different types
-	// But if we can avoid having custom data per type for instantiating remote objects as we would like to be able to fully express the state of a replicated object using the define protocol alone. 
-	// This will probably be specified as a Iris generated struct
-	struct FReplicationTestObjectCreationHeader : public FCreationHeader
-	{
-		FString ArchetypeName;
-		uint32 NumComponentsToSpawn;
-		uint32 NumIrisComponentsToSpawn;
-		uint32 NumDynamicComponentsToSpawn;
-		uint32 NumConnectionFilteredComponentsToSpawn;
-		uint32 NumObjectReferenceComponentsToSpawn;
-		bool bForceFailCreateRemoteInstance;
-	};
+protected:
 
-	virtual TUniquePtr<FCreationHeader> GetCreationHeader(FNetRefHandle Handle) override;
-	virtual bool WriteCreationHeader(UE::Net::FNetSerializationContext& Context, FNetRefHandle Handle) override;
-	virtual bool WriteCreationHeader(UE::Net::FNetSerializationContext& Context, const FCreationHeader* Header) override;
-	virtual TUniquePtr<FCreationHeader> ReadCreationHeader(UE::Net::FNetSerializationContext& Context) override;
+	friend FSupressCreateInstanceFailedEnsureScope;
 
-	virtual FObjectReplicationBridgeInstantiateResult BeginInstantiateFromRemote(FNetRefHandle RootObjectOfSubObject, const UE::Net::FNetObjectResolveContext& ResolveContext, const FCreationHeader* InHeader) override;
 	virtual void EndInstantiateFromRemote(FNetRefHandle Handle) override;
 	virtual void DestroyInstanceFromRemote(const FDestroyInstanceParams& Params) override;
 	virtual bool IsAllowedToDestroyInstance(const UObject* Instance) const override;
 
-	TArray<TStrongObjectPtr<UObject>>* CreatedObjectsOnNode;
-
 	TFunction<void(FNetRefHandle NetHandle, const UObject* ReplicatedObject, FVector& OutLocation, float& OutCullDistance)> WorldLocationUpdateFunc;
 	bool bForceFailCreateRemoteInstance = false;
+
+	UE::Net::FNetObjectFactoryId ReplicatedObjectFactoryId = UE::Net::InvalidNetObjectFactoryId;
 };
 
 extern const UE::Net::FRepTag RepTag_FakeGeneratedReplicationState_IntB;
diff --git a/Engine/Plugins/Runtime/ReplicationSystemTestPlugin/Source/Private/Tests/ReplicationSystem/ReplicatedTestObjectFactory.cpp b/Engine/Plugins/Runtime/ReplicationSystemTestPlugin/Source/Private/Tests/ReplicationSystem/ReplicatedTestObjectFactory.cpp
new file mode 100644
--- /dev/null
+++ b/Engine/Plugins/Runtime/ReplicationSystemTestPlugin/Source/Private/Tests/ReplicationSystem/ReplicatedTestObjectFactory.cpp
@@ -0,0 +1,180 @@
+// Copyright Epic Games, Inc. All Rights Reserved.
+
+#include "ReplicatedTestObjectFactory.h"
+#include "ReplicatedTestObject.h"
+
+#include "Iris/Serialization/NetBitStreamWriter.h"
+#include "Iris/Serialization/NetSerializationContext.h"
+#include "Iris/Serialization/NetBitStreamUtil.h"
+
+
+namespace UE::Net
+{
+
+//-------------------------------------------------------------------------
+// FReplicationTestCreationHeader
+//------------------------------------------------------------------------
+
+FString FReplicationTestCreationHeader::ToString() const
+{
+	return FString::Printf(TEXT("FReplicationTestCreationHeader (ProtocolId:0x%x)\n\t"
+								"ArchetypeName=%s"
+								"NumComponentsToSpawn=%u"
+								"NumIrisComponentsToSpawn=%u"
+								"NumDynamicComponentsToSpawn=%u"
+								"NumConnectionFilteredComponentsToSpawn=%u"
+								"NumObjectReferenceComponentsToSpawn=%u"
+								"bForceFailCreationRemoteInstance=%u"),
+						   		GetProtocolId(),
+						   		*ArchetypeName,
+						   		NumComponentsToSpawn,
+						   		NumIrisComponentsToSpawn,
+						   		NumDynamicComponentsToSpawn,
+						   		NumConnectionFilteredComponentsToSpawn,
+						   		NumObjectReferenceComponentsToSpawn,
+						   		bForceFailCreateRemoteInstance);
+}
+
+bool FReplicationTestCreationHeader::Serialize(const FCreationHeaderContext& Context) const
+{
+	FNetBitStreamWriter& Writer = *Context.Serialization.GetBitStreamWriter();
+
+	WriteString(&Writer, ArchetypeName);
+	Writer.WriteBits(NumComponentsToSpawn, 16);
+	Writer.WriteBits(NumIrisComponentsToSpawn, 16);
+	Writer.WriteBits(NumDynamicComponentsToSpawn, 16);
+	Writer.WriteBits(NumConnectionFilteredComponentsToSpawn, 16);
+	Writer.WriteBits(NumObjectReferenceComponentsToSpawn, 16);
+	Writer.WriteBool(bForceFailCreateRemoteInstance);
+
+	return !Writer.IsOverflown();
+}
+
+bool FReplicationTestCreationHeader::Deserialize(const FCreationHeaderContext& Context)
+{
+	FNetBitStreamReader& Reader = *Context.Serialization.GetBitStreamReader();
+
+	ReadString(&Reader, this->ArchetypeName);
+
+	this->NumComponentsToSpawn = Reader.ReadBits(16);
+	this->NumIrisComponentsToSpawn = Reader.ReadBits(16);
+	this->NumDynamicComponentsToSpawn = Reader.ReadBits(16);
+	this->NumConnectionFilteredComponentsToSpawn = Reader.ReadBits(16);
+	this->NumObjectReferenceComponentsToSpawn = Reader.ReadBits(16);
+	this->bForceFailCreateRemoteInstance = Reader.ReadBool();
+
+	return !Reader.IsOverflown();
+}
+
+} // end namespace UE::Net
+
+//------------------------------------------------------------------------
+// UReplicatedTestObjectFactory
+//------------------------------------------------------------------------
+
+FName UReplicatedTestObjectFactory::GetFactoryName()
+{
+	return TEXT("TestObjectFactory");
+}
+
+TUniquePtr<UE::Net::FNetObjectCreationHeader> UReplicatedTestObjectFactory::CreateAndFillHeader(UE::Net::FNetRefHandle Handle)
+{
+	using namespace UE::Net;
+	using namespace UE::Net::Private;
+
+	const UObject* Object = Bridge->GetReplicatedObject(Handle);
+	if (!ensure(Object))
+	{
+		return nullptr;
+	}
+
+	UObject* Archetype = Object->GetArchetype();
+	if (!Archetype)
+	{
+		check(Archetype);
+		return nullptr;
+	}
+
+	TUniquePtr<FReplicationTestCreationHeader> HeaderRef(new FReplicationTestCreationHeader);
+	FReplicationTestCreationHeader* Header = HeaderRef.Get();
+
+	if (const UReplicatedTestObject* ReplicatedTestObject = Cast<UReplicatedTestObject>(Object))
+	{
+		Header->bForceFailCreateRemoteInstance = ReplicatedTestObject->bForceFailToInstantiateOnRemote;
+	}
+
+	if (const UTestReplicatedIrisObject* TestReplicatedIrisObject = Cast<UTestReplicatedIrisObject>(Object))
+	{
+		Header->NumComponentsToSpawn = IntCastChecked<uint16>(TestReplicatedIrisObject->Components.Num());
+		Header->NumIrisComponentsToSpawn = IntCastChecked<uint16>(TestReplicatedIrisObject->IrisComponents.Num());
+		Header->NumDynamicComponentsToSpawn = IntCastChecked<uint16>(TestReplicatedIrisObject->DynamicStateComponents.Num());
+		Header->NumConnectionFilteredComponentsToSpawn = IntCastChecked<uint16>(TestReplicatedIrisObject->ConnectionFilteredComponents.Num());
+		Header->NumObjectReferenceComponentsToSpawn = IntCastChecked<uint16>(TestReplicatedIrisObject->ObjectReferenceComponents.Num());
+	}
+
+	Header->ArchetypeName = Archetype->GetPathName();
+
+	return HeaderRef;
+}
+
+bool UReplicatedTestObjectFactory::SerializeHeader(const UE::Net::FCreationHeaderContext& Context, const UE::Net::FNetObjectCreationHeader* Header)
+{
+	const UE::Net::FReplicationTestCreationHeader* TestHeader = static_cast<const UE::Net::FReplicationTestCreationHeader*>(Header);
+	return TestHeader->Serialize(Context);
+}
+
+TUniquePtr<UE::Net::FNetObjectCreationHeader> UReplicatedTestObjectFactory::CreateAndDeserializeHeader(const UE::Net::FCreationHeaderContext& Context)
+{
+	TUniquePtr<UE::Net::FReplicationTestCreationHeader> Header(new UE::Net::FReplicationTestCreationHeader);
+	Header->Deserialize(Context);
+	return Header;
+}
+
+UNetObjectFactory::FInstantiateResult UReplicatedTestObjectFactory::InstantiateReplicatedObjectFromHeader(const FInstantiateContext& Context, const UE::Net::FNetObjectCreationHeader* InHeader)
+{
+	using namespace UE::Net;
+
+	const FReplicationTestCreationHeader* Header = static_cast<const FReplicationTestCreationHeader*>(InHeader);
+
+	// Force fail to create this remote instance
+	if (Header->bForceFailCreateRemoteInstance)
+	{
+		return FInstantiateResult();
+	}
+
+	UObject* ArcheType = StaticFindObject(UObject::StaticClass(), nullptr, *Header->ArchetypeName, false);
+	check(ArcheType);
+
+	FStaticConstructObjectParameters ConstructObjectParameters(ArcheType->GetClass());
+	UObject* CreatedObject = StaticConstructObject_Internal(ConstructObjectParameters);
+
+	if (UReplicatedTestObject* BaseTestObject = Cast<UReplicatedTestObject>(CreatedObject))
+	{
+		BaseTestObject->bIsSubObject = Context.RootObjectOfSubObject.IsValid();
+	}
+	
+	if (UTestReplicatedIrisObject* CreatedTestObject = Cast<UTestReplicatedIrisObject>(CreatedObject))
+	{
+		UTestReplicatedIrisObject::FComponents Components;
+		Components.PropertyComponentCount = Header->NumComponentsToSpawn;
+		Components.IrisComponentCount = Header->NumIrisComponentsToSpawn;
+		Components.DynamicStateComponentCount = Header->NumDynamicComponentsToSpawn;
+		Components.ConnectionFilteredComponentCount = Header->NumConnectionFilteredComponentsToSpawn;
+		Components.ObjectReferenceComponentCount = Header->NumObjectReferenceComponentsToSpawn;
+
+		CreatedTestObject->AddComponents(Components);
+	}
+
+	UReplicatedTestObjectBridge* TestBridge = Cast<UReplicatedTestObjectBridge>(Bridge);
+
+	// Store the object so that we can find detached/torn off instances from tests
+	if (TestBridge->CreatedObjectsOnNode)
+	{
+		TestBridge->CreatedObjectsOnNode->Add(TStrongObjectPtr<UObject>(CreatedObject));
+	}
+	
+	FInstantiateResult InstantiateResult;
+	InstantiateResult.Instance = CreatedObject;
+	InstantiateResult.Flags |= EReplicationBridgeCreateNetRefHandleResultFlags::AllowDestroyInstanceFromRemote;
+	return InstantiateResult;
+}
diff --git a/Engine/Plugins/Runtime/ReplicationSystemTestPlugin/Source/Private/Tests/ReplicationSystem/ReplicatedTestObjectFactory.h b/Engine/Plugins/Runtime/ReplicationSystemTestPlugin/Source/Private/Tests/ReplicationSystem/ReplicatedTestObjectFactory.h
new file mode 100644
--- /dev/null
+++ b/Engine/Plugins/Runtime/ReplicationSystemTestPlugin/Source/Private/Tests/ReplicationSystem/ReplicatedTestObjectFactory.h
@@ -0,0 +1,53 @@
+// Copyright Epic Games, Inc. All Rights Reserved.
+
+#pragma once
+
+#include "Iris/ReplicationSystem/NetObjectFactory.h"
+
+#include "ReplicatedTestObjectFactory.generated.h"
+
+namespace UE::Net
+{
+
+class FReplicationTestCreationHeader : public FNetObjectCreationHeader
+{
+public:
+	FString ArchetypeName;
+	uint32 NumComponentsToSpawn = 0;
+	uint32 NumIrisComponentsToSpawn = 0;
+	uint32 NumDynamicComponentsToSpawn = 0;
+	uint32 NumConnectionFilteredComponentsToSpawn = 0;
+	uint32 NumObjectReferenceComponentsToSpawn = 0;
+	bool bForceFailCreateRemoteInstance = false;
+
+	virtual FString ToString() const override;
+
+	bool Serialize(const FCreationHeaderContext& Context) const;
+	bool Deserialize(const FCreationHeaderContext& Context);
+};
+
+} // end namespace UE::Net
+
+UCLASS()
+class UReplicatedTestObjectFactory : public UNetObjectFactory
+{
+	GENERATED_BODY()
+
+public:
+
+	static FName GetFactoryName();
+
+	virtual FInstantiateResult InstantiateReplicatedObjectFromHeader(const FInstantiateContext& Context, const UE::Net::FNetObjectCreationHeader* InHeader) override;
+
+protected:
+
+	virtual TUniquePtr<UE::Net::FNetObjectCreationHeader> CreateAndFillHeader(UE::Net::FNetRefHandle Handle) override;
+	
+	virtual bool SerializeHeader(const UE::Net::FCreationHeaderContext& Context, const UE::Net::FNetObjectCreationHeader* Header) override;
+
+	virtual TUniquePtr<UE::Net::FNetObjectCreationHeader> CreateAndDeserializeHeader(const UE::Net::FCreationHeaderContext& Context) override;
+
+};
+
+	
+
diff --git a/Engine/Source/Runtime/Engine/Classes/Engine/World.h b/Engine/Source/Runtime/Engine/Classes/Engine/World.h
--- a/Engine/Source/Runtime/Engine/Classes/Engine/World.h
+++ b/Engine/Source/Runtime/Engine/Classes/Engine/World.h
@@ -449,7 +449,7 @@ private:
 	friend class UPackageMapClient;
 
 #if UE_WITH_IRIS
-	friend class UActorReplicationBridge;
+	friend class UNetActorFactory;
 #endif // UE_WITH_IRIS
 
 	/* Is the actor remotely owned. This should only be set true by the package map when it is creating an actor on a client that was replicated from the server. */
diff --git a/Engine/Source/Runtime/Engine/Private/Net/Iris/ReplicationSystem/ActorReplicationBridge.cpp b/Engine/Source/Runtime/Engine/Private/Net/Iris/ReplicationSystem/ActorReplicationBridge.cpp
--- a/Engine/Source/Runtime/Engine/Private/Net/Iris/ReplicationSystem/ActorReplicationBridge.cpp
+++ b/Engine/Source/Runtime/Engine/Private/Net/Iris/ReplicationSystem/ActorReplicationBridge.cpp
@@ -5,25 +5,34 @@
 
 #if UE_WITH_IRIS
 
-#include "Net/Iris/ReplicationSystem/ActorReplicationBridgeInternal.h"
+#include "Net/Iris/ReplicationSystem/NetActorFactory.h"
+#include "Net/Iris/ReplicationSystem/NetSubObjectFactory.h"
+
 #include "Iris/IrisConfig.h"
 #include "Iris/IrisConstants.h"
+
 #include "Iris/Core/IrisLog.h"
 #include "Iris/Core/IrisProfiler.h"
 #include "Iris/Core/NetObjectReference.h"
 #include "Iris/Core/IrisMemoryTracker.h"
+
 #include "Iris/ReplicationSystem/NetToken.h"
 #include "Iris/ReplicationSystem/ObjectReplicationBridgeConfig.h"
 #include "Iris/ReplicationSystem/StringTokenStore.h"
 #include "Iris/ReplicationSystem/ReplicationOperations.h"
 #include "Iris/ReplicationSystem/ReplicationSystem.h"
 #include "Iris/ReplicationSystem/Filtering/NetObjectFilter.h"
+
 #include "Iris/Serialization/NetBitStreamReader.h"
 #include "Iris/Serialization/NetBitStreamUtil.h"
 #include "Iris/Serialization/NetBitStreamWriter.h"
 #include "Iris/Serialization/ObjectNetSerializer.h"
 #include "Iris/Serialization/IrisObjectReferencePackageMap.h"
+
 #include "Iris/Metrics/NetMetrics.h"
+
+#include "AnalyticsEventAttribute.h"
+
 #include "Engine/Engine.h"
 #include "Engine/EngineTypes.h"
 #include "Engine/Level.h"
@@ -31,10 +40,12 @@
 #include "Engine/NetDriver.h"
 #include "Engine/Level.h"
 #include "Engine/World.h"
-#include "AnalyticsEventAttribute.h"
+
 #include "GameFramework/Actor.h"
 #include "GameFramework/PlayerController.h"
+
 #include "HAL/LowLevelMemStats.h"
+
 #include "Net/DataBunch.h"
 #include "Net/DataChannel.h"
 #include "Net/Core/Connection/NetEnums.h"
@@ -42,8 +53,11 @@
 #include "Net/Core/Misc/NetSubObjectRegistry.h"
 #include "Net/Core/Trace/NetDebugName.h"
 #include "Net/NetSubObjectRegistryGetter.h"
+
 #include "ProfilingDebugging/AssetMetadataTrace.h"
+
 #include "Templates/Casts.h"
+
 #include "UObject/Package.h"
 
 #include <limits>
@@ -107,7 +121,8 @@ bool ShouldIncludeActorInLevelGroups(const AActor* Actor)
 
 UActorReplicationBridge::UActorReplicationBridge()
 : UObjectReplicationBridge()
-, SpawnInfoFlags(0U)
+, ActorFactoryId(UE::Net::InvalidNetObjectFactoryId)
+, SubObjectFactoryId(UE::Net::InvalidNetObjectFactoryId)
 {
 	SetInstancePreUpdateFunction(UE::Net::Private::ActorReplicationBridgePreUpdateFunction);
 	SetInstanceGetWorldObjectInfoFunction(UE::Net::Private::ActorReplicationBridgeGetActorWorldObjectInfo);
@@ -115,10 +130,18 @@ UActorReplicationBridge::UActorReplicationBridge()
 
 void UActorReplicationBridge::Initialize(UReplicationSystem* InReplicationSystem)
 {
+	using namespace UE::Net;
+
 	Super::Initialize(InReplicationSystem);
 
 	ensureMsgf(GDefaultUseSubObjectReplicationList, TEXT("Iris requires replicated actors to use registered subobjectslists. Add \n[SystemSettings]\nnet.SubObjects.DefaultUseSubObjectReplicationList=1\n to your DefaultEngine.ini"));
 
+	ActorFactoryId = FNetObjectFactoryRegistry::GetFactoryIdFromName(UNetActorFactory::GetFactoryName());
+	checkf(ActorFactoryId != InvalidNetObjectFactoryId, TEXT("UNetActorFactory with name %s was not registered"), ToCStr(UNetActorFactory::GetFactoryName().ToString()));
+
+	SubObjectFactoryId = FNetObjectFactoryRegistry::GetFactoryIdFromName(UNetSubObjectFactory::GetFactoryName());
+	checkf(SubObjectFactoryId != InvalidNetObjectFactoryId, TEXT("UNetSubObjectFactory with name %s was not registered"), ToCStr(UNetSubObjectFactory::GetFactoryName().ToString()));
+
 	{
 		auto ShouldSpatialize = [](const UClass* Class)
 		{
@@ -160,9 +183,6 @@ void UActorReplicationBridge::Initialize(UReplicationSystem* InReplicationSystem
 	}
 
 	ObjectReferencePackageMap = NewObject<UIrisObjectReferencePackageMap>();
-
-	// Get spawn info flags from cvars
-	SpawnInfoFlags = UE::Net::Private::GetActorReplicationBridgeSpawnInfoFlags();
 }
 
 UActorReplicationBridge::~UActorReplicationBridge()
@@ -279,7 +299,7 @@ UE::Net::FNetRefHandle UActorReplicationBridge::StartReplicatingActor(AActor* Ac
 	ensureMsgf(!(Actor->bAlwaysRelevant || Actor->bOnlyRelevantToOwner) || RootObjectParams.StaticPriority >= 1.0f, TEXT("Very low NetPriority %.02f for always relevant or owner relevant Actor %s. Set it to 1.0f or higher."), Actor->NetPriority, ToCStr(Actor->GetName()));
 #endif
 
-	FNetRefHandle ActorRefHandle = StartReplicatingRootObject(Actor, RootObjectParams);
+	FNetRefHandle ActorRefHandle = StartReplicatingRootObject(Actor, RootObjectParams, ActorFactoryId);
 
 	if (!ActorRefHandle.IsValid())
 	{
@@ -330,13 +350,14 @@ UE::Net::FNetRefHandle UActorReplicationBridge::StartReplicatingActor(AActor* Ac
 
 	if (ActorSubObjects.GetRegistryList().Num() != 0 || ReplicatedComponents.Num() != 0)
 	{
+		const FSubObjectReplicationParams SubObjectParams { .RootObjectHandle = ActorRefHandle };
 		// Start with the Actor's SubObjects (that is SubObjects that are not ActorComponents)
 		for (const FSubObjectRegistry::FEntry& SubObjectInfo : ActorSubObjects.GetRegistryList())
 		{
 			UObject* SubObjectToReplicate = SubObjectInfo.GetSubObject();
 			if (IsValid(SubObjectToReplicate) && SubObjectInfo.NetCondition != ELifetimeCondition::COND_Never)
 			{
-				FNetRefHandle SubObjectRefHandle = StartReplicatingSubObject(ActorRefHandle, SubObjectToReplicate);
+				FNetRefHandle SubObjectRefHandle = Super::StartReplicatingSubObject(SubObjectToReplicate, SubObjectParams, SubObjectFactoryId);
 				if (SubObjectRefHandle.IsValid() && SubObjectInfo.NetCondition != ELifetimeCondition::COND_None)
 				{
 					UObjectReplicationBridge::SetSubObjectNetCondition(SubObjectRefHandle, SubObjectInfo.NetCondition);
@@ -359,11 +380,11 @@ UE::Net::FNetRefHandle UActorReplicationBridge::StartReplicatingActor(AActor* Ac
 	return ActorRefHandle;
 }
 
-UE::Net::FNetRefHandle UActorReplicationBridge::StartReplicatingComponent(FNetRefHandle OwnerHandle, UActorComponent* SubObject)
+UE::Net::FNetRefHandle UActorReplicationBridge::StartReplicatingComponent(FNetRefHandle RootObjectHandle, UActorComponent* SubObject)
 {
 	using namespace UE::Net;
 
-	if (!OwnerHandle.IsValid())
+	if (!RootObjectHandle.IsValid())
 	{
 		return FNetRefHandle::GetInvalid();
 	}
@@ -395,7 +416,8 @@ UE::Net::FNetRefHandle UActorReplicationBridge::StartReplicatingComponent(FNetRe
 		}
 
 		// Start replicating the subobject with its owner.
-		ReplicatedComponentHandle = StartReplicatingSubObject(OwnerHandle, SubObject);
+		const FSubObjectReplicationParams Params{ .RootObjectHandle = RootObjectHandle };
+		ReplicatedComponentHandle = Super::StartReplicatingSubObject(SubObject, Params, SubObjectFactoryId);
 	}
 
 	if (!ReplicatedComponentHandle.IsValid())
@@ -411,15 +433,24 @@ UE::Net::FNetRefHandle UActorReplicationBridge::StartReplicatingComponent(FNetRe
 	}
 
 	// Begin replication for any SubObjects registered by the component
-	for (const FSubObjectRegistry::FEntry& SubObjectInfo : RepComponentInfo->SubObjects.GetRegistryList())
 	{
-		UObject* SubObjectToReplicate = SubObjectInfo.GetSubObject();
-		if (IsValid(SubObjectToReplicate) && SubObjectInfo.NetCondition != ELifetimeCondition::COND_Never)
+		const FSubObjectReplicationParams Params
+		{
+			.RootObjectHandle = RootObjectHandle,
+			.InsertRelativeToSubObjectHandle = ReplicatedComponentHandle,
+			.InsertionOrder = ESubObjectInsertionOrder::ReplicateWith
+		};
+	
+		for (const FSubObjectRegistry::FEntry& SubObjectInfo : RepComponentInfo->SubObjects.GetRegistryList())
 		{
-			FNetRefHandle SubObjectHandle = StartReplicatingSubObject(OwnerHandle, SubObjectToReplicate, ReplicatedComponentHandle, UReplicationBridge::ESubObjectInsertionOrder::ReplicateWith);
-			if (SubObjectHandle.IsValid() && SubObjectInfo.NetCondition != ELifetimeCondition::COND_None)
+			UObject* SubObjectToReplicate = SubObjectInfo.GetSubObject();
+			if (IsValid(SubObjectToReplicate) && SubObjectInfo.NetCondition != ELifetimeCondition::COND_Never)
 			{
-				SetSubObjectNetCondition(SubObjectHandle, SubObjectInfo.NetCondition);
+				FNetRefHandle SubObjectHandle = Super::StartReplicatingSubObject(SubObjectToReplicate, Params, SubObjectFactoryId);
+				if (SubObjectHandle.IsValid() && SubObjectInfo.NetCondition != ELifetimeCondition::COND_None)
+				{
+					SetSubObjectNetCondition(SubObjectHandle, SubObjectInfo.NetCondition);
+				}
 			}
 		}
 	}
@@ -427,6 +458,11 @@ UE::Net::FNetRefHandle UActorReplicationBridge::StartReplicatingComponent(FNetRe
 	return ReplicatedComponentHandle;
 }
 
+UE::Net::FNetRefHandle UActorReplicationBridge::StartReplicatingSubObject(UObject* SubObject, const FSubObjectReplicationParams& Params)
+{
+	return Super::StartReplicatingSubObject(SubObject, Params, SubObjectFactoryId);
+}
+
 void UActorReplicationBridge::StopReplicatingActor(AActor* Actor, EEndPlayReason::Type EndPlayReason)
 {
 	using namespace UE::Net;
@@ -486,385 +522,6 @@ void UActorReplicationBridge::StopReplicatingComponent(UActorComponent* ActorCom
 	}
 }
 
-bool UActorReplicationBridge::WriteCreationHeader(UE::Net::FNetSerializationContext& Context, FNetRefHandle Handle)
-{
-	using namespace UE::Net;
-	using namespace UE::Net::Private;
-
-	FNetBitStreamWriter* Writer = Context.GetBitStreamWriter();
-	ensure(IsReplicatedHandle(Handle));
-	const UObject* Object = GetReplicatedObject(Handle);
-	if (const AActor* Actor = Cast<AActor>(Object))
-	{
-		// Get Header
-		FActorCreationHeader Header;
-		GetActorCreationHeader(Actor, Header);
-	
-		// Serialize the data
-		// Indicate that this is an actor
-		Writer->WriteBool(true);
-		WriteActorCreationHeader(Context, Header, SpawnInfoFlags);
-
-		return !Writer->IsOverflown();
-	}
-	else if (Object)
-	{
-		const UObject* RootObject = GetReplicatedObject(GetRootObjectOfSubObject(Handle));
-
-		// Get Header
-		FSubObjectCreationHeader Header;
-		GetSubObjectCreationHeader(Object, RootObject, Header);
-
-		// Serialize the data
-		// Indicate that this is a SubObject
-		Writer->WriteBool(false);
-		WriteSubObjectCreationHeader(Context, Header);
-
-		return !Writer->IsOverflown();
-	}
-
-	ensureMsgf(false, TEXT("UActorReplicationBridge::WriteCreationHeader Failed to write creationHeader for NetRefHandle (Id=%u) %s"), Handle.GetId(), ToCStr(GetReplicationSystem()->GetDebugName(Handle)));
-
-	return false;
-}
-
-TUniquePtr<UObjectReplicationBridge::FCreationHeader> UActorReplicationBridge::GetCreationHeader(FNetRefHandle Handle)
-{
-	using namespace UE::Net;
-	using namespace UE::Net::Private;
-
-	if (!ensure(IsReplicatedHandle(Handle)))
-	{
-		return nullptr;
-	}
-
-	const UObject* Object = GetReplicatedObject(Handle);
-	if (const AActor* Actor = Cast<AActor>(Object))
-	{
-		// Get Header
-		TUniquePtr<FActorCreationHeader> Header(new FActorCreationHeader);
-		GetActorCreationHeader(Actor, *(Header.Get()));
-	
-		return Header;
-	}
-	else if (Object)
-	{
-		const UObject* RootObject = GetReplicatedObject(GetRootObjectOfSubObject(Handle));
-
-		// Get Header
-		TUniquePtr<FSubObjectCreationHeader> Header(new FSubObjectCreationHeader);
-		GetSubObjectCreationHeader(Object, RootObject, *(Header.Get()));
-
-		return Header;
-	}
-
-	ensureMsgf(false, TEXT("UActorReplicationBridge::GetCreationHeader Failed to get creationHeader for NetRefHandle (Id=%u) %s"), Handle.GetId(), ToCStr(GetReplicationSystem()->GetDebugName(Handle)));
-
-	return nullptr;
-}
-
-bool UActorReplicationBridge::WriteCreationHeader(UE::Net::FNetSerializationContext& Context, const FCreationHeader* InHeader)
-{
-	using namespace UE::Net;
-	using namespace UE::Net::Private;
-
-	FNetBitStreamWriter* Writer = Context.GetBitStreamWriter();
-
-	if (!InHeader)
-	{
-		return false;
-	}
-
-	const FActorReplicationBridgeCreationHeader* BridgeHeader = static_cast<const FActorReplicationBridgeCreationHeader*>(InHeader);
-
-	if (BridgeHeader->bIsActor)
-	{
-		const FActorCreationHeader* Header = static_cast<const FActorCreationHeader*>(BridgeHeader);
-
-		// Serialize the data
-		// Indicate that this is an actor
-		Writer->WriteBool(true);
-		WriteActorCreationHeader(Context, *Header, SpawnInfoFlags);
-
-		return !Writer->IsOverflown();
-	}
-	else
-	{
-		const FSubObjectCreationHeader* Header = static_cast<const FSubObjectCreationHeader*>(BridgeHeader);
-
-		// Serialize the data
-		// Indicate that this is a SubObject
-		Writer->WriteBool(false);
-		WriteSubObjectCreationHeader(Context, *Header);
-
-		return !Writer->IsOverflown();
-	}
-}
-
-TUniquePtr<UObjectReplicationBridge::FCreationHeader> UActorReplicationBridge::ReadCreationHeader(UE::Net::FNetSerializationContext& Context)
-{
-	using namespace UE::Net;
-	using namespace UE::Net::Private;
-
-	FNetBitStreamReader* Reader = Context.GetBitStreamReader();
-
-	if (const bool bIsActor = Reader->ReadBool())
-	{
-		TUniquePtr<FActorCreationHeader> Header(new FActorCreationHeader);
-		ReadActorCreationHeader(Context, *(Header.Get()));
-
-		if (!Reader->IsOverflown())
-		{
-			return Header;
-		}
-	}
-	else
-	{
-		TUniquePtr<FSubObjectCreationHeader> Header(new FSubObjectCreationHeader);
-		ReadSubObjectCreationHeader(Context, *(Header.Get()));
-
-		if (!Reader->IsOverflown())
-		{
-			return Header;
-		}
-	}
-
-	return nullptr;
-}
-
-FObjectReplicationBridgeInstantiateResult UActorReplicationBridge::BeginInstantiateFromRemote(FNetRefHandle RootObjectOfSubObject, const UE::Net::FNetObjectResolveContext& ResolveContext, const UObjectReplicationBridge::FCreationHeader* InHeader)
-{
-	LLM_SCOPE(ELLMTag::EngineMisc);
-
-	using namespace UE::Net::Private;
-
-	IRIS_PROFILER_SCOPE(ActorReplicationBridge_OnBeginInstantiateFromRemote);
-
-	FObjectReplicationBridgeInstantiateResult InstantiateResult;
-
-	const FActorReplicationBridgeCreationHeader* BridgeHeader = static_cast<const FActorReplicationBridgeCreationHeader*>(InHeader);
-
-	if (BridgeHeader->bIsActor)
-	{
-		const FActorCreationHeader* Header = static_cast<const FActorCreationHeader*>(InHeader);
-
-		// Spawn actor
-		// Try to instantiate the actor (or find it if it already exists)
-		if (Header->bIsDynamic)
-		{
-			// Find archetype
-			UObject* Archetype = ResolveObjectReference(Header->ArchetypeReference, ResolveContext);
-
-			// Find level
-			UObject* Level = nullptr;
-			if (!Header->bUsePersistentLevel)
-			{
-				Level = ResolveObjectReference(Header->LevelReference, ResolveContext);
-			}
-
-			if (Archetype)
-			{
-				LLM_SCOPE_DYNAMIC_STAT_OBJECTPATH(Archetype->GetPackage(), ELLMTagSet::Assets);
-				LLM_SCOPE_DYNAMIC_STAT_OBJECTPATH(Archetype->GetClass(), ELLMTagSet::AssetClasses);
-				UE_TRACE_METADATA_SCOPE_ASSET(Archetype, Archetype->GetClass());
-				// For streaming levels, it's possible that the owning level has been made not-visible but is still loaded.
-				// In that case, the level will still be found but the owning world will be invalid.
-				// If that happens, wait to spawn the Actor until the next time the level is streamed in.
-				// At that point, the Server should resend any dynamic Actors.
-				ULevel* SpawnLevel = Cast<ULevel>(Level);
-				check(SpawnLevel == nullptr || SpawnLevel->GetWorld() != nullptr);
-
-				FActorSpawnParameters SpawnInfo;
-				SpawnInfo.Template = Cast<AActor>(Archetype);
-				SpawnInfo.OverrideLevel = SpawnLevel;
-				SpawnInfo.SpawnCollisionHandlingOverride = ESpawnActorCollisionHandlingMethod::AlwaysSpawn;
-				SpawnInfo.bRemoteOwned = true;
-				SpawnInfo.bNoFail = true;
-
-				UWorld* World = NetDriver->GetWorld();
-				FVector SpawnLocation = FRepMovement::RebaseOntoLocalOrigin(Header->SpawnInfo.Location, World->OriginLocation);
-		
-				AActor* Actor = World->SpawnActorAbsolute(Archetype->GetClass(), FTransform(Header->SpawnInfo.Rotation, SpawnLocation), SpawnInfo);
-
-				// For Iris we expect that we will be able to spawn the actor as streaming always is controlled from server
-				if (ensure(Actor))
-				{
-					const FActorReplicationBridgeSpawnInfo& DefaultSpawnInfo = GetDefaultActorReplicationBridgeSpawnInfo();
-					static constexpr float Epsilon = UE_KINDA_SMALL_NUMBER;
-
-					// Set Velocity if it differs from Default
-					if (!Header->SpawnInfo.Velocity.Equals(DefaultSpawnInfo.Velocity, Epsilon))
-					{
-						Actor->PostNetReceiveVelocity(Header->SpawnInfo.Velocity);
-					}
-					
-					// Set Scale if it differs from Default
-					if (!Header->SpawnInfo.Scale.Equals(DefaultSpawnInfo.Scale, Epsilon))
-					{
-						Actor->SetActorRelativeScale3D(Header->SpawnInfo.Scale);
-					}
-
-					UE_LOG_ACTORREPLICATIONBRIDGE(Verbose, TEXT("OnBeginInstantiateFromRemote Spawned Actor %s"), *Actor->GetPathName());
-				}
-
-				InstantiateResult.Object = Actor;
-				if (NetDriver->ShouldClientDestroyActor(Actor))
-				{
-					InstantiateResult.Flags |= EReplicationBridgeCreateNetRefHandleResultFlags::AllowDestroyInstanceFromRemote;
-				}
-				return InstantiateResult;
-			}
-			else
-			{
-				UE_LOG_ACTORREPLICATIONBRIDGE(Error, TEXT("OnBeginInstantiateFromRemote Unable to spawn object, failed to resolve archetype with %s"), *DescribeObjectReference(Header->ArchetypeReference, ResolveContext));
-			}
-		}
-		else
-		{
-			if (UObject* Object = ResolveObjectReference(Header->ObjectReference, ResolveContext))
-			{
-				UE_LOG_ACTORREPLICATIONBRIDGE(Verbose, TEXT("OnBeginInstantiateFromRemote Found static Actor using path %s"), ToCStr(Object->GetPathName()));
-				InstantiateResult.Object = Object;
-				if (NetDriver->ShouldClientDestroyActor(Cast<AActor>(Object)))
-				{
-					InstantiateResult.Flags |= EReplicationBridgeCreateNetRefHandleResultFlags::AllowDestroyInstanceFromRemote;
-				}
-				return InstantiateResult;
-			}
-			else
-			{
-				UE_LOG_ACTORREPLICATIONBRIDGE(Error, TEXT("OnBeginInstantiateFromRemote Failed to find Resolve ObjectReference for static Actor %s"), *DescribeObjectReference(Header->ObjectReference, ResolveContext));
-			}
-		}
-	}
-	else
-	{
-		const FSubObjectCreationHeader* Header = static_cast<const FSubObjectCreationHeader*>(InHeader);
-
-		// Spawn sub object
-		if (Header->bIsDynamic)
-		{
-			// Resolve root object
-			UObject* RootObject = GetReplicatedObject(RootObjectOfSubObject);
-			AActor* RootActor = CastChecked<AActor>(RootObject);
-			
-			UObject* SubObj = nullptr;
-
-			if (Header->bIsNameStableForNetworking)
-			{
-				// Resolve by finding object relative to owner. We do not allow this object to be destroyed.
-				SubObj = ResolveObjectReference(Header->ObjectReference, ResolveContext);
-
-				if (!SubObj)
-				{
-					UE_LOG_ACTORREPLICATIONBRIDGE(Error, TEXT("BeginInstantiateFromRemote Failed to find stable name reference of dynamic SubObject %s, Owner %s, RootObject %s"), *DescribeObjectReference(Header->ObjectReference, ResolveContext), *RootObjectOfSubObject.ToString(), *GetPathNameSafe(RootActor));
-				}
-			}
-			else
-			{
-				// Find the proper Outer
-				UObject* OuterObject = nullptr;
-				if (Header->bOuterIsTransientLevel)
-				{
-					OuterObject = GetTransientPackage();
-				}
-				else if (Header->bOuterIsRootObject)
-				{
-					OuterObject = RootActor;
-				}				
-				else
-				{
-					OuterObject = ResolveObjectReference(Header->OuterReference, ResolveContext);
-
-					if (!OuterObject)
-					{
-						UE_LOG_ACTORREPLICATIONBRIDGE(Error, TEXT("BeginInstantiateFromRemote Failed to find Outer %s for dynamic subobject %s"), *DescribeObjectReference(Header->OuterReference, ResolveContext), *DescribeObjectReference(Header->ObjectReference, ResolveContext))
-
-						// Fallback to the rootobject instead
-						OuterObject = RootActor;
-					}
-				}
-
-				// We need to spawn the subobject
-				UObject* SubObjClassObj = ResolveObjectReference(Header->ObjectClassReference, ResolveContext);
-				UClass * SubObjClass = Cast<UClass>(SubObjClassObj);
-
-				// Try to spawn SubObject
- 				SubObj = NewObject<UObject>(OuterObject, SubObjClass);
-
-				// Sanity check some things
-				checkf(SubObj != nullptr, TEXT("UActorReplicationBridge::BeginInstantiateFromRemote: Subobject is NULL after instantiating. Class: %s, Outer %s, Actor %s"), *GetNameSafe(SubObjClass), *GetNameSafe(OuterObject), *GetNameSafe(RootActor));
-				checkf(!OuterObject || SubObj->IsIn(OuterObject), TEXT("UActorReplicationBridge::BeginInstantiateFromRemote: Subobject is not in Outer. SubObject: %s, Outer %s, Actor %s"), *SubObj->GetName(), *GetNameSafe(OuterObject), *GetNameSafe(RootActor));
-				checkf(Cast<AActor>(SubObj) == nullptr, TEXT("UActorReplicationBridge::BeginInstantiateFromRemote: Subobject is an Actor. SubObject: %s, Outer %s, Actor %s"), *SubObj->GetName(), *GetNameSafe(OuterObject), *GetNameSafe(RootActor));
-
-				// We must defer call OnSubObjectCreatedFromReplication after the state has been applied to the owning actor in order to behave like old system.
-				InstantiateResult.Flags |= EReplicationBridgeCreateNetRefHandleResultFlags::ShouldCallSubObjectCreatedFromReplication;
-
-				// Created objects may be destroyed.
-				InstantiateResult.Flags |= EReplicationBridgeCreateNetRefHandleResultFlags::AllowDestroyInstanceFromRemote;
-			}
-			
-			InstantiateResult.Object = SubObj;
-			return InstantiateResult;
-		}
-		else
-		{
-			// try to find the object by path
-			if (UObject* Object = ResolveObjectReference(Header->ObjectReference, ResolveContext))
-			{
-				UE_LOG_ACTORREPLICATIONBRIDGE(Verbose, TEXT("BeginInstantiateFromRemote Found static SubObject using path %s"), ToCStr(Object->GetPathName()));
-				// We do not allow this object to be destroyed.
-				InstantiateResult.Object = Object;
-				return InstantiateResult;
-			}
-			else
-			{
-				UE_LOG_ACTORREPLICATIONBRIDGE(Error, TEXT("BeginInstantiateFromRemote Failed to find Resolve SubObjectReference for static SubObject %s, Owner %s (%s)"), 
-				*DescribeObjectReference(Header->ObjectReference, ResolveContext), *RootObjectOfSubObject.ToString(), *GetPathNameSafe(GetReplicatedObject(RootObjectOfSubObject)));
-			}			
-		}
-	}
-
-	return InstantiateResult;
-}
-
-bool UActorReplicationBridge::OnInstantiatedFromRemote(UObject* Instance, const UObjectReplicationBridge::FCreationHeader* InHeader, uint32 ConnectionId) const
-{
-	LLM_SCOPE(ELLMTag::EngineMisc);
-
-	using namespace UE::Net::Private;
-
-	const FActorReplicationBridgeCreationHeader* BridgeHeader = static_cast<const FActorReplicationBridgeCreationHeader*>(InHeader);
-
-	if (BridgeHeader->bIsActor)
-	{
-		const FActorCreationHeader* Header = static_cast<const FActorCreationHeader*>(InHeader);
-		
-		AActor* Actor = CastChecked<AActor>(Instance);
-		if (Actor)
-		{
-			// OnActorChannelOpen
-			{
-				UNetConnection* Connection = NetDriver->GetConnectionById(ConnectionId);
-				FInBunch Bunch(Connection, const_cast<uint8*>(Header->CustomCreationData.GetData()), Header->CustomCreationDataBitCount);
-				Actor->OnActorChannelOpen(Bunch, Connection);
-
-				if (Bunch.IsError() || Bunch.GetBitsLeft() != 0)
-				{
-					// $IRIS TODO Report an intelligent error
-					check(false);
-					return false;
-				}
-			}
-
-			// Wake up from dormancy. This is important for client replays.
-			WakeUpObjectInstantiatedFromRemote(Actor);
-		}
-	}
-	
-	return true;
-}
-
 void UActorReplicationBridge::EndInstantiateFromRemote(FNetRefHandle Handle)
 {
 	if (AActor* Actor = Cast<AActor>(GetReplicatedObject(Handle)))
@@ -876,6 +533,7 @@ void UActorReplicationBridge::EndInstantiateFromRemote(FNetRefHandle Handle)
 	}
 }
 
+//$IRIS todo: move this to be a factory callback
 void UActorReplicationBridge::OnSubObjectCreatedFromReplication(FNetRefHandle SubObjectHandle)
 {
 	AActor* RootObject = Cast<AActor>(GetReplicatedObject(GetRootObjectOfSubObject(SubObjectHandle)));
@@ -1034,144 +692,6 @@ void UActorReplicationBridge::OnMaxTickRateChanged(UNetDriver* InNetDriver, int3
 	ReinitPollFrequency();
 }
 
-void UActorReplicationBridge::GetActorCreationHeader(const AActor* Actor, UE::Net::Private::FActorCreationHeader& Header) const
-{
-	Header.bIsActor = true;
-	
-	UE::Net::FNetObjectReference ActorReference = GetOrCreateObjectReference(Actor);
-
-	Header.bIsDynamic = ActorReference.GetRefHandle().IsDynamic();
-
-	// Dynamic actor?
-	if (Header.bIsDynamic)
-	{
-		// Get creation data
-
-		// This is more or less a straight copy from ClientPackageMap and needs to be updated accordingly
-		UObject* Archetype = nullptr;
-		UObject* ActorLevel = nullptr;
-
-		// ChildActor's need to be spawned from the ChildActorTemplate otherwise any non-replicated 
-		// customized properties will be incorrect on the Client.
-		if (UChildActorComponent* CAC = Actor->GetParentComponent())
-		{
-			Archetype = CAC->GetSpawnableChildActorTemplate();
-		}
-		if (Archetype == nullptr)
-		{
-			Archetype = Actor->GetArchetype();
-		}
-		ActorLevel = Actor->GetLevel();
-
-		check(Archetype != nullptr);
-		check(Actor->NeedsLoadForClient());			// We have no business sending this unless the client can load
-		check(Archetype->NeedsLoadForClient());		// We have no business sending this unless the client can load
-
-		// Fill in Header
-		Header.ArchetypeReference = GetOrCreateObjectReference(Archetype);
-		Header.bUsePersistentLevel = (UE::Net::Private::SerializeNewActorOverrideLevel == 0) || (NetDriver->GetWorld()->PersistentLevel == ActorLevel);
-		if (!Header.bUsePersistentLevel)
-		{
-			Header.LevelReference = GetOrCreateObjectReference(ActorLevel);
-		}
-
-		if (USceneComponent* RootComponent = Actor->GetRootComponent())
-		{
-			Header.SpawnInfo.Location = FRepMovement::RebaseOntoZeroOrigin(Actor->GetActorLocation(), Actor);
-			Header.SpawnInfo.Rotation = Actor->GetActorRotation();
-			Header.SpawnInfo.Scale = Actor->GetActorScale();
-			FVector Scale = Actor->GetActorScale();
-
-			if (USceneComponent* AttachParent = RootComponent->GetAttachParent())
-			{
-				// If this actor is attached, when the scale is serialized on the client, the attach parent property won't be set yet.
-				// USceneComponent::SetWorldScale3D (which got called by AActor::SetActorScale3D, which we used to do but no longer).
-				// would perform this transformation so that what is sent is relative to the parent. If we don't do this, we will
-				// apply the world scale on the client, which will then get applied a second time when the attach parent property is received.
-				FTransform ParentToWorld = AttachParent->GetSocketTransform(RootComponent->GetAttachSocketName());
-				Scale = Scale * ParentToWorld.GetSafeScaleReciprocal(ParentToWorld.GetScale3D());
-			}
-
-			Header.SpawnInfo.Scale = Scale;
-			Header.SpawnInfo.Velocity = Actor->GetVelocity();
-		}
-		else
-		{
-			Header.SpawnInfo = UE::Net::Private::GetDefaultActorReplicationBridgeSpawnInfo();
-		}
-	}
-	else // Refer by path
-	{
-		Header.ObjectReference = ActorReference;
-	}
-
-	// Custom actor creation data
-	{
-		constexpr int64 MaxBunchBits = 1024;
-		FOutBunch Bunch(MaxBunchBits);
-		const_cast<AActor*>(Actor)->OnSerializeNewActor(Bunch);
-		Header.CustomCreationDataBitCount = Bunch.GetNumBits();
-		if (Header.CustomCreationDataBitCount > 0)
-		{
-			check(Header.CustomCreationDataBitCount == Bunch.GetNumBits());
-			Header.CustomCreationData.SetNumZeroed(Align(Bunch.GetNumBytes(), 4));
-			FMemory::Memcpy(Header.CustomCreationData.GetData(), Bunch.GetData(), Bunch.GetNumBytes());
-		}
-	}
-}
-
-void UActorReplicationBridge::GetSubObjectCreationHeader(const UObject* Object, const UObject* RootObject, UE::Net::Private::FSubObjectCreationHeader& Header) const
-{
-	using namespace UE::Net;
-	using namespace UE::Net::Private;
-
-	// SubObject
-	FNetObjectReference ObjectRef = GetOrCreateObjectReference(Object);
-
-	// It's Outer
-	UObject* OuterObject = Object->GetOuter();
-
-	Header.bIsActor = false;
-	Header.bOuterIsTransientLevel = false;
-	Header.bOuterIsRootObject = false;
-	Header.bIsDynamic = ObjectRef.GetRefHandle().IsDynamic();
-	Header.bUsePersistentLevel = false;
-	Header.bIsNameStableForNetworking = Object->IsNameStableForNetworking();
-
-	if (Header.bIsDynamic)
-	{
-		check(Object->NeedsLoadForClient() );				// We have no business sending this unless the client can load
-		check(Object->GetClass()->NeedsLoadForClient());	// We have no business sending this unless the client can load
-
-		// We need to spawn this object on the receiving end, so we include the class
-		if (!Header.bIsNameStableForNetworking)
-		{
-			Header.ObjectClassReference = GetOrCreateObjectReference(Object->GetClass());
-
-			// Find the Outer
-			if (OuterObject == GetTransientPackage())
-			{
-				Header.bOuterIsTransientLevel = true;
-			}
-			else if (OuterObject == RootObject)
-			{
-				Header.bOuterIsRootObject = true;
-			}
-			else
-			{
-				Header.OuterReference = GetOrCreateObjectReference(OuterObject);
-
-				if (!ensure(Header.OuterReference.IsValid()))
-				{
-					UE_LOG_ACTORREPLICATIONBRIDGE(Error, TEXT("Could not create NetReference to Outer %s of subobject %s"), *GetNameSafe(OuterObject), *GetNameSafe(Object));
-				}
-			}
-		}
-	}
-
-	Header.ObjectReference = ObjectRef;
-}
-
 bool UActorReplicationBridge::RemapPathForPIE(uint32 ConnectionId, FString& Str, bool bReading) const
 {
 	if (ConnectionId == UE::Net::InvalidConnectionId)
diff --git a/Engine/Source/Runtime/Engine/Private/Net/Iris/ReplicationSystem/ActorReplicationBridgeInternal.cpp b/Engine/Source/Runtime/Engine/Private/Net/Iris/ReplicationSystem/ActorReplicationBridgeInternal.cpp
deleted file mode 100644
diff --git a/Engine/Source/Runtime/Engine/Private/Net/Iris/ReplicationSystem/ActorReplicationBridgeInternal.h b/Engine/Source/Runtime/Engine/Private/Net/Iris/ReplicationSystem/ActorReplicationBridgeInternal.h
deleted file mode 100644
diff --git a/Engine/Source/Runtime/Engine/Private/Net/Iris/ReplicationSystem/NetActorFactory.cpp b/Engine/Source/Runtime/Engine/Private/Net/Iris/ReplicationSystem/NetActorFactory.cpp
new file mode 100644
--- /dev/null
+++ b/Engine/Source/Runtime/Engine/Private/Net/Iris/ReplicationSystem/NetActorFactory.cpp
@@ -0,0 +1,567 @@
+// Copyright Epic Games, Inc. All Rights Reserved.
+
+#include "Net/Iris/ReplicationSystem/NetActorFactory.h"
+
+#if UE_WITH_IRIS
+
+#include "Engine/Level.h"
+#include "Engine/NetDriver.h"
+#include "Engine/World.h"
+
+#include "GameFramework/Actor.h"
+
+#include "HAL/UnrealMemory.h"
+#include "HAL/LowLevelMemStats.h"
+
+#include "Iris/Core/IrisLog.h"
+#include "Iris/Core/IrisProfiler.h"
+
+#include "Iris/Serialization/NetBitStreamUtil.h"
+#include "Iris/Serialization/NetBitStreamReader.h"
+#include "Iris/Serialization/NetBitStreamWriter.h"
+#include "Iris/Serialization/ObjectNetSerializer.h"
+#include "Iris/Serialization/PackedVectorNetSerializers.h"
+#include "Iris/Serialization/VectorNetSerializers.h"
+
+#include "Net/DataBunch.h"
+
+#include "Net/Iris/ReplicationSystem/ActorReplicationBridge.h"
+
+#include "ProfilingDebugging/AssetMetadataTrace.h"
+
+#include "UObject/Package.h"
+
+#endif // UE_WITH_IRIS
+
+namespace UE::Net::Private
+{
+#if UE_WITH_IRIS
+	static int32 SerializeNewActorMaxBunchSize = 1024;
+	static FAutoConsoleVariableRef CVarSerializeNewActorMaxBunchSize(
+	TEXT("net.Iris.SerializeNewActorMaxBunchSize"),
+	SerializeNewActorMaxBunchSize,
+	TEXT("Max allowed bits that can be added to the creation header via OnSerializeNewActor."));
+
+	enum class EActorNetSpawnInfoFlags : uint32
+	{
+		None = 0U,
+		QuantizeScale = 1U,
+		QuantizeLocation = QuantizeScale << 1U,
+		QuantizeVelocity = QuantizeLocation << 1U,
+		BitCount = 3U,
+	};
+	ENUM_CLASS_FLAGS(EActorNetSpawnInfoFlags)
+	
+	EActorNetSpawnInfoFlags GetSpawnInfoFlags()
+	{
+		// Init spawninfo flags from CVARs
+		EActorNetSpawnInfoFlags Flags = EActorNetSpawnInfoFlags::None;
+		{
+			bool bQuantizeActorScaleOnSpawn = false;
+			IConsoleVariable* CVar = IConsoleManager::Get().FindConsoleVariable(TEXT("net.QuantizeActorScaleOnSpawn"));
+			if (ensure(CVar))
+			{
+				bQuantizeActorScaleOnSpawn = CVar->GetBool();
+			}
+			Flags |= bQuantizeActorScaleOnSpawn ? EActorNetSpawnInfoFlags::QuantizeScale : EActorNetSpawnInfoFlags::None;
+		}
+
+		{
+			bool bQuantizeActorLocationOnSpawn = true;
+			IConsoleVariable* CVar = IConsoleManager::Get().FindConsoleVariable(TEXT("net.QuantizeActorLocationOnSpawn"));
+			if (ensure(CVar))
+			{
+				bQuantizeActorLocationOnSpawn = CVar->GetBool();
+			}
+			Flags |= bQuantizeActorLocationOnSpawn ? EActorNetSpawnInfoFlags::QuantizeLocation : EActorNetSpawnInfoFlags::None;
+		}
+
+		{
+			bool bQuantizeActorVelocityOnSpawn = true;
+			IConsoleVariable* CVar = IConsoleManager::Get().FindConsoleVariable(TEXT("net.QuantizeActorVelocityOnSpawn"));
+			if (ensure(CVar))
+			{
+				bQuantizeActorVelocityOnSpawn = CVar->GetBool();
+			}
+			Flags |= bQuantizeActorVelocityOnSpawn ? EActorNetSpawnInfoFlags::QuantizeVelocity : EActorNetSpawnInfoFlags::None;
+		}
+
+		return Flags;
+	}
+
+	// Write vector using default value compression using specified serializer
+	static void WriteConditionallyQuantizedVector(FNetBitStreamWriter* Writer, const FVector& Vector, const FVector& DefaultValue, bool bQuantize)
+	{
+		// We use 0.01f for comparing when using quantization, because we will only send a single point of precision anyway.
+		// We could probably get away with 0.1f, but that may introduce edge cases for rounding.
+		static constexpr float Epsilon_Quantized = 0.01f;
+				
+		// We use KINDA_SMALL_NUMBER for comparing when not using quantization, because that's the default for FVector::Equals.
+		const float Epsilon = bQuantize ? Epsilon_Quantized : UE_KINDA_SMALL_NUMBER;
+	
+		if (Writer->WriteBool(!Vector.Equals(DefaultValue, Epsilon)))
+		{
+			Writer->WriteBool(bQuantize);
+
+			const FNetSerializer& Serializer = bQuantize ? UE_NET_GET_SERIALIZER(FVectorNetQuantize10NetSerializer) : UE_NET_GET_SERIALIZER(FVectorNetSerializer);
+
+			FNetSerializationContext Context(Writer);
+
+			alignas(16) uint8 QuantizedState[32] = {};
+			checkSlow(sizeof(QuantizedState) >= Serializer.QuantizedTypeSize);
+
+			FNetQuantizeArgs QuantizeArgs;
+			QuantizeArgs.Version = Serializer.Version;
+			QuantizeArgs.Source = NetSerializerValuePointer(&Vector);
+			QuantizeArgs.Target = NetSerializerValuePointer(&QuantizedState[0]);
+			QuantizeArgs.NetSerializerConfig = NetSerializerConfigParam(Serializer.DefaultConfig);
+			Serializer.Quantize(Context, QuantizeArgs);
+
+			FNetSerializeArgs SerializeArgs;
+			SerializeArgs.Version = Serializer.Version;
+			SerializeArgs.Source = QuantizeArgs.Target;
+			SerializeArgs.NetSerializerConfig = NetSerializerConfigParam(Serializer.DefaultConfig);
+			Serializer.Serialize(Context, SerializeArgs);
+		}
+	}
+
+	// Read default value compressed vector using specified serializer
+	static void ReadConditionallyQuantizedVector(FNetBitStreamReader* Reader, FVector& OutVector, const FVector& DefaultValue)
+	{
+		if (Reader->ReadBool())
+		{
+			const bool bIsQuantized = Reader->ReadBool();
+
+			const FNetSerializer& Serializer = bIsQuantized ? UE_NET_GET_SERIALIZER(FVectorNetQuantize10NetSerializer) : UE_NET_GET_SERIALIZER(FVectorNetSerializer);
+
+			FNetSerializationContext Context(Reader);
+
+			alignas(16) uint8 QuantizedState[32] = {};
+			checkSlow(sizeof(QuantizedState) >= Serializer.QuantizedTypeSize);
+
+			FNetDeserializeArgs DeserializeArgs;
+			DeserializeArgs.Version = Serializer.Version;
+			DeserializeArgs.Target = NetSerializerValuePointer(&QuantizedState[0]);
+			DeserializeArgs.NetSerializerConfig = NetSerializerConfigParam(Serializer.DefaultConfig);
+			Serializer.Deserialize(Context, DeserializeArgs);
+
+			FNetDequantizeArgs DequantizeArgs;
+			DequantizeArgs.Version = Serializer.Version;
+			DequantizeArgs.Source = DeserializeArgs.Target;
+			DequantizeArgs.Target = NetSerializerValuePointer(&OutVector);
+			DequantizeArgs.NetSerializerConfig = NetSerializerConfigParam(Serializer.DefaultConfig);
+			Serializer.Dequantize(Context, DequantizeArgs);
+		}
+		else
+		{
+			OutVector = DefaultValue;
+		}
+	}
+#endif // UE_WITH_IRIS
+} // end namespace UE::Net::Private
+
+#if UE_WITH_IRIS
+
+//------------------------------------------------------------------------
+// UNetActorFactory
+//------------------------------------------------------------------------
+
+void UNetActorFactory::OnInit()
+{
+	Super::OnInit();
+	SpawnInfoFlags = UE::Net::Private::GetSpawnInfoFlags();
+}
+
+TUniquePtr<UE::Net::FNetObjectCreationHeader> UNetActorFactory::CreateAndFillHeader(UE::Net::FNetRefHandle Handle)
+{
+	using namespace UE::Net;
+
+	AActor* Actor = Cast<AActor>(Bridge->GetReplicatedObject(Handle));
+	
+	if (!Actor)
+	{
+		ensureMsgf(Actor, TEXT("UNetActorFactory::CreateAndFillHeader could not find actor tied to handle: %s"), *Bridge->PrintObjectFromNetRefHandle(Handle));
+		return nullptr;
+	}
+
+	UActorReplicationBridge* ActorBridge = CastChecked<UActorReplicationBridge>(Bridge);
+
+	TUniquePtr<FActorNetCreationHeader> Header( new FActorNetCreationHeader );
+
+	const FNetObjectReference ActorReference = ActorBridge->GetOrCreateObjectReference(Actor);	
+	Header->bIsDynamic = ActorReference.GetRefHandle().IsDynamic();
+	//$IRIS TODO: Create two different header classes for dynamic and static actors.
+	if (Header->bIsDynamic)
+	{
+		// This is more or less a straight copy from ClientPackageMap and needs to be updated accordingly
+		UObject* Archetype = nullptr;
+		ULevel* ActorLevel = nullptr;
+
+		// ChildActor's need to be spawned from the ChildActorTemplate otherwise any non-replicated 
+		// customized properties will be incorrect on the Client.
+		if (const UChildActorComponent* CAC = Actor->GetParentComponent())
+		{
+			Archetype = CAC->GetSpawnableChildActorTemplate();
+		}
+
+		if (Archetype == nullptr)
+		{
+			Archetype = Actor->GetArchetype();
+		}
+
+		ActorLevel = Actor->GetLevel();
+
+		check(Archetype != nullptr);
+		check(Actor->NeedsLoadForClient());			// We have no business sending this unless the client can load
+		check(Archetype->NeedsLoadForClient());		// We have no business sending this unless the client can load
+
+		// Fill in Header
+		Header->ArchetypeReference = ActorBridge->GetOrCreateObjectReference(Archetype);
+		Header->bUsePersistentLevel = (UE::Net::Private::SerializeNewActorOverrideLevel == 0) || (ActorBridge->GetNetDriver()->GetWorld()->PersistentLevel == ActorLevel);
+
+		if (!Header->bUsePersistentLevel)
+		{
+			Header->LevelReference = ActorBridge->GetOrCreateObjectReference(ActorLevel);
+		}
+
+		if (const USceneComponent* RootComponent = Actor->GetRootComponent())
+		{
+			Header->SpawnInfo.Location = FRepMovement::RebaseOntoZeroOrigin(Actor->GetActorLocation(), Actor);
+			Header->SpawnInfo.Rotation = Actor->GetActorRotation();
+			Header->SpawnInfo.Scale = Actor->GetActorScale();
+			FVector Scale = Actor->GetActorScale();
+
+			if (USceneComponent* AttachParent = RootComponent->GetAttachParent())
+			{
+				// If this actor is attached, when the scale is serialized on the client, the attach parent property won't be set yet.
+				// USceneComponent::SetWorldScale3D (which got called by AActor::SetActorScale3D, which we used to do but no longer).
+				// would perform this transformation so that what is sent is relative to the parent. If we don't do this, we will
+				// apply the world scale on the client, which will then get applied a second time when the attach parent property is received.
+				FTransform ParentToWorld = AttachParent->GetSocketTransform(RootComponent->GetAttachSocketName());
+				Scale = Scale * ParentToWorld.GetSafeScaleReciprocal(ParentToWorld.GetScale3D());
+			}
+
+			Header->SpawnInfo.Scale = Scale;
+			Header->SpawnInfo.Velocity = Actor->GetVelocity();
+		}
+		else
+		{
+			//OutHeader.SpawnInfo = FactoryGetDefaultActorReplicationBridgeSpawnInfo();
+			check(!Header->SpawnInfo.Location.ContainsNaN());
+		}
+	}
+	else 
+	{
+		// Refer by path for static actors
+		Header->ObjectReference = ActorReference;
+	}
+
+	// Custom actor creation data
+	{
+		FOutBunch Bunch(UE::Net::Private::SerializeNewActorMaxBunchSize);
+		Actor->OnSerializeNewActor(Bunch);
+		Header->CustomCreationDataBitCount = IntCastChecked<uint16>(Bunch.GetNumBits());
+		if (Header->CustomCreationDataBitCount > 0)
+		{
+			Header->CustomCreationData.SetNumZeroed(Align(Bunch.GetNumBytes(), 4));
+			FMemory::Memcpy(Header->CustomCreationData.GetData(), Bunch.GetData(), Bunch.GetNumBytes());
+		}
+	}
+
+	return Header;
+}
+
+bool UNetActorFactory::SerializeHeader(const UE::Net::FCreationHeaderContext& Context, const UE::Net::FNetObjectCreationHeader* Header) 
+{
+	using namespace UE::Net;
+
+	const FActorNetCreationHeader* ActorHeader = static_cast<const FActorNetCreationHeader*>(Header);
+	
+	return ActorHeader->Serialize(Context, SpawnInfoFlags, DefaultSpawnInfo);;
+}
+
+TUniquePtr<UE::Net::FNetObjectCreationHeader> UNetActorFactory::CreateAndDeserializeHeader(const UE::Net::FCreationHeaderContext& Context)
+{
+	using namespace UE::Net;
+
+	TUniquePtr<FActorNetCreationHeader> Header(new FActorNetCreationHeader);
+
+	Header->Deserialize(Context, DefaultSpawnInfo);
+	return Header;
+}
+
+UNetObjectFactory::FInstantiateResult UNetActorFactory::InstantiateReplicatedObjectFromHeader(const FInstantiateContext& Context, const UE::Net::FNetObjectCreationHeader* Header)
+{
+	using namespace UE::Net;
+	using namespace UE::Net::Private;
+
+	LLM_SCOPE(ELLMTag::EngineMisc);
+	IRIS_PROFILER_SCOPE(NetActorFactory_InstantiateReplicatedObjectFromHeader);
+
+	UActorReplicationBridge* ActorBridge = CastChecked<UActorReplicationBridge>(Bridge);
+	UNetDriver* NetDriver = ActorBridge->GetNetDriver();
+
+	const FActorNetCreationHeader* ActorHeader = static_cast<const FActorNetCreationHeader*>(Header);
+
+	// For static actors, just find the object using the path
+	if (!ActorHeader->bIsDynamic)
+	{
+		AActor* Actor = Cast<AActor>(ActorBridge->ResolveObjectReference(ActorHeader->ObjectReference, Context.ResolveContext));
+		if (!Actor)
+		{
+			UE_LOG(LogIris, Error, TEXT("UNetActorFactory::InstantiateNetObjectFromHeader Failed to resolve ObjectReference: %s . Could not find static actor."), *ActorBridge->DescribeObjectReference(ActorHeader->ObjectReference, Context.ResolveContext));
+			return FInstantiateResult();
+		}
+
+		UE_LOG(LogIris, Verbose, TEXT("UNetActorFactory::InstantiateNetObjectFromHeader Found static Actor: %s using ObjectReference: %s"), ToCStr(Actor->GetPathName()), *ActorBridge->DescribeObjectReference(ActorHeader->ObjectReference, Context.ResolveContext));
+
+		FInstantiateResult Result {  .Instance = Actor };
+			
+		if (NetDriver->ShouldClientDestroyActor(Actor))
+		{
+			Result.Flags |= EReplicationBridgeCreateNetRefHandleResultFlags::AllowDestroyInstanceFromRemote;
+		}
+
+		return Result;
+	}
+
+	// For dynamic actors, spawn a new actor using the provided information
+
+	// Find archetype
+	AActor* Archetype = Cast<AActor>(Bridge->ResolveObjectReference(ActorHeader->ArchetypeReference, Context.ResolveContext));
+	if (!Archetype)
+	{
+		UE_LOG(LogIris, Error, TEXT("UNetActorFactory::InstantiateNetObjectFromHeader Unable to instantiate actor, failed to resolve archetype: %s"), *Bridge->DescribeObjectReference(ActorHeader->ArchetypeReference, Context.ResolveContext));
+		return FInstantiateResult();
+	}
+
+	// Find level
+	ULevel* Level = nullptr;
+	if (!ActorHeader->bUsePersistentLevel)
+	{
+		Level = Cast<ULevel>(Bridge->ResolveObjectReference(ActorHeader->LevelReference, Context.ResolveContext));
+	}
+
+	LLM_SCOPE_DYNAMIC_STAT_OBJECTPATH(Archetype->GetPackage(), ELLMTagSet::Assets);
+	LLM_SCOPE_DYNAMIC_STAT_OBJECTPATH(Archetype->GetClass(), ELLMTagSet::AssetClasses);
+	UE_TRACE_METADATA_SCOPE_ASSET(Archetype, Archetype->GetClass());
+			
+	// For streaming levels, it's possible that the owning level has been made not-visible but is still loaded.
+	// In that case, the level will still be found but the owning world will be invalid.
+	// If that happens, wait to spawn the Actor until the next time the level is streamed in.
+	// At that point, the Server should resend any dynamic Actors.
+			
+	check(Level == nullptr || Level->GetWorld() != nullptr);
+
+	FActorSpawnParameters SpawnInfo;
+	SpawnInfo.Template = Archetype;
+	SpawnInfo.OverrideLevel = Level;
+	SpawnInfo.SpawnCollisionHandlingOverride = ESpawnActorCollisionHandlingMethod::AlwaysSpawn;
+	SpawnInfo.bRemoteOwned = true;
+	SpawnInfo.bNoFail = true;
+
+	UWorld* World = NetDriver->GetWorld();
+	FVector SpawnLocation = FRepMovement::RebaseOntoLocalOrigin(ActorHeader->SpawnInfo.Location, World->OriginLocation);
+		
+	AActor* Actor = World->SpawnActorAbsolute(Archetype->GetClass(), FTransform(ActorHeader->SpawnInfo.Rotation, SpawnLocation), SpawnInfo);
+
+	// For Iris we expect that we will be able to spawn the actor as streaming always is controlled from server
+	if (!Actor)
+	{
+		ensureMsgf(Actor, TEXT("UNetActorFactory::InstantiateNetObjectFromHeader SpawnActor failed. Used Archetype: %s ObjectReference: %s"), *GetNameSafe(Archetype), *ActorBridge->DescribeObjectReference(ActorHeader->ArchetypeReference, Context.ResolveContext));
+		return FInstantiateResult();
+	}
+
+	FActorNetCreationHeader::FActorNetSpawnInfo DefaultNetSpawnInfo;
+	static constexpr float Epsilon = UE_KINDA_SMALL_NUMBER;
+
+	// Set Velocity if it differs from Default
+	if (!ActorHeader->SpawnInfo.Velocity.Equals(DefaultNetSpawnInfo.Velocity, Epsilon))
+	{
+		Actor->PostNetReceiveVelocity(ActorHeader->SpawnInfo.Velocity);
+	}
+					
+	// Set Scale if it differs from Default
+	if (!ActorHeader->SpawnInfo.Scale.Equals(DefaultNetSpawnInfo.Scale, Epsilon))
+	{
+		Actor->SetActorRelativeScale3D(ActorHeader->SpawnInfo.Scale);
+	}
+
+	FInstantiateResult Result
+	{ 
+		.Instance = Actor
+	};
+
+	if (NetDriver->ShouldClientDestroyActor(Actor))
+	{
+		Result.Flags |= EReplicationBridgeCreateNetRefHandleResultFlags::AllowDestroyInstanceFromRemote;
+	}
+
+	UE_LOG(LogIris, Verbose, TEXT("UNetActorFactory::InstantiateNetObjectFromHeader Spawned Actor: %s using Archetype: %s"), *Actor->GetPathName(), *GetNameSafe(Archetype));
+
+	return Result;
+}
+
+void UNetActorFactory::PostInstantiation(const FPostInstantiationContext& Context)
+{
+	using namespace UE::Net;
+
+	AActor* Actor = CastChecked<AActor>(Context.Instance);
+	if (!Actor)
+	{
+		return;
+	}
+
+	UActorReplicationBridge* ActorBridge = CastChecked<UActorReplicationBridge>(Bridge);
+	UNetDriver* NetDriver = ActorBridge->GetNetDriver();
+
+	const FActorNetCreationHeader* ActorHeader = static_cast<const FActorNetCreationHeader*>(Context.Header);
+
+	// OnActorChannelOpen
+	{
+		UNetConnection* Connection = NetDriver->GetConnectionById(Context.ConnectionId);
+		FInBunch Bunch(Connection, const_cast<uint8*>(ActorHeader->CustomCreationData.GetData()),ActorHeader->CustomCreationDataBitCount);
+		Actor->OnActorChannelOpen(Bunch, Connection);
+
+		if (Bunch.IsError() || Bunch.GetBitsLeft() != 0)
+		{
+			UE_LOG(LogIris, Error, TEXT("UNetActorFactory::PostInstantiation deserialization error in OnActorChannelOpen for Actor: %s"), *Actor->GetPathName());
+			check(false);
+			return;
+		}
+	}
+
+	// Wake up from dormancy. This is important for client replays.
+	ActorBridge->WakeUpObjectInstantiatedFromRemote(Actor);
+}
+
+
+//------------------------------------------------------------------------
+// FActorNetCreationHeader
+//------------------------------------------------------------------------
+ namespace UE::Net
+ {
+
+bool FActorNetCreationHeader::Serialize(const FCreationHeaderContext& Context, UE::Net::Private::EActorNetSpawnInfoFlags SpawnFlags, const FActorNetSpawnInfo& DefaultSpawnInfo) const
+{
+	using namespace UE::Net::Private;
+
+	FNetBitStreamWriter* Writer = Context.Serialization.GetBitStreamWriter();
+
+	if (Writer->WriteBool(bIsDynamic))
+	{
+		// Write Archetype and LevelPath
+		WriteFullNetObjectReference(Context.Serialization, ArchetypeReference);
+		
+		// Only write the LevelPath if it differs from the persistent level
+		if (!Writer->WriteBool(bUsePersistentLevel))
+		{
+			WriteFullNetObjectReference(Context.Serialization, LevelReference);
+		}
+
+		const bool bQuantizeLocation = (SpawnFlags&EActorNetSpawnInfoFlags::QuantizeLocation)!=EActorNetSpawnInfoFlags::None;
+		const bool bQuantizeScale = (SpawnFlags&EActorNetSpawnInfoFlags::QuantizeScale)!=EActorNetSpawnInfoFlags::None;
+		const bool bQuantizeVelocity = (SpawnFlags&EActorNetSpawnInfoFlags::QuantizeVelocity)!=EActorNetSpawnInfoFlags::None;
+
+		// Write actor spawn info
+		WriteConditionallyQuantizedVector(Writer, SpawnInfo.Location, DefaultSpawnInfo.Location, bQuantizeLocation);
+		WriteConditionallyQuantizedVector(Writer, SpawnInfo.Scale, DefaultSpawnInfo.Scale, bQuantizeScale);
+		WriteConditionallyQuantizedVector(Writer, SpawnInfo.Velocity, DefaultSpawnInfo.Velocity, bQuantizeVelocity);
+
+		// For rotation we use 0.001f for Rotation comparison to keep consistency with old behavior.
+		static constexpr float RotationEpsilon = 0.001f;
+		WriteRotator(Writer, SpawnInfo.Rotation, DefaultSpawnInfo.Rotation, RotationEpsilon);
+
+
+	}
+	else // should be possible to refer by path??
+	{
+		WriteFullNetObjectReference(Context.Serialization, ObjectReference);
+	}
+
+	if (Writer->WriteBool(CustomCreationDataBitCount > 0))
+	{
+		Writer->WriteBits(CustomCreationDataBitCount - 1U, 16U);
+		Writer->WriteBitStream(reinterpret_cast<const uint32*>(CustomCreationData.GetData()), 0U, CustomCreationDataBitCount);
+	}
+
+	return true;
+}
+
+bool FActorNetCreationHeader::Deserialize(const FCreationHeaderContext& Context, const FActorNetSpawnInfo& DefaultSpawnInfo)
+{
+	using namespace UE::Net::Private;
+
+	FNetBitStreamReader* Reader = Context.Serialization.GetBitStreamReader();
+
+	// Dynamic actor?
+	bIsDynamic = Reader->ReadBool();
+
+	if (bIsDynamic)
+	{
+		// Read Archetype
+		ReadFullNetObjectReference(Context.Serialization, ArchetypeReference);
+		
+		bUsePersistentLevel = Reader->ReadBool();
+		if (!bUsePersistentLevel)
+		{
+			ReadFullNetObjectReference(Context.Serialization, LevelReference);
+		}
+
+		// Read actor spawn info
+		ReadConditionallyQuantizedVector(Reader, SpawnInfo.Location, DefaultSpawnInfo.Location);
+		ReadConditionallyQuantizedVector(Reader, SpawnInfo.Scale, DefaultSpawnInfo.Scale);
+		ReadConditionallyQuantizedVector(Reader, SpawnInfo.Velocity, DefaultSpawnInfo.Velocity);
+
+		ReadRotator(Reader, SpawnInfo.Rotation, DefaultSpawnInfo.Rotation);
+	}
+	else
+	{
+		ReadFullNetObjectReference(Context.Serialization, ObjectReference);
+	}
+
+	CustomCreationDataBitCount = 0;
+	if (Reader->ReadBool())
+	{
+		CustomCreationDataBitCount = 1U + Reader->ReadBits(16U);
+		CustomCreationData.SetNumZeroed(((CustomCreationDataBitCount + 31U) & ~31U) >> 3U);
+		Reader->ReadBitStream(reinterpret_cast<uint32*>(CustomCreationData.GetData()), CustomCreationDataBitCount);
+	}
+	
+	return true;
+}
+
+
+FString FActorNetCreationHeader::ToString() const
+{
+	return FString::Printf(TEXT("FDynamicActorNetCreationHeader (ProtocolId:0x%x):\n\t"
+								"bIsDynamic=%u\n\t"
+								"ObjectReference=%s\n\t"
+								"ArchetypeReference=%s\n\t"
+								"SpawnInfo.Location=%s\n\t"
+								"SpawnInfo.Rotation=%s\n\t"
+								"SpawnInfo.Scale=%s\n\t"
+								"SpawnInfo.Velocity=%s\n\t"		
+								"bUsePersistentLevel=%s\n\t"
+								"LevelReference=%s\n\t"
+								"CustomCreationData=%u bits"
+								),
+						   		GetProtocolId(),
+						   		bIsDynamic,
+						   		*ObjectReference.ToString(),
+								*ArchetypeReference.ToString(),
+								*SpawnInfo.Location.ToCompactString(),
+								*SpawnInfo.Rotation.ToCompactString(),
+								*SpawnInfo.Scale.ToCompactString(),
+								*SpawnInfo.Velocity.ToCompactString(),		
+								bUsePersistentLevel?TEXT("True"):TEXT("False"),
+								*LevelReference.ToString(),
+								CustomCreationDataBitCount
+								);
+}
+
+
+} // end namespace UE::Net
+
+#endif // UE_WITH_IRIS
+
diff --git a/Engine/Source/Runtime/Engine/Private/Net/Iris/ReplicationSystem/NetEngineFactories.cpp b/Engine/Source/Runtime/Engine/Private/Net/Iris/ReplicationSystem/NetEngineFactories.cpp
new file mode 100644
--- /dev/null
+++ b/Engine/Source/Runtime/Engine/Private/Net/Iris/ReplicationSystem/NetEngineFactories.cpp
@@ -0,0 +1,31 @@
+// Copyright Epic Games, Inc. All Rights Reserved.
+
+#include "Net/Iris/ReplicationSystem/NetEngineFactories.h"
+
+#if UE_WITH_IRIS
+#include "Net/Iris/ReplicationSystem/NetActorFactory.h"
+#include "Net/Iris/ReplicationSystem/NetSubObjectFactory.h"
+
+#include "Iris/ReplicationSystem/NetObjectFactoryRegistry.h"
+#endif 
+
+#if UE_WITH_IRIS
+
+namespace UE::Net
+{
+
+void InitEngineNetObjectFactories()
+{
+	FNetObjectFactoryRegistry::RegisterFactory(UNetActorFactory::StaticClass(), UNetActorFactory::GetFactoryName());
+	FNetObjectFactoryRegistry::RegisterFactory(UNetSubObjectFactory::StaticClass(), UNetSubObjectFactory::GetFactoryName());
+}
+
+void ShutdownEngineNetObjectFactories()
+{
+	FNetObjectFactoryRegistry::UnregisterFactory(UNetActorFactory::GetFactoryName());
+	FNetObjectFactoryRegistry::UnregisterFactory(UNetSubObjectFactory::GetFactoryName());
+}
+
+} // end namespace UE::Net
+
+#endif // UE_WITH_IRIS
\ No newline at end of file
diff --git a/Engine/Source/Runtime/Engine/Private/Net/Iris/ReplicationSystem/NetSubObjectFactory.cpp b/Engine/Source/Runtime/Engine/Private/Net/Iris/ReplicationSystem/NetSubObjectFactory.cpp
new file mode 100644
--- /dev/null
+++ b/Engine/Source/Runtime/Engine/Private/Net/Iris/ReplicationSystem/NetSubObjectFactory.cpp
@@ -0,0 +1,286 @@
+// Copyright Epic Games, Inc. All Rights Reserved.
+
+#include "Net/Iris/ReplicationSystem/NetSubObjectFactory.h"
+
+#if UE_WITH_IRIS
+
+#include "GameFramework/Actor.h"
+
+#include "HAL/UnrealMemory.h"
+
+#include "Iris/Core/IrisLog.h"
+
+#include "Iris/ReplicationSystem/ObjectReplicationBridge.h"
+
+#include "Iris/Serialization/NetBitStreamUtil.h"
+#include "Iris/Serialization/NetBitStreamReader.h"
+#include "Iris/Serialization/NetBitStreamWriter.h"
+#include "Iris/Serialization/NetSerializationContext.h"
+#include "Iris/Serialization/ObjectNetSerializer.h"
+
+#include "UObject/Package.h"
+
+#include "Net/DataBunch.h"
+
+#endif // UE_WITH_IRIS
+
+#if UE_WITH_IRIS
+
+TUniquePtr<UE::Net::FNetObjectCreationHeader> UNetSubObjectFactory::CreateAndFillHeader(UE::Net::FNetRefHandle Handle)
+{
+	using namespace UE::Net;
+
+	TUniquePtr<FNetObjectCreationHeader> Header(new FNetSubObjectCreationHeader);
+	FNetSubObjectCreationHeader* SubObjectHeader = static_cast<FNetSubObjectCreationHeader*>(Header.Get());
+
+	UObject* SubObject = Bridge->GetReplicatedObject(Handle);
+	if (!SubObject)
+	{
+		ensureMsgf(SubObject, TEXT("UNetSubObjectFactory::CreateAndFillHeader could not find object tied to handle: %s"), *Bridge->PrintObjectFromNetRefHandle(Handle));
+		return nullptr;
+	}
+
+	FNetObjectReference ObjectRef = Bridge->GetOrCreateObjectReference(SubObject);
+
+	SubObjectHeader->bIsDynamic = ObjectRef.GetRefHandle().IsDynamic();
+	SubObjectHeader->bIsNameStableForNetworking = SubObject->IsNameStableForNetworking();
+	SubObjectHeader->ObjectReference = ObjectRef;
+
+	if (!SubObjectHeader->bIsDynamic || SubObjectHeader->bIsNameStableForNetworking)
+	{
+		// No more information needed since we don't need to spawn the object on the remote.
+		return Header;
+	}
+
+	check(SubObject->NeedsLoadForClient() ); // We have no business sending this unless the client can load
+	check(SubObject->GetClass()->NeedsLoadForClient());	// We have no business sending this unless the client can load
+
+	SubObjectHeader->ObjectClassReference = Bridge->GetOrCreateObjectReference(SubObject->GetClass());
+
+	// Find the right Outer
+	UObject* OuterObject = SubObject->GetOuter();	
+	if (OuterObject == GetTransientPackage())
+	{
+		SubObjectHeader->bOuterIsTransientLevel = true;
+	}
+	else 
+	{
+		FNetRefHandle RootObjectHandle = Bridge->GetRootObjectOfSubObject(Handle);
+		UObject* RootObject = Bridge->GetReplicatedObject(RootObjectHandle);
+		check(RootObject);
+
+		if (OuterObject == RootObject)
+		{
+			SubObjectHeader->bOuterIsRootObject = true;
+		}
+		else
+		{
+			SubObjectHeader->OuterReference = Bridge->GetOrCreateObjectReference(OuterObject);
+
+			if (!SubObjectHeader->OuterReference.IsValid())
+			{
+				UE_LOG(LogIris, Error, TEXT("UNetSubObjectFactory::CreateAndFillHeader %s could not create NetReference to Outer %s"), *Bridge->PrintObjectFromNetRefHandle(Handle), *GetNameSafe(OuterObject));
+			}
+		}
+	}
+
+	return Header;
+}
+
+TUniquePtr<UE::Net::FNetObjectCreationHeader> UNetSubObjectFactory::CreateAndDeserializeHeader(const UE::Net::FCreationHeaderContext& Context)
+{
+	using namespace UE::Net;
+
+	TUniquePtr<FNetSubObjectCreationHeader> Header(new FNetSubObjectCreationHeader);
+	Header->Deserialize(Context);
+	return Header;
+}
+
+
+UNetSubObjectFactory::FInstantiateResult UNetSubObjectFactory::InstantiateReplicatedObjectFromHeader(const FInstantiateContext& Context, const UE::Net::FNetObjectCreationHeader* Header)
+{
+	using namespace UE::Net;
+
+	const FNetSubObjectCreationHeader* SubObjectHeader = static_cast<const FNetSubObjectCreationHeader*>(Header);
+
+	if (!SubObjectHeader->bIsDynamic || SubObjectHeader->bIsNameStableForNetworking)
+	{
+		// Resolve by finding object relative to owner. We do not allow this object to be destroyed.
+		UObject* SubObject = Bridge->ResolveObjectReference(SubObjectHeader->ObjectReference, Context.ResolveContext);
+
+		if (!SubObject)
+		{
+			if (!SubObjectHeader->bIsDynamic)
+			{
+				UE_LOG(LogIris, Error, TEXT("UNetSubObjectFactory::InstantiateNetObjectFromHeader %s: Failed to find static object reference for static SubObject: %s, Owner: %s, RootObject: %s"), *Context.Handle.ToString(), *Bridge->DescribeObjectReference(SubObjectHeader->ObjectReference, Context.ResolveContext), *Bridge->PrintObjectFromNetRefHandle(Context.RootObjectOfSubObject), *GetPathNameSafe(Bridge->GetReplicatedObject(Context.RootObjectOfSubObject)));
+			}
+			else if (SubObjectHeader->bIsNameStableForNetworking)
+			{
+				UE_LOG(LogIris, Error, TEXT("UNetSubObjectFactory::InstantiateNetObjectFromHeader %s: Failed to find stable name reference of dynamic SubObject: %s, Owner: %s, RootObject: %s"), *Context.Handle.ToString(), *Bridge->DescribeObjectReference(SubObjectHeader->ObjectReference, Context.ResolveContext), *Bridge->PrintObjectFromNetRefHandle(Context.RootObjectOfSubObject), *GetPathNameSafe(Bridge->GetReplicatedObject(Context.RootObjectOfSubObject)));
+			}
+
+			return FInstantiateResult();
+		}
+
+		UE_LOG(LogIris, Verbose, TEXT("UNetSubObjectFactory::InstantiateNetObjectFromHeader %s: Found static SubObject using path %s"), *Context.Handle.ToString(), ToCStr(SubObject->GetPathName()));
+
+		FInstantiateResult Result { .Instance = SubObject };
+		return Result;
+	}
+
+	// For dynamic objects we have to spawn them
+		
+	UObject* RootObject = Bridge->GetReplicatedObject(Context.RootObjectOfSubObject);
+	AActor* RootActor = CastChecked<AActor>(RootObject);
+			
+	// Find the proper Outer
+	UObject* OuterObject = nullptr;
+	if (SubObjectHeader->bOuterIsTransientLevel)
+	{
+		OuterObject = GetTransientPackage();
+	}
+	else if (SubObjectHeader->bOuterIsRootObject)
+	{
+		OuterObject = RootActor;
+	}				
+	else
+	{
+		OuterObject = Bridge->ResolveObjectReference(SubObjectHeader->OuterReference, Context.ResolveContext);
+
+		if (!OuterObject)
+		{
+			UE_LOG(LogIris, Error, TEXT("BeginInstantiateFromRemote Failed to find Outer %s for dynamic subobject %s"), *Bridge->DescribeObjectReference(SubObjectHeader->OuterReference, Context.ResolveContext), *Bridge->DescribeObjectReference(SubObjectHeader->ObjectReference, Context.ResolveContext))
+
+			// Fallback to the rootobject instead
+			OuterObject = RootActor;
+		}
+	}
+
+	// We need to spawn the subobject
+	UObject* SubObjClassObj = Bridge->ResolveObjectReference(SubObjectHeader->ObjectClassReference, Context.ResolveContext);
+	UClass * SubObjClass = Cast<UClass>(SubObjClassObj);
+
+	// Try to spawn SubObject
+	UObject* SubObj = NewObject<UObject>(OuterObject, SubObjClass);
+
+	// Sanity check some things
+	checkf(SubObj != nullptr, TEXT("UNetSubObjectFactory::InstantiateNetObjectFromHeader: Subobject is NULL after instantiating. Class: %s, Outer %s, Actor %s"), *GetNameSafe(SubObjClass), *GetNameSafe(OuterObject), *GetNameSafe(RootActor));
+	checkf(!OuterObject || SubObj->IsIn(OuterObject), TEXT("UNetSubObjectFactory::InstantiateNetObjectFromHeader: Subobject is not in Outer. SubObject: %s, Outer %s, Actor %s"), *SubObj->GetName(), *GetNameSafe(OuterObject), *GetNameSafe(RootActor));
+	checkf(Cast<AActor>(SubObj) == nullptr, TEXT("UNetSubObjectFactory::InstantiateNetObjectFromHeader: Subobject is an Actor. SubObject: %s, Outer %s, Actor %s"), *SubObj->GetName(), *GetNameSafe(OuterObject), *GetNameSafe(RootActor));
+
+	FInstantiateResult Result { .Instance = SubObj };
+
+	// We must defer call OnSubObjectCreatedFromReplication after the state has been applied to the owning actor in order to behave like old system.
+	Result.Flags |= EReplicationBridgeCreateNetRefHandleResultFlags::ShouldCallSubObjectCreatedFromReplication;
+
+	// Created objects may be destroyed.
+	Result.Flags |= EReplicationBridgeCreateNetRefHandleResultFlags::AllowDestroyInstanceFromRemote;
+			
+	return Result;
+}
+
+bool UNetSubObjectFactory::SerializeHeader(const UE::Net::FCreationHeaderContext& Context, const UE::Net::FNetObjectCreationHeader* Header)
+{
+	using namespace UE::Net;
+	const FNetSubObjectCreationHeader* SubObjectHeader = static_cast<const FNetSubObjectCreationHeader*>(Header);
+	return SubObjectHeader->Serialize(Context);
+}
+
+//------------------------------------------------------------------------
+
+namespace UE::Net
+{
+
+FString FNetSubObjectCreationHeader::ToString() const
+{
+	return FString::Printf(TEXT("FNetSubObjectCreationHeader (ProtocolId:0x%x):\n\t"
+								"ObjectReference=%s\n\t"
+								"ObjectClassReference=%s\n\t"
+								"OuterReference=%s\n\t"
+								"bIsDynamic=%u\n\t"
+								"bIsNameStableForNetworking=%u\n\t"
+								"bUsePersistenLevel=%u\n\t"
+								"bOuterIsTransientLevel=%u\n\t"
+								"bOuterIsRootObject=%u\n\t"),
+								GetProtocolId(),
+								*ObjectReference.ToString(),
+								*ObjectClassReference.ToString(),
+								*OuterReference.ToString(),
+								bIsDynamic,
+								bIsNameStableForNetworking,
+								bUsePersistentLevel,
+								bOuterIsTransientLevel,
+								bOuterIsRootObject);
+
+}
+
+bool FNetSubObjectCreationHeader::Serialize(const FCreationHeaderContext& Context) const
+{
+	FNetBitStreamWriter* Writer = Context.Serialization.GetBitStreamWriter();
+
+	// Write required references to either find or instantiate the subobject
+	if (Writer->WriteBool(bIsDynamic))
+	{
+		if (Writer->WriteBool(bIsNameStableForNetworking))
+		{
+			WriteFullNetObjectReference(Context.Serialization, ObjectReference);
+		}
+		else
+		{
+			WriteFullNetObjectReference(Context.Serialization, ObjectClassReference);
+
+			if (!Writer->WriteBool(bOuterIsTransientLevel))
+			{
+				if (!Writer->WriteBool(bOuterIsRootObject))
+				{
+					WriteFullNetObjectReference(Context.Serialization, OuterReference);
+				}
+			}
+		}
+	}
+	else
+	{
+		WriteFullNetObjectReference(Context.Serialization, ObjectReference);
+	}
+
+	return true;
+}
+
+bool FNetSubObjectCreationHeader::Deserialize(const FCreationHeaderContext& Context)
+{
+	FNetBitStreamReader* Reader = Context.Serialization.GetBitStreamReader();
+
+	bIsDynamic = Reader->ReadBool();
+	if (bIsDynamic)
+	{
+		bIsNameStableForNetworking = Reader->ReadBool();
+		if (bIsNameStableForNetworking)
+		{
+			ReadFullNetObjectReference(Context.Serialization, ObjectReference);
+		}
+		else
+		{
+			ReadFullNetObjectReference(Context.Serialization, ObjectClassReference);
+
+			bOuterIsTransientLevel = Reader->ReadBool();
+			if (!bOuterIsTransientLevel)
+			{
+				bOuterIsRootObject = Reader->ReadBool();
+				if (!bOuterIsRootObject)
+				{
+					ReadFullNetObjectReference(Context.Serialization, OuterReference);
+				}
+			}
+		}
+	}
+	else
+	{
+		ReadFullNetObjectReference(Context.Serialization, ObjectReference);
+	}
+
+	return true;
+}
+
+} // end namespace UE::Net
+
+#endif //UE_WITH_IRIS
\ No newline at end of file
diff --git a/Engine/Source/Runtime/Engine/Private/Net/Iris/ReplicationSystem/ReplicationSystemUtil.cpp b/Engine/Source/Runtime/Engine/Private/Net/Iris/ReplicationSystem/ReplicationSystemUtil.cpp
--- a/Engine/Source/Runtime/Engine/Private/Net/Iris/ReplicationSystem/ReplicationSystemUtil.cpp
+++ b/Engine/Source/Runtime/Engine/Private/Net/Iris/ReplicationSystem/ReplicationSystemUtil.cpp
@@ -237,12 +237,13 @@ void FReplicationSystemUtil::BeginReplicationForActorSubObject(const AActor* Act
 		{
 			if (ReplicationSystem->IsServer())
 			{
-				if (UObjectReplicationBridge* Bridge = ReplicationSystem->GetReplicationBridgeAs<UObjectReplicationBridge>())
+				if (UActorReplicationBridge* Bridge = ReplicationSystem->GetReplicationBridgeAs<UActorReplicationBridge>())
 				{
 					const FNetRefHandle ActorRefHandle = Bridge->GetReplicatedRefHandle(ActorHandle);
 					if (ActorRefHandle.IsValid())
 					{
-						const FNetRefHandle SubObjectRefHandle = Bridge->StartReplicatingSubObject(ActorRefHandle, ActorSubObject);
+						const UObjectReplicationBridge::FSubObjectReplicationParams Params { .RootObjectHandle = ActorRefHandle };
+						const FNetRefHandle SubObjectRefHandle = Bridge->StartReplicatingSubObject(ActorSubObject, Params);
 						if (SubObjectRefHandle.IsValid() && NetCondition != ELifetimeCondition::COND_None)
 						{
 							Bridge->SetSubObjectNetCondition(SubObjectRefHandle, NetCondition);
@@ -271,13 +272,19 @@ void FReplicationSystemUtil::BeginReplicationForActorComponentSubObject(UActorCo
 			{
 				if (ReplicationSystem->IsServer())
 				{
-					if (UObjectReplicationBridge* Bridge = ReplicationSystem->GetReplicationBridgeAs<UObjectReplicationBridge>())
+					if (UActorReplicationBridge* Bridge = ReplicationSystem->GetReplicationBridgeAs<UActorReplicationBridge>())
 					{
 						const FNetRefHandle ActorRefHandle = Bridge->GetReplicatedRefHandle(ActorHandle);
 						const FNetRefHandle ActorComponentRefHandle = Bridge->GetReplicatedRefHandle(ActorComponent);
 						if (ActorRefHandle.IsValid() && ActorComponentRefHandle.IsValid())
 						{
-							const FNetRefHandle SubObjectRefHandle = Bridge->StartReplicatingSubObject(ActorRefHandle, SubObject, ActorComponentRefHandle, UReplicationBridge::ESubObjectInsertionOrder::ReplicateWith);
+							const UObjectReplicationBridge::FSubObjectReplicationParams Params
+							{ 
+								.RootObjectHandle = ActorRefHandle,
+								.InsertRelativeToSubObjectHandle = ActorComponentRefHandle,
+								.InsertionOrder = UReplicationBridge::ESubObjectInsertionOrder::ReplicateWith
+							};
+							const FNetRefHandle SubObjectRefHandle = Bridge->StartReplicatingSubObject(SubObject, Params);
 							if (SubObjectRefHandle.IsValid() && NetCondition != ELifetimeCondition::COND_None)
 							{
 								Bridge->SetSubObjectNetCondition(SubObjectRefHandle, NetCondition);
diff --git a/Engine/Source/Runtime/Engine/Private/UnrealEngine.cpp b/Engine/Source/Runtime/Engine/Private/UnrealEngine.cpp
--- a/Engine/Source/Runtime/Engine/Private/UnrealEngine.cpp
+++ b/Engine/Source/Runtime/Engine/Private/UnrealEngine.cpp
@@ -129,6 +129,7 @@ UnrealEngine.cpp: Implements the UEngine class and helpers.
 
 #if UE_WITH_IRIS
 #include "Iris/IrisConfig.h"
+#include "Iris/ReplicationSystem/NetObjectFactoryRegistry.h"
 #endif
 
 #include "IUniversalObjectLocatorModule.h"
@@ -232,6 +233,10 @@ UnrealEngine.cpp: Implements the UEngine class and helpers.
 #include "Streaming/StreamingManagerTexture.h"
 #endif
 
+#if UE_WITH_IRIS
+#include "Net/Iris/ReplicationSystem/NetEngineFactories.h"
+#endif
+
 #include "HAL/FileManagerGeneric.h"
 #include "UObject/ReferenceChainSearch.h"
 
@@ -351,6 +356,10 @@ void FEngineModule::StartupModule()
 
 	UE::Anim::FSkeletonRemappingRegistry::Init();
 
+#if UE_WITH_IRIS
+	UE::Net::InitEngineNetObjectFactories();
+#endif
+
 	IUniversalObjectLocatorModule& UolModule = FModuleManager::Get().LoadModuleChecked<IUniversalObjectLocatorModule>("UniversalObjectLocator");
 
 	FDelayedAutoRegisterHelper(EDelayedRegisterRunPhase::ObjectSystemReady,
@@ -382,6 +391,10 @@ void FEngineModule::StartupModule()
 
 void FEngineModule::ShutdownModule()
 {
+#if UE_WITH_IRIS
+	UE::Net::ShutdownEngineNetObjectFactories();
+#endif
+
 #if TRACE_FILTERING_ENABLED
 	FTraceFilter::Destroy();
 #endif
@@ -14315,6 +14328,13 @@ namespace UE::Private
 				const FName DriverName = InNetDriverName.IsNone() ? ReturnVal->GetFName() : InNetDriverName;
 				const bool bInitializeWithIris = Engine->WillNetDriverUseIris(Context, NetDriverDefinition, DriverName);
 
+#if UE_WITH_IRIS
+				if (bInitializeWithIris)
+				{
+					UE::Net::FNetObjectFactoryRegistry::SetFactoryRegistrationAllowed(false);
+				}
+#endif
+
 				ReturnVal->SetNetDriverName(DriverName);
 				ReturnVal->SetNetDriverDefinition(NetDriverDefinition);
 				ReturnVal->PostCreation(bInitializeWithIris);
@@ -14384,6 +14404,29 @@ namespace UE::Private
 				break;
 			}
 		}
+
+#if UE_WITH_IRIS
+		auto IsAnyNetDriverUsingIris = []()
+		{
+			const TIndirectArray<FWorldContext>& Worlds = GEngine->GetWorldContexts();
+			for (const FWorldContext& Context : Worlds)
+			{
+				for (const FNamedNetDriver& NetDriver : Context.ActiveNetDrivers)
+				{
+					if (NetDriver.NetDriver && NetDriver.NetDriver->IsUsingIrisReplication())
+					{
+						return true;
+					}
+				}
+			}
+			return false;
+		};
+		
+		if (!IsAnyNetDriverUsingIris())
+		{
+			UE::Net::FNetObjectFactoryRegistry::SetFactoryRegistrationAllowed(true);
+		}
+#endif // UE_WITH_IRIS
 	}
 } // end namespace UE::Private
 
diff --git a/Engine/Source/Runtime/Engine/Public/Net/Iris/ReplicationSystem/ActorReplicationBridge.h b/Engine/Source/Runtime/Engine/Public/Net/Iris/ReplicationSystem/ActorReplicationBridge.h
--- a/Engine/Source/Runtime/Engine/Public/Net/Iris/ReplicationSystem/ActorReplicationBridge.h
+++ b/Engine/Source/Runtime/Engine/Public/Net/Iris/ReplicationSystem/ActorReplicationBridge.h
@@ -76,7 +76,10 @@ public:
 	 * Begin replication of an ActorComponent and its registered SubObjects, 
 	 * if the ActorComponent already is replicated any set NetObjectConditions will be updated.
 	*/
-	ENGINE_API FNetRefHandle StartReplicatingComponent(FNetRefHandle OwnerHandle, UActorComponent* ActorComponent);
+	ENGINE_API FNetRefHandle StartReplicatingComponent(FNetRefHandle RootObjectHandle, UActorComponent* ActorComponent);
+
+	/** Begin replication of a subobject. */
+	ENGINE_API FNetRefHandle StartReplicatingSubObject(UObject* SubObject, const FSubObjectReplicationParams& Params);
 
 	/** Stop replicating an ActorComponent and its associated SubObjects. */
 	ENGINE_API void StopReplicatingComponent(UActorComponent* ActorComponent, EEndReplicationFlags EndReplicationFlags = EEndReplicationFlags::None);
@@ -93,6 +96,8 @@ public:
 	/** Called when NetUpdateFrequency has changed on the Actor. */
 	void OnNetUpdateFrequencyChanged(const AActor* Actor);
 
+	void WakeUpObjectInstantiatedFromRemote(AActor* Actor) const;
+
 	/**
 	 * Add relevant network metrics gathered since the last call to ConsumeNetMetrics.
 	 * Any periodic stat will be reset here too.
@@ -105,12 +110,6 @@ protected:
 	// UObjectReplicationBridge
 	virtual void Initialize(UReplicationSystem* ReplicationSystem) override;
 	virtual void Deinitialize() override;
-	virtual bool WriteCreationHeader(UE::Net::FNetSerializationContext& Context, FNetRefHandle Handle) override;
-	virtual TUniquePtr<FCreationHeader> GetCreationHeader(FNetRefHandle Handle) override;
-	virtual bool WriteCreationHeader(UE::Net::FNetSerializationContext& Context, const FCreationHeader* Header) override;
-	virtual TUniquePtr<FCreationHeader> ReadCreationHeader(UE::Net::FNetSerializationContext& Context) override;
-	virtual FObjectReplicationBridgeInstantiateResult BeginInstantiateFromRemote(FNetRefHandle RootObjectOfSubObject, const UE::Net::FNetObjectResolveContext& ResolveContext, const FCreationHeader* InHeader) override;
-	virtual bool OnInstantiatedFromRemote(UObject* Instance, const FCreationHeader* InHeader, uint32 ConnectionId) const override;
 	virtual void OnSubObjectCreatedFromReplication(FNetRefHandle SubObjectHandle) override;
 	virtual void EndInstantiateFromRemote(FNetRefHandle Handle) override;
 	virtual void DestroyInstanceFromRemote(const FDestroyInstanceParams& Params) override;
@@ -126,17 +125,15 @@ protected:
 	[[nodiscard]] virtual FString PrintConnectionInfo(uint32 ConnectionId) const override;
 
 private:
-	void GetActorCreationHeader(const AActor* Actor, UE::Net::Private::FActorCreationHeader& Header) const;
-	void GetSubObjectCreationHeader(const UObject* Object, const UObject* RootObject, UE::Net::Private::FSubObjectCreationHeader& Header) const;
-
+	
 	void OnMaxTickRateChanged(UNetDriver* InNetDriver, int32 NewMaxTickRate, int32 OldMaxTickRate);
 
-	void WakeUpObjectInstantiatedFromRemote(AActor* Actor) const;
 	void AddActorToLevelGroup(const AActor* Actor);
 
 private:
 
-	uint32 SpawnInfoFlags;
+	UE::Net::FNetObjectFactoryId ActorFactoryId;
+	UE::Net::FNetObjectFactoryId SubObjectFactoryId;
 
 #endif // UE_WITH_IRIS
 
diff --git a/Engine/Source/Runtime/Engine/Public/Net/Iris/ReplicationSystem/NetActorFactory.h b/Engine/Source/Runtime/Engine/Public/Net/Iris/ReplicationSystem/NetActorFactory.h
new file mode 100644
--- /dev/null
+++ b/Engine/Source/Runtime/Engine/Public/Net/Iris/ReplicationSystem/NetActorFactory.h
@@ -0,0 +1,105 @@
+// Copyright Epic Games, Inc. All Rights Reserved.
+
+#pragma once
+
+#include "Iris/ReplicationSystem/NetObjectFactory.h"
+
+#if UE_WITH_IRIS
+#include "Iris/Core/NetObjectReference.h"
+#endif
+
+#include "NetActorFactory.generated.h"
+
+namespace UE::Net::Private
+{
+	enum class EActorNetSpawnInfoFlags : uint32;
+}
+
+
+namespace UE::Net
+{
+
+#if UE_WITH_IRIS
+
+/**
+* Header information representing dynamic actors
+*/
+class FActorNetCreationHeader : public FNetObjectCreationHeader
+{
+public:
+
+	struct FActorNetSpawnInfo
+	{
+		FActorNetSpawnInfo()
+		: Location(EForceInit::ForceInitToZero)
+		, Rotation(EForceInit::ForceInitToZero)
+		, Scale(FVector::OneVector)
+		, Velocity(EForceInit::ForceInitToZero)
+		{}
+
+		FVector Location;
+		FRotator Rotation;
+		FVector Scale;
+		FVector Velocity;
+	};
+
+	FActorNetSpawnInfo SpawnInfo;
+
+	FNetObjectReference ObjectReference;
+	FNetObjectReference ArchetypeReference;
+	FNetObjectReference LevelReference; // Only when bUsePersistentLevel is false
+	
+	bool bIsDynamic = false;
+	bool bUsePersistentLevel = false;
+
+	TArray<uint8> CustomCreationData;
+	uint16 CustomCreationDataBitCount = 0;
+	
+	bool Serialize(const FCreationHeaderContext& Context, UE::Net::Private::EActorNetSpawnInfoFlags SpawnFlags, const FActorNetSpawnInfo& DefaultSpawnInfo) const;
+	bool Deserialize(const FCreationHeaderContext& Context, const FActorNetSpawnInfo& DefaultSpawnInfo);
+
+	virtual FString ToString() const override;
+};
+
+#endif // UE_WITH_IRIS
+
+} // end namespace UE::Net
+
+/**
+ * Responsible for creating headers allowing remote factories to spawn replicated actors
+ */
+UCLASS()
+class UNetActorFactory : public UNetObjectFactory
+{
+	GENERATED_BODY()
+
+#if UE_WITH_IRIS
+
+public:
+
+	static FName GetFactoryName() { return TEXT("NetActorFactory"); }
+
+	virtual void OnInit() override;
+
+	virtual FInstantiateResult InstantiateReplicatedObjectFromHeader(const FInstantiateContext& Context, const UE::Net::FNetObjectCreationHeader* Header) override;
+
+	virtual void PostInstantiation(const FPostInstantiationContext& Context) override;
+
+protected:
+
+	virtual TUniquePtr<UE::Net::FNetObjectCreationHeader> CreateAndFillHeader(UE::Net::FNetRefHandle Handle) override;
+	virtual TUniquePtr<UE::Net::FNetObjectCreationHeader> CreateAndDeserializeHeader(const UE::Net::FCreationHeaderContext& Context) override;
+
+	virtual bool SerializeHeader(const UE::Net::FCreationHeaderContext& Context, const UE::Net::FNetObjectCreationHeader* Header)  override;
+
+private:
+
+	UE::Net::Private::EActorNetSpawnInfoFlags SpawnInfoFlags;
+
+	const UE::Net::FActorNetCreationHeader::FActorNetSpawnInfo DefaultSpawnInfo;
+
+#endif // UE_WITH_IRIS
+};
+
+
+
diff --git a/Engine/Source/Runtime/Engine/Public/Net/Iris/ReplicationSystem/NetEngineFactories.h b/Engine/Source/Runtime/Engine/Public/Net/Iris/ReplicationSystem/NetEngineFactories.h
new file mode 100644
--- /dev/null
+++ b/Engine/Source/Runtime/Engine/Public/Net/Iris/ReplicationSystem/NetEngineFactories.h
@@ -0,0 +1,13 @@
+// Copyright Epic Games, Inc. All Rights Reserved.
+
+#pragma once
+
+#if UE_WITH_IRIS
+
+namespace UE::Net
+{
+	void InitEngineNetObjectFactories();
+	void ShutdownEngineNetObjectFactories();
+}
+
+#endif //UE_WITH_IRIS
\ No newline at end of file
diff --git a/Engine/Source/Runtime/Engine/Public/Net/Iris/ReplicationSystem/NetSubObjectFactory.h b/Engine/Source/Runtime/Engine/Public/Net/Iris/ReplicationSystem/NetSubObjectFactory.h
new file mode 100644
--- /dev/null
+++ b/Engine/Source/Runtime/Engine/Public/Net/Iris/ReplicationSystem/NetSubObjectFactory.h
@@ -0,0 +1,73 @@
+// Copyright Epic Games, Inc. All Rights Reserved.
+
+#pragma once
+
+#include "Iris/ReplicationSystem/NetObjectFactory.h"
+
+#if UE_WITH_IRIS
+#include "Iris/Core/NetObjectReference.h"
+#endif
+
+#include "NetSubObjectFactory.generated.h"
+
+/**
+ * Responsible for creating headers allowing remote factories to spawn replicated actors
+ */
+UCLASS()
+class UNetSubObjectFactory : public UNetObjectFactory
+{
+	GENERATED_BODY()
+
+#if UE_WITH_IRIS
+
+public:
+
+	static FName GetFactoryName() { return TEXT("NetSubObjectFactory"); }
+
+	virtual FInstantiateResult InstantiateReplicatedObjectFromHeader(const FInstantiateContext& Context, const UE::Net::FNetObjectCreationHeader* Header) override;
+
+protected:
+
+	virtual TUniquePtr<UE::Net::FNetObjectCreationHeader> CreateAndFillHeader(UE::Net::FNetRefHandle Handle) override;
+	virtual TUniquePtr<UE::Net::FNetObjectCreationHeader> CreateAndDeserializeHeader(const UE::Net::FCreationHeaderContext& Context) override;
+
+	virtual bool SerializeHeader(const UE::Net::FCreationHeaderContext& Context, const UE::Net::FNetObjectCreationHeader* Header)  override;
+#endif // UE_WITH_IRIS
+};
+
+
+namespace UE::Net
+{
+#if UE_WITH_IRIS
+/**
+* Header information representing subobjects
+*/
+class FNetSubObjectCreationHeader : public FNetObjectCreationHeader
+{
+public:
+
+	bool Serialize(const FCreationHeaderContext& Context) const;
+	bool Deserialize(const FCreationHeaderContext& Context);
+
+	virtual FString ToString() const override;
+
+	FNetSubObjectCreationHeader()
+		: bIsDynamic(false)
+		, bIsNameStableForNetworking(false)
+		, bUsePersistentLevel(false)
+		, bOuterIsTransientLevel(false)
+		, bOuterIsRootObject(false)
+	{}
+
+	FNetObjectReference ObjectReference; // Only for static objects
+	FNetObjectReference ObjectClassReference; // Only for dynamic objects
+	FNetObjectReference OuterReference; // Optional: Outer ref sent only for dynamic subobjects.
+	uint8 bIsDynamic : 1;
+	uint8 bIsNameStableForNetworking : 1;
+	uint8 bUsePersistentLevel : 1;
+	uint8 bOuterIsTransientLevel : 1;	// When set the OuterReference was not sent because the Outer is the default transient level.
+	uint8 bOuterIsRootObject : 1;		// When set the OuterReference was not sent because the Outer is the known RootObject.
+
+};
+#endif // UE_WITH_IRIS
+} // end namespace UE::Net
\ No newline at end of file
diff --git a/Engine/Source/Runtime/Experimental/Iris/Core/Private/Iris/ReplicationSystem/NetObjectFactory.cpp b/Engine/Source/Runtime/Experimental/Iris/Core/Private/Iris/ReplicationSystem/NetObjectFactory.cpp
new file mode 100644
--- /dev/null
+++ b/Engine/Source/Runtime/Experimental/Iris/Core/Private/Iris/ReplicationSystem/NetObjectFactory.cpp
@@ -0,0 +1,137 @@
+// Copyright Epic Games, Inc. All Rights Reserved.
+
+#include "Iris/ReplicationSystem/NetObjectFactory.h"
+
+#include "Iris/Core/IrisLog.h"
+
+#include "Iris/ReplicationSystem/ObjectReplicationBridge.h"
+
+#include "Iris/Serialization/NetBitStreamReader.h"
+#include "Iris/Serialization/NetBitStreamWriter.h"
+#include "Iris/Serialization/NetBitStreamUtil.h"
+#include "Iris/Serialization/NetSerializationContext.h"
+
+// Define this on both client and server to add extra header information and quickly detect errors in the header serialization/deserialization code
+#ifndef IRIS_CREATIONHEADER_BITGUARD
+	#define IRIS_CREATIONHEADER_BITGUARD 0
+#endif
+
+
+//------------------------------------------------------------------------
+// UNetObjectFactory
+//------------------------------------------------------------------------
+
+void UNetObjectFactory::Init(UE::Net::FNetObjectFactoryId InId, UObjectReplicationBridge* InBridge)
+{
+	FactoryId = InId;
+	Bridge = InBridge;
+
+	OnInit();
+}
+
+void UNetObjectFactory::Deinit()
+{
+	OnDeinit();
+
+	Bridge = nullptr;
+}
+
+TUniquePtr<UE::Net::FNetObjectCreationHeader> UNetObjectFactory::CreateHeader(UE::Net::FNetRefHandle Handle, UE::Net::FReplicationProtocolIdentifier ProtocolId)
+{
+	using namespace UE::Net;
+
+	const bool bIsValid = Bridge->IsReplicatedHandle(Handle);
+	if (UNLIKELY(!bIsValid))
+	{
+		ensureMsgf(bIsValid, TEXT("%s::CreateLocalHeader received invalid or non-replicated handle: %s"), *GetNameSafe(GetClass()), *Handle.ToString());
+		return nullptr;
+	}
+
+	// Ask the derived class to allocate and fill the header
+	TUniquePtr<FNetObjectCreationHeader> Header = CreateAndFillHeader(Handle);
+
+	if (LIKELY(Header.IsValid()))
+	{
+		Header->SetProtocolId(ProtocolId);
+		Header->SetFactoryId(FactoryId);
+		return Header;
+	}
+
+	return nullptr;
+}
+
+bool UNetObjectFactory::WriteHeader(UE::Net::FNetRefHandle Handle, UE::Net::FNetSerializationContext& Serialization, const UE::Net::FNetObjectCreationHeader* Header)
+{
+	using namespace UE::Net;
+
+	FNetBitStreamWriter* Writer = Serialization.GetBitStreamWriter();
+	check(Writer);
+
+	check(Header->GetNetFactoryId() == FactoryId);
+	check(FactoryId != InvalidNetObjectFactoryId);
+
+	// FactoryID is always serialized first since the bridge will read it first to find the right factory 
+	Writer->WriteBits(Header->GetNetFactoryId(), FNetObjectFactoryRegistry::GetMaxBits());
+	Writer->WriteBits(Header->GetProtocolId(), 32);
+
+#if IRIS_CREATIONHEADER_BITGUARD
+	// Preserialize bits that will hold the header bit size
+	const uint32 StartPos = Writer->GetPosBits();
+	Writer->WriteBits(0, 32);
+#endif
+
+	
+	const bool bSuccess = SerializeHeader(FCreationHeaderContext(Handle, Bridge, this, Serialization), Header);
+
+#if IRIS_CREATIONHEADER_BITGUARD
+	if (bSuccess)
+	{
+		// Go back and write the final header bit size
+		const uint32 BitsWritten = Writer->GetPosBits() - StartPos;
+		FNetBitStreamWriteScope WriteScope(*Writer, StartPos);
+		Writer->WriteBits(BitsWritten, 32);
+	}
+#endif
+
+	return bSuccess && !Writer->IsOverflown();
+}
+
+TUniquePtr<UE::Net::FNetObjectCreationHeader> UNetObjectFactory::ReadHeader(UE::Net::FNetRefHandle Handle, UE::Net::FNetSerializationContext& Serialization)
+{
+	using namespace UE::Net;
+
+	FNetBitStreamReader* Reader = Serialization.GetBitStreamReader();
+	check(Reader);
+
+	// FactoryId was already read by the Bridge
+	const FReplicationProtocolIdentifier ProtocolId = Reader->ReadBits(32);
+
+#if IRIS_CREATIONHEADER_BITGUARD
+	// Find the amount of bits we are expected to read
+	const uint32 StartPos = Reader->GetPosBits();
+	const uint32 ExpectedReadBits = Reader->ReadBits(32);
+#endif
+
+	TUniquePtr<FNetObjectCreationHeader> Header = CreateAndDeserializeHeader(FCreationHeaderContext(Handle, Bridge, this, Serialization));	
+
+#if IRIS_CREATIONHEADER_BITGUARD
+	const uint32 ActualReadBits = Reader->GetPosBits() - StartPos;
+	if (ActualReadBits != ExpectedReadBits)
+	{
+		Reader->DoOverflow();
+		UE_LOG(LogIris, Error, TEXT("Found deserialization error in %s for %s. Header: %s. Source wrote %u bits but we read %u bits (delta: %d)"), *GetName(), *Handle.ToString(), Header.IsValid() ? *Header->ToString() : TEXT("invalid"), ExpectedReadBits, ActualReadBits, (ActualReadBits-ExpectedReadBits));
+		ensureMsgf(ActualReadBits == ExpectedReadBits, TEXT("Found deserialization error in %s for %s"), *GetName(), *Handle.ToString());
+		return nullptr;
+	}
+#endif
+
+	if (LIKELY(Header.IsValid() && !Reader->IsOverflown()))
+	{
+		Header->SetFactoryId(FactoryId);
+		Header->SetProtocolId(ProtocolId);
+		return Header;
+	}
+
+	return nullptr;
+}
+
diff --git a/Engine/Source/Runtime/Experimental/Iris/Core/Private/Iris/ReplicationSystem/NetObjectFactoryRegistry.cpp b/Engine/Source/Runtime/Experimental/Iris/Core/Private/Iris/ReplicationSystem/NetObjectFactoryRegistry.cpp
new file mode 100644
--- /dev/null
+++ b/Engine/Source/Runtime/Experimental/Iris/Core/Private/Iris/ReplicationSystem/NetObjectFactoryRegistry.cpp
@@ -0,0 +1,126 @@
+// Copyright Epic Games, Inc. All Rights Reserved.
+
+#include "Iris/ReplicationSystem/NetObjectFactoryRegistry.h"
+
+#include "Iris/ReplicationSystem/NetObjectFactory.h"
+#include "Iris/Core/IrisLog.h"
+
+namespace UE::Net::Private
+{
+	bool bIsFactoryRegistrationAllowed = true;
+} // end namespace UE::Net::Private
+
+namespace UE::Net 
+{
+
+TArray<FNetObjectFactoryRegistry::FFactoryData, TFixedAllocator<FNetObjectFactoryRegistry::MaxFactories>>  FNetObjectFactoryRegistry::NetFactories;
+
+void FNetObjectFactoryRegistry::SetFactoryRegistrationAllowed(bool bAllowed)
+{
+	UE::Net::Private::bIsFactoryRegistrationAllowed = bAllowed;
+}
+
+void FNetObjectFactoryRegistry::RegisterFactory(UClass* FactoryClass, FName FactoryName)
+{
+	using namespace UE::Net::Private;
+
+#if UE_WITH_IRIS
+
+	// Factories cannot be modified while Iris NetDrivers exist
+	if (!bIsFactoryRegistrationAllowed)
+	{
+		UE_LOG(LogIris, Error, TEXT("FNetObjectFactoryRegistry::RegisterFactory cannot register factory: %s name: %s because it was called while Iris replication was already started"), *GetNameSafe(FactoryClass), *FactoryName.ToString());
+		check(bIsFactoryRegistrationAllowed);
+		return;
+	}
+
+	check(FactoryClass);
+
+	if (FactoryName.IsNone())
+	{
+		checkf(!FactoryName.IsNone(), TEXT("FNetObjectFactoryRegistry::RegisterFactory cannot register %s due to invalid name"), *GetNameSafe(FactoryClass));
+		return;
+	}
+
+	if (NetFactories.Num() >= MaxFactories)
+	{
+		UE_LOG(LogIris, Error, TEXT("FNetObjectFactoryRegistry::RegisterFactory already has %u factories registered. Cannot register factory: %s with name: %s"), MaxFactories, *GetNameSafe(FactoryClass), *FactoryName.ToString());
+		ensure(false);
+		return;
+	}
+
+	// Make sure the name is unique
+	for (const FFactoryData& FactoryData : NetFactories)
+	{
+		if (FactoryData.Name == FactoryName)
+		{
+			UE_LOG(LogIris, Error, TEXT("FNetObjectFactoryRegistry::RegisterFactory cannot register factory: %s with name: %s. This name is already used by factory: %s id: %u"), *GetNameSafe(FactoryClass), *FactoryName.ToString(), *GetNameSafe(FactoryData.NetFactoryClass), FactoryData.Id);
+			ensure(false);
+			return;
+		}
+	}
+
+	// Make sure the class is of the correct type
+	if (!FactoryClass->IsChildOf<UNetObjectFactory>())
+	{
+		checkf(FactoryClass->IsChildOf<UNetObjectFactory>(), TEXT("FNetObjectFactoryRegistry::RegisterFactory factory: %s name: %s is not derived from UNetObjectFactory."), *GetNameSafe(FactoryClass), *FactoryName.ToString());
+		return;
+	}
+	
+	const int32 Index = NetFactories.Emplace( 
+		FFactoryData 
+		{
+			.Name = FactoryName,
+			.NetFactoryClass = FactoryClass
+		});
+	NetFactories[Index].Id = IntCastChecked<FNetObjectFactoryId>(Index);
+
+	UE_LOG(LogIris, Log, TEXT("FNetObjectFactoryRegistry::UnregisterFactory is unregistering factory: %s name: %s id: %u"), *GetNameSafe(NetFactories[Index].NetFactoryClass), *NetFactories[Index].Name.ToString(), NetFactories[Index].Id);
+#endif
+}
+
+void FNetObjectFactoryRegistry::UnregisterFactory(FName FactoryName)
+{
+	using namespace UE::Net::Private;
+#if UE_WITH_IRIS
+	
+	// Factories cannot be modified while Iris NetDrivers exist
+	if (!bIsFactoryRegistrationAllowed)
+	{
+		UE_LOG(LogIris, Error, TEXT("FNetObjectFactoryRegistry::UnregisterFactory cannot unregister factory name: %s because it was called while Iris replication was already started"), *FactoryName.ToString());
+		check(bIsFactoryRegistrationAllowed);
+		return;
+	}
+
+	for (int32 Index=0; Index < NetFactories.Num(); ++Index)
+	{
+		const FFactoryData& Data = NetFactories[Index];
+		if (Data.Name == FactoryName)
+		{
+			UE_LOG(LogIris, Log, TEXT("FNetObjectFactoryRegistry::UnregisterFactory is unregistering factory: %s name: %s id: %u"), *GetNameSafe(Data.NetFactoryClass), *Data.Name.ToString(), Data.Id);
+			NetFactories.RemoveAt(Index, EAllowShrinking::No);
+			return;
+		}
+	}
+
+	UE_LOG(LogIris, Error, TEXT("FNetObjectFactoryRegistry::UnregisterFactory could not find any factories using name: %s"), *FactoryName.ToString());
+	ensure(false);
+#endif
+}
+	  
+FNetObjectFactoryId FNetObjectFactoryRegistry::GetFactoryIdFromName(FName FactoryName)
+{
+#if UE_WITH_IRIS
+	for (const FFactoryData& Data : NetFactories)
+	{
+		if (Data.Name == FactoryName)
+		{
+			return Data.Id;
+		}
+	}
+#endif
+
+	return InvalidNetObjectFactoryId;
+}
+
+} // end namespace UE::Net
\ No newline at end of file
diff --git a/Engine/Source/Runtime/Experimental/Iris/Core/Private/Iris/ReplicationSystem/NetRefHandleManager.h b/Engine/Source/Runtime/Experimental/Iris/Core/Private/Iris/ReplicationSystem/NetRefHandleManager.h
--- a/Engine/Source/Runtime/Experimental/Iris/Core/Private/Iris/ReplicationSystem/NetRefHandleManager.h
+++ b/Engine/Source/Runtime/Experimental/Iris/Core/Private/Iris/ReplicationSystem/NetRefHandleManager.h
@@ -13,6 +13,7 @@
 #include "Iris/ReplicationState/ReplicationStateDescriptor.h"
 #include "Iris/ReplicationSystem/NetRefHandle.h"
 #include "Iris/ReplicationSystem/NetDependencyData.h"
+#include "Iris/ReplicationSystem/NetObjectFactoryRegistry.h"
 #include "UObject/ObjectPtr.h"
 #include "Delegates/Delegate.h"
 
@@ -65,40 +66,41 @@ public:
 	struct FReplicatedObjectData
 	{
 		FReplicatedObjectData()
-		: Protocol(nullptr)
-		, InstanceProtocol(nullptr)
-		, ReceiveStateBuffer(nullptr)
-		, SubObjectRootIndex(InvalidInternalIndex)
-		, SubObjectParentIndex(InvalidInternalIndex)
-		, Flags(0U)
+		: Flags(0U)
 		{
 		}
 
 		FNetRefHandle RefHandle;
 		FNetHandle NetHandle;
-		const FReplicationProtocol* Protocol;
-		const FReplicationInstanceProtocol* InstanceProtocol;
-		uint8* ReceiveStateBuffer;
+		
+		const FReplicationProtocol* Protocol = nullptr;
+		const FReplicationInstanceProtocol* InstanceProtocol = nullptr;
+		uint8* ReceiveStateBuffer = nullptr ;
+		
 		/** Subobjects only: Internal index of the RootObject of this subobject */
-		FInternalNetRefIndex SubObjectRootIndex;
+		FInternalNetRefIndex SubObjectRootIndex = InvalidInternalIndex;
+		
 		/** Subobjects only: Internal index of the ParentObject of this subobject */
-		FInternalNetRefIndex SubObjectParentIndex;
+		FInternalNetRefIndex SubObjectParentIndex = InvalidInternalIndex;
+
+		/** The factory responsible for instantiating this object */
+		FNetObjectFactoryId NetFactoryId = InvalidNetObjectFactoryId;
 
 		union
 		{
-			uint32 Flags : 32U;
+			uint16 Flags : 16U;
 			struct
 			{
-				uint32 bShouldPropagateChangedStates : 1U;
-				uint32 bTearOff : 1U;
-				uint32 bDestroySubObjectWithOwner : 1U;
-				uint32 bIsDependentObject : 1U;
-				uint32 bHasDependentObjects : 1U;
-				uint32 bAllowDestroyInstanceFromRemote : 1U;
-				uint32 bNeedsFullCopyAndQuantize : 1U;
-				uint32 bWantsFullPoll : 1U;
-				uint32 bPendingEndReplication : 1U;
-				uint32 bHasCachedCreationInfo : 1U;
+				uint16 bShouldPropagateChangedStates : 1U;
+				uint16 bTearOff : 1U;
+				uint16 bDestroySubObjectWithOwner : 1U;
+				uint16 bIsDependentObject : 1U;
+				uint16 bHasDependentObjects : 1U;
+				uint16 bAllowDestroyInstanceFromRemote : 1U;
+				uint16 bNeedsFullCopyAndQuantize : 1U;
+				uint16 bWantsFullPoll : 1U;
+				uint16 bPendingEndReplication : 1U;
+				uint16 bHasCachedCreationInfo : 1U;
 			};
 		};
 	
diff --git a/Engine/Source/Runtime/Experimental/Iris/Core/Private/Iris/ReplicationSystem/ObjectReplicationBridge.cpp b/Engine/Source/Runtime/Experimental/Iris/Core/Private/Iris/ReplicationSystem/ObjectReplicationBridge.cpp
--- a/Engine/Source/Runtime/Experimental/Iris/Core/Private/Iris/ReplicationSystem/ObjectReplicationBridge.cpp
+++ b/Engine/Source/Runtime/Experimental/Iris/Core/Private/Iris/ReplicationSystem/ObjectReplicationBridge.cpp
@@ -36,6 +36,8 @@
 #include "Iris/ReplicationSystem/ReplicationOperations.h"
 #include "Iris/ReplicationSystem/ReplicationOperationsInternal.h"
 #include "Iris/ReplicationSystem/ReplicationWriter.h"
+#include "Iris/ReplicationSystem/NetObjectFactory.h"
+#include "Iris/ReplicationSystem/NetObjectFactoryRegistry.h"
 #include "Iris/ReplicationSystem/RepTag.h"
 #include "Iris/ReplicationSystem/Polling/ObjectPoller.h"
 #include "Iris/Serialization/NetSerializationContext.h"
@@ -163,12 +165,15 @@ void UObjectReplicationBridge::Initialize(UReplicationSystem* InReplicationSyste
 	LoadConfig();
 
 	InitConditionalPropertyDelegates();
+	InitNetObjectFactories();
 }
 
 void UObjectReplicationBridge::Deinitialize()
 {
 	NetRefHandleManager->GetOnMaxInternalNetRefIndexIncreasedDelegate().RemoveAll(this);
 
+	DeinitNetObjectFactories();
+
 	UE::Net::Private::FPropertyConditionDelegates::GetOnPropertyCustomConditionChangedDelegate().Remove(OnCustomConditionChangedHandle);
 	UE::Net::Private::FPropertyConditionDelegates::GetOnPropertyDynamicConditionChangedDelegate().Remove(OnDynamicConditionChangedHandle);
 	OnCustomConditionChangedHandle.Reset();
@@ -177,6 +182,42 @@ void UObjectReplicationBridge::Deinitialize()
 	Super::Deinitialize();
 }
 
+void UObjectReplicationBridge::InitNetObjectFactories()
+{
+	using namespace UE::Net;
+	TConstArrayView<FNetObjectFactoryRegistry::FFactoryData> RegisteredFactories = FNetObjectFactoryRegistry::GetRegisteredFactories();
+
+	check(NetObjectFactories.IsEmpty());
+
+	for(int32 Index=0; Index < RegisteredFactories.Num(); ++Index )
+	{
+		const FNetObjectFactoryRegistry::FFactoryData& Data = RegisteredFactories[Index];
+		check(Data.Id == (FNetObjectFactoryId)Index);
+		check(Data.NetFactoryClass);
+
+		UNetObjectFactory* Factory = NewObject<UNetObjectFactory>(this, Data.NetFactoryClass);
+		NetObjectFactories.Add(Factory);
+
+		Factory->Init(Data.Id, this);
+	}
+}
+
+void UObjectReplicationBridge::DeinitNetObjectFactories()
+{
+	for (UNetObjectFactory* Factory : NetObjectFactories)
+	{
+		Factory->Deinit();
+	}
+
+	NetObjectFactories.Empty();
+}
+
+UNetObjectFactory* UObjectReplicationBridge::GetNetFactory(UE::Net::FNetObjectFactoryId FactoryId) const
+{
+	check(FactoryId != UE::Net::InvalidNetObjectFactoryId);
+	return NetObjectFactories[FactoryId];
+}
+
 void UObjectReplicationBridge::OnMaxInternalNetRefIndexIncreased(UE::Net::Private::FInternalNetRefIndex NewMaxInternalIndex)
 {
 	PollFrequencyLimiter->OnMaxInternalNetRefIndexIncreased(NewMaxInternalIndex);
@@ -245,7 +286,7 @@ UE::Net::FNetRefHandle UObjectReplicationBridge::GetReplicatedRefHandle(FNetHand
 	return ObjectData.RefHandle;
 }
 
-UE::Net::FNetRefHandle UObjectReplicationBridge::StartReplicatingNetObject(UObject* Instance, UE::Net::EReplicationFragmentTraits Traits)
+UE::Net::FNetRefHandle UObjectReplicationBridge::StartReplicatingNetObject(UObject* Instance, UE::Net::EReplicationFragmentTraits Traits, UE::Net::FNetObjectFactoryId NetFactoryId)
 {
 	using namespace UE::Net;
 	using namespace UE::Net::Private;
@@ -347,6 +388,9 @@ UE::Net::FNetRefHandle UObjectReplicationBridge::StartReplicatingNetObject(UObje
 
 	const FInternalNetRefIndex InternalReplicationIndex = NetRefHandleManager->GetInternalIndex(RefHandle);
 
+	check(NetFactoryId != InvalidNetObjectFactoryId);
+	NetRefHandleManager->GetReplicatedObjectDataNoCheck(InternalReplicationIndex).NetFactoryId = NetFactoryId;
+
 	// Initialize conditionals
 	FReplicationSystemInternal* ReplicationSystemInternal = GetReplicationSystem()->GetReplicationSystemInternal();
 	ReplicationSystemInternal->GetConditionals().InitPropertyCustomConditions(InternalReplicationIndex);
@@ -365,7 +409,7 @@ UE::Net::FNetRefHandle UObjectReplicationBridge::StartReplicatingNetObject(UObje
 	return RefHandle;
 }
 
-UE::Net::FNetRefHandle UObjectReplicationBridge::StartReplicatingRootObject(UObject* Instance, const FRootObjectReplicationParams& Params)
+UE::Net::FNetRefHandle UObjectReplicationBridge::StartReplicatingRootObject(UObject* Instance, const FRootObjectReplicationParams& Params, UE::Net::FNetObjectFactoryId NetFactoryId)
 {
 	LLM_SCOPE_BYTAG(IrisState);
 
@@ -383,7 +427,7 @@ UE::Net::FNetRefHandle UObjectReplicationBridge::StartReplicatingRootObject(UObj
 	Traits |= Params.bNeedsWorldLocationUpdate ? EReplicationFragmentTraits::NeedsWorldLocationUpdate : EReplicationFragmentTraits::None;
 
 	// Register the object
-	FNetRefHandle RefHandle = StartReplicatingNetObject(Instance, Traits);
+	FNetRefHandle RefHandle = StartReplicatingNetObject(Instance, Traits, NetFactoryId);
 
 	if (!RefHandle.IsValid())
 	{
@@ -475,7 +519,7 @@ void UObjectReplicationBridge::AssignDynamicFilter(UObject* Instance, const FRoo
 	
 }
 
-UE::Net::FNetRefHandle UObjectReplicationBridge::StartReplicatingSubObject(FNetRefHandle OwnerRefHandle, UObject* Instance, FNetRefHandle InsertRelativeToSubObjectRefHandle, ESubObjectInsertionOrder InsertionOrder)
+UE::Net::FNetRefHandle UObjectReplicationBridge::StartReplicatingSubObject(UObject* Instance, const FSubObjectReplicationParams& Params , UE::Net::FNetObjectFactoryId NetFactoryId)
 {
 	LLM_SCOPE_BYTAG(IrisState);
 
@@ -485,22 +529,22 @@ UE::Net::FNetRefHandle UObjectReplicationBridge::StartReplicatingSubObject(FNetR
 	const FReplicationSystemInternal* ReplicationSystemInternal = GetReplicationSystem()->GetReplicationSystemInternal();
 	const FNetRefHandleManager& LocalNetRefHandleManager = ReplicationSystemInternal->GetNetRefHandleManager();
 
-	checkf(IsReplicatedHandle(OwnerRefHandle), TEXT("Owner %s (%s) must be replicated for subobject %s to replicate."), 
-		*GetNameSafe(LocalNetRefHandleManager.GetReplicatedObjectInstance(LocalNetRefHandleManager.GetInternalIndex(OwnerRefHandle))), *OwnerRefHandle.ToString(), *GetNameSafe(Instance));
+	checkf(IsReplicatedHandle(Params.RootObjectHandle), TEXT("Owner %s (%s) must be replicated for subobject %s to replicate."), 
+		*GetNameSafe(LocalNetRefHandleManager.GetReplicatedObjectInstance(LocalNetRefHandleManager.GetInternalIndex(Params.RootObjectHandle))), *Params.RootObjectHandle.ToString(), *GetNameSafe(Instance));
 
 	checkf(!Instance->HasAnyFlags(RF_ArchetypeObject | RF_ClassDefaultObject), TEXT("Iris cannot replicate subobject %s owned by %s because it's an %s"), 
-		*GetNameSafe(Instance), *GetNameSafe(LocalNetRefHandleManager.GetReplicatedObjectInstance(LocalNetRefHandleManager.GetInternalIndex(OwnerRefHandle))), Instance->HasAnyFlags(RF_ArchetypeObject)?TEXT("Archetype"):TEXT("DefaultObject"));
+		*GetNameSafe(Instance), *GetNameSafe(LocalNetRefHandleManager.GetReplicatedObjectInstance(LocalNetRefHandleManager.GetInternalIndex(Params.RootObjectHandle))), Instance->HasAnyFlags(RF_ArchetypeObject)?TEXT("Archetype"):TEXT("DefaultObject"));
 
 	FNetRefHandle SubObjectRefHandle = GetReplicatedRefHandle(Instance);
 	if (SubObjectRefHandle.IsValid())
 	{
 		// Verify that the existing object is a subobject of the owner
-		check(OwnerRefHandle == LocalNetRefHandleManager.GetRootObjectOfSubObject(SubObjectRefHandle));
+		check(Params.RootObjectHandle == LocalNetRefHandleManager.GetRootObjectOfSubObject(SubObjectRefHandle));
 		return SubObjectRefHandle;
 	}
 	
 	const EReplicationFragmentTraits Traits = EReplicationFragmentTraits::CanReplicate;
-	SubObjectRefHandle = StartReplicatingNetObject(Instance, Traits);
+	SubObjectRefHandle = StartReplicatingNetObject(Instance, Traits, NetFactoryId);
 
 	if (!SubObjectRefHandle.IsValid())
 	{
@@ -508,15 +552,15 @@ UE::Net::FNetRefHandle UObjectReplicationBridge::StartReplicatingSubObject(FNetR
 	}
 
 	// Add subobject
-	InternalAddSubObject(OwnerRefHandle, SubObjectRefHandle, InsertRelativeToSubObjectRefHandle, InsertionOrder);
+	InternalAddSubObject(Params.RootObjectHandle, SubObjectRefHandle, Params.InsertRelativeToSubObjectHandle, Params.InsertionOrder);
 
-	UE_LOG_OBJECTREPLICATIONBRIDGE(Verbose, TEXT("StartReplicatingSubObject Added %s RelativeToSubObjectHandle %s"), *PrintObjectFromNetRefHandle(SubObjectRefHandle), *PrintObjectFromNetRefHandle(InsertRelativeToSubObjectRefHandle));
+	UE_LOG_OBJECTREPLICATIONBRIDGE(Verbose, TEXT("StartReplicatingSubObject Added %s RelativeToSubObjectHandle %s"), *PrintObjectFromNetRefHandle(SubObjectRefHandle), *PrintObjectFromNetRefHandle(Params.InsertRelativeToSubObjectHandle));
 
 	// SubObjects should always poll with owner
-	SetPollWithObject(OwnerRefHandle, SubObjectRefHandle);
+	SetPollWithObject(Params.RootObjectHandle, SubObjectRefHandle);
 
 	// Copy pending dormancy from owner
-	SetSubObjectDormancyStatus(SubObjectRefHandle, OwnerRefHandle);
+	SetSubObjectDormancyStatus(SubObjectRefHandle, Params.RootObjectHandle);
 	
 	return SubObjectRefHandle;
 }
@@ -610,45 +654,78 @@ void UObjectReplicationBridge::RemoveDependentObject(FNetRefHandle ParentHandle,
 
 bool UObjectReplicationBridge::WriteNetRefHandleCreationInfo(FReplicationBridgeSerializationContext& Context, FNetRefHandle Handle)
 {
-	// Write Type header, if there is a cached one, use it!
-	if (TUniquePtr<const FCreationHeader>* CachedHeader = CachedCreationHeaders.Find(Handle))
+	using namespace UE::Net;
+	using namespace UE::Net::Private;
+
+	// If there is a cached header for this handle, use it!
+	if (TUniquePtr<const FNetObjectCreationHeader>* CachedHeader = CachedCreationHeaders.Find(Handle))
 	{
-		Context.SerializationContext.GetBitStreamWriter()->WriteBits((*CachedHeader)->ProtocolIdentifier, 32);
-		return WriteCreationHeader(Context.SerializationContext, (*CachedHeader).Get());
+		const FNetObjectCreationHeader* Header = CachedHeader->Get();
+		UNetObjectFactory* Factory = GetNetFactory(Header->GetNetFactoryId());
+		UE_LOG(LogIrisBridge, VeryVerbose, TEXT("WriteNetRefHandleCreationInfo (cached) for %s : Factory: %s (id:%u) Header: %s"), *Handle.ToString(), *GetNameSafe(Factory), Header->GetNetFactoryId(), *Header->ToString());
+		return Factory->WriteHeader(Handle, Context.SerializationContext, Header);
 	}
-	else
+	
+	// Find the ProtocolId
+	const UE::Net::FReplicationProtocol* Protocol = GetReplicationSystem()->GetReplicationProtocol(Handle);
+	if (!Protocol)
 	{
-		// Write ProtocolId
-		const UE::Net::FReplicationProtocol* Protocol = GetReplicationSystem()->GetReplicationProtocol(Handle);
-		if (!Protocol)
-		{
-			UE_LOG_OBJECTREPLICATIONBRIDGE(Error, TEXT("WriteNetRefHandleCreationInfo: Cannot write creationinfo for %s, since protocol has been detached"), *Handle.ToString());
-			ensureMsgf(Protocol, TEXT("WriteNetRefHandleCreationInfo: Cannot write creationinfo for %s, since protocol has been detached"), *Handle.ToString());
-			return false;
-		}
-		Context.SerializationContext.GetBitStreamWriter()->WriteBits(Protocol->ProtocolIdentifier, 32);
-		return WriteCreationHeader(Context.SerializationContext, Handle);
+		UE_LOG_OBJECTREPLICATIONBRIDGE(Error, TEXT("WriteNetRefHandleCreationInfo: Cannot write creationinfo for %s, since protocol has been detached"), *Handle.ToString());
+		ensureMsgf(Protocol, TEXT("WriteNetRefHandleCreationInfo: Cannot write creationinfo for %s, since protocol has been detached"), *Handle.ToString());
+		return false;
+	}
+
+	const FInternalNetRefIndex ObjectIndex = NetRefHandleManager->GetInternalIndex(Handle);
+	if (ObjectIndex == FNetRefHandleManager::InvalidInternalIndex)
+	{
+		UE_LOG(LogIrisBridge, Error, TEXT("WriteNetRefHandleCreationInfo: Cannot write creationinfo for %s: No internal index assigned to it"), *Handle.ToString());
+		return false;
 	}
+
+	const FNetRefHandleManager::FReplicatedObjectData& ObjectData = NetRefHandleManager->GetReplicatedObjectDataNoCheck(ObjectIndex);
+	const FNetObjectFactoryId FactoryId = ObjectData.NetFactoryId;
+
+	UNetObjectFactory* Factory = GetNetFactory(FactoryId);
+
+	TUniquePtr<FNetObjectCreationHeader> Header = Factory->CreateHeader(Handle, Protocol->ProtocolIdentifier);
+
+	UE_LOG(LogIrisBridge, VeryVerbose, TEXT("WriteNetRefHandleCreationInfo (new) for %s : Factory: %s (id:%u) Header: %s"), *Handle.ToString(), *GetNameSafe(Factory), FactoryId, *Header->ToString());
+
+	return Factory->WriteHeader(Handle, Context.SerializationContext, Header.Get());
 }
 
 bool UObjectReplicationBridge::CacheNetRefHandleCreationInfo(FNetRefHandle Handle)
 {
-	const UE::Net::FReplicationProtocol* Protocol = GetReplicationSystem()->GetReplicationProtocol(Handle);	
+	using namespace UE::Net;
+	using namespace UE::Net::Private;
+
+	const FReplicationProtocol* Protocol = GetReplicationSystem()->GetReplicationProtocol(Handle);	
 	if (!Protocol)
 	{
 		return false;
 	}
 
-	TUniquePtr<FCreationHeader> Header(GetCreationHeader(Handle));
-
-	if (Header)
+	const FInternalNetRefIndex ObjectIndex = NetRefHandleManager->GetInternalIndex(Handle);
+	if (ObjectIndex == FNetRefHandleManager::InvalidInternalIndex)
 	{
-		Header->ProtocolIdentifier = Protocol->ProtocolIdentifier;
-		CachedCreationHeaders.Add(Handle, MoveTemp(Header));
-		return true;
+		UE_LOG(LogIrisBridge, Error, TEXT("CacheNetRefHandleCreationInfo cannot cache creation info for %s: no internal index assigned to it"), *Handle.ToString());
+		ensureMsgf(false, TEXT("CacheNetRefHandleCreationInfo cannot cache creation info for %s: no internal index assigned to it"), *Handle.ToString());
+		return false;
 	}
 
-	return false;
+	const FNetObjectFactoryId FactoryId = NetRefHandleManager->GetReplicatedObjectDataNoCheck(ObjectIndex).NetFactoryId;
+
+	UNetObjectFactory* Factory = GetNetFactory(FactoryId);
+
+	TUniquePtr<FNetObjectCreationHeader> Header = Factory->CreateHeader(Handle, Protocol->ProtocolIdentifier);
+
+	if (!Header.IsValid())
+	{
+		return false;
+	}
+	
+	CachedCreationHeaders.Add(Handle, MoveTemp(Header));
+	return true;
 }
 
 void UObjectReplicationBridge::StopReplicatingNetObject(UObject* Instance, EEndReplicationFlags EndReplicationFlags)
@@ -709,7 +786,7 @@ void UObjectReplicationBridge::UnregisterInstance(FNetRefHandle RefHandle)
 	}
 }
 
-void UObjectReplicationBridge::RegisterRemoteInstance(FNetRefHandle RefHandle, UObject* Instance, const UE::Net::FReplicationProtocol* Protocol, UE::Net::FReplicationInstanceProtocol* InstanceProtocol, const FCreationHeader* Header, uint32 ConnectionId)
+void UObjectReplicationBridge::RegisterRemoteInstance(FNetRefHandle RefHandle, UObject* Instance, const UE::Net::FReplicationProtocol* Protocol, UE::Net::FReplicationInstanceProtocol* InstanceProtocol, uint32 ConnectionId)
 {
 	// Attach the instance protocol and instance to the handle
 	constexpr bool bBindInstanceProtocol = false;
@@ -733,59 +810,68 @@ FReplicationBridgeCreateNetRefHandleResult UObjectReplicationBridge::CreateNetRe
 
 	FNetBitStreamReader* Reader = Context.SerializationContext.GetBitStreamReader();
 
-	FReplicationProtocolIdentifier ReceivedProtocolId = Context.SerializationContext.GetBitStreamReader()->ReadBits(32);
+	const FNetObjectFactoryId FactoryId = IntCastChecked<FNetObjectFactoryId>(Reader->ReadBits(FNetObjectFactoryRegistry::GetMaxBits()));
+	check(FactoryId != InvalidNetObjectFactoryId);
 
-	// Read creation header
-	TUniquePtr<FCreationHeader> Header(ReadCreationHeader(Context.SerializationContext));
-	if (Context.SerializationContext.HasErrorOrOverflow())
+	UNetObjectFactory* Factory = GetNetFactory(FactoryId);
+
+	TUniquePtr<FNetObjectCreationHeader> Header = Factory->ReadHeader(WantedNetHandle, Context.SerializationContext);
+		
+	if (!Header.IsValid() || Context.SerializationContext.HasErrorOrOverflow())
 	{
 		return FReplicationBridgeCreateNetRefHandleResult();
 	}
 
+	const FReplicationProtocolIdentifier ReceivedProtocolId = Header->GetProtocolId();	
+
 	// Currently remote objects can only receive replicated data
 	FFragmentRegistrationContext FragmentRegistrationContext(GetReplicationStateDescriptorRegistry(), GetReplicationSystem(), EReplicationFragmentTraits::CanReceive);
 	FReplicationProtocolManager* ProtocolManager = GetReplicationProtocolManager();
 
-	FReplicationBridgeCreateNetRefHandleResult CreateResult;
-
 	if (UE_LOG_ACTIVE(LogIrisBridge, Verbose))
 	{
 		if (RootObjectOfSubObject.IsValid())
 		{
-			UE_LOG(LogIrisBridge, Verbose, TEXT("CreateNetRefHandleFromRemote: SubObject: %s of RootObject: %s"), *WantedNetHandle.ToString(), *RootObjectOfSubObject.ToString());
+			UE_LOG(LogIrisBridge, Verbose, TEXT("CreateNetRefHandleFromRemote: SubObject: %s of RootObject: %s using header: %s"), *WantedNetHandle.ToString(), *RootObjectOfSubObject.ToString(), ToCStr(Header->ToString()));
 		}
 		else
 		{
-			UE_LOG(LogIrisBridge, Verbose, TEXT("CreateNetRefHandleFromRemote: RootObject: %s"), *WantedNetHandle.ToString());
+			UE_LOG(LogIrisBridge, Verbose, TEXT("CreateNetRefHandleFromRemote: RootObject: %s using header: %s"), *WantedNetHandle.ToString(), ToCStr(Header->ToString()));
 		}
 	}
 	
+	// Currently we need to always instantiate remote objects, moving forward we want to make this optional so that can be deferred until it is time to apply received state data.
+	// https://jira.it.epicgames.com/browse/UE-127369	
+
+	UObject* InstancePtr = nullptr;
+	FReplicationBridgeCreateNetRefHandleResult CreateResult;
 	
+	// Instantiate the replicated object
+	const UNetObjectFactory::FInstantiateContext InstantiateContext(WantedNetHandle, Context.SerializationContext.GetInternalContext()->ResolveContext, RootObjectOfSubObject);
 	
+	const UNetObjectFactory::FInstantiateResult Result = Factory->InstantiateReplicatedObjectFromHeader(InstantiateContext, Header.Get());
 
-	// Currently we need to always instantiate remote objects, moving forward we want to make this optional so that can be deferred until it is time to apply received state data.
-	// https://jira.it.epicgames.com/browse/UE-127369	
-	FObjectReplicationBridgeInstantiateResult InstantiateResult = BeginInstantiateFromRemote(RootObjectOfSubObject, Context.SerializationContext.GetInternalContext()->ResolveContext, Header.Get());
-	UObject* InstancePtr = InstantiateResult.Object;
-	if (!InstancePtr)
+	// Failed ?
+	if (!Result.Instance)
 	{
 		if (UE_LOG_ACTIVE(LogIrisBridge, Warning) && !bSuppressCreateInstanceFailedEnsure)
 		{
 			if (RootObjectOfSubObject.IsValid())
 			{
-				UE_LOG(LogIrisBridge, Warning, TEXT("CreateNetRefHandleFromRemote: Failed to instantiate SubObject NetHandle: %s of %s"), *WantedNetHandle.ToString(), *PrintObjectFromNetRefHandle(RootObjectOfSubObject));
+				UE_LOG(LogIrisBridge, Warning, TEXT("CreateNetRefHandleFromRemote: Failed to instantiate SubObject NetHandle: %s of %s using header: %s"), *WantedNetHandle.ToString(), *PrintObjectFromNetRefHandle(RootObjectOfSubObject), ToCStr(Header->ToString()));
 			}
 			else
 			{
-				UE_LOG(LogIrisBridge, Warning, TEXT("CreateNetRefHandleFromRemote: Failed to instantiate RootObject NetHandle: %s"), *WantedNetHandle.ToString());
+				UE_LOG(LogIrisBridge, Warning, TEXT("CreateNetRefHandleFromRemote: Failed to instantiate RootObject NetHandle: %s using header: %s"), *WantedNetHandle.ToString(), ToCStr(Header->ToString()));
 			}
 		}
 		
 		ensureMsgf(bSuppressCreateInstanceFailedEnsure, TEXT("Failed to instantiate Handle: %s"), *WantedNetHandle.ToString());
-		return CreateResult;
+		return FReplicationBridgeCreateNetRefHandleResult();
 	}
 
-	CreateResult.Flags |= InstantiateResult.Flags;
+	InstancePtr = Result.Instance;
+	CreateResult.Flags |= Result.Flags;
 
 	// Register all fragments
 	CallRegisterReplicationFragments(InstancePtr, FragmentRegistrationContext, EFragmentRegistrationFlags::None);
@@ -836,13 +922,12 @@ FReplicationBridgeCreateNetRefHandleResult UObjectReplicationBridge::CreateNetRe
 		CreateResult.NetRefHandle = Handle;
 		if (Handle.IsValid())
 		{
-			RegisterRemoteInstance(Handle, InstancePtr, ReplicationProtocol, InstanceProtocol.Get(), Header.Get(), Context.ConnectionId);
+			RegisterRemoteInstance(Handle, InstancePtr, ReplicationProtocol, InstanceProtocol.Get(), Context.ConnectionId);
 
 			// Release instance protocol from the uniquePtr as it is now successfully bound to the handle
 			(void)InstanceProtocol.Release();
 
-			// Now it is safe to issue OnActorChannelOpen callback
-			ensureMsgf(OnInstantiatedFromRemote(InstancePtr, Header.Get(), Context.ConnectionId), TEXT("Failed to invoke OnInstantiatedFromRemote for Instance named %s %s"), *InstancePtr->GetName(), *Handle.ToString());
+			Factory->PostInstantiation(UNetObjectFactory::FPostInstantiationContext{ .Instance = InstancePtr, .Header = Header.Get(), .ConnectionId = Context.ConnectionId });
 		}
 	}
 
diff --git a/Engine/Source/Runtime/Experimental/Iris/Core/Public/Iris/ReplicationSystem/NetObjectFactory.h b/Engine/Source/Runtime/Experimental/Iris/Core/Public/Iris/ReplicationSystem/NetObjectFactory.h
new file mode 100644
--- /dev/null
+++ b/Engine/Source/Runtime/Experimental/Iris/Core/Public/Iris/ReplicationSystem/NetObjectFactory.h
@@ -0,0 +1,199 @@
+// Copyright Epic Games, Inc. All Rights Reserved.
+
+#pragma once
+
+#include "UObject/ObjectMacros.h"
+
+#include "Iris/ReplicationSystem/NetObjectFactoryRegistry.h"
+#include "Iris/ReplicationSystem/ReplicationBridgeTypes.h"
+#include "Iris/ReplicationSystem/NetRefHandle.h"
+
+#include "NetObjectFactory.generated.h"
+
+class UObjectReplicationBridge;
+class UNetObjectFactory;
+
+namespace UE::Net
+{
+	class FNetSerializationContext;
+
+	struct FNetObjectResolveContext;
+
+	typedef uint32 FReplicationProtocolIdentifier;
+}
+
+namespace UE::Net
+{
+
+/** Contextual information passed to an header so it can serialize/deserialize itself  */
+struct FCreationHeaderContext
+{
+	/** The handle of the replicated object represented by the header */
+	FNetRefHandle Handle;
+	/** The bridge responsible for the replicated object */
+	UObjectReplicationBridge* Bridge;
+	/** The factory that allocated the header */
+	UNetObjectFactory* Factory;
+	/** Access to the bitstream reader or writer */
+	FNetSerializationContext& Serialization;
+
+	FCreationHeaderContext(FNetRefHandle InHandle, UObjectReplicationBridge* InBridge, UNetObjectFactory* InFactory, FNetSerializationContext& InSerialization) : Handle(InHandle), Bridge(InBridge), Factory(InFactory), Serialization(InSerialization) {}
+};
+
+/*
+ * Class holding the raw information allowing any client from retrieving or allocating a replicated UObject instance.
+ * Can also implement the serialization of it's data into a bitstream
+ */ 
+class FNetObjectCreationHeader
+{
+public:
+
+	virtual ~FNetObjectCreationHeader() {};
+
+	void SetProtocolId(uint32 InId) { ProtocolIdentifier = InId; }
+	void SetFactoryId(FNetObjectFactoryId InId) { FactoryId = InId; }
+
+	FReplicationProtocolIdentifier GetProtocolId() const { return ProtocolIdentifier; }
+	FNetObjectFactoryId GetNetFactoryId() const { return FactoryId; }
+
+	/** Transform the header information into a readable format */
+	virtual FString ToString() const { return TEXT("NotImplemented"); }
+
+private:
+
+	FReplicationProtocolIdentifier ProtocolIdentifier = 0;
+	FNetObjectFactoryId  FactoryId = InvalidNetObjectFactoryId;
+};
+
+} // end namespace UE::Net
+
+/**
+ * The class is responsible for creating the header representing specific replicated object types.
+ * Also responsible for instantiating the UObject from a replicated header.
+ */
+UCLASS(MinimalAPI, transient, abstract)
+class UNetObjectFactory : public UObject
+{
+	GENERATED_BODY()
+
+public:
+
+	void Init(UE::Net::FNetObjectFactoryId InId, UObjectReplicationBridge* InBridge);
+	void Deinit();
+
+	struct FInstantiateResult;
+	struct FInstantiateContext;
+	struct FPostInstantiationContext;
+
+	/**
+	 * Creates the header containing all information required to instantiate a remote version of the object represented by the handle.
+	 * @param Handle The handle of the object represented by the header
+	 * @param ProtocolId The protocol id to add to the header
+	 * @return Return a valid header filled with the information that will be sent to remote connections
+	 */
+	TUniquePtr<UE::Net::FNetObjectCreationHeader> CreateHeader(UE::Net::FNetRefHandle Handle, UE::Net::FReplicationProtocolIdentifier ProtocolId);
+
+	/** 
+	 * Serializes a valid header so it can be replicated to remote connections 
+	 * @param Handle The handle of the object represented by the header
+	 * @param Header The filled header that will be serialized
+	 * @return Returns true if the serialization was a success
+	 */
+	bool WriteHeader(UE::Net::FNetRefHandle Handle, UE::Net::FNetSerializationContext& Serialization, const UE::Net::FNetObjectCreationHeader* Header);
+
+	/** 
+	 * Deserialize the header data received and return a valid header
+	 * @param Handle The handle of the object represented by the header
+	 * @param Serialization Gives access to the bitstream reader
+	 * @return Return a valid header filled with the information that represents the remote object
+	 */
+	TUniquePtr<UE::Net::FNetObjectCreationHeader> ReadHeader(UE::Net::FNetRefHandle Handle, UE::Net::FNetSerializationContext& Serialization);
+
+	/**
+	* Create or bind a replicated object from the received creation header.
+	* @param Context Gives you access to useful info on the object to instantiate
+	* @param Header The filled header information to use to spawn the object
+	* @return The instantiated object if successful and relevant flags for the bridge to act on
+	*/
+	virtual FInstantiateResult InstantiateReplicatedObjectFromHeader(const FInstantiateContext& Context, const UE::Net::FNetObjectCreationHeader* Header) PURE_VIRTUAL(UNetObjectFactory::InstantiateReplicatedObjectFromHeader, return FInstantiateResult(););
+
+	/**
+	 * Optional callback triggered at the end of the instantiation process and before any replicated properties were applied.
+	 * Useful to apply any additional data included with the header or to signal the new object to other systems.
+	 * @param Context Gives you access to the new object and the header it was created from.
+	 */
+	virtual void PostInstantiation(const FPostInstantiationContext& Context) {}
+
+public:
+
+	/** Result of the instantiate request */
+	struct FInstantiateResult
+	{
+		/** The instantiated object represented by the header */
+		UObject* Instance = nullptr;
+		/** Flags to pass back to the bridge */
+		EReplicationBridgeCreateNetRefHandleResultFlags Flags = EReplicationBridgeCreateNetRefHandleResultFlags::None;
+	};
+
+	/** Contextual information to use during instantiation */
+	struct FInstantiateContext
+	{
+		/** The handle tied to the replicated object to instantiate */
+		UE::Net::FNetRefHandle Handle;
+
+		const UE::Net::FNetObjectResolveContext& ResolveContext;
+
+		/** The handle of the object's root object (only if instantiating a subobject) */
+		UE::Net::FNetRefHandle RootObjectOfSubObject;
+
+		FInstantiateContext(UE::Net::FNetRefHandle InHandle, const UE::Net::FNetObjectResolveContext& InResolveContext, UE::Net::FNetRefHandle InRootObjectHandle) : Handle(InHandle), ResolveContext(InResolveContext), RootObjectOfSubObject(InRootObjectHandle) {}
+
+		/** Tells if we instantiating a root object or a subobject. */
+		bool IsRootObject() const { return !RootObjectOfSubObject.IsValid(); }
+		bool IsSubObject() const { return !IsRootObject(); }
+	};
+
+	/** Contextual information to use in the PostInstantiation callback */
+	struct FPostInstantiationContext
+	{
+		/** The object instantiated */
+		UObject* Instance = nullptr;
+		/** The header representing the replicated object */
+		const UE::Net::FNetObjectCreationHeader* Header = nullptr;
+		/** The connection that owns the replicated object */
+		uint32 ConnectionId = 0;
+	};
+
+protected:
+
+	virtual void OnInit() {}
+	virtual void OnDeinit() {}
+
+	/**
+	 * Create the correct header type for a given replicated object and fill the header with the information representing it.
+	 * @param Handle The net handle of the replicated object.
+	 * @return A valid creation header if successful. 
+	 */
+	virtual TUniquePtr<UE::Net::FNetObjectCreationHeader> CreateAndFillHeader(UE::Net::FNetRefHandle Handle) PURE_VIRTUAL(UNetObjectFactory::CreateAndFillHeader, return nullptr;);
+
+	/**
+	* Serialize the header into the bitstream
+	* @param Context Gives access to the bitstream writer and other useful information.
+	* @param Header A valid and filled header to serialize
+	* @return Return true if the header serialization was successful.
+	*/
+	virtual bool SerializeHeader(const UE::Net::FCreationHeaderContext& Context, const UE::Net::FNetObjectCreationHeader* Header) PURE_VIRTUAL(UNetObjectFactory::SerializeHeader, return false;);
+
+	/**
+	* Create a new header and deserialize it's data from the incoming bitstream
+	* @param Serialization Gives you access to the bit reader
+	* @return Return a valid and filled header if successful.
+	*/
+	virtual TUniquePtr<UE::Net::FNetObjectCreationHeader> CreateAndDeserializeHeader(const UE::Net::FCreationHeaderContext& Context) PURE_VIRTUAL(UNetObjectFactory::CreateAndDeserializeHeader, return nullptr;);
+
+protected:
+
+	UObjectReplicationBridge* Bridge = nullptr;
+	UE::Net::FNetObjectFactoryId FactoryId = UE::Net::InvalidNetObjectFactoryId;
+
+};
\ No newline at end of file
diff --git a/Engine/Source/Runtime/Experimental/Iris/Core/Public/Iris/ReplicationSystem/NetObjectFactoryRegistry.h b/Engine/Source/Runtime/Experimental/Iris/Core/Public/Iris/ReplicationSystem/NetObjectFactoryRegistry.h
new file mode 100644
--- /dev/null
+++ b/Engine/Source/Runtime/Experimental/Iris/Core/Public/Iris/ReplicationSystem/NetObjectFactoryRegistry.h
@@ -0,0 +1,74 @@
+// Copyright Epic Games, Inc. All Rights Reserved.
+
+#pragma once
+
+#include "Containers/ArrayView.h"
+#include "UObject/NameTypes.h"
+
+// Define this to increase the maximum amount of NetObjectFactories that can exist
+#ifndef UE_IRIS_MAX_NETOBJECT_FACTORIES
+	#define UE_IRIS_MAX_NETOBJECT_FACTORIES 8
+#endif
+
+namespace UE::Net
+{
+
+typedef uint8 FNetObjectFactoryId;
+enum { InvalidNetObjectFactoryId = MAX_uint8 };
+
+
+/**
+ * Keeps track of Iris NetObjectFactory templates
+ * NetObjectFactories must all be registered before any Iris ReplicationSystem is created.
+ */
+class FNetObjectFactoryRegistry
+{
+public:
+
+	/** Register a UNetObjectFactory class and associate with a specific name. */
+	IRISCORE_API static void RegisterFactory(UClass* FactoryClass, FName FactoryName);
+
+	/** Unregister the UNetObjectFactory class associated with the name */
+	IRISCORE_API static void UnregisterFactory(FName FactoryName);
+
+	/** Find the FNetFactoryID that was assigned to name on registration */
+	IRISCORE_API static FNetObjectFactoryId GetFactoryIdFromName(FName FactoryName);
+
+	/** The engine sets this false after it created an Iris replication system since its now illegal to register new factories */
+	IRISCORE_API static void SetFactoryRegistrationAllowed(bool bAllowed);
+
+	/** Limit how many factories can be registered */
+	static constexpr uint32 MaxFactories = UE_IRIS_MAX_NETOBJECT_FACTORIES;
+
+	/** The amount of bits to serialize FNetFactoryIds with  */
+	static constexpr uint32 GetMaxBits() { return GetNumBits(MaxFactories-1); }
+
+	struct FFactoryData
+	{
+		/** Name associated with this factory class */
+		FName Name;
+		/** Class representing a concrete UNetObjectFactory */
+		UClass* NetFactoryClass = nullptr;
+		/** FactoryId assigned to this factory */
+		FNetObjectFactoryId Id = InvalidNetObjectFactoryId;
+	};
+
+	/** The registered factories ready to be instantiated by the replication bridge */
+	static const TConstArrayView<FFactoryData> GetRegisteredFactories() { return MakeConstArrayView(NetFactories.GetData(), NetFactories.Num()); }
+
+private:
+	
+	FNetObjectFactoryRegistry() = delete;
+	~FNetObjectFactoryRegistry() = delete;
+
+	static constexpr uint32 GetNumBits(uint32 Number)
+	{
+		return Number==0 ? 0 : 1 + GetNumBits(Number >> 1);
+	}
+
+private:
+
+	static TArray<FFactoryData, TFixedAllocator<MaxFactories>> NetFactories;
+};
+
+} // end namespace UE::Net
\ No newline at end of file
diff --git a/Engine/Source/Runtime/Experimental/Iris/Core/Public/Iris/ReplicationSystem/ObjectReplicationBridge.h b/Engine/Source/Runtime/Experimental/Iris/Core/Public/Iris/ReplicationSystem/ObjectReplicationBridge.h
--- a/Engine/Source/Runtime/Experimental/Iris/Core/Public/Iris/ReplicationSystem/ObjectReplicationBridge.h
+++ b/Engine/Source/Runtime/Experimental/Iris/Core/Public/Iris/ReplicationSystem/ObjectReplicationBridge.h
@@ -8,6 +8,8 @@
 #include "Iris/ReplicationSystem/ReplicationBridge.h"
 #include "Delegates/IDelegateInstance.h"
 #include "Iris/ReplicationSystem/ReplicationSystemTypes.h"
+#include "Iris/ReplicationSystem/NetObjectFactory.h"
+#include "Iris/ReplicationSystem/NetObjectFactoryRegistry.h"
 #include "Containers/ArrayView.h"
 
 #include "ObjectReplicationBridge.generated.h"
@@ -17,6 +19,7 @@ DECLARE_LOG_CATEGORY_EXTERN(LogIrisFilterConfig, Log, All);
 namespace UE::Net
 {
 	enum class ENetRefHandleError : uint32;
+	enum class EReplicationFragmentTraits : uint32;
 
 	struct FNetObjectResolveContext;
 
@@ -26,8 +29,6 @@ namespace UE::Net
 
 	class FNetObjectReference;
 
-	enum class EReplicationFragmentTraits : uint32;
-
 	namespace Private
 	{
 		typedef uint32 FInternalNetRefIndex;
@@ -55,42 +56,8 @@ public:
 
 	using EGetRefHandleFlags = UE::Net::EGetRefHandleFlags;
 
-	struct FRootObjectReplicationParams
-	{
-		/** When true it means the object wants to receive a PreUpdate callback just before it gets polled. */
-		bool bNeedsPreUpdate = false;
-
-		/** When true the object has a dynamic world location and we should ask the bridge to update its current location everytime it is polled. */
-		bool bNeedsWorldLocationUpdate = false;
-
-		/** Whether the object is dormant or not */
-		bool bIsDormant = false;
-
-		/** Ask the class config for a dynamic filter assigned to this class or one of it's parent class. Default is true. */
-		bool bUseClassConfigDynamicFilter = true;
-
-		/** When enabled we ignore the class config for this object and instead use the one specified by ExplicitDynamicFilter */
-		bool bUseExplicitDynamicFilter = false;
-
-		/** 
-		 * The name of the dynamic filter to use for this object (instead of asking the class config).
-		 * Can be none so that no dynamic filter is assigned to the object.
-		 * Only used when bUseExplicitDynamicFilter is true.
-		 */
-		FName ExplicitDynamicFilterName;
-
-		/**
-		 * If StaticPriority is > 0 the ReplicationSystem will use that as priority when scheduling objects. 
-		 * If it's <= 0.0f one will look for a world location support and then use the default spatial prioritizer.
-		 */
-		float StaticPriority = 0.0f;
-
-		/**
-		 * How often per second the object should be polled for dirtiness, including calling the InstancePreUpdate function. 
-		 * When set to zero it will be polled every frame.
-		 */
-		float PollFrequency = 0.0f;
-	};
+	struct FRootObjectReplicationParams;
+	struct FSubObjectReplicationParams;
 
 	IRISCORE_API UObjectReplicationBridge();
 
@@ -122,7 +89,7 @@ public:
 	 * @param Params Optional configuration parameters to specify how the root object will be replicated.
 	 * @return The NetRefHandle associated to this object if the operation succeeded.
 	 */
-	IRISCORE_API FNetRefHandle StartReplicatingRootObject(UObject* Instance, const FRootObjectReplicationParams& Params);
+	IRISCORE_API FNetRefHandle StartReplicatingRootObject(UObject* Instance, const FRootObjectReplicationParams& Params, UE::Net::FNetObjectFactoryId NetFactoryId);
 
 	/**
 	 * Start replicating a SubObject and return a valid NetRefHandle if successful.
@@ -133,7 +100,7 @@ public:
 	 * @param InsertionOrder When none will always add the new subobject at the end of the list. 
 	 * @return The NetRefHandle of the subobject if successful.
 	 */
-	IRISCORE_API FNetRefHandle StartReplicatingSubObject(FNetRefHandle OwnerHandle, UObject* Instance, FNetRefHandle InsertRelativeToSubObjectHandle = FNetRefHandle::GetInvalid(), ESubObjectInsertionOrder InsertionOrder = UReplicationBridge::ESubObjectInsertionOrder::None);
+	IRISCORE_API FNetRefHandle StartReplicatingSubObject(UObject* Instance, const FSubObjectReplicationParams& Params, UE::Net::FNetObjectFactoryId NetFactoryId);
 
 	/** 
 	 * Stop replicating any type of NetObject. @see UReplicationBridge::StopReplicatingNetObject for more details.
@@ -162,15 +129,6 @@ public:
 	/** Returns true of the level that the Object belongs to has finished loading. */
 	virtual bool ObjectLevelHasFinishedLoading(UObject* Object) const { return true; }
 
-	struct FCreationHeader
-	{
-		virtual ~FCreationHeader() {};
-
-	private:
-		friend UObjectReplicationBridge;
-		UE::Net::FReplicationProtocolIdentifier ProtocolIdentifier = 0U;
-	};
-	
 	/**
 	 * Adds a dependent object. A dependent object can replicate separately or if a parent replicates.
 	 * Dependent objects cannot be filtered out by dynamic filtering unless the parent is also filtered out.
@@ -206,6 +164,57 @@ public:
 	/** Getter for CVar net.Iris.UseVerboseIrisCsvStats which is true if we want per-class iris stats output to CSV */
 	IRISCORE_API bool ShouldUseVerboseCsvStats() const;
 
+public:
+
+	struct FRootObjectReplicationParams
+	{
+		/** When true it means the object wants to receive a PreUpdate callback just before it gets polled. */
+		bool bNeedsPreUpdate = false;
+
+		/** When true the object has a dynamic world location and we should ask the bridge to update its current location everytime it is polled. */
+		bool bNeedsWorldLocationUpdate = false;
+
+		/** Whether the object is dormant or not */
+		bool bIsDormant = false;
+
+		/** Ask the class config for a dynamic filter assigned to this class or one of it's parent class. Default is true. */
+		bool bUseClassConfigDynamicFilter = true;
+
+		/** When enabled we ignore the class config for this object and instead use the one specified by ExplicitDynamicFilter */
+		bool bUseExplicitDynamicFilter = false;
+
+		/** 
+		* The name of the dynamic filter to use for this object (instead of asking the class config).
+		* Can be none so that no dynamic filter is assigned to the object.
+		* Only used when bUseExplicitDynamicFilter is true.
+		*/
+		FName ExplicitDynamicFilterName;
+
+		/**
+		* If StaticPriority is > 0 the ReplicationSystem will use that as priority when scheduling objects. 
+		* If it's <= 0.0f one will look for a world location support and then use the default spatial prioritizer.
+		*/
+		float StaticPriority = 0.0f;
+
+		/**
+		* How often per second the object should be polled for dirtiness, including calling the InstancePreUpdate function. 
+		* When set to zero it will be polled every frame.
+		*/
+		float PollFrequency = 0.0f;
+	};
+
+	struct FSubObjectReplicationParams
+	{
+		/** The root object that the subobject will be attached to*/
+		FNetRefHandle RootObjectHandle;
+
+		/** Optional handle to another subobject indicating where in the subobject list we inserted */
+		FNetRefHandle InsertRelativeToSubObjectHandle; 
+
+		/** Optional enum to select where in the list this subobject will be inserted */
+		ESubObjectInsertionOrder InsertionOrder = ESubObjectInsertionOrder::None;
+	};
+
 public:
 
 	// Debug functions that are triggered via console commands
@@ -241,31 +250,7 @@ protected:
 	IRISCORE_API virtual void OnErrorWithNetRefHandleReported(UE::Net::ENetRefHandleError ErrorType, FNetRefHandle RefHandle, uint32 ConnectionId) override;
 	
 protected:
-	/**
-	* $IRIS TODO:
-	* Consider moving these methods to separate interface to allow for a single ObjectReplicationBridge to support different type specific methods to its own interface and allow for a single bridge to support multiple subtypes? 
-	* We will see what we need when we start to Serialize data we might not even want to expose virtual methods for doing this, and instead provide support for specifying a "creation state" that works exactly as any other state
-	* but only ever is replicated when we first instantiate an object of the specific type 
-	*/
-
-	/** Allocate and fill in creation header. */
-	virtual TUniquePtr<FCreationHeader> GetCreationHeader(FNetRefHandle Handle) { return nullptr; };
-
-	/** Write creation header. */
-	virtual bool WriteCreationHeader(UE::Net::FNetSerializationContext& Context, const FCreationHeader* CreationHeader) { return false; };
-
-	/** Write data required to instantiate the Handle. */
-	virtual bool WriteCreationHeader(UE::Net::FNetSerializationContext& Context, FNetRefHandle Handle) { return false; };
-
-	/** Create Header and read data required to instantiate the Handle. */
-	virtual TUniquePtr<FCreationHeader> ReadCreationHeader(UE::Net::FNetSerializationContext& Context) { return nullptr; };
-
-	/** Called when we instantiate/find object instance requested by remote. */
-	virtual FObjectReplicationBridgeInstantiateResult BeginInstantiateFromRemote(FNetRefHandle RootObjectOfSubObject, const UE::Net::FNetObjectResolveContext& ResolveContext, const FCreationHeader* Header) { return FObjectReplicationBridgeInstantiateResult(); };
-
-	/** Invoked before we start applying state data to instance on remote end. */
-	virtual bool OnInstantiatedFromRemote(UObject* Instance, const FCreationHeader* InHeader, uint32 ConnectionId) const { return true; }
-
+	
 	/** Invoked for new replicated SubObjects after state has been applied to owner */
 	virtual void OnSubObjectCreatedFromReplication(FNetRefHandle SubObjectHandle) {};
 
@@ -385,7 +370,7 @@ private:
 	/** Remove mapping between handle and object instance. */
 	void UnregisterInstance(FNetRefHandle RefHandle);
 
-	void RegisterRemoteInstance(FNetRefHandle RefHandle, UObject* InstancePtr, const UE::Net::FReplicationProtocol* Protocol, UE::Net::FReplicationInstanceProtocol* InstanceProtocol, const FCreationHeader* Header, uint32 ConnectionId);
+	void RegisterRemoteInstance(FNetRefHandle RefHandle, UObject* InstancePtr, const UE::Net::FReplicationProtocol* Protocol, UE::Net::FReplicationInstanceProtocol* InstanceProtocol, uint32 ConnectionId);
 
 	void SetNetPushIdOnInstance(UE::Net::FReplicationInstanceProtocol* InstanceProtocol, FNetHandle NetHandle);
 
@@ -393,7 +378,7 @@ private:
 	void FindClassesInPollPeriodOverrides();
 
 	/** Does all that is necessary to replicate any given UObject. Returns the object's handle if successful  */
-	UE::Net::FNetRefHandle StartReplicatingNetObject(UObject* Instance, UE::Net::EReplicationFragmentTraits Traits);
+	UE::Net::FNetRefHandle StartReplicatingNetObject(UObject* Instance, UE::Net::EReplicationFragmentTraits Traits, UE::Net::FNetObjectFactoryId NetFactoryId);
 
 protected:
 	/** Retrieves the dynamic filter to set for the given class. Will return an invalid handle if no dynamic filter should be set. */
@@ -424,9 +409,13 @@ private:
 	FName GetConfigClassPathName(const UClass* Class);
 
 	void InitConditionalPropertyDelegates();
+	void InitNetObjectFactories();
+	void DeinitNetObjectFactories();
 
 	void OnMaxInternalNetRefIndexIncreased(UE::Net::Private::FInternalNetRefIndex NewMaxInternalIndex);
 
+	UNetObjectFactory* GetNetFactory(UE::Net::FNetObjectFactoryId FactoryId) const;
+
 private:
 
 	friend UE::Net::Private::FObjectPoller;
@@ -488,7 +477,7 @@ private:
 	TMap<FName, FName> ClassesWithTypeStats;
 
 	// When we flush objects, we might need to defer sending creation info.
-	TMap<FNetRefHandle, TUniquePtr<const FCreationHeader>> CachedCreationHeaders;
+	TMap<FNetRefHandle, TUniquePtr<const UE::Net::FNetObjectCreationHeader>> CachedCreationHeaders;
 
 	// Objects which has object references and could be affected by garbage collection.
 	UE::Net::FNetBitArray ObjectsWithObjectReferences;
@@ -500,6 +489,9 @@ private:
 	FDelegateHandle OnCustomConditionChangedHandle;
 	FDelegateHandle OnDynamicConditionChangedHandle;
 
+	UPROPERTY()
+	TArray<TObjectPtr<UNetObjectFactory>> NetObjectFactories;
+
 	bool bHasPollOverrides = false;
 	bool bHasDirtyClassesInPollPeriodOverrides = false;
 
diff --git a/Engine/Source/Runtime/Experimental/Iris/Stub/Public/Iris/ReplicationSystem/NetObjectFactory.h b/Engine/Source/Runtime/Experimental/Iris/Stub/Public/Iris/ReplicationSystem/NetObjectFactory.h
new file mode 100644
--- /dev/null
+++ b/Engine/Source/Runtime/Experimental/Iris/Stub/Public/Iris/ReplicationSystem/NetObjectFactory.h
@@ -0,0 +1,16 @@
+// Copyright Epic Games, Inc. All Rights Reserved.
+
+#pragma once
+
+#include "UObject/ObjectMacros.h"
+
+#include "NetObjectFactory.generated.h"
+
+// Stub used when we are not compiling with iris to workaround UHT dependencies
+static_assert(UE_WITH_IRIS == 0, "IrisStub module should not be used when iris is enabled");
+
+UCLASS(MinimalAPI, transient)
+class UNetObjectFactory : public UObject
+{
+	GENERATED_BODY()
+};
\ No newline at end of file
