Fix QAGame AllInPIE failure due to UStereoLayerComponent not calling Super::EndPlay() in its override.

diff --git a/Engine/Source/Runtime/Engine/Private/Components/StereoLayerComponent.cpp b/Engine/Source/Runtime/Engine/Private/Components/StereoLayerComponent.cpp
--- a/Engine/Source/Runtime/Engine/Private/Components/StereoLayerComponent.cpp
+++ b/Engine/Source/Runtime/Engine/Private/Components/StereoLayerComponent.cpp
@@ -38,6 +38,8 @@ UStereoLayerComponent::UStereoLayerComponent(const FObjectInitializer& ObjectIni
 
 void UStereoLayerComponent::EndPlay(const EEndPlayReason::Type EndPlayReason)
 {
+	Super::EndPlay(EndPlayReason);
+
 	if (EndPlayReason == EEndPlayReason::EndPlayInEditor || EndPlayReason == EEndPlayReason::Quit)
 	{
 		FStereoLayerAdditionalFlagsManager::Destroy();
