Fix soundwave failing to play in editor when force-inline

diff --git a/Engine/Source/Runtime/Engine/Private/SoundWave.cpp b/Engine/Source/Runtime/Engine/Private/SoundWave.cpp
--- a/Engine/Source/Runtime/Engine/Private/SoundWave.cpp
+++ b/Engine/Source/Runtime/Engine/Private/SoundWave.cpp
@@ -2955,9 +2955,6 @@ void USoundWave::PostEditChangeProperty(FPropertyChangedEvent& PropertyChangedEv
 		//an empty event property field might mean the update comes from an undo
 		//we can't discern what properties where reverted so we update the asset and bakes
 	
-		// Avoid modifying FSoundWaveProxy that is still in use by outside systems. 
-		CreateNewSoundWaveData();
-
 		UpdateAsset();
 		BakeFFTAnalysis();
 		BakeEnvelopeAnalysis();
@@ -2991,6 +2988,7 @@ void USoundWave::PostEditChangeProperty(FPropertyChangedEvent& PropertyChangedEv
 		{
 			// Avoid modifying FSoundWaveProxy that is still in use by outside systems. 
 			CreateNewSoundWaveData();
+			bool bSoundWaveDataNeedsInitialization = true;
 
 			// Regenerate on save any compressed sound formats or if analysis needs to be re-done
 			if (Name == LoadingBehaviorFName || Name == InlinedAudioInSecondsFName)
@@ -3013,6 +3011,7 @@ void USoundWave::PostEditChangeProperty(FPropertyChangedEvent& PropertyChangedEv
 			if(Name == SoundAssetCompressionTypeFName)
 			{
 				SetSoundAssetCompressionType(SoundAssetCompressionType);
+				bSoundWaveDataNeedsInitialization = false; //< SetSoundAssetCompressionType eventually calls InitializeDataFromSoundWave(...)
 			}
 
 			if (Name == CompressionQualityFName
@@ -3026,6 +3025,12 @@ void USoundWave::PostEditChangeProperty(FPropertyChangedEvent& PropertyChangedEv
 				|| Name == CloudStreamingFName)
 			{
 				UpdateAsset();
+				bSoundWaveDataNeedsInitialization = false; //< UpdateAsset() eventually calls InitializeDataFromSoundWave(...)
+			}
+
+			if (bSoundWaveDataNeedsInitialization)	
+			{
+				SoundWaveDataPtr->InitializeDataFromSoundWave(*this);
 			}
 
 			if (AnyFFTAnalysisPropertiesChanged(Name))
