[AutoRTFM] Use transactionally-safe critical sections while streaming.

This code is reachable from within a Verse call.

diff --git a/Engine/Source/Runtime/Engine/Private/StreamableManager.cpp b/Engine/Source/Runtime/Engine/Private/StreamableManager.cpp
--- a/Engine/Source/Runtime/Engine/Private/StreamableManager.cpp
+++ b/Engine/Source/Runtime/Engine/Private/StreamableManager.cpp
@@ -6,6 +6,8 @@
 #include "UObject/Package.h"
 #include "Misc/PackageAccessTrackingOps.h"
 #include "Misc/PackageName.h"
+#include "Misc/TransactionallySafeCriticalSection.h"
+#include "Misc/TransactionallySafeScopeLock.h"
 #include "UObject/UObjectThreadContext.h"
 #include "HAL/IConsoleManager.h"
 #include "Tickable.h"
@@ -59,7 +61,7 @@ public:
 		FPendingDelegate* PendingDelegate = new FPendingDelegate(Forward<InStreamableDelegate>(Delegate), Forward<InStreamableDelegate>(CancelDelegate), AssociatedHandle);
 
 		{
-			FScopeLock Lock(&DataLock);
+			FTransactionallySafeScopeLock Lock(&DataLock);
 
 			FPendingDelegateList& DelegatesForHandle = PendingDelegatesByHandle.FindOrAdd(AssociatedHandle);
 			FPendingDelegateList::AddTail(PendingDelegate, PendingDelegates, DelegatesForHandle);
@@ -72,7 +74,7 @@ public:
 		FPendingDelegateList PendingDelegatesToDelete;
 
 		{
-			FScopeLock _(&DataLock);
+			FTransactionallySafeScopeLock _(&DataLock);
 
 			FPendingDelegateList* DelegatesForHandle = PendingDelegatesByHandle.Find(AssociatedHandle);
 			if (!DelegatesForHandle)
@@ -142,7 +144,7 @@ public:
 
 		FPendingDelegateList DelegatesToCall;
 		{
-			FScopeLock Lock(&DataLock);
+			FTransactionallySafeScopeLock Lock(&DataLock);
 
 			FPendingDelegate* CurrentNode = PendingDelegates.Head;
 			while (CurrentNode)
@@ -328,7 +330,7 @@ private:
 	FPendingDelegateList PendingDelegates;
 	TMap<TSharedPtr<FStreamableHandle>, FPendingDelegateList> PendingDelegatesByHandle;
 
-	FCriticalSection DataLock;
+	FTransactionallySafeCriticalSection DataLock;
 };
 
 static FStreamableDelegateDelayHelper* StreamableDelegateDelayHelper = nullptr;
