Fix rhivalidation for FParticleAttributesTexture: CreateAttributesTextureRHI sets up .SetInitialState(ERHIAccess::RTV); but we were using ERHIAccess::SRVMask as previous state

diff --git a/Engine/Source/Runtime/Engine/Private/Particles/ParticleGpuSimulation.cpp b/Engine/Source/Runtime/Engine/Private/Particles/ParticleGpuSimulation.cpp
--- a/Engine/Source/Runtime/Engine/Private/Particles/ParticleGpuSimulation.cpp
+++ b/Engine/Source/Runtime/Engine/Private/Particles/ParticleGpuSimulation.cpp
@@ -478,7 +478,7 @@ public:
 		// allocate the new textures
 		FTextureRHIRef NewAttributesTextureRHI = CreateAttributesTextureRHI();
 
-		RHICmdList.Transition(FRHITransitionInfo(NewAttributesTextureRHI, ERHIAccess::SRVMask, ERHIAccess::CopyDest));
+		RHICmdList.Transition(FRHITransitionInfo(NewAttributesTextureRHI, ERHIAccess::RTV, ERHIAccess::CopyDest));
 		RHICmdList.Transition(FRHITransitionInfo(TextureRHI, ERHIAccess::SRVMask, ERHIAccess::CopySrc));
 
 		FRHICopyTextureInfo CopyTextureInfo;
