Fixed incorrect log message

#rnx

diff --git a/Engine/Source/Runtime/Engine/Private/UnrealEngine.cpp b/Engine/Source/Runtime/Engine/Private/UnrealEngine.cpp
--- a/Engine/Source/Runtime/Engine/Private/UnrealEngine.cpp
+++ b/Engine/Source/Runtime/Engine/Private/UnrealEngine.cpp
@@ -14171,7 +14171,7 @@ bool UEngine::WillNetDriverUseIris(const FWorldContext& Context, FName InNetDriv
 	}
 	else if (CmdlineRequest == EReplicationSystem::Generic)
 	{
-		UE_CLOG(bUseIrisRepSystem, LogNet, Log, TEXT("Cmdline -UseIrisReplication=0 is forcing NetDriver %s (NetDefinition %s) to use the Iris replication system."), *InNetDriverName.ToString(), *InNetDriverDefinition.ToString());
+		UE_CLOG(bUseIrisRepSystem, LogNet, Log, TEXT("Cmdline -UseIrisReplication=0 is forcing NetDriver %s (NetDefinition %s) to use the Generic replication system."), *InNetDriverName.ToString(), *InNetDriverDefinition.ToString());
 		bUseIrisRepSystem = false;
 	}
 
