[World Partition] Implement the possibility for actor classes to be non-PIE editor-only.
Make sure that classes enabling this are non-spatially loaded and can't have associated data layers to minimize impacts on streaming generation, since we don't cluster actors from references coming from editor-only actors or editor-only properties.

diff --git a/Engine/Source/Runtime/Engine/Classes/GameFramework/Actor.h b/Engine/Source/Runtime/Engine/Classes/GameFramework/Actor.h
--- a/Engine/Source/Runtime/Engine/Classes/GameFramework/Actor.h
+++ b/Engine/Source/Runtime/Engine/Classes/GameFramework/Actor.h
@@ -2283,7 +2283,8 @@ public:
 	ENGINE_API virtual bool CanBeInCluster() const override;
 	static ENGINE_API void AddReferencedObjects(UObject* InThis, FReferenceCollector& Collector);
 	ENGINE_API virtual bool IsEditorOnly() const override;
-	virtual bool IsRuntimeOnly() const { return false; }
+	ENGINE_API virtual bool IsNonPIEEditorOnly() const { return false; }
+	ENGINE_API virtual bool IsRuntimeOnly() const { return false; }
 	ENGINE_API virtual bool IsAsset() const override;
 
 	PRAGMA_DISABLE_DEPRECATION_WARNINGS // Suppress compiler warning on override of deprecated function
diff --git a/Engine/Source/Runtime/Engine/Private/WorldPartition/WorldPartitionActorDesc.cpp b/Engine/Source/Runtime/Engine/Private/WorldPartition/WorldPartitionActorDesc.cpp
--- a/Engine/Source/Runtime/Engine/Private/WorldPartition/WorldPartitionActorDesc.cpp
+++ b/Engine/Source/Runtime/Engine/Private/WorldPartition/WorldPartitionActorDesc.cpp
@@ -429,6 +429,11 @@ void FWorldPartitionActorDesc::SerializeTo(TArray<uint8>& OutData, FWorldPartiti
 	OutData.Append(PayloadData);
 }
 
+bool FWorldPartitionActorDesc::GetActorIsNonPIEEditorOnly() const
+{
+	return ActorNativeClass->GetDefaultObject<AActor>()->IsNonPIEEditorOnly();
+}
+
 TArray<FName> FWorldPartitionActorDesc::GetDataLayers(bool bIncludeExternalDataLayer) const
 {
 	const FName ExternalDataLayer = GetExternalDataLayer();
diff --git a/Engine/Source/Runtime/Engine/Private/WorldPartition/WorldPartitionStreamingGeneration.cpp b/Engine/Source/Runtime/Engine/Private/WorldPartition/WorldPartitionStreamingGeneration.cpp
--- a/Engine/Source/Runtime/Engine/Private/WorldPartition/WorldPartitionStreamingGeneration.cpp
+++ b/Engine/Source/Runtime/Engine/Private/WorldPartition/WorldPartitionStreamingGeneration.cpp
@@ -800,7 +800,10 @@ class FWorldPartitionStreamingGenerator
 				OutContainerInstances.Add(InActorDescView);
 			}
 
-			if (!InActorDescView.GetActorIsEditorOnly() || (InActor && !InActor->IsEditorOnly()))
+			const bool bActorIsEditorOnly = InActorDescView.GetActorIsEditorOnly() || (InActor && !InActor->IsEditorOnly());
+			const bool bActorIsNonPIEEditorOnly = InActorDescView.GetActorIsNonPIEEditorOnly();
+
+			if (!bActorIsEditorOnly || bActorIsNonPIEEditorOnly)
 			{
 				const FGuid ActorGuid = InActorDescView.GetGuid();
 				OutActorDescViewMap.Emplace(ActorGuid, MoveTemp(InActorDescView));
@@ -1320,6 +1323,23 @@ class FWorldPartitionStreamingGenerator
 					}
 				}
 
+				// Non-PIE editor-only actors can only be non-spatially loaded and can't have data layers
+				if (ActorDescView.GetActorIsNonPIEEditorOnly() && (ActorDescView.GetIsSpatiallyLoaded() || ActorDescView.GetRuntimeDataLayerInstanceNames().GetNonExternalDataLayers().Num()))
+				{
+					if (ensure(PassType == EPassType::ErrorReporting))
+					{
+						// No error reporting here because this is not a data issue rather than a code issue. Actor classes returning true to IsNonPIEEditorOnly should be non-spatially loaded
+						// and can't be associated to data layers, so make sure they also have CanChangeIsSpatiallyLoadedFlag and ActorTypeSupportsDataLayer return false.
+					}
+					else
+					{
+						ActorDescView.SetForcedNonSpatiallyLoaded();
+						ActorDescView.SetForcedNoDataLayers();
+					}
+
+					NbErrorsDetected++;
+				}
+
 				for (FActorReferenceInfo& Info : References)
 				{
 					FStreamingGenerationActorDescView* RefererActorDescView = Info.ActorDesc;
diff --git a/Engine/Source/Runtime/Engine/Public/WorldPartition/ActorDescContainerSubsystem.h b/Engine/Source/Runtime/Engine/Public/WorldPartition/ActorDescContainerSubsystem.h
--- a/Engine/Source/Runtime/Engine/Public/WorldPartition/ActorDescContainerSubsystem.h
+++ b/Engine/Source/Runtime/Engine/Public/WorldPartition/ActorDescContainerSubsystem.h
@@ -46,7 +46,7 @@ public:
 	void RegisterContainer(UActorDescContainer* Container) { ContainerManager.RegisterContainer(Container); }
 	void UnregisterContainer(UActorDescContainer* Container) { ContainerManager.UnregisterContainer(Container); }
 	FBox GetContainerBounds(const FString& ContainerName, bool bIsEditorBounds = false) const { return ContainerManager.GetContainerBounds(ContainerName, bIsEditorBounds); }
-	void UpdateContainerBounds(const FString& ContainerName) { ContainerManager.UpdateContainerBounds(ContainerName); }
+	ENGINE_API void UpdateContainerBounds(const FString& ContainerName) { ContainerManager.UpdateContainerBounds(ContainerName); }
 	void SetContainerPackage(UActorDescContainer* Container, FName ContainerPackageName) { ContainerManager.SetContainerPackage(Container, ContainerPackageName); }
 	void NotifyContainerUpdated(FName ContainerPackage)
 	{
diff --git a/Engine/Source/Runtime/Engine/Public/WorldPartition/WorldPartitionActorDesc.h b/Engine/Source/Runtime/Engine/Public/WorldPartition/WorldPartitionActorDesc.h
--- a/Engine/Source/Runtime/Engine/Public/WorldPartition/WorldPartitionActorDesc.h
+++ b/Engine/Source/Runtime/Engine/Public/WorldPartition/WorldPartitionActorDesc.h
@@ -165,6 +165,7 @@ public:
 	inline bool GetIsSpatiallyLoaded() const { return RuntimeBounds.IsValid ? bIsSpatiallyLoaded : false; }
 	inline bool GetIsSpatiallyLoadedRaw() const { return bIsSpatiallyLoaded; }
 	inline bool GetActorIsEditorOnly() const { return bActorIsEditorOnly; }
+	ENGINE_API bool GetActorIsNonPIEEditorOnly() const;
 	inline bool GetActorIsRuntimeOnly() const { return bActorIsRuntimeOnly; }
 
 	inline bool GetActorIsHLODRelevant() const { return bActorIsHLODRelevant; }
diff --git a/Engine/Source/Runtime/Engine/Public/WorldPartition/WorldPartitionActorDescInstance.h b/Engine/Source/Runtime/Engine/Public/WorldPartition/WorldPartitionActorDescInstance.h
--- a/Engine/Source/Runtime/Engine/Public/WorldPartition/WorldPartitionActorDescInstance.h
+++ b/Engine/Source/Runtime/Engine/Public/WorldPartition/WorldPartitionActorDescInstance.h
@@ -45,6 +45,7 @@ public:
 	virtual FName GetRuntimeGrid() const { return ActorDesc->GetRuntimeGrid(); }
 	virtual bool GetIsSpatiallyLoaded() const { return !GetForceNonSpatiallyLoaded() && ActorDesc->GetIsSpatiallyLoaded(); }
 	virtual bool GetActorIsEditorOnly() const { return ActorDesc->GetActorIsEditorOnly(); }
+	virtual bool GetActorIsNonPIEEditorOnly() const { return ActorDesc->GetActorIsNonPIEEditorOnly(); }
 	virtual bool GetActorIsRuntimeOnly() const { return ActorDesc->GetActorIsRuntimeOnly(); }
 	ENGINE_API virtual bool IsRuntimeRelevant() const;
 	ENGINE_API virtual bool IsEditorRelevant() const;
diff --git a/Engine/Source/Runtime/Engine/Public/WorldPartition/WorldPartitionActorDescInstanceView.h b/Engine/Source/Runtime/Engine/Public/WorldPartition/WorldPartitionActorDescInstanceView.h
--- a/Engine/Source/Runtime/Engine/Public/WorldPartition/WorldPartitionActorDescInstanceView.h
+++ b/Engine/Source/Runtime/Engine/Public/WorldPartition/WorldPartitionActorDescInstanceView.h
@@ -26,6 +26,7 @@ public:
 	virtual FName GetRuntimeGrid() const override { return ActorDescInstance->GetRuntimeGrid(); }
 	virtual bool GetIsSpatiallyLoaded() const override { return ActorDescInstance->GetIsSpatiallyLoaded(); }
 	virtual bool GetActorIsEditorOnly() const override { return ActorDescInstance->GetActorIsEditorOnly(); }
+	virtual bool GetActorIsNonPIEEditorOnly() const override { return ActorDescInstance->GetActorIsNonPIEEditorOnly(); }
 	virtual bool GetActorIsRuntimeOnly() const override { return ActorDescInstance->GetActorIsRuntimeOnly(); }
 	virtual bool IsRuntimeRelevant() const override { return ActorDescInstance->IsRuntimeRelevant(); }
 	virtual bool IsEditorRelevant() const override { return ActorDescInstance->IsEditorRelevant(); }
diff --git a/Engine/Source/Runtime/Engine/Public/WorldPartition/WorldPartitionActorDescInstanceViewInterface.h b/Engine/Source/Runtime/Engine/Public/WorldPartition/WorldPartitionActorDescInstanceViewInterface.h
--- a/Engine/Source/Runtime/Engine/Public/WorldPartition/WorldPartitionActorDescInstanceViewInterface.h
+++ b/Engine/Source/Runtime/Engine/Public/WorldPartition/WorldPartitionActorDescInstanceViewInterface.h
@@ -29,6 +29,7 @@ public:
 	virtual FName GetRuntimeGrid() const = 0;
 	virtual bool GetIsSpatiallyLoaded() const = 0;
 	virtual bool GetActorIsEditorOnly() const = 0;
+	virtual bool GetActorIsNonPIEEditorOnly() const = 0;
 	virtual bool GetActorIsRuntimeOnly() const = 0;
 	virtual bool IsRuntimeRelevant() const = 0;
 	virtual bool IsEditorRelevant() const = 0;
