Fix Substrat conversion could happening several times or in unintended fashion.

diff --git a/Engine/Source/Runtime/Engine/Private/Materials/Material.cpp b/Engine/Source/Runtime/Engine/Private/Materials/Material.cpp
--- a/Engine/Source/Runtime/Engine/Private/Materials/Material.cpp
+++ b/Engine/Source/Runtime/Engine/Private/Materials/Material.cpp
@@ -3520,7 +3520,8 @@ void UMaterial::ConvertMaterialToSubstrateMaterial()
 	if(!bLayersMaterial)
 #endif //ENABLE_MATERIAL_LAYER_PROTOTYPE
 	{
-		if (bUseMaterialAttributes && EditorOnly->MaterialAttributes.Expression && !EditorOnly->FrontMaterial.IsConnected() && !EditorOnly->MaterialAttributes.Expression->IsResultSubstrateMaterial(EditorOnly->MaterialAttributes.OutputIndex)) // M_Rifle cause issues there
+		const bool bHasAnySubstrateNodes = HasAnyExpressionsInMaterialAndFunctionsOfType<UMaterialExpressionSubstrateBSDF>();
+		if (bUseMaterialAttributes && EditorOnly->MaterialAttributes.Expression && !EditorOnly->FrontMaterial.IsConnected() && !EditorOnly->MaterialAttributes.Expression->IsResultSubstrateMaterial(EditorOnly->MaterialAttributes.OutputIndex) && !bHasAnySubstrateNodes) // M_Rifle cause issues there
 		{
 			UMaterialExpressionSubstrateConvertMaterialAttributes* ConvertAttributeNode = NewObject<UMaterialExpressionSubstrateConvertMaterialAttributes>(this);
 			ConvertAttributeNode->Material = this;
@@ -3646,7 +3647,7 @@ void UMaterial::ConvertMaterialToSubstrateMaterial()
 			bRelinkCustomOutputNodes = false;
 			bInvalidateShader = true;
 		}
-		else if (!bUseMaterialAttributes && !EditorOnly->FrontMaterial.IsConnected())
+		else if (!bUseMaterialAttributes && !EditorOnly->FrontMaterial.IsConnected() && !bHasAnySubstrateNodes)
 		{
 			if (MaterialDomain == MD_Surface)
 			{
