Breakpoints in Animation Blueprints can now get hit when not debugging a specific object

Ensure that anim BP execution occurs on the game thread if breakpoints are currently set. This ensures that users can break into running code. As per regular BP behavior, this requires no specific debug target object to be set for the breakpoint to trigger.

[RN] Breakpoints in Animation Blueprints can now get hit when not debugging a specific object

diff --git a/Engine/Source/Runtime/Engine/Private/Animation/AnimInstance.cpp b/Engine/Source/Runtime/Engine/Private/Animation/AnimInstance.cpp
--- a/Engine/Source/Runtime/Engine/Private/Animation/AnimInstance.cpp
+++ b/Engine/Source/Runtime/Engine/Private/Animation/AnimInstance.cpp
@@ -29,6 +29,7 @@
 #include "Animation/AnimSubsystem_SharedLinkedAnimLayers.h"
 #if WITH_EDITOR
 #include "Engine/Blueprint.h"
+#include "BlueprintEditorSettings.h"
 #endif
 
 #include UE_INLINE_GENERATED_CPP_BY_NAME(AnimInstance)
@@ -811,9 +812,11 @@ bool UAnimInstance::NeedsImmediateUpdate(float DeltaSeconds, bool bNeedsValidRoo
 		GIntraFrameDebuggingGameThread ||
 #if WITH_EDITOR
 		// Force the debugged object to run its anim graph on the game thread if it is being debugged
+		// Also force onto the game thread if breakpoints are set. This will force the anim BP to break even if
+		// the target object is not set for debug
 		// This ensures that it uses the persistent ubergraph frame and debugging facilities are available like
 		// watches, breakpoints etc.
-		(Blueprint && Blueprint->GetObjectBeingDebugged() == this) ||
+		(Blueprint && (Blueprint->GetObjectBeingDebugged() == this || FKismetDebugUtilities::BlueprintHasBreakpoints(Blueprint))) ||
 #endif
 		CVarUseParallelAnimUpdate.GetValueOnGameThread() == 0 ||
 		CVarUseParallelAnimationEvaluation.GetValueOnGameThread() == 0 ||
