UWaterSubsystem::FindWaterZone optimization: go through ForEachIntersectingActorDescInstance instead of doing the ForEachActorDescInstance/Intersect combo.

diff --git a/Engine/Plugins/Experimental/Water/Source/Runtime/Private/WaterSubsystem.cpp b/Engine/Plugins/Experimental/Water/Source/Runtime/Private/WaterSubsystem.cpp
--- a/Engine/Plugins/Experimental/Water/Source/Runtime/Private/WaterSubsystem.cpp
+++ b/Engine/Plugins/Experimental/Water/Source/Runtime/Private/WaterSubsystem.cpp
@@ -623,17 +623,11 @@ TSoftObjectPtr<AWaterZone> UWaterSubsystem::FindWaterZone(const UWorld* World, c
 	{
 		if (UWorldPartition* WorldPartition = World->GetWorldPartition())
 		{
-			FWorldPartitionHelpers::ForEachActorDescInstance<AWaterZone>(WorldPartition, [Bounds, &ViableZones](const FWorldPartitionActorDescInstance* ActorDescInstance)
+			const FBox Bounds3D(FVector(Bounds.Min.X, Bounds.Min.Y, -HALF_WORLD_MAX), FVector(Bounds.Max.X, Bounds.Max.Y, HALF_WORLD_MAX));
+			FWorldPartitionHelpers::ForEachIntersectingActorDescInstance<AWaterZone>(WorldPartition, Bounds3D, [&Bounds, &ViableZones](const FWorldPartitionActorDescInstance* ActorDescInstance)
 			{
 				FWaterZoneActorDesc* WaterZoneActorDesc = (FWaterZoneActorDesc*)ActorDescInstance->GetActorDesc();
-				const FBox WaterZoneBounds = ActorDescInstance->GetEditorBounds();
-				const FBox2D WaterZoneBounds2D(FVector2D(WaterZoneBounds.Min), FVector2D(WaterZoneBounds.Max));
-
-				if (Bounds.Intersect(WaterZoneBounds2D))
-				{
-					ViableZones.Emplace(ActorDescInstance->GetActorSoftPath(), WaterZoneActorDesc->GetOverlapPriority());
-				}
-
+				ViableZones.Emplace(ActorDescInstance->GetActorSoftPath(), WaterZoneActorDesc->GetOverlapPriority());
 				return true;
 			});
 		}
diff --git a/Engine/Source/Runtime/Engine/Private/WorldPartition/WorldPartitionHelpers.cpp b/Engine/Source/Runtime/Engine/Private/WorldPartition/WorldPartitionHelpers.cpp
--- a/Engine/Source/Runtime/Engine/Private/WorldPartition/WorldPartitionHelpers.cpp
+++ b/Engine/Source/Runtime/Engine/Private/WorldPartition/WorldPartitionHelpers.cpp
@@ -96,15 +96,24 @@ bool FWorldPartitionHelpers::IsActorDescClassCompatibleWith(const FWorldPartitio
 
 void FWorldPartitionHelpers::ForEachIntersectingActorDescInstance(UWorldPartition* WorldPartition, const FBox& Box, TSubclassOf<AActor> ActorClass, TFunctionRef<bool(const FWorldPartitionActorDescInstance*)> Func)
 {
-	bool bProcessNextActors = true;
-
-	WorldPartition->EditorHash->ForEachIntersectingActor(Box, [&ActorClass, Func, &bProcessNextActors](const FWorldPartitionActorDescInstance* ActorDescInstance)
+	if (WorldPartition->EditorHash)
 	{
-		if (bProcessNextActors && IsActorDescClassCompatibleWith(ActorDescInstance->GetActorDesc(), ActorClass))
+		bool bProcessNextActors = true;
+		WorldPartition->EditorHash->ForEachIntersectingActor(Box, [&ActorClass, Func, &bProcessNextActors](const FWorldPartitionActorDescInstance* ActorDescInstance)
 		{
-			bProcessNextActors = Func(ActorDescInstance);
-		}
-	});
+			if (bProcessNextActors && IsActorDescClassCompatibleWith(ActorDescInstance->GetActorDesc(), ActorClass))
+			{
+				bProcessNextActors = Func(ActorDescInstance);
+			}
+		});
+	}
+	else
+	{
+		ForEachActorDescInstance(WorldPartition, ActorClass, [&Box, Func](const FWorldPartitionActorDescInstance* ActorDescInstance)
+		{
+			return !ActorDescInstance->GetEditorBounds().Intersect(Box) || Func(ActorDescInstance);
+		});
+	}
 }
 
 void FWorldPartitionHelpers::ForEachActorDescInstance(UWorldPartition* WorldPartition, TSubclassOf<AActor> ActorClass, TFunctionRef<bool(const FWorldPartitionActorDescInstance*)> Func)
