Iris - Don't ensure on clients ending replication on actors due to level being streamed out.

diff --git a/Engine/Source/Runtime/Engine/Private/Net/Iris/ReplicationSystem/ActorReplicationBridge.cpp b/Engine/Source/Runtime/Engine/Private/Net/Iris/ReplicationSystem/ActorReplicationBridge.cpp
--- a/Engine/Source/Runtime/Engine/Private/Net/Iris/ReplicationSystem/ActorReplicationBridge.cpp
+++ b/Engine/Source/Runtime/Engine/Private/Net/Iris/ReplicationSystem/ActorReplicationBridge.cpp
@@ -486,7 +486,9 @@ void UActorReplicationBridge::StopReplicatingActor(AActor* Actor, EEndPlayReason
 
 	// If we are shutting down or the actor is a nettemporary we do not need to validate that we are not detaching remote instances by accident.
 	const bool bIsShuttingDown = (EndPlayReason == EEndPlayReason::EndPlayInEditor) || (EndPlayReason == EEndPlayReason::Quit);
-	if (bIsShuttingDown || Actor->bNetTemporary)
+	// Streaming out a level can happen prior to all actors in it being destroyed/ending replication through replication so we don't want to validate they're detached by accident.
+	const bool bIsStreamingOutLevel = (EndPlayReason == EEndPlayReason::RemovedFromWorld);
+	if (bIsShuttingDown || bIsStreamingOutLevel || Actor->bNetTemporary)
 	{
 		Flags |= EEndReplicationFlags::SkipPendingEndReplicationValidation;
 	}
