Visual logger improvements in order to better handle high item number scenarios:
* Added a variant of UE_VLOG_ARROW where the size of the arrow head can be specified
* Added frustum culling for all types of rendered debug items, instead of just 2 of them
* Added FarClippingDistance to visual logger tool, in order to set a smaller far clipping distance only for the visual log items
* Added wire variants of UE_VLOG_CIRCLE...

Misc improvements :
* Added FarClippingDistance to FDebugRenderSceneProxy in order to be able to control, per-primitive component using such render proxies, the clipping distance of all debug renders made by a FDebugRenderSceneProxy
* Added a CVar r.DebugRender.AllowFrustumCulling, to eventually disable frustum culling on debug render items
* Added a CVar r.DebugRender.OverrideFarClippingPlane, to override the far clipping distance of all debug renders made by a FDebugRenderSceneProxy
* Circles are now handled as a whole instead of a list of lines. This accelerates frustum culling since we perform it on the circle rather than per-line

diff --git a/Engine/Source/Developer/LogVisualizer/Private/SVisualLoggerToolbar.cpp b/Engine/Source/Developer/LogVisualizer/Private/SVisualLoggerToolbar.cpp
--- a/Engine/Source/Developer/LogVisualizer/Private/SVisualLoggerToolbar.cpp
+++ b/Engine/Source/Developer/LogVisualizer/Private/SVisualLoggerToolbar.cpp
@@ -4,8 +4,12 @@
 #include "Styling/AppStyle.h"
 #include "Framework/MultiBox/MultiBoxDefs.h"
 #include "Framework/MultiBox/MultiBoxBuilder.h"
+#include "Widgets/Input/SNumericEntryBox.h"
 #include "LogVisualizerStyle.h"
+#include "LogVisualizerPublic.h"
 #include "VisualLoggerCommands.h"
+#include "VisualLoggerRenderingActor.h"
+#include "VisualLoggerDatabase.h"
 #define LOCTEXT_NAMESPACE "SVisualLoggerToolbar"
 
 /* SVisualLoggerToolbar interface
@@ -42,9 +46,38 @@ TSharedRef<SWidget> SVisualLoggerToolbar::MakeToolbar(const TSharedRef<FUIComman
 		ToolBarBuilder.AddToolBarButton(FVisualLoggerCommands::Get().ResetData, NAME_None, LOCTEXT("ResetData", "Clear"), TAttribute<FText>(), FSlateIcon(FLogVisualizerStyle::Get().GetStyleSetName(), TEXT("Toolbar.Remove")));
 		ToolBarBuilder.AddToolBarButton(FVisualLoggerCommands::Get().ToggleGraphs, NAME_None, LOCTEXT("ToggleGraphs", "Graphs"), TAttribute<FText>(), FSlateIcon(FLogVisualizerStyle::Get().GetStyleSetName(), TEXT("Toolbar.Graphs")));
 		ToolBarBuilder.AddToolBarButton(FVisualLoggerCommands::Get().Refresh, NAME_None, LOCTEXT("ForceRefresh", "Refresh"), TAttribute<FText>(), FSlateIcon(FAppStyle::GetAppStyleSetName(), "Icons.Refresh"));
+
+		ToolBarBuilder.AddSeparator();
+		ToolBarBuilder.AddToolBarWidget(MakeClippingDistanceWidget(), LOCTEXT("FarClippingDistance", "Far Clipping Distance"), /*InTutorialHighlightName = */NAME_None, /*bInSearchable = */true, LOCTEXT("FarClipDistanceTooltip", "Max distance after which visual log items will stop being rendered (ignored if <=0). This can help performance/visibility in situations where lots of visual logs are displayed on screen."));
 	}
 
 	return ToolBarBuilder.MakeWidget();
 }
 
+TSharedRef<SWidget> SVisualLoggerToolbar::MakeClippingDistanceWidget()
+{
+	return SNew(SSpinBox<double>)
+		.MinValue(0.0)
+		.Delta(10.0f)
+		.MinDesiredWidth(70.0f)
+		.Justification(ETextJustify::Center)
+		.Value_Lambda([]()
+			{
+				AVisualLoggerRenderingActor* HelperActor = Cast<AVisualLoggerRenderingActor>(FVisualLoggerEditorInterface::Get()->GetHelperActor(FLogVisualizer::Get().GetWorld()));
+				if (ensure(HelperActor != nullptr))
+				{
+					return HelperActor->GetFarClippingDistance();
+				}
+				return 0.0;
+			})
+		.OnValueChanged_Lambda([](double Value)
+			{
+				AVisualLoggerRenderingActor* HelperActor = Cast<AVisualLoggerRenderingActor>(FVisualLoggerEditorInterface::Get()->GetHelperActor(FLogVisualizer::Get().GetWorld()));
+				if (ensure(HelperActor != nullptr))
+				{
+					HelperActor->SetFarClippingDistance(Value);
+				}
+			});
+}
+
 #undef LOCTEXT_NAMESPACE
diff --git a/Engine/Source/Developer/LogVisualizer/Private/SVisualLoggerToolbar.h b/Engine/Source/Developer/LogVisualizer/Private/SVisualLoggerToolbar.h
--- a/Engine/Source/Developer/LogVisualizer/Private/SVisualLoggerToolbar.h
+++ b/Engine/Source/Developer/LogVisualizer/Private/SVisualLoggerToolbar.h
@@ -38,4 +38,6 @@ protected:
 	* @return The toolbar widget.
 	*/
 	TSharedRef<SWidget> MakeToolbar(const TSharedRef<FUICommandList>& CommandList);
+
+	TSharedRef<SWidget> MakeClippingDistanceWidget();
 };
diff --git a/Engine/Source/Developer/LogVisualizer/Private/VisualLoggerRenderingActorBase.cpp b/Engine/Source/Developer/LogVisualizer/Private/VisualLoggerRenderingActorBase.cpp
--- a/Engine/Source/Developer/LogVisualizer/Private/VisualLoggerRenderingActorBase.cpp
+++ b/Engine/Source/Developer/LogVisualizer/Private/VisualLoggerRenderingActorBase.cpp
@@ -33,6 +33,7 @@ public:
 		ViewFlagName = TEXT("VisLog");
 		ViewFlagIndex = uint32(FEngineShowFlags::FindIndexByName(*ViewFlagName));
 		bWantsSelectionOutline = false;
+		FarClippingDistance = InComponent->GetFarClippingDistance();
 	}
 
 	virtual FPrimitiveViewRelevance GetViewRelevance(const FSceneView* View) const override
@@ -48,10 +49,21 @@ public:
 	virtual uint32 GetMemoryFootprint(void) const override { return sizeof(*this) + GetAllocatedSize(); }
 };
 
-UVisualLoggerRenderingComponent::UVisualLoggerRenderingComponent(const FObjectInitializer& ObjectInitializer) : Super(ObjectInitializer)
+UVisualLoggerRenderingComponent::UVisualLoggerRenderingComponent(const FObjectInitializer& ObjectInitializer)
+	: Super(ObjectInitializer)
+{
+}
+
+void UVisualLoggerRenderingComponent::SetFarClippingDistance(double Distance)
 {
+	if (FarClippingDistance != Distance)
+	{
+		FarClippingDistance = Distance;
+		MarkRenderStateDirty();
+	}
 }
 
+
 FDebugRenderSceneProxy* UVisualLoggerRenderingComponent::CreateDebugSceneProxy()
 {
 	AVisualLoggerRenderingActorBase* RenderingActor = Cast<AVisualLoggerRenderingActorBase>(GetOuter());
@@ -68,6 +80,7 @@ FDebugRenderSceneProxy* UVisualLoggerRenderingComponent::CreateDebugSceneProxy()
 		{
 			VLogSceneProxy->Spheres.Append(Shapes.Points);
 			VLogSceneProxy->Lines.Append(Shapes.Lines);
+			VLogSceneProxy->Circles.Append(Shapes.Circles);
 			VLogSceneProxy->Boxes.Append(Shapes.Boxes);
 			VLogSceneProxy->Meshes.Append(Shapes.Meshes);
 			VLogSceneProxy->Cones.Append(Shapes.Cones);
@@ -103,6 +116,19 @@ AVisualLoggerRenderingActorBase::~AVisualLoggerRenderingActorBase()
 {
 }
 
+void AVisualLoggerRenderingActorBase::SetFarClippingDistance(double Distance)
+{
+	if (RenderingComponent != nullptr)
+	{
+		RenderingComponent->SetFarClippingDistance(Distance);
+	}
+}
+
+double AVisualLoggerRenderingActorBase::GetFarClippingDistance() const
+{
+	return (RenderingComponent != nullptr) ? RenderingComponent->GetFarClippingDistance() : 0.0;
+}
+
 namespace
 {
 	static bool IsPolygonWindingCorrect(const TArray<FVector>& Verts)
@@ -208,7 +234,8 @@ void AVisualLoggerRenderingActorBase::GetDebugShapes(const FVisualLogEntry& InEn
 			|| (ElementType == EVisualLoggerShapeElement::WireSphere)
 			|| (ElementType == EVisualLoggerShapeElement::WireCapsule)
 			|| (ElementType == EVisualLoggerShapeElement::WireCone)
-			|| (ElementType == EVisualLoggerShapeElement::WireCylinder)) ? FDebugRenderSceneProxy::EDrawType::WireMesh : FDebugRenderSceneProxy::EDrawType::Invalid;
+			|| (ElementType == EVisualLoggerShapeElement::WireCylinder)
+			|| (ElementType == EVisualLoggerShapeElement::WireCircle)) ? FDebugRenderSceneProxy::EDrawType::WireMesh : FDebugRenderSceneProxy::EDrawType::Invalid;
 
 		switch (ElementType)
 		{
@@ -264,7 +291,9 @@ void AVisualLoggerRenderingActorBase::GetDebugShapes(const FVisualLogEntry& InEn
 			int32 EndIndex = StartIndex + HeaderData.VerticesNum;
 			for (int32 VIdx = StartIndex; VIdx < EndIndex; VIdx++)
 			{
-				TestMesh.Vertices.Add(FVector3f(ElementToDraw->Points[VIdx]));
+				const FVector& CurrentPoint = ElementToDraw->Points[VIdx];
+				TestMesh.Vertices.Add(FVector3f(CurrentPoint));
+				TestMesh.Box += CurrentPoint;
 			}
 
 
@@ -489,10 +518,11 @@ void AVisualLoggerRenderingActorBase::GetDebugShapes(const FVisualLogEntry& InEn
 			const bool bDrawLabel = (ElementToDraw->Description.IsEmpty() == false);
 			const FVector* Location = ElementToDraw->Points.GetData();
 			const int32 NumPoints = ElementToDraw->Points.Num();
+			const float Mag = float(ElementToDraw->Mag);
 
 			for (int32 Index = 0; Index + 1 < NumPoints; Index += 2, Location += 2)
 			{
-				DebugShapes.Arrows.Add(FDebugRenderSceneProxy::FArrowLine(*Location, *(Location + 1), Color));
+				DebugShapes.Arrows.Add(FDebugRenderSceneProxy::FArrowLine(*Location, *(Location + 1), Color, Mag));
 
 				if (bDrawLabel)
 				{
@@ -504,34 +534,19 @@ void AVisualLoggerRenderingActorBase::GetDebugShapes(const FVisualLogEntry& InEn
 			break;
 
 		case EVisualLoggerShapeElement::Circle:
+		case EVisualLoggerShapeElement::WireCircle:
 		{
-			const bool bDrawLabel = (ElementToDraw->Description.IsEmpty() == false);
+			const float Thickness = float(ElementToDraw->Thicknes);
+			const bool bDrawLabel = ElementToDraw->Description.IsEmpty() == false;
 			const int32 NumPoints = ElementToDraw->Points.Num();
-
 			for (int32 Index = 0; Index + 2 < NumPoints; Index += 3)
 			{
 				const FVector Center = ElementToDraw->Points[Index + 0];
 				const FVector UpAxis = ElementToDraw->Points[Index + 1];
 				const double Radius = ElementToDraw->Points[Index + 2].X;
-				const float Thickness = float(ElementToDraw->Thicknes);
 
-				const FQuat Rotation = FQuat::FindBetweenNormals(FVector::UpVector, UpAxis);
-				const FVector XAxis = Rotation.RotateVector(FVector::XAxisVector);
-				const FVector YAxis = Rotation.RotateVector(FVector::YAxisVector);
+				DebugShapes.Circles.Add(FDebugRenderSceneProxy::FCircle(Center, UpAxis.GetSafeNormal(), Radius, Color, DrawTypeOverride)); 
 				
-				static constexpr int32 CircleDivs = 12;
-				FVector PrevPosition = FVector::ZeroVector;
-				for (int32 Div = 0; Div <= CircleDivs; Div++)
-				{
-					const float Angle = (float)Div / (float)CircleDivs * UE_PI * 2.0f;
-					const FVector Position = Center + (FMath::Cos(Angle) * XAxis + FMath::Sin(Angle) * YAxis) * Radius;
-					if (Div > 0)
-					{
-						DebugShapes.Lines.Add(FDebugRenderSceneProxy::FDebugLine(PrevPosition, Position, Color, Thickness));
-					}
-					PrevPosition = Position;
-				}
-
 				if (bDrawLabel)
 				{
 					const FString PrintString = NumPoints == 3 ? ElementToDraw->Description : FString::Printf(TEXT("%s_%d"), *ElementToDraw->Description, Index / 3);
diff --git a/Engine/Source/Developer/LogVisualizer/Private/VisualLoggerRenderingComponent.h b/Engine/Source/Developer/LogVisualizer/Private/VisualLoggerRenderingComponent.h
--- a/Engine/Source/Developer/LogVisualizer/Private/VisualLoggerRenderingComponent.h
+++ b/Engine/Source/Developer/LogVisualizer/Private/VisualLoggerRenderingComponent.h
@@ -16,4 +16,11 @@ public:
 
 	virtual FDebugRenderSceneProxy* CreateDebugSceneProxy() override;
 	virtual FBoxSphereBounds CalcBounds(const FTransform &LocalToWorld) const override;
+
+	// Allows to override the far clipping distance for all the log elements by this actor :
+	void SetFarClippingDistance(double Distance);
+	double GetFarClippingDistance() const { return FarClippingDistance; }
+
+protected:
+	double FarClippingDistance = 0.0;
 };
diff --git a/Engine/Source/Developer/LogVisualizer/Public/VisualLoggerRenderingActorBase.h b/Engine/Source/Developer/LogVisualizer/Public/VisualLoggerRenderingActorBase.h
--- a/Engine/Source/Developer/LogVisualizer/Public/VisualLoggerRenderingActorBase.h
+++ b/Engine/Source/Developer/LogVisualizer/Public/VisualLoggerRenderingActorBase.h
@@ -9,6 +9,8 @@
 #include "DebugRenderSceneProxy.h"
 #include "VisualLoggerRenderingActorBase.generated.h"
 
+class UVisualLoggerRenderingComponent;
+
 /**
 *	Transient actor used to draw visual logger data on level
 *	Base class shared between RewindDebugger and Visual Logger implementations
@@ -24,6 +26,7 @@ public:
 	struct FTimelineDebugShapes
 	{
 		TArray<FDebugRenderSceneProxy::FDebugLine> Lines;
+		TArray<FDebugRenderSceneProxy::FCircle> Circles;
 		TArray<FDebugRenderSceneProxy::FCone> Cones;
 		TArray<FDebugRenderSceneProxy::FDebugBox> Boxes;
 		TArray<FDebugRenderSceneProxy::FSphere> Points;
@@ -37,6 +40,7 @@ public:
 		void Reset()
 		{
 			Lines.Reset();
+			Circles.Reset();
 			Cones.Reset();
 			Boxes.Reset();
 			Points.Reset();
@@ -57,10 +61,14 @@ public:
 	virtual bool IsSelectable() const override { return false; }
 #endif // WITH_EDITOR
 
+// Allows to override the far clipping distance for all the log elements by this actor :
+	void SetFarClippingDistance(double Distance);
+	double GetFarClippingDistance() const;
+
 protected:
 #if ENABLE_VISUAL_LOG
 	void GetDebugShapes(const FVisualLogEntry& EntryItem, bool bAddEntryLocationPointer, FTimelineDebugShapes& OutDebugShapes);
 #endif
 
-	UPrimitiveComponent* RenderingComponent;
+	UVisualLoggerRenderingComponent* RenderingComponent = nullptr;
 };
diff --git a/Engine/Source/Runtime/Engine/Classes/VisualLogger/VisualLoggerKismetLibrary.h b/Engine/Source/Runtime/Engine/Classes/VisualLogger/VisualLoggerKismetLibrary.h
--- a/Engine/Source/Runtime/Engine/Classes/VisualLogger/VisualLoggerKismetLibrary.h
+++ b/Engine/Source/Runtime/Engine/Classes/VisualLogger/VisualLoggerKismetLibrary.h
@@ -51,11 +51,11 @@ class UVisualLoggerKismetLibrary : public UBlueprintFunctionLibrary
 
 	/** Logs arrow - recording for Visual Logs has to be enabled to record this data */
 	UFUNCTION(BlueprintCallable, Category = "Debug|VisualLogger", meta = (DisplayName = "VisLog Arrow", AdvancedDisplay = "WorldContextObject, bAddToMessageLog, Thickness", CallableWithoutWorldContext, DevelopmentOnly, DefaultToSelf = "WorldContextObject"))
-	static void LogArrow(UObject* WorldContextObject, const FVector SegmentStart, const FVector SegmentEnd, FString Text, FLinearColor ObjectColor = FLinearColor::Blue, FName CategoryName = TEXT("VisLogBP"), bool bAddToMessageLog = false);
+	static void LogArrow(UObject* WorldContextObject, const FVector SegmentStart, const FVector SegmentEnd, FString Text, FLinearColor ObjectColor = FLinearColor::Blue, FName CategoryName = TEXT("VisLogBP"), bool bAddToMessageLog = false, float ArrowHeadSize = 8.0f);
 
 	/** Logs circle - recording for Visual Logs has to be enabled to record this data */
 	UFUNCTION(BlueprintCallable, Category = "Debug|VisualLogger", meta = (DisplayName = "VisLog Circle", AdvancedDisplay = "WorldContextObject, bAddToMessageLog, Thickness", CallableWithoutWorldContext, DevelopmentOnly, DefaultToSelf = "WorldContextObject"))
-	static void LogCircle(UObject* WorldContextObject, FVector Center, FVector UpAxis, float Radius, FString Text, FLinearColor ObjectColor = FLinearColor::Blue, const float Thickness = 0.f, FName CategoryName = TEXT("VisLogBP"), bool bAddToMessageLog = false);
+	static void LogCircle(UObject* WorldContextObject, FVector Center, FVector UpAxis, float Radius, FString Text, FLinearColor ObjectColor = FLinearColor::Blue, const float Thickness = 0.f, FName CategoryName = TEXT("VisLogBP"), bool bAddToMessageLog = false, bool bWireframe = false);
 
 	/** Logs segment - recording for Visual Logs has to be enabled to record this data */
 	UFUNCTION(BlueprintCallable, Category = "Debug|VisualLogger", meta = (DisplayName = "VisLog Segment", AdvancedDisplay = "WorldContextObject, bAddToMessageLog, Thickness", CallableWithoutWorldContext, DevelopmentOnly, DefaultToSelf = "WorldContextObject"))
diff --git a/Engine/Source/Runtime/Engine/Private/DebugRenderSceneProxy.cpp b/Engine/Source/Runtime/Engine/Private/DebugRenderSceneProxy.cpp
--- a/Engine/Source/Runtime/Engine/Private/DebugRenderSceneProxy.cpp
+++ b/Engine/Source/Runtime/Engine/Private/DebugRenderSceneProxy.cpp
@@ -15,6 +15,16 @@
 #include "Materials/Material.h"
 #include "Materials/MaterialRenderProxy.h"
 
+static TAutoConsoleVariable<bool> CVarDebugRenderAllowFrustumCulling(
+	TEXT("r.DebugRender.AllowFrustumCulling"), true,
+	TEXT("Allows to cull debug shapes against the view frustum. This helps in high item number situations but incurs a price on the CPU."),
+	ECVF_RenderThreadSafe);
+
+static TAutoConsoleVariable<float> CVarDebugRenderOverrideFarClippingPlane(
+	TEXT("r.DebugRender.OverrideFarClippingPlane"), 0.0f,
+	TEXT("Allows to override the far clipping plane for debug shapes and text (in cm), only effective if > 0.0 and if r.DebugRender.AllowFrustumCulling is enabled"),
+	ECVF_RenderThreadSafe);
+
 // FPrimitiveSceneProxy interface.
 
 FDebugRenderSceneProxy::FDebugRenderSceneProxy(const UPrimitiveComponent* InComponent)
@@ -88,6 +98,7 @@ uint32 FDebugRenderSceneProxy::GetAllocatedSize(void) const
 { 
 	return	FPrimitiveSceneProxy::GetAllocatedSize() + 
 		Cylinders.GetAllocatedSize() + 
+		Circles.GetAllocatedSize() +
 		ArrowLines.GetAllocatedSize() + 
 		Stars.GetAllocatedSize() + 
 		DashedLines.GetAllocatedSize() + 
@@ -100,6 +111,8 @@ uint32 FDebugRenderSceneProxy::GetAllocatedSize(void) const
 
 void FDebugDrawDelegateHelper::DrawDebugLabels(UCanvas* Canvas, APlayerController*)
 {
+	TRACE_CPUPROFILER_EVENT_SCOPE(FDebugDrawDelegateHelper::DrawDebugLabels);
+
 	const FColor OldDrawColor = Canvas->DrawColor;
 	const FFontRenderInfo FontRenderInfo = Canvas->CreateFontRenderInfo(true, false);
 	const FFontRenderInfo FontRenderInfoWithShadow = Canvas->CreateFontRenderInfo(true, true);
@@ -109,9 +122,11 @@ void FDebugDrawDelegateHelper::DrawDebugLabels(UCanvas* Canvas, APlayerControlle
 	UFont* RenderFont = GEngine->GetSmallFont();
 
 	const FSceneView* View = Canvas->SceneView;
+	const bool bAllowFrustumCulling = CVarDebugRenderAllowFrustumCulling.GetValueOnGameThread();
+	const FConvexVolume AdjustedFrustum = FDebugRenderSceneProxy::AdjustViewFrustumForFarClipping(View, FarClippingDistance); 
 	for (auto It = Texts.CreateConstIterator(); It; ++It)
 	{
-		if (FDebugRenderSceneProxy::PointInView(It->Location, View))
+		if (!bAllowFrustumCulling || FDebugRenderSceneProxy::PointInFrustum(It->Location, AdjustedFrustum))
 		{
 			const FVector ScreenLoc = Canvas->Project(It->Location);
 			const FFontRenderInfo& FontInfo = TextWithoutShadowDistance >= 0 ? (FDebugRenderSceneProxy::PointInRange(It->Location, View, TextWithoutShadowDistance) ? FontRenderInfoWithShadow : FontRenderInfo) : FontRenderInfo;
@@ -148,20 +163,41 @@ void FDebugRenderSceneProxy::GetDynamicMeshElements(const TArray<const FSceneVie
 
 void FDebugRenderSceneProxy::GetDynamicMeshElementsForView(const FSceneView* View, const int32 ViewIndex, const FSceneViewFamily& ViewFamily, const uint32 VisibilityMap, FMeshElementCollector& Collector, FMaterialCache& DefaultMaterialCache, FMaterialCache& SolidMeshMaterialCache) const
 {
+	TRACE_CPUPROFILER_EVENT_SCOPE(FDebugRenderSceneProxy::GetDynamicMeshElementsForView);
+
 	FPrimitiveDrawInterface* PDI = Collector.GetPDI(ViewIndex);
 
+	const bool bAllowFrustumCulling = CVarDebugRenderAllowFrustumCulling.GetValueOnRenderThread();
+	const FConvexVolume AdjustedFrustum = AdjustViewFrustumForFarClipping(View, FarClippingDistance);;
+
 	// Draw Lines
 	const int32 LinesNum = Lines.Num();
 	PDI->AddReserveLines(SDPG_World, LinesNum, false, false);
 	for (const FDebugLine& Line : Lines)
 	{
-		Line.Draw(PDI);
+		if (!bAllowFrustumCulling || SegmentInFrustum(Line.Start, Line.End, AdjustedFrustum))
+		{
+			Line.Draw(PDI);
+		}
 	}
 
 	// Draw Dashed Lines
 	for (const FDashedLine& Dash : DashedLines)
 	{
-		Dash.Draw(PDI);
+		if (!bAllowFrustumCulling || SegmentInFrustum(Dash.Start, Dash.End, AdjustedFrustum))
+		{
+			Dash.Draw(PDI);
+		}
+	}
+
+	// Draw Circles
+	for (const FCircle& Circle : Circles)
+	{
+		// The idea is to have cheap test, not a precise one : a bounding sphere will do : 
+		if (!bAllowFrustumCulling || SphereInFrustum(Circle.Center, Circle.Radius, AdjustedFrustum))
+		{
+			Circle.Draw(PDI, (Circle.DrawTypeOverride != EDrawType::Invalid) ? Circle.DrawTypeOverride : DrawType, DrawAlpha, DefaultMaterialCache, ViewIndex, Collector);
+		}
 	}
 
 	// Draw Arrows
@@ -169,38 +205,79 @@ void FDebugRenderSceneProxy::GetDynamicMeshElementsForView(const FSceneView* Vie
 	PDI->AddReserveLines(SDPG_World, 5 * ArrowsNum, false, false);
 	for (const FArrowLine& ArrowLine : ArrowLines)
 	{
-		ArrowLine.Draw(PDI, 8.f);
+		if (!bAllowFrustumCulling || SegmentInFrustum(ArrowLine.Start, ArrowLine.End, AdjustedFrustum))
+		{
+			ArrowLine.Draw(PDI);
+		}
 	}
 
 	// Draw Stars
 	for (const FWireStar& Star : Stars)
 	{
-		Star.Draw(PDI);
+		if (!bAllowFrustumCulling || SphereInFrustum(Star.Position, Star.Size, AdjustedFrustum))
+		{
+			Star.Draw(PDI);
+		}
 	}
 
 	// Draw Cylinders
 	for (const FWireCylinder& Cylinder : Cylinders)
 	{
-		Cylinder.Draw(PDI, (Cylinder.DrawTypeOverride != EDrawType::Invalid) ? Cylinder.DrawTypeOverride : DrawType, DrawAlpha, DefaultMaterialCache, ViewIndex, Collector);
+		// The idea is to have cheap test, not a precise one : a bounding sphere will do : 
+		double BoundingSphereRadius = FMath::Sqrt(FMath::Square(Cylinder.HalfHeight) + 2.0 * FMath::Square(Cylinder.Radius));
+		if (!bAllowFrustumCulling || SphereInFrustum(Cylinder.Base, BoundingSphereRadius, AdjustedFrustum))
+		{
+			Cylinder.Draw(PDI, (Cylinder.DrawTypeOverride != EDrawType::Invalid) ? Cylinder.DrawTypeOverride : DrawType, DrawAlpha, DefaultMaterialCache, ViewIndex, Collector);
+		}
 	}
 
 	// Draw Boxes
 	for (const FDebugBox& Box : Boxes)
 	{
-		Box.Draw(PDI, (Box.DrawTypeOverride != EDrawType::Invalid) ? Box.DrawTypeOverride : DrawType, DrawAlpha, DefaultMaterialCache, ViewIndex, Collector);
+		FVector Center, Extents;
+		FBox WorldBox = Box.Box.TransformBy(Box.Transform);
+		WorldBox.GetCenterAndExtents(Center, Extents);
+		if (!bAllowFrustumCulling || BoxInFrustum(Center, Extents, AdjustedFrustum))
+		{
+			Box.Draw(PDI, (Box.DrawTypeOverride != EDrawType::Invalid) ? Box.DrawTypeOverride : DrawType, DrawAlpha, DefaultMaterialCache, ViewIndex, Collector);
+		}
 	}
 
 	// Draw Cones
-	TArray<FVector> Verts;
 	for (const FCone& Cone : Cones)
 	{
-		Cone.Draw(PDI, (Cone.DrawTypeOverride != EDrawType::Invalid) ? Cone.DrawTypeOverride : DrawType, DrawAlpha, DefaultMaterialCache, ViewIndex, Collector, &Verts);
+		// The idea is to have a cheap test, not a precise one : a bounding sphere will do : 
+		double Angle = FMath::Max(Cone.Angle1, Cone.Angle2) * UE_PI / 180.0;
+		double CosAngle = FMath::Cos(Angle);
+		FVector SegmentStart = Cone.ConeToWorld.TransformPosition(FVector::ZeroVector);
+		FVector SegmentEnd = Cone.ConeToWorld.TransformPosition(FVector::XAxisVector);
+		FVector Forward = (SegmentEnd - SegmentStart);
+		double ConeLength = Forward.Size();
+
+		FVector SphereCenter;
+		double Radius;
+		if (Angle > UE_PI / 4.0)
+		{
+			SphereCenter = SegmentStart + CosAngle * Forward;
+			Radius = FMath::Sin(Angle) * ConeLength;
+		}
+		else
+		{
+			SphereCenter = SegmentStart + Forward / (2.0 * CosAngle);
+			Radius = ConeLength / (2.0 * CosAngle);
+		}
+
+		if (!bAllowFrustumCulling || SphereInFrustum(SphereCenter, Radius, AdjustedFrustum))
+		{
+			TArray<FVector> Verts;
+			Cone.Draw(PDI, (Cone.DrawTypeOverride != EDrawType::Invalid) ? Cone.DrawTypeOverride : DrawType, DrawAlpha, DefaultMaterialCache, ViewIndex, Collector, &Verts);
+		}
 	}
 
 	// Draw spheres
 	for (const FSphere& Sphere : Spheres)
 	{
-		if (PointInView(Sphere.Location, View))
+		if (!bAllowFrustumCulling || SphereInFrustum(Sphere.Location, Sphere.Radius, AdjustedFrustum))
 		{
 			Sphere.Draw(PDI, (Sphere.DrawTypeOverride != EDrawType::Invalid) ? Sphere.DrawTypeOverride : DrawType, DrawAlpha, DefaultMaterialCache, ViewIndex, Collector);
 		}
@@ -209,7 +286,15 @@ void FDebugRenderSceneProxy::GetDynamicMeshElementsForView(const FSceneView* Vie
 	// Draw Capsules
 	for (const FCapsule& Capsule : Capsules)
 	{
-		if (PointInView(Capsule.Base, View))
+		// The idea is to have cheap test, not a precise one : a bounding sphere will do : 
+		const FVector Origin = Capsule.Base;
+		const double HalfAxis = FMath::Max<double>(Capsule.HalfHeight - Capsule.Radius, 1.0);
+		const FVector BottomEnd = Capsule.Base;
+		const double TotalHalfLength = (HalfAxis + Capsule.Radius);
+		const FVector TopEnd = Capsule.Base + 2.0 * TotalHalfLength * Capsule.Z;
+		const FVector SphereCenter = (TopEnd + BottomEnd) / 2.0;
+
+		if (!bAllowFrustumCulling || SphereInFrustum(SphereCenter, TotalHalfLength, AdjustedFrustum))
 		{
 			Capsule.Draw(PDI, (Capsule.DrawTypeOverride != EDrawType::Invalid) ? Capsule.DrawTypeOverride : DrawType, DrawAlpha, DefaultMaterialCache, ViewIndex, Collector);
 		}
@@ -218,15 +303,48 @@ void FDebugRenderSceneProxy::GetDynamicMeshElementsForView(const FSceneView* Vie
 	// Draw Meshes
 	for (const FMesh& Mesh : Meshes)
 	{
+		if (!bAllowFrustumCulling || Mesh.Box.IsValid)
+		{
+			FVector Center, Extents;
+			Mesh.Box.GetCenterAndExtents(Center, Extents);
+			if (!BoxInFrustum(Center, Extents, AdjustedFrustum))
+			{
+				continue;
+			}
+		}
 		FDynamicMeshBuilder MeshBuilder(View->GetFeatureLevel());
 		MeshBuilder.AddVertices(Mesh.Vertices);
 		MeshBuilder.AddTriangles(Mesh.Indices);
 
-		FMaterialCache& MeshMaterialCache = Mesh.Color.A == 255 ? SolidMeshMaterialCache : DefaultMaterialCache; 
+		FMaterialCache& MeshMaterialCache = Mesh.Color.A == 255 ? SolidMeshMaterialCache : DefaultMaterialCache;
 		MeshBuilder.GetMesh(FMatrix::Identity, MeshMaterialCache[Mesh.Color], SDPG_World, false, false, ViewIndex, Collector);
 	}
 }
 
+FConvexVolume FDebugRenderSceneProxy::AdjustViewFrustumForFarClipping(const FSceneView* InView, double InFarClippingDistance)
+{
+	if (InView == nullptr)
+	{
+		return FConvexVolume();
+	}
+
+	FConvexVolume AdjustedFrustum = InView->ViewFrustum;
+	double FinalFarClippingDistance = InFarClippingDistance;
+	if (const double OverrideFarClippingPlane = CVarDebugRenderOverrideFarClippingPlane.GetValueOnAnyThread(); OverrideFarClippingPlane > 0.0f)
+	{ 
+		FinalFarClippingDistance = OverrideFarClippingPlane;
+	}
+
+	if (FinalFarClippingDistance > 0.0f)
+	{
+		const FPlane FarPlane(InView->ViewMatrices.GetViewOrigin() + InView->GetViewDirection() * FinalFarClippingDistance, InView->GetViewDirection());
+		// Derive the view frustum from the view projection matrix, overriding the far plane
+		GetViewFrustumBounds(AdjustedFrustum, InView->ViewMatrices.GetViewProjectionMatrix(), FarPlane, true, false);
+	}
+
+	return AdjustedFrustum;
+}
+
 /**
 * Draws a line with an arrow at the end.
 *
@@ -237,8 +355,8 @@ void FDebugRenderSceneProxy::GetDynamicMeshElementsForView(const FSceneView* Vie
 */
 void FDebugRenderSceneProxy::DrawLineArrow(FPrimitiveDrawInterface* PDI,const FVector &Start,const FVector &End,const FColor &Color,float Mag) const
 {
-	const FArrowLine ArrowLine = {Start, End, Color};
-	ArrowLine.Draw(PDI, Mag);
+	const FArrowLine ArrowLine(Start, End, Color, Mag);
+	ArrowLine.Draw(PDI);
 }
 
 FDebugRenderSceneProxy::FMaterialCache::FMaterialCache(FMeshElementCollector& InCollector, bool bUseLight, UMaterial* InMaterial)
@@ -287,7 +405,20 @@ void FDebugRenderSceneProxy::FWireStar::Draw(FPrimitiveDrawInterface* PDI) const
 	DrawWireStar(PDI, Position, Size, Color, SDPG_World);
 }
 
-void FDebugRenderSceneProxy::FArrowLine::Draw(FPrimitiveDrawInterface* PDI, const float Mag) const
+void FDebugRenderSceneProxy::FArrowLine::Draw(FPrimitiveDrawInterface* PDI) const
+{
+	// draw a pretty arrow
+	FVector Dir = End - Start;
+	const float DirMag = Dir.Size();
+	Dir /= DirMag;
+	FVector YAxis, ZAxis;
+	Dir.FindBestAxisVectors(YAxis, ZAxis);
+	FMatrix ArrowTM(Dir, YAxis, ZAxis, Start);
+	DrawDirectionalArrow(PDI, ArrowTM, Color, DirMag, Mag, SDPG_World);
+}
+
+// Deprecated : set the Mag member and call the base Draw function instead
+void FDebugRenderSceneProxy::FArrowLine::Draw(FPrimitiveDrawInterface* PDI, const float InMag) const
 {
 	// draw a pretty arrow
 	FVector Dir = End - Start;
@@ -316,6 +447,21 @@ void FDebugRenderSceneProxy::FDebugBox::Draw(FPrimitiveDrawInterface* PDI, EDraw
 	}
 }
 
+void FDebugRenderSceneProxy::FCircle::Draw(FPrimitiveDrawInterface* PDI, EDrawType InDrawType, uint32 InDrawAlpha, FMaterialCache& MaterialCache, int32 ViewIndex, FMeshElementCollector& Collector) const
+{
+	FVector XAxis, YAxis;
+	Axis.FindBestAxisVectors(XAxis, YAxis);
+	if (InDrawType == SolidAndWireMeshes || InDrawType == WireMesh)
+	{
+		DrawCircle(PDI, Center, XAxis, YAxis, Color, Radius, (InDrawType == SolidAndWireMeshes) ? 9 : 12, SDPG_World, InDrawType == SolidAndWireMeshes ? 2 : 0, 0, true);
+	}
+
+	if (InDrawType == SolidAndWireMeshes || InDrawType == SolidMesh)
+	{
+		GetDiscMesh(Center, XAxis, YAxis, Radius, 12, MaterialCache[Color.WithAlpha(InDrawAlpha * Color.A)], SDPG_World, ViewIndex, Collector);
+	}
+}
+
 void FDebugRenderSceneProxy::FWireCylinder::Draw(FPrimitiveDrawInterface* PDI, EDrawType InDrawType, uint32 InDrawAlpha, FMaterialCache& MaterialCache, int32 ViewIndex, FMeshElementCollector& Collector) const
 {
 	FVector XAxis, YAxis;
diff --git a/Engine/Source/Runtime/Engine/Private/PrimitiveDrawingUtils.cpp b/Engine/Source/Runtime/Engine/Private/PrimitiveDrawingUtils.cpp
--- a/Engine/Source/Runtime/Engine/Private/PrimitiveDrawingUtils.cpp
+++ b/Engine/Source/Runtime/Engine/Private/PrimitiveDrawingUtils.cpp
@@ -500,6 +500,73 @@ void DrawCone(FPrimitiveDrawInterface* PDI,const FMatrix& ConeToWorld, float Ang
 	}
 }
 
+void BuildDiscVerts(const FVector& Center, const FVector& XAxis, const FVector& YAxis, double Radius, uint32 Sides, TArray<FDynamicMeshVertex>& OutVerts, TArray<uint32>& OutIndices)
+{
+	const float	AngleDelta = 2.0f * UE_PI / Sides;
+	FVector	LastVertex = Center + XAxis * Radius;
+	FVector	ZAxis = XAxis ^ YAxis;
+
+	FVector2D TC = FVector2D(0.0f, 0.0f);
+	float TCStep = 1.0f / Sides;
+
+	int32 BaseVertIndex = OutVerts.Num();
+
+	for (uint32 SideIndex = 0; SideIndex < Sides; SideIndex++)
+	{
+		const FVector Vertex = Center + (XAxis * FMath::Cos(AngleDelta * (SideIndex + 1)) + YAxis * FMath::Sin(AngleDelta * (SideIndex + 1))) * Radius;
+		FVector Normal = Vertex - Center;
+		Normal.Normalize();
+
+		FDynamicMeshVertex MeshVertex;
+
+		MeshVertex.Position = FVector3f(Vertex);
+		MeshVertex.TextureCoordinate[0] = FVector2f(TC);
+
+		MeshVertex.SetTangents(
+			(FVector3f)-ZAxis,
+			FVector3f((-ZAxis) ^ Normal),
+			(FVector3f)Normal
+		);
+
+		OutVerts.Add(MeshVertex); //Add bottom vertex
+
+		LastVertex = Vertex;
+		TC.X += TCStep;
+	}
+
+	//Add top triangles, in the style of a fan.
+	//Note if we wanted nice rendering of the caps then we need to duplicate the vertices and modify
+	//texture/tangent coordinates.
+	for (uint32 SideIndex = 1; SideIndex < Sides; SideIndex++)
+	{
+		int32 V0 = BaseVertIndex;
+		int32 V1 = BaseVertIndex + SideIndex;
+		int32 V2 = BaseVertIndex + ((SideIndex + 1) % Sides);
+
+		//bottom
+		OutIndices.Add(V0);
+		OutIndices.Add(V1);
+		OutIndices.Add(V2);
+	}
+}
+
+void GetDiscMesh(const FVector& Center, const FVector& XAxis, const FVector& YAxis, double Radius, int32 Sides, const FMaterialRenderProxy* MaterialRenderProxy, uint8 DepthPriority, int32 ViewIndex, FMeshElementCollector& Collector, HHitProxy* HitProxy)
+{
+	GetDiscMesh(FMatrix::Identity, Center, XAxis, YAxis, Radius, Sides, MaterialRenderProxy, DepthPriority, ViewIndex, Collector, HitProxy);
+}
+
+void GetDiscMesh(const FMatrix& ToWorld, const FVector& Center, const FVector& XAxis, const FVector& YAxis, double Radius, int32 Sides, const FMaterialRenderProxy* MaterialRenderProxy, uint8 DepthPriority, int32 ViewIndex, FMeshElementCollector& Collector, HHitProxy* HitProxy)
+{
+	TArray<FDynamicMeshVertex> MeshVerts;
+	TArray<uint32> MeshIndices;
+	BuildDiscVerts(Center, XAxis, YAxis, Radius, Sides, MeshVerts, MeshIndices);
+
+	FDynamicMeshBuilder MeshBuilder(Collector.GetFeatureLevel());
+	MeshBuilder.AddVertices(MeshVerts);
+	MeshBuilder.AddTriangles(MeshIndices);
+
+	MeshBuilder.GetMesh(ToWorld, MaterialRenderProxy, DepthPriority, false, false, true, ViewIndex, Collector, HitProxy);
+}
 
 void BuildCylinderVerts(const FVector& Base, const FVector& XAxis, const FVector& YAxis, const FVector& ZAxis, double Radius, double HalfHeight, uint32 Sides, TArray<FDynamicMeshVertex>& OutVerts, TArray<uint32>& OutIndices)
 {
@@ -988,7 +1055,7 @@ void DrawCircle(FPrimitiveDrawInterface* PDI, const FVector& Base, const FVector
 	const float	AngleDelta = 2.0f * UE_PI / NumSides;
 	FVector	LastVertex = Base + X * Radius;
 
-	for(int32 SideIndex = 0;SideIndex < NumSides;SideIndex++)
+	for(int32 SideIndex = 0; SideIndex < NumSides; SideIndex++)
 	{
 		const FVector Vertex = Base + (X * FMath::Cos(AngleDelta * (SideIndex + 1)) + Y * FMath::Sin(AngleDelta * (SideIndex + 1))) * Radius;
 		PDI->DrawTranslucentLine(LastVertex, Vertex, Color, DepthPriority, Thickness, DepthBias, bScreenSpace);
@@ -1089,7 +1156,7 @@ void DrawRectangle(class FPrimitiveDrawInterface* PDI, const FVector& Center, co
 
 void DrawWireSphere(class FPrimitiveDrawInterface* PDI, const FVector& Base, const FLinearColor& Color, double Radius, int32 NumSides, uint8 DepthPriority, float Thickness, float DepthBias, bool bScreenSpace)
 {
-	DrawCircle(PDI, Base, FVector(1,0,0), FVector(0,1,0), Color, Radius, NumSides, DepthPriority, Thickness, DepthBias, bScreenSpace);
+	DrawCircle(PDI, Base, FVector(1, 0, 0), FVector(0, 1, 0), Color, Radius, NumSides, DepthPriority, Thickness, DepthBias, bScreenSpace);
 	DrawCircle(PDI, Base, FVector(1, 0, 0), FVector(0, 0, 1), Color, Radius, NumSides, DepthPriority, Thickness, DepthBias, bScreenSpace);
 	DrawCircle(PDI, Base, FVector(0, 1, 0), FVector(0, 0, 1), Color, Radius, NumSides, DepthPriority, Thickness, DepthBias, bScreenSpace);
 }
@@ -1108,7 +1175,6 @@ void DrawWireSphere(class FPrimitiveDrawInterface* PDI, const FTransform& Transf
 	DrawCircle(PDI, Transform.GetLocation(), Transform.GetScaledAxis(EAxis::Y), Transform.GetScaledAxis(EAxis::Z), Color, Radius, NumSides, DepthPriority, Thickness, DepthBias, bScreenSpace);
 }
 
-
 void DrawWireSphereAutoSides(class FPrimitiveDrawInterface* PDI, const FTransform& Transform, const FLinearColor& Color, double Radius, uint8 DepthPriority, float Thickness, float DepthBias, bool bScreenSpace)
 {
 	// Guess a good number of sides
diff --git a/Engine/Source/Runtime/Engine/Private/VisualLogger/VisualLogger.cpp b/Engine/Source/Runtime/Engine/Private/VisualLogger/VisualLogger.cpp
--- a/Engine/Source/Runtime/Engine/Private/VisualLogger/VisualLogger.cpp
+++ b/Engine/Source/Runtime/Engine/Private/VisualLogger/VisualLogger.cpp
@@ -92,13 +92,22 @@ void FVisualLogger::SegmentLogfImpl(const UObject* Object, const FName& Category
 	);
 }
 
+// Deprecated
 void FVisualLogger::ArrowLogfImpl(const UObject* Object, const FName& CategoryName, ELogVerbosity::Type Verbosity, const FVector& Start, const FVector& End, const FColor& Color, const TCHAR* Fmt, ...)
 {
 	COLLAPSED_LOGF(
-		CurrentEntry->AddArrow(Start, End, CategoryName, Verbosity, Color, Buffer);
+		CurrentEntry->AddArrow(Start, End, CategoryName, Verbosity, Color, Buffer, 8);
 	);
 }
 
+void FVisualLogger::ArrowLineLogfImpl(const UObject* Object, const FName& CategoryName, ELogVerbosity::Type Verbosity, const FVector& Start, const FVector& End, const FColor& Color, const uint16 Mag, const TCHAR* Fmt, ...)
+{
+	COLLAPSED_LOGF(
+		CurrentEntry->AddArrow(Start, End, CategoryName, Verbosity, Color, Buffer, Mag);
+	);
+}
+
+// Deprecated
 void FVisualLogger::CircleLogfImpl(const UObject* Object, const FName& CategoryName, ELogVerbosity::Type Verbosity, const FVector& Center, const FVector& UpAxis, const float Radius, const FColor& Color, const uint16 Thickness, const TCHAR* Fmt, ...)
 {
 	COLLAPSED_LOGF(
@@ -106,6 +115,13 @@ void FVisualLogger::CircleLogfImpl(const UObject* Object, const FName& CategoryN
 	);
 }
 
+void FVisualLogger::DiscLogfImpl(const UObject* Object, const FName& CategoryName, ELogVerbosity::Type Verbosity, const FVector& Center, const FVector& UpAxis, const float Radius, const FColor& Color, const uint16 Thickness, bool bWireframe, const TCHAR* Fmt, ...)
+{
+	COLLAPSED_LOGF(
+		CurrentEntry->AddCircle(Center, UpAxis, Radius, CategoryName, Verbosity, Color, Buffer, Thickness, bWireframe);
+	);
+}
+
 void FVisualLogger::LocationLogfImpl(const UObject* Object, const FName& CategoryName, ELogVerbosity::Type Verbosity, const FVector& Location, uint16 Thickness, const FColor& Color, const TCHAR* Fmt, ...)
 {
 	COLLAPSED_LOGF(
diff --git a/Engine/Source/Runtime/Engine/Private/VisualLogger/VisualLoggerKismetLibrary.cpp b/Engine/Source/Runtime/Engine/Private/VisualLogger/VisualLoggerKismetLibrary.cpp
--- a/Engine/Source/Runtime/Engine/Private/VisualLogger/VisualLoggerKismetLibrary.cpp
+++ b/Engine/Source/Runtime/Engine/Private/VisualLogger/VisualLoggerKismetLibrary.cpp
@@ -142,11 +142,11 @@ void UVisualLoggerKismetLibrary::LogOrientedBox(UObject* WorldContextObject, FBo
 	}
 }
 
-void UVisualLoggerKismetLibrary::LogArrow(UObject* WorldContextObject, const FVector SegmentStart, const FVector SegmentEnd, FString Text, FLinearColor ObjectColor, FName CategoryName, bool bAddToMessageLog)
+void UVisualLoggerKismetLibrary::LogArrow(UObject* WorldContextObject, const FVector SegmentStart, const FVector SegmentEnd, FString Text, FLinearColor ObjectColor, FName CategoryName, bool bAddToMessageLog, float ArrowHeadSize)
 {
 #if ENABLE_VISUAL_LOG
 	const ELogVerbosity::Type DefaultVerbosity = ELogVerbosity::Log;
-	FVisualLogger::ArrowLogf(WorldContextObject, CategoryName, DefaultVerbosity, SegmentStart, SegmentEnd, ObjectColor.ToFColor(true), TEXT("%s"), *Text);
+	FVisualLogger::ArrowLineLogf(WorldContextObject, CategoryName, DefaultVerbosity, SegmentStart, SegmentEnd, ObjectColor.ToFColor(true), ArrowHeadSize, TEXT("%s"), *Text);
 #endif
 	if (bAddToMessageLog)
 	{
@@ -155,11 +155,11 @@ void UVisualLoggerKismetLibrary::LogArrow(UObject* WorldContextObject, const FVe
 	}
 }
 
-void UVisualLoggerKismetLibrary::LogCircle(UObject* WorldContextObject, FVector Center, FVector UpAxis, float Radius, FString Text, FLinearColor ObjectColor, const float Thickness, FName CategoryName, bool bAddToMessageLog)
+void UVisualLoggerKismetLibrary::LogCircle(UObject* WorldContextObject, FVector Center, FVector UpAxis, float Radius, FString Text, FLinearColor ObjectColor, const float Thickness, FName CategoryName, bool bAddToMessageLog, bool bWireframe)
 {
 #if ENABLE_VISUAL_LOG
 	const ELogVerbosity::Type DefaultVerbosity = ELogVerbosity::Log;
-	FVisualLogger::CircleLogf(WorldContextObject, CategoryName, DefaultVerbosity, Center, UpAxis, Radius, ObjectColor.ToFColor(true), (uint16)Thickness, TEXT("%s"), *Text);
+	FVisualLogger::DiscLogf(WorldContextObject, CategoryName, DefaultVerbosity, Center, UpAxis, Radius, ObjectColor.ToFColor(true), (uint16)Thickness, bWireframe, TEXT("%s"), *Text);
 #endif
 	if (bAddToMessageLog)
 	{
diff --git a/Engine/Source/Runtime/Engine/Private/VisualLogger/VisualLoggerTypes.cpp b/Engine/Source/Runtime/Engine/Private/VisualLogger/VisualLoggerTypes.cpp
--- a/Engine/Source/Runtime/Engine/Private/VisualLogger/VisualLoggerTypes.cpp
+++ b/Engine/Source/Runtime/Engine/Private/VisualLogger/VisualLoggerTypes.cpp
@@ -250,7 +250,7 @@ void FVisualLogEntry::AddSegment(const FVector& Start, const FVector& End, const
 	ElementsToDraw.Add(Element);
 }
 
-void FVisualLogEntry::AddArrow(const FVector& Start, const FVector& End, const FName& CategoryName, ELogVerbosity::Type Verbosity, const FColor& Color, const FString& Description)
+void FVisualLogEntry::AddArrow(const FVector& Start, const FVector& End, const FName& CategoryName, ELogVerbosity::Type Verbosity, const FColor& Color, const FString& Description, const uint16 Mag)
 {
 	FVisualLogShapeElement Element(EVisualLoggerShapeElement::Arrow);
 	Element.Category = CategoryName;
@@ -260,10 +260,11 @@ void FVisualLogEntry::AddArrow(const FVector& Start, const FVector& End, const F
 	Element.Points.Add(Start);
 	Element.Points.Add(End);
 	Element.Verbosity = Verbosity;
+	Element.Mag = Mag;
 	ElementsToDraw.Add(Element);
 }
 
-void FVisualLogEntry::AddCircle(const FVector& Center, const FVector& UpAxis, const float Radius, const FName& CategoryName, ELogVerbosity::Type Verbosity, const FColor& Color, const FString& Description, const uint16 Thickness)
+void FVisualLogEntry::AddCircle(const FVector& Center, const FVector& UpAxis, const float Radius, const FName& CategoryName, ELogVerbosity::Type Verbosity, const FColor& Color, const FString& Description, const uint16 Thickness, bool bInUseWires)
 {
 	FVisualLogShapeElement Element(EVisualLoggerShapeElement::Circle);
 	Element.Category = CategoryName;
@@ -274,6 +275,7 @@ void FVisualLogEntry::AddCircle(const FVector& Center, const FVector& UpAxis, co
 	Element.Points.Add(Center);
 	Element.Points.Add(UpAxis);
 	Element.Points.Add(FVector(Radius, 0., 0.));
+	Element.Type = bInUseWires ? EVisualLoggerShapeElement::WireCircle : EVisualLoggerShapeElement::Circle;
 	Element.Verbosity = Verbosity;
 	ElementsToDraw.Add(Element);
 }
diff --git a/Engine/Source/Runtime/Engine/Public/DebugRenderSceneProxy.h b/Engine/Source/Runtime/Engine/Public/DebugRenderSceneProxy.h
--- a/Engine/Source/Runtime/Engine/Public/DebugRenderSceneProxy.h
+++ b/Engine/Source/Runtime/Engine/Public/DebugRenderSceneProxy.h
@@ -62,14 +62,52 @@ public:
 
 	FORCEINLINE static bool PointInView(const FVector& Location, const FSceneView* View)
 	{
-		return View ? View->ViewFrustum.IntersectBox(Location, FVector::ZeroVector) : false;
+		return View ? PointInFrustum(Location, View->ViewFrustum) : false;
 	}
 
-	FORCEINLINE static bool PointInRange(const FVector& Start, const FSceneView* View, float Range)
+	FORCEINLINE static bool SegmentInView(const FVector& StartPoint, const FVector& EndPoint, const FSceneView* View)
+	{
+		return View ? SegmentInFrustum(StartPoint, EndPoint, View->ViewFrustum) : false;
+	}
+
+	FORCEINLINE static bool BoxInView(const FVector& Origin, const FVector& Extent, const FSceneView* View)
+	{
+		return View ? BoxInFrustum(Origin, Extent, View->ViewFrustum) : false;
+	}
+
+	FORCEINLINE static bool SphereInView(const FVector& Center, double Radius, const FSceneView* View)
+	{
+		return View ? SphereInFrustum(Center, Radius, View->ViewFrustum) : false;
+	}
+
+	FORCEINLINE static bool PointInFrustum(const FVector& Location, const FConvexVolume& ViewFrustum)
+	{
+		return ViewFrustum.IntersectPoint(Location);
+	}
+
+	FORCEINLINE static bool SegmentInFrustum(const FVector& StartPoint, const FVector& EndPoint, const FConvexVolume& ViewFrustum)
+	{
+		return ViewFrustum.IntersectLineSegment(StartPoint, EndPoint);
+	}
+
+	FORCEINLINE static bool BoxInFrustum(const FVector& Origin, const FVector& Extent, const FConvexVolume& ViewFrustum)
+	{
+		return ViewFrustum.IntersectBox(Origin, Extent);
+	}
+
+	FORCEINLINE static bool SphereInFrustum(const FVector& Center, double Radius, const FConvexVolume& ViewFrustum)
+	{
+		return ViewFrustum.IntersectSphere(Center, static_cast<float>(Radius));
+	}
+
+	FORCEINLINE static bool PointInRange(const FVector& Start, const FSceneView* View, double Range)
 	{
 		return FVector::DistSquared(Start, View->ViewMatrices.GetViewOrigin()) <= FMath::Square(Range);
 	}
 
+	// Create a new frustum out of the view, taking into account the far clipping distance (if InFarClippingDistance > 0.0) 
+	ENGINE_API static FConvexVolume AdjustViewFrustumForFarClipping(const FSceneView* InView, double InFarClippingDistance);
+	
 	struct FMaterialCache
 	{
 		FMaterialCache(FMeshElementCollector& InCollector, bool bUseLight = false, UMaterial* InMaterial = nullptr);
@@ -99,6 +137,26 @@ public:
 		ENGINE_API void Draw(FPrimitiveDrawInterface* PDI) const;
 	};
 
+	/** Struct to hold info about circles to render. */
+	struct FCircle
+	{
+		FCircle(const FVector& InCenter, const FVector& InAxis, const float InRadius, const FColor& InColor, EDrawType InDrawTypeOverride = EDrawType::Invalid)
+			: Center(InCenter)
+			, Axis(InAxis)
+			, Radius(InRadius)
+			, Color(InColor)
+			, DrawTypeOverride(InDrawTypeOverride)
+		{}
+
+		ENGINE_API void Draw(FPrimitiveDrawInterface* PDI, EDrawType InDrawType, uint32 InDrawAlpha, FMaterialCache& MaterialCache, int32 ViewIndex, FMeshElementCollector& Collector) const;
+
+		FVector Center;
+		FVector Axis;
+		float Radius;
+		FColor Color;
+		EDrawType DrawTypeOverride = EDrawType::Invalid;
+	};
+
 	/** Struct to hold info about boxes to render. */
 	struct FDebugBox
 	{
@@ -148,7 +206,7 @@ public:
 	/** Struct to hold info about lined stars to render. */
 	struct FWireStar
 	{
-		FWireStar(const FVector &InPosition, const FColor &InColor, const float &InSize)
+		FWireStar(const FVector &InPosition, const FColor &InColor, const float InSize)
 			: Position(InPosition)
 			, Color(InColor)
 			, Size(InSize) 
@@ -164,17 +222,22 @@ public:
 	/** Struct to hold info about arrowed lines to render. */
 	struct FArrowLine
 	{
-		FArrowLine(const FVector &InStart, const FVector &InEnd, const FColor &InColor) 
+		FArrowLine(const FVector &InStart, const FVector &InEnd, const FColor &InColor, const float InMag = 0.0f) 
 			: Start(InStart)
 			, End(InEnd)
 			, Color(InColor) 
+			, Mag(InMag > 0.0f ? InMag : 8.0f) // default to a Mag of 8.0 as it used to be the default in the past
 		{}
 
-		ENGINE_API void Draw(FPrimitiveDrawInterface* PDI, const float Mag) const;
+		ENGINE_API void Draw(FPrimitiveDrawInterface* PDI) const;
+
+		UE_DEPRECATED(5.6, "Use the other Draw function and pass the proper data in the constructor")
+		ENGINE_API void Draw(FPrimitiveDrawInterface* PDI, const float InMag) const;
 
 		FVector Start;
 		FVector End;
 		FColor Color;
+		float Mag;
 	};
 
 	/** Struct to gold info about dashed lines to render. */
@@ -251,8 +314,9 @@ public:
 
 	struct FMesh
 	{
-		TArray<FDynamicMeshVertex>	Vertices;
+		TArray<FDynamicMeshVertex> Vertices;
 		TArray <uint32> Indices;
+		FBox Box = FBox(ForceInit);
 		FColor Color;
 	};
 
@@ -283,6 +347,7 @@ public:
 	TArray<FDebugLine> Lines;
 	TArray<FDashedLine>	DashedLines;
 	TArray<FArrowLine> ArrowLines;
+	TArray<FCircle> Circles;
 	TArray<FWireCylinder> Cylinders;
 	TArray<FWireStar> Stars;
 	TArray<FDebugBox> Boxes;
@@ -299,6 +364,7 @@ public:
 	FDelegateHandle DebugTextDrawingDelegateHandle;
 	EDrawType DrawType;
 	uint32 DrawAlpha;
+	double FarClippingDistance = 0.0;
 
 	TWeakObjectPtr<UMaterial> SolidMeshMaterial;
 
@@ -313,6 +379,7 @@ struct FDebugDrawDelegateHelper
 		: State(UndefinedState)
 		, ViewFlagName(TEXT("Game"))
 		, TextWithoutShadowDistance(1500)
+		, FarClippingDistance(0.0)
 	{}
 
 	virtual ~FDebugDrawDelegateHelper() {}
@@ -330,6 +397,7 @@ public:
 		ViewFlagName = InSceneProxy->ViewFlagName;
 		TextWithoutShadowDistance = InSceneProxy->TextWithoutShadowDistance;
 		State = (State == UndefinedState) ? InitializedState : State;
+		FarClippingDistance = InSceneProxy->FarClippingDistance;
 	}
 
 	UE_DEPRECATED(5.0, "This method is deprecated. Call RequestRegisterDebugDrawDelegate or override RegisterDebugDrawDelegateInternal instead.")
@@ -369,6 +437,7 @@ protected:
 	void ResetTexts() { Texts.Reset(); }
 	const TextArray& GetTexts() const { return Texts; }
 	float GetTextWithoutShadowDistance() const {return TextWithoutShadowDistance; }
+	double GetFarClippingDistance() const { return FarClippingDistance; }
 
 protected:
 	FDebugDrawDelegate DebugTextDrawingDelegate;
@@ -386,4 +455,5 @@ private:
 	TextArray Texts;
 	FString ViewFlagName;
 	float TextWithoutShadowDistance;
+	double FarClippingDistance;
 };
diff --git a/Engine/Source/Runtime/Engine/Public/SceneManagement.h b/Engine/Source/Runtime/Engine/Public/SceneManagement.h
--- a/Engine/Source/Runtime/Engine/Public/SceneManagement.h
+++ b/Engine/Source/Runtime/Engine/Public/SceneManagement.h
@@ -2371,27 +2371,33 @@ extern ENGINE_API void DrawCylinder(class FPrimitiveDrawInterface* PDI, const FM
 extern ENGINE_API void DrawCylinder(class FPrimitiveDrawInterface* PDI, const FVector& Start, const FVector& End, double Radius, int32 Sides, const FMaterialRenderProxy* MaterialInstance, uint8 DepthPriority);
 
 
-extern ENGINE_API void GetBoxMesh(const FMatrix& BoxToWorld,const FVector& Radii,const FMaterialRenderProxy* MaterialRenderProxy,uint8 DepthPriority,int32 ViewIndex,FMeshElementCollector& Collector, HHitProxy* HitProxy = NULL);
+extern ENGINE_API void GetBoxMesh(const FMatrix& BoxToWorld, const FVector& Radii, const FMaterialRenderProxy* MaterialRenderProxy, uint8 DepthPriority, int32 ViewIndex, FMeshElementCollector& Collector, HHitProxy* HitProxy = NULL);
 extern ENGINE_API void GetOrientedHalfSphereMesh(const FVector& Center, const FRotator& Orientation, const FVector& Radii, int32 NumSides, int32 NumRings, float StartAngle, float EndAngle, const FMaterialRenderProxy* MaterialRenderProxy, uint8 DepthPriority, bool bDisableBackfaceCulling,
-									int32 ViewIndex, FMeshElementCollector& Collector, bool bUseSelectionOutline = false, HHitProxy* HitProxy = NULL);
+	int32 ViewIndex, FMeshElementCollector& Collector, bool bUseSelectionOutline = false, HHitProxy* HitProxy = NULL);
 extern ENGINE_API void GetHalfSphereMesh(const FVector& Center, const FVector& Radii, int32 NumSides, int32 NumRings, float StartAngle, float EndAngle, const FMaterialRenderProxy* MaterialRenderProxy, uint8 DepthPriority, bool bDisableBackfaceCulling,
-									int32 ViewIndex, FMeshElementCollector& Collector, bool bUseSelectionOutline=false, HHitProxy* HitProxy=NULL);
+	int32 ViewIndex, FMeshElementCollector& Collector, bool bUseSelectionOutline = false, HHitProxy* HitProxy = NULL);
 extern ENGINE_API void GetSphereMesh(const FVector& Center, const FVector& Radii, int32 NumSides, int32 NumRings, const FMaterialRenderProxy* MaterialRenderProxy, uint8 DepthPriority,
 	bool bDisableBackfaceCulling, int32 ViewIndex, FMeshElementCollector& Collector);
-extern ENGINE_API void GetSphereMesh(const FVector& Center,const FVector& Radii,int32 NumSides,int32 NumRings,const FMaterialRenderProxy* MaterialRenderProxy,uint8 DepthPriority,
-									bool bDisableBackfaceCulling,int32 ViewIndex,FMeshElementCollector& Collector, bool bUseSelectionOutline, HHitProxy* HitProxy);
+extern ENGINE_API void GetSphereMesh(const FVector& Center, const FVector& Radii, int32 NumSides, int32 NumRings, const FMaterialRenderProxy* MaterialRenderProxy, uint8 DepthPriority,
+	bool bDisableBackfaceCulling, int32 ViewIndex, FMeshElementCollector& Collector, bool bUseSelectionOutline, HHitProxy* HitProxy);
+
+extern ENGINE_API void GetDiscMesh(const FVector& Center, const FVector& XAxis, const FVector& YAxis, double Radius, int32 Sides, const FMaterialRenderProxy* MaterialRenderProxy, uint8 DepthPriority, 
+	int32 ViewIndex, FMeshElementCollector& Collector, HHitProxy* HitProxy = nullptr);
+extern ENGINE_API void GetDiscMesh(const FMatrix& ToWorld, const FVector& Center, const FVector& XAxis, const FVector& YAxis, double Radius, int32 Sides, const FMaterialRenderProxy* MaterialRenderProxy, uint8 DepthPriority,
+	int32 ViewIndex, FMeshElementCollector& Collector, HHitProxy* HitProxy = nullptr);
+
 extern ENGINE_API void GetCylinderMesh(const FVector& Base, const FVector& XAxis, const FVector& YAxis, const FVector& ZAxis,
-									double Radius, double HalfHeight, int32 Sides, const FMaterialRenderProxy* MaterialInstance, uint8 DepthPriority, int32 ViewIndex, FMeshElementCollector& Collector, HHitProxy* HitProxy = NULL);
+	double Radius, double HalfHeight, int32 Sides, const FMaterialRenderProxy* MaterialInstance, uint8 DepthPriority, int32 ViewIndex, FMeshElementCollector& Collector, HHitProxy* HitProxy = NULL);
 extern ENGINE_API void GetCylinderMesh(const FMatrix& CylToWorld, const FVector& Base, const FVector& XAxis, const FVector& YAxis, const FVector& ZAxis,
-									double Radius, double HalfHeight, uint32 Sides, const FMaterialRenderProxy* MaterialInstance, uint8 DepthPriority, int32 ViewIndex, FMeshElementCollector& Collector, HHitProxy* HitProxy = NULL);
+	double Radius, double HalfHeight, uint32 Sides, const FMaterialRenderProxy* MaterialInstance, uint8 DepthPriority, int32 ViewIndex, FMeshElementCollector& Collector, HHitProxy* HitProxy = NULL);
 //Draws a cylinder along the axis from Start to End
 extern ENGINE_API void GetCylinderMesh(const FVector& Start, const FVector& End, double Radius, int32 Sides, const FMaterialRenderProxy* MaterialInstance, uint8 DepthPriority, int32 ViewIndex, FMeshElementCollector& Collector, HHitProxy* HitProxy = NULL);
 
 
 extern ENGINE_API void GetConeMesh(const FMatrix& LocalToWorld, float AngleWidth, float AngleHeight, uint32 NumSides,
-									const FMaterialRenderProxy* MaterialRenderProxy, uint8 DepthPriority, int32 ViewIndex, FMeshElementCollector& Collector, HHitProxy* HitProxy = NULL);
+	const FMaterialRenderProxy* MaterialRenderProxy, uint8 DepthPriority, int32 ViewIndex, FMeshElementCollector& Collector, HHitProxy* HitProxy = NULL);
 extern ENGINE_API void GetCapsuleMesh(const FVector& Origin, const FVector& XAxis, const FVector& YAxis, const FVector& ZAxis, const FLinearColor& Color, double Radius, double HalfHeight, int32 NumSides,
-									const FMaterialRenderProxy* MaterialRenderProxy, uint8 DepthPriority, bool bDisableBackfaceCulling, int32 ViewIndex, FMeshElementCollector& Collector, HHitProxy* HitProxy = NULL);
+	const FMaterialRenderProxy* MaterialRenderProxy, uint8 DepthPriority, bool bDisableBackfaceCulling, int32 ViewIndex, FMeshElementCollector& Collector, HHitProxy* HitProxy = NULL);
 
 
 /**
@@ -2489,23 +2495,6 @@ extern ENGINE_API void DrawFlatArrow(class FPrimitiveDrawInterface* PDI,const FV
 extern ENGINE_API void DrawWireBox(class FPrimitiveDrawInterface* PDI, const FBox& Box, const FLinearColor& Color, uint8 DepthPriority, float Thickness = 0.0f, float DepthBias = 0.0f, bool bScreenSpace = false);
 extern ENGINE_API void DrawWireBox(class FPrimitiveDrawInterface* PDI, const FMatrix& Matrix, const FBox& Box, const FLinearColor& Color, uint8 DepthPriority, float Thickness = 0.0f, float DepthBias = 0.0f, bool bScreenSpace = false);
 
-/**
- * Draws a circle using lines.
- *
- * @param	PDI				Draw interface.
- * @param	Base			Center of the circle.
- * @param	X				X alignment axis to draw along.
- * @param	Y				Y alignment axis to draw along.
- * @param	Z				Z alignment axis to draw along.
- * @param	Color			Color of the circle.
- * @param	Radius			Radius of the circle.
- * @param	NumSides		Numbers of sides that the circle has.
- * @param	DepthPriority	Depth priority for the circle.
- * @param	Thickness		Thickness of the lines comprising the circle
- */
-extern ENGINE_API void DrawCircle(class FPrimitiveDrawInterface* PDI, const FVector& Base, const FVector& X, const FVector& Y, const FLinearColor& Color, double Radius, int32 NumSides, uint8 DepthPriority, float Thickness = 0.0f, float DepthBias = 0.0f, bool bScreenSpace = false);
-
-
 /**
  * Draws an arc using lines.
  *
@@ -2566,6 +2555,21 @@ extern ENGINE_API void DrawWireSphere(class FPrimitiveDrawInterface* PDI, const
 extern ENGINE_API void DrawWireSphereAutoSides(class FPrimitiveDrawInterface* PDI, const FVector& Base, const FLinearColor& Color, double Radius, uint8 DepthPriority, float Thickness = 0.0f, float DepthBias = 0.0f, bool bScreenSpace = false);
 extern ENGINE_API void DrawWireSphereAutoSides(class FPrimitiveDrawInterface* PDI, const FTransform& Transform, const FLinearColor& Color, double Radius, uint8 DepthPriority, float Thickness = 0.0f, float DepthBias = 0.0f, bool bScreenSpace = false);
 
+/**
+ * Draws a wireframe circle.
+ *
+ * @param	PDI				Draw interface.
+ * @param	Center			Center of the circle.
+ * @param	X				X alignment axis to draw along.
+ * @param	Y				Y alignment axis to draw along.
+ * @param	Color			Color of the cylinder.
+ * @param	Radius			Radius of the cylinder.
+ * @param	NumSides		Numbers of sides that the cylinder has.
+ * @param	DepthPriority	Depth priority for the cylinder.
+ * @param	Thickness		Thickness of the lines comprising the circle
+ */
+extern ENGINE_API void DrawCircle(class FPrimitiveDrawInterface* PDI, const FVector& Base, const FVector& X, const FVector& Y, const FLinearColor& Color, double Radius, int32 NumSides, uint8 DepthPriority, float Thickness = 0.0f, float DepthBias = 0.0f, bool bScreenSpace = false);
+
 /**
  * Draws a wireframe cylinder.
  *
diff --git a/Engine/Source/Runtime/Engine/Public/VisualLogger/VisualLogger.h b/Engine/Source/Runtime/Engine/Public/VisualLogger/VisualLogger.h
--- a/Engine/Source/Runtime/Engine/Public/VisualLogger/VisualLogger.h
+++ b/Engine/Source/Runtime/Engine/Public/VisualLogger/VisualLogger.h
@@ -51,31 +51,31 @@
 #define UE_CVLOG_BOX(Condition, LogOwner, CategoryName, Verbosity, Box, Color, Format, ...) if(FVisualLogger::IsRecording() && Condition) {UE_VLOG_BOX(LogOwner, CategoryName, Verbosity, Box, Color, Format, ##__VA_ARGS__);} 
 // Wire box shape
 #define UE_VLOG_WIREBOX(LogOwner, CategoryName, Verbosity, Box, Color, Format, ...) if(FVisualLogger::IsRecording()) FVisualLogger::BoxLogf(LogOwner, CategoryName, ELogVerbosity::Verbosity, Box, FMatrix::Identity, Color, /*bWireframe = */true, Format, ##__VA_ARGS__)
-#define UE_CVLOG_WIREBOX(Condition, LogOwner, CategoryName, Verbosity, Box, Color, Format, ...) if(FVisualLogger::IsRecording() && Condition) {UE_VLOG_BOX(LogOwner, CategoryName, Verbosity, Box, Color, Format, ##__VA_ARGS__);} 
+#define UE_CVLOG_WIREBOX(Condition, LogOwner, CategoryName, Verbosity, Box, Color, Format, ...) if(FVisualLogger::IsRecording() && Condition) {UE_VLOG_WIREBOX(LogOwner, CategoryName, Verbosity, Box, Color, Format, ##__VA_ARGS__);} 
 // Oriented box shape
 #define UE_VLOG_OBOX(LogOwner, CategoryName, Verbosity, Box, Matrix, Color, Format, ...) if(FVisualLogger::IsRecording()) FVisualLogger::BoxLogf(LogOwner, CategoryName, ELogVerbosity::Verbosity, Box, Matrix, Color, /*bWireframe = */false, Format, ##__VA_ARGS__)
 #define UE_CVLOG_OBOX(Condition, LogOwner, CategoryName, Verbosity, Box, Matrix, Color, Format, ...) if(FVisualLogger::IsRecording() && Condition) {UE_VLOG_OBOX(LogOwner, CategoryName, Verbosity, Box, Matrix, Color, Format, ##__VA_ARGS__);} 
 // Wire oriented box shape
 #define UE_VLOG_WIREOBOX(LogOwner, CategoryName, Verbosity, Box, Matrix, Color, Format, ...) if(FVisualLogger::IsRecording()) FVisualLogger::BoxLogf(LogOwner, CategoryName, ELogVerbosity::Verbosity, Box, Matrix, Color, /*bWireframe = */true, Format, ##__VA_ARGS__)
-#define UE_CVLOG_WIREOBOX(Condition, LogOwner, CategoryName, Verbosity, Box, Matrix, Color, Format, ...) if(FVisualLogger::IsRecording() && Condition) {UE_VLOG_OBOX(LogOwner, CategoryName, Verbosity, Box, Matrix, Color, Format, ##__VA_ARGS__);} 
+#define UE_CVLOG_WIREOBOX(Condition, LogOwner, CategoryName, Verbosity, Box, Matrix, Color, Format, ...) if(FVisualLogger::IsRecording() && Condition) {UE_VLOG_WIREOBOX(LogOwner, CategoryName, Verbosity, Box, Matrix, Color, Format, ##__VA_ARGS__);} 
 // Cone shape
 #define UE_VLOG_CONE(LogOwner, CategoryName, Verbosity, Origin, Direction, Length, Angle, Color, Format, ...) if(FVisualLogger::IsRecording()) FVisualLogger::ConeLogf(LogOwner, CategoryName, ELogVerbosity::Verbosity, Origin, Direction, Length, Angle, Color, /*bWireframe = */false, Format, ##__VA_ARGS__)
 #define UE_CVLOG_CONE(Condition, LogOwner, CategoryName, Verbosity, Origin, Direction, Length, Angle, Color, Format, ...) if(FVisualLogger::IsRecording() && Condition) {UE_VLOG_CONE(LogOwner, CategoryName, Verbosity, Origin, Direction, Length, Angle, Color, Format, ##__VA_ARGS__);} 
 // Wire cone shape
 #define UE_VLOG_WIRECONE(LogOwner, CategoryName, Verbosity, Origin, Direction, Length, Angle, Color, Format, ...) if(FVisualLogger::IsRecording()) FVisualLogger::ConeLogf(LogOwner, CategoryName, ELogVerbosity::Verbosity, Origin, Direction, Length, Angle, Color, /*bWireframe = */true, Format, ##__VA_ARGS__)
-#define UE_CVLOG_WIRECONE(Condition, LogOwner, CategoryName, Verbosity, Origin, Direction, Length, Angle, Color, Format, ...) if(FVisualLogger::IsRecording() && Condition) {UE_VLOG_CONE(LogOwner, CategoryName, Verbosity, Origin, Direction, Length, Angle, Color, Format, ##__VA_ARGS__);} 
+#define UE_CVLOG_WIRECONE(Condition, LogOwner, CategoryName, Verbosity, Origin, Direction, Length, Angle, Color, Format, ...) if(FVisualLogger::IsRecording() && Condition) {UE_VLOG_WIRECONE(LogOwner, CategoryName, Verbosity, Origin, Direction, Length, Angle, Color, Format, ##__VA_ARGS__);} 
 // Cylinder shape
 #define UE_VLOG_CYLINDER(LogOwner, CategoryName, Verbosity, Start, End, Radius, Color, Format, ...) if(FVisualLogger::IsRecording()) FVisualLogger::CylinderLogf(LogOwner, CategoryName, ELogVerbosity::Verbosity, Start, End, Radius, Color, /*bWireframe = */false, Format, ##__VA_ARGS__)
 #define UE_CVLOG_CYLINDER(Condition, LogOwner, CategoryName, Verbosity, Start, End, Radius, Color, Format, ...) if(FVisualLogger::IsRecording() && Condition) {UE_VLOG_CYLINDER(LogOwner, CategoryName, Verbosity, Start, End, Radius, Color, Format, ##__VA_ARGS__);} 
 // Wire cylinder shape
 #define UE_VLOG_WIRECYLINDER(LogOwner, CategoryName, Verbosity, Start, End, Radius, Color, Format, ...) if(FVisualLogger::IsRecording()) FVisualLogger::CylinderLogf(LogOwner, CategoryName, ELogVerbosity::Verbosity, Start, End, Radius, Color, /*bWireframe = */true, Format, ##__VA_ARGS__)
-#define UE_CVLOG_WIRECYLINDER(Condition, LogOwner, CategoryName, Verbosity, Start, End, Radius, Color, Format, ...) if(FVisualLogger::IsRecording() && Condition) {UE_VLOG_CYLINDER(LogOwner, CategoryName, Verbosity, Start, End, Radius, Color, Format, ##__VA_ARGS__);} 
+#define UE_CVLOG_WIRECYLINDER(Condition, LogOwner, CategoryName, Verbosity, Start, End, Radius, Color, Format, ...) if(FVisualLogger::IsRecording() && Condition) {UE_VLOG_WIRECYLINDER(LogOwner, CategoryName, Verbosity, Start, End, Radius, Color, Format, ##__VA_ARGS__);} 
 // Capsule shape
 #define UE_VLOG_CAPSULE(LogOwner, CategoryName, Verbosity, Base, HalfHeight, Radius, Rotation, Color, Format, ...) if(FVisualLogger::IsRecording()) FVisualLogger::CapsuleLogf(LogOwner, CategoryName, ELogVerbosity::Verbosity, Base, HalfHeight, Radius, Rotation, Color, /*bWireframe = */false, Format, ##__VA_ARGS__)
 #define UE_CVLOG_CAPSULE(Condition, LogOwner, CategoryName, Verbosity, Base, HalfHeight, Radius, Rotation, Color, Format, ...) if(FVisualLogger::IsRecording() && Condition) {UE_VLOG_CAPSULE(LogOwner, CategoryName, Verbosity, Base, HalfHeight, Radius, Rotation, Color, Format, ##__VA_ARGS__);} 
 // Wire capsule shape
 #define UE_VLOG_WIRECAPSULE(LogOwner, CategoryName, Verbosity, Base, HalfHeight, Radius, Rotation, Color, Format, ...) if(FVisualLogger::IsRecording()) FVisualLogger::CapsuleLogf(LogOwner, CategoryName, ELogVerbosity::Verbosity, Base, HalfHeight, Radius, Rotation, Color, /*bWireframe = */true, Format, ##__VA_ARGS__)
-#define UE_CVLOG_WIRECAPSULE(Condition, LogOwner, CategoryName, Verbosity, Base, HalfHeight, Radius, Rotation, Color, Format, ...) if(FVisualLogger::IsRecording() && Condition) {UE_VLOG_CAPSULE(LogOwner, CategoryName, Verbosity, Base, HalfHeight, Radius, Rotation, Color, Format, ##__VA_ARGS__);} 
+#define UE_CVLOG_WIRECAPSULE(Condition, LogOwner, CategoryName, Verbosity, Base, HalfHeight, Radius, Rotation, Color, Format, ...) if(FVisualLogger::IsRecording() && Condition) {UE_VLOG_WIRECAPSULE(LogOwner, CategoryName, Verbosity, Base, HalfHeight, Radius, Rotation, Color, Format, ##__VA_ARGS__);} 
 // Histogram data for 2d graphs 
 #define UE_VLOG_HISTOGRAM(LogOwner, CategoryName, Verbosity, GraphName, DataName, Data) if(FVisualLogger::IsRecording()) FVisualLogger::HistogramDataLogf(LogOwner, CategoryName, ELogVerbosity::Verbosity, GraphName, DataName, Data, FColor::White, TEXT(""))
 #define UE_CVLOG_HISTOGRAM(Condition, LogOwner, CategoryName, Verbosity, GraphName, DataName, Data) if(FVisualLogger::IsRecording() && Condition) {UE_VLOG_HISTOGRAM(LogOwner, CategoryName, Verbosity, GraphName, DataName, Data);} 
@@ -89,14 +89,23 @@
 #define UE_VLOG_CONVEXPOLY(LogOwner, CategoryName, Verbosity, Points, Color, Format, ...) if(FVisualLogger::IsRecording()) FVisualLogger::ConvexLogf(LogOwner, CategoryName, ELogVerbosity::Verbosity, Points, Color, Format, ##__VA_ARGS__)
 #define UE_CVLOG_CONVEXPOLY(Condition, LogOwner, CategoryName, Verbosity, Points, Color, Format, ...) if(FVisualLogger::IsRecording() && Condition) {UE_VLOG_CONVEXPOLY(LogOwner, CategoryName, Verbosity, Points, Color, Format, ##__VA_ARGS__);}
 // Segment with an arrowhead
-#define UE_VLOG_ARROW(LogOwner, CategoryName, Verbosity, SegmentStart, SegmentEnd, Color, Format, ...) if(FVisualLogger::IsRecording()) FVisualLogger::ArrowLogf(LogOwner, CategoryName, ELogVerbosity::Verbosity, SegmentStart, SegmentEnd, Color, Format, ##__VA_ARGS__)
+#define UE_VLOG_ARROW(LogOwner, CategoryName, Verbosity, SegmentStart, SegmentEnd, Color, Format, ...) if(FVisualLogger::IsRecording()) FVisualLogger::ArrowLineLogf(LogOwner, CategoryName, ELogVerbosity::Verbosity, SegmentStart, SegmentEnd, Color, 0, Format, ##__VA_ARGS__)
 #define UE_CVLOG_ARROW(Condition, LogOwner, CategoryName, Verbosity, SegmentStart, SegmentEnd, Color, Format, ...) if(FVisualLogger::IsRecording() && Condition) {UE_VLOG_ARROW(LogOwner, CategoryName, Verbosity, SegmentStart, SegmentEnd, Color, Format, ##__VA_ARGS__);} 
+// Segment with an arrowhead with a custom arrow size
+#define UE_VLOG_ARROW_MAG(LogOwner, CategoryName, Verbosity, SegmentStart, SegmentEnd, Color, Mag, Format, ...) if(FVisualLogger::IsRecording()) FVisualLogger::ArrowLineLogf(LogOwner, CategoryName, ELogVerbosity::Verbosity, SegmentStart, SegmentEnd, Color, Mag, Format, ##__VA_ARGS__)
+#define UE_CVLOG_ARROW_MAG(Condition, LogOwner, CategoryName, Verbosity, SegmentStart, SegmentEnd, Color, Mag, Format, ...) if(FVisualLogger::IsRecording() && Condition) {UE_VLOG_ARROW_MAG(LogOwner, CategoryName, Verbosity, SegmentStart, SegmentEnd, Color, Mag, Format, ##__VA_ARGS__);} 
 // Circle shape
-#define UE_VLOG_CIRCLE(LogOwner, CategoryName, Verbosity, Center, UpAxis, Radius, Color, Format, ...) if(FVisualLogger::IsRecording()) FVisualLogger::CircleLogf(LogOwner, CategoryName, ELogVerbosity::Verbosity, Center, UpAxis, Radius, Color, 0, Format, ##__VA_ARGS__)
+#define UE_VLOG_CIRCLE(LogOwner, CategoryName, Verbosity, Center, UpAxis, Radius, Color, Format, ...) if(FVisualLogger::IsRecording()) FVisualLogger::DiscLogf(LogOwner, CategoryName, ELogVerbosity::Verbosity, Center, UpAxis, Radius, Color, 0, /*bWireframe = */false, Format, ##__VA_ARGS__)
 #define UE_CVLOG_CIRCLE(Condition, LogOwner, CategoryName, Verbosity, Center, UpAxis, Radius, Color, Format, ...) if(FVisualLogger::IsRecording() && Condition) {UE_VLOG_CIRCLE(LogOwner, CategoryName, Verbosity, Center, UpAxis, Radius, Color, Format, ##__VA_ARGS__);} 
 // Circle shape
-#define UE_VLOG_CIRCLE_THICK(LogOwner, CategoryName, Verbosity, Center, UpAxis, Radius, Color, Thickness, Format, ...) if(FVisualLogger::IsRecording()) FVisualLogger::CircleLogf(LogOwner, CategoryName, ELogVerbosity::Verbosity, Center, UpAxis, Radius, Color, Thickness, Format, ##__VA_ARGS__)
+#define UE_VLOG_CIRCLE_THICK(LogOwner, CategoryName, Verbosity, Center, UpAxis, Radius, Color, Thickness, Format, ...) if(FVisualLogger::IsRecording()) FVisualLogger::DiscLogf(LogOwner, CategoryName, ELogVerbosity::Verbosity, Center, UpAxis, Radius, Color, Thickness, /*bWireframe = */false, Format, ##__VA_ARGS__)
 #define UE_CVLOG_CIRCLE_THICK(Condition, LogOwner, CategoryName, Verbosity, Center, UpAxis, Radius, Color, Thickness, Format, ...) if(FVisualLogger::IsRecording() && Condition) {UE_VLOG_CIRCLE_THICK(LogOwner, CategoryName, Verbosity, Center, UpAxis, Radius, Color, Thickness, Format, ##__VA_ARGS__);} 
+// Wire circle shape
+#define UE_VLOG_WIRECIRCLE(LogOwner, CategoryName, Verbosity, Center, UpAxis, Radius, Color, Format, ...) if(FVisualLogger::IsRecording()) FVisualLogger::DiscLogf(LogOwner, CategoryName, ELogVerbosity::Verbosity, Center, UpAxis, Radius, Color, 0, /*bWireframe = */true, Format, ##__VA_ARGS__)
+#define UE_CVLOG_WIRECIRCLE(Condition, LogOwner, CategoryName, Verbosity, Center, UpAxis, Radius, Color, Format, ...) if(FVisualLogger::IsRecording() && Condition) {UE_VLOG_WIRECIRCLE(LogOwner, CategoryName, Verbosity, Center, UpAxis, Radius, Color, Format, ##__VA_ARGS__);} 
+// Wire circle shape
+#define UE_VLOG_WIRECIRCLE_THICK(LogOwner, CategoryName, Verbosity, Center, UpAxis, Radius, Color, Thickness, Format, ...) if(FVisualLogger::IsRecording()) FVisualLogger::DiscLogf(LogOwner, CategoryName, ELogVerbosity::Verbosity, Center, UpAxis, Radius, Color, Thickness, /*bWireframe = */true, Format, ##__VA_ARGS__)
+#define UE_CVLOG_WIRECIRCLE_THICK(Condition, LogOwner, CategoryName, Verbosity, Center, UpAxis, Radius, Color, Thickness, Format, ...) if(FVisualLogger::IsRecording() && Condition) {UE_VLOG_WIRECIRCLE_THICK(LogOwner, CategoryName, Verbosity, Center, UpAxis, Radius, Color, Thickness, Format, ##__VA_ARGS__);} 
 
 #define DECLARE_VLOG_EVENT(EventName) extern FVisualLogEventBase EventName;
 #define DEFINE_VLOG_EVENT(EventName, Verbosity, UserFriendlyDesc) FVisualLogEventBase EventName(TEXT(#EventName), TEXT(UserFriendlyDesc), ELogVerbosity::Verbosity); 
@@ -164,6 +173,10 @@
 #define UE_CVLOG_CIRCLE(Condition, LogOwner, CategoryName, Verbosity, Center, UpAxis, Radius, Color, Format, ...)
 #define UE_VLOG_CIRCLE_THICK(LogOwner, CategoryName, Verbosity, Center, UpAxis, Radius, Color, Thickness, Format, ...)
 #define UE_CVLOG_CIRCLE_THICK(Condition, LogOwner, CategoryName, Verbosity, Center, UpAxis, Radius, Color, Thickness, Format, ...)
+#define UE_VLOG_WIRECIRCLE(LogOwner, CategoryName, Verbosity, Center, UpAxis, Radius, Color, Format, ...)
+#define UE_CVLOG_WIRECIRCLE(Condition, LogOwner, CategoryName, Verbosity, Center, UpAxis, Radius, Color, Format, ...)
+#define UE_VLOG_WIRECIRCLE_THICK(LogOwner, CategoryName, Verbosity, Center, UpAxis, Radius, Color, Thickness, Format, ...)
+#define UE_CVLOG_WIRECIRCLE_THICK(Condition, LogOwner, CategoryName, Verbosity, Center, UpAxis, Radius, Color, Thickness, Format, ...)
 
 #define DECLARE_VLOG_EVENT(EventName)
 #define DEFINE_VLOG_EVENT(EventName, Verbosity, UserFriendlyDesc)
@@ -217,8 +230,12 @@ class FVisualLogger : public FOutputDevice
 	static ENGINE_API void MeshLogfImpl(const UObject* LogOwner, const FName& CategoryName, ELogVerbosity::Type Verbosity, const TArray<FVector>& Vertices, const TArray<int32>& Indices, const FColor& Color, const TCHAR* Fmt, ...);
 	static ENGINE_API void ConvexLogfImpl(const UObject* LogOwner, const FName& CategoryName, ELogVerbosity::Type Verbosity, const TArray<FVector>& Points, const FColor& Color, const TCHAR* Fmt, ...);
 	static ENGINE_API void HistogramDataLogfImpl (const UObject* LogOwner, const FName& CategoryName, ELogVerbosity::Type Verbosity, FName GraphName, FName DataName, const FVector2D& Data, const FColor& Color, const TCHAR* Fmt, ...);
+	UE_DEPRECATED(5.6, "Use ArrowLineLogfImpl")
 	static ENGINE_API void ArrowLogfImpl(const UObject* LogOwner, const FName& CategoryName, ELogVerbosity::Type Verbosity, const FVector& Start, const FVector& End, const FColor& Color, const TCHAR* Fmt, ...);
+	static ENGINE_API void ArrowLineLogfImpl(const UObject* LogOwner, const FName& CategoryName, ELogVerbosity::Type Verbosity, const FVector& Start, const FVector& End, const FColor& Color, const uint16 Mag, const TCHAR* Fmt, ...);
+	UE_DEPRECATED(5.6, "Use DiscLogfImpl")
 	static ENGINE_API void CircleLogfImpl(const UObject* LogOwner, const FName& CategoryName, ELogVerbosity::Type Verbosity, const FVector& Center, const FVector& UpAxis, const float Radius, const FColor& Color, const uint16 Thickness, const TCHAR* Fmt, ...);
+	static ENGINE_API void DiscLogfImpl(const UObject* LogOwner, const FName& CategoryName, ELogVerbosity::Type Verbosity, const FVector& Center, const FVector& UpAxis, const float Radius, const FColor& Color, const uint16 Thickness, bool bWireframe, const TCHAR* Fmt, ...);
 
 public:
 	// Regular text log
@@ -278,38 +295,77 @@ public:
 
 	// Arrow
 	template <typename FmtType, typename... Types>
+	UE_DEPRECATED(5.6, "Use ArrowLineLogfImpl")
 	static void ArrowLogf(const UObject* LogOwner, const FLogCategoryBase& Category, ELogVerbosity::Type Verbosity, const FVector& Start, const FVector& End, const FColor& Color, const FmtType& Fmt, Types... Args)
 	{
 		static_assert(TIsArrayOrRefOfTypeByPredicate<FmtType, TIsCharEncodingCompatibleWithTCHAR>::Value, "Formatting string must be a TCHAR array.");
 		static_assert((TIsValidVariadicFunctionArg<Types>::Value && ...), "Invalid argument(s) passed to FVisualLogger::ArrowLogf");
-		ArrowLogfImpl(LogOwner, Category.GetCategoryName(), Verbosity, Start, End, Color, (const TCHAR*)Fmt, Args...);
+		ArrowLineLogfImpl(LogOwner, Category.GetCategoryName(), Verbosity, Start, End, Color, 8, (const TCHAR*)Fmt, Args...);
 	}
 	template <typename FmtType, typename... Types>
+	UE_DEPRECATED(5.6, "Use ArrowLineLogfImpl")
 	static void ArrowLogf(const UObject* LogOwner, const FName& CategoryName, ELogVerbosity::Type Verbosity, const FVector& Start, const FVector& End, const FColor& Color, const FmtType& Fmt, Types... Args)
 	{
 		static_assert(TIsArrayOrRefOfTypeByPredicate<FmtType, TIsCharEncodingCompatibleWithTCHAR>::Value, "Formatting string must be a TCHAR array.");
 		static_assert((TIsValidVariadicFunctionArg<Types>::Value && ...), "Invalid argument(s) passed to FVisualLogger::ArrowLogf");
 
-		ArrowLogfImpl(LogOwner, CategoryName, Verbosity, Start, End, Color, (const TCHAR*)Fmt, Args...);
+		ArrowLineLogfImpl(LogOwner, CategoryName, Verbosity, Start, End, Color, 8, (const TCHAR*)Fmt, Args...);
+	}
+
+	template <typename FmtType, typename... Types>
+	static void ArrowLineLogf(const UObject* LogOwner, const FLogCategoryBase& Category, ELogVerbosity::Type Verbosity, const FVector& Start, const FVector& End, const FColor& Color, const uint16 Mag, const FmtType& Fmt, Types... Args)
+	{
+		static_assert(TIsArrayOrRefOfTypeByPredicate<FmtType, TIsCharEncodingCompatibleWithTCHAR>::Value, "Formatting string must be a TCHAR array.");
+		static_assert((TIsValidVariadicFunctionArg<Types>::Value && ...), "Invalid argument(s) passed to FVisualLogger::ArrowLineLogf");
+		
+		ArrowLineLogfImpl(LogOwner, Category.GetCategoryName(), Verbosity, Start, End, Color, Mag, (const TCHAR*)Fmt, Args...);
+	}
+	template <typename FmtType, typename... Types>
+	static void ArrowLineLogf(const UObject* LogOwner, const FName& CategoryName, ELogVerbosity::Type Verbosity, const FVector& Start, const FVector& End, const FColor& Color, const uint16 Mag, const FmtType& Fmt, Types... Args)
+	{
+		static_assert(TIsArrayOrRefOfTypeByPredicate<FmtType, TIsCharEncodingCompatibleWithTCHAR>::Value, "Formatting string must be a TCHAR array.");
+		static_assert((TIsValidVariadicFunctionArg<Types>::Value && ...), "Invalid argument(s) passed to FVisualLogger::ArrowLineLogf");
+
+		ArrowLineLogfImpl(LogOwner, CategoryName, Verbosity, Start, End, Color, Mag, (const TCHAR*)Fmt, Args...);
 	}
 
 
 	// Circle log
 	template <typename FmtType, typename... Types>
+	UE_DEPRECATED(5.6, "Use DiscLogfImpl")
 	static void CircleLogf(const UObject* LogOwner, const FLogCategoryBase& Category, ELogVerbosity::Type Verbosity, const FVector& Center, const FVector& UpAxis, const float Radius, const FColor& Color, const uint16 Thickness, const FmtType& Fmt, Types... Args)
 	{
 		static_assert(TIsArrayOrRefOfTypeByPredicate<FmtType, TIsCharEncodingCompatibleWithTCHAR>::Value, "Formatting string must be a TCHAR array.");
 		static_assert((TIsValidVariadicFunctionArg<Types>::Value && ...), "Invalid argument(s) passed to FVisualLogger::CircleLogf");
 
-		CircleLogfImpl(LogOwner, Category.GetCategoryName(), Verbosity, Center, UpAxis, Radius, Color, Thickness, (const TCHAR*)Fmt, Args...);
+		DiscLogfImpl(LogOwner, Category.GetCategoryName(), Verbosity, Center, UpAxis, Radius, Color, Thickness, /*bWireframe = */false, (const TCHAR*)Fmt, Args...);
 	}
 	template <typename FmtType, typename... Types>
+	UE_DEPRECATED(5.6, "Use DiscLogfImpl")
 	static void CircleLogf(const UObject* LogOwner, const FName& CategoryName, ELogVerbosity::Type Verbosity, const FVector& Center, const FVector& UpAxis, const float Radius, const FColor& Color, const uint16 Thickness, const FmtType& Fmt, Types... Args)
 	{
 		static_assert(TIsArrayOrRefOfTypeByPredicate<FmtType, TIsCharEncodingCompatibleWithTCHAR>::Value, "Formatting string must be a TCHAR array.");
 		static_assert((TIsValidVariadicFunctionArg<Types>::Value && ...), "Invalid argument(s) passed to FVisualLogger::CircleLogf");
 
-		CircleLogfImpl(LogOwner, CategoryName, Verbosity, Center, UpAxis, Radius, Color, Thickness, (const TCHAR*)Fmt, Args...);
+		DiscLogfImpl(LogOwner, CategoryName, Verbosity, Center, UpAxis, Radius, Color, Thickness, /*bWireframe = */false, (const TCHAR*)Fmt, Args...);
+	}
+
+	// Disc/circle log
+	template <typename FmtType, typename... Types>
+	static void DiscLogf(const UObject* LogOwner, const FLogCategoryBase& Category, ELogVerbosity::Type Verbosity, const FVector& Center, const FVector& UpAxis, const float Radius, const FColor& Color, const uint16 Thickness, bool bWireframe, const FmtType& Fmt, Types... Args)
+	{
+		static_assert(TIsArrayOrRefOfTypeByPredicate<FmtType, TIsCharEncodingCompatibleWithTCHAR>::Value, "Formatting string must be a TCHAR array.");
+		static_assert((TIsValidVariadicFunctionArg<Types>::Value && ...), "Invalid argument(s) passed to FVisualLogger::DiscLogf");
+
+		DiscLogfImpl(LogOwner, Category.GetCategoryName(), Verbosity, Center, UpAxis, Radius, Color, Thickness, bWireframe, (const TCHAR*)Fmt, Args...);
+	}
+	template <typename FmtType, typename... Types>
+	static void DiscLogf(const UObject* LogOwner, const FName& CategoryName, ELogVerbosity::Type Verbosity, const FVector& Center, const FVector& UpAxis, const float Radius, const FColor& Color, const uint16 Thickness, bool bWireframe, const FmtType& Fmt, Types... Args)
+	{
+		static_assert(TIsArrayOrRefOfTypeByPredicate<FmtType, TIsCharEncodingCompatibleWithTCHAR>::Value, "Formatting string must be a TCHAR array.");
+		static_assert((TIsValidVariadicFunctionArg<Types>::Value && ...), "Invalid argument(s) passed to FVisualLogger::DiscLogf");
+
+		DiscLogfImpl(LogOwner, CategoryName, Verbosity, Center, UpAxis, Radius, Color, Thickness, bWireframe,(const TCHAR*)Fmt, Args...);
 	}
 
 	// Location log
diff --git a/Engine/Source/Runtime/Engine/Public/VisualLogger/VisualLoggerTypes.h b/Engine/Source/Runtime/Engine/Public/VisualLogger/VisualLoggerTypes.h
--- a/Engine/Source/Runtime/Engine/Public/VisualLogger/VisualLoggerTypes.h
+++ b/Engine/Source/Runtime/Engine/Public/VisualLogger/VisualLoggerTypes.h
@@ -56,6 +56,7 @@ enum class EVisualLoggerShapeElement : uint8
 	NavAreaMesh, // convex based mesh with min and max Z values
 	Arrow,
 	Circle,
+	WireCircle,
 	// note that in order to remain backward compatibility in terms of log
 	// serialization new enum values need to be added at the end
 };
@@ -139,6 +140,7 @@ struct FVisualLogShapeElement
 	{
 		uint16 Thicknes;
 		uint16 Radius;
+		uint16 Mag;
 	};
 
 	ENGINE_API FVisualLogShapeElement(EVisualLoggerShapeElement InType = EVisualLoggerShapeElement::Invalid);
@@ -262,11 +264,11 @@ struct FVisualLogEntry
 	// histogram sample
 	ENGINE_API void AddHistogramData(const FVector2D& DataSample, const FName& CategoryName, ELogVerbosity::Type Verbosity, const FName& GraphName, const FName& DataName);
 	// arrow
-	ENGINE_API void AddArrow(const FVector& Start, const FVector& End, const FName& CategoryName, ELogVerbosity::Type Verbosity, const FColor& Color = FColor::White, const FString& Description = TEXT(""));
+	ENGINE_API void AddArrow(const FVector& Start, const FVector& End, const FName& CategoryName, ELogVerbosity::Type Verbosity, const FColor& Color = FColor::White, const FString& Description = TEXT(""), const uint16 Mag = 0);
 	// boxes
 	ENGINE_API void AddBoxes(const TArray<FBox>& Boxes, const FName& CategoryName, ELogVerbosity::Type Verbosity, const FColor& Color = FColor::White);
 	// circle
-	ENGINE_API void AddCircle(const FVector& Center, const FVector& UpAxis, const float Radius, const FName& CategoryName, ELogVerbosity::Type Verbosity, const FColor& Color, const FString& Description = TEXT(""), uint16 Thickness = 0);
+	ENGINE_API void AddCircle(const FVector& Center, const FVector& UpAxis, const float Radius, const FName& CategoryName, ELogVerbosity::Type Verbosity, const FColor& Color, const FString& Description = TEXT(""), uint16 Thickness = 0, bool bInUseWires = false);
 
 	// Custom data block
 	ENGINE_API FVisualLogDataBlock& AddDataBlock(const FString& TagName, const TArray<uint8>& BlobDataArray, const FName& CategoryName, ELogVerbosity::Type Verbosity);
