Update warning message in UAnimationAsset::ReplaceSkeleton for clarity

diff --git a/Engine/Source/Runtime/Engine/Private/AnimationAsset.cpp b/Engine/Source/Runtime/Engine/Private/AnimationAsset.cpp
--- a/Engine/Source/Runtime/Engine/Private/AnimationAsset.cpp
+++ b/Engine/Source/Runtime/Engine/Private/AnimationAsset.cpp
@@ -493,7 +493,13 @@ bool UAnimationAsset::ReplaceSkeleton(USkeleton* NewSkeleton, bool bConvertSpace
 
 			if (AnimAsset->GetSkeleton() != GetSkeleton())
 			{
-				UE_LOG(LogAnimation, Warning, TEXT("AnimationAsset referencing asset using different skeleton. This will generate undeterministic builds. Please Fix the Asset : AnimationAsset: [%s] - ReferencedAsset : [%s]"), *GetName(), *AnimAsset->GetName());
+				UE_LOG(LogAnimation, Warning, TEXT("AnimationAsset referencing asset using different skeleton. This will generate undeterministic builds and may cause additional warnings from undeterministic ordering of skeleton conforming. Please Fix the Asset : AnimationAsset: [%s] (Skeleton: [%s]) - ReferencedAsset : [%s] (Skeleton: [%s]) (Conforming Skeleton from [%s] to [%s]"),
+					*GetName(),
+					*GetSkeleton()->GetName(),
+					*AnimAsset->GetName(),
+					*AnimAsset->GetSkeleton()->GetName(),
+					*AnimAsset->GetSkeleton()->GetName(),
+					*NewSkeleton->GetName());
 			}
 
 			// This ensure that in subsequent behaviour the RawData GUID is never 'new-ed' but always calculated from the 
