[NetPhysics] Balancing the default physics tick sync values to handle Bad networking.

diff --git a/Engine/Source/Runtime/Engine/Private/PlayerController.cpp b/Engine/Source/Runtime/Engine/Private/PlayerController.cpp
--- a/Engine/Source/Runtime/Engine/Private/PlayerController.cpp
+++ b/Engine/Source/Runtime/Engine/Private/PlayerController.cpp
@@ -147,17 +147,17 @@ namespace NetworkPhysicsCvars
 	int32 TickOffsetUpdateIntervalTime = 100;
 	FAutoConsoleVariableRef CVarTickOffsetUpdateIntervalTime(TEXT("np2.TickOffsetUpdateIntervalTime"), TickOffsetUpdateIntervalTime, TEXT("Value in milliseconds, default 100. How long time between syncing the tick offset between client and server. Deactivate syncing by setting value 0."));
 
-	int32 TickOffsetBufferTime = 60;
+	int32 TickOffsetBufferTime = 90;
 	FAutoConsoleVariableRef CVarTickOffsetBufferTime(TEXT("np2.TickOffsetBufferTime"), TickOffsetBufferTime, TEXT("Value in milliseconds, default 60. Additional offset to be added to the local to server offset used by network prediction, this results in a buffer server-side for incoming data that uses the client/server physics offset."));
 
-	int32 TickOffsetCorrectionSizeTimeLimit = 200;
-	FAutoConsoleVariableRef CVarTickOffsetCorrectionSizeTimeLimit(TEXT("np2.TickOffsetCorrectionSizeTimeLimit"), TickOffsetCorrectionSizeTimeLimit, TEXT("Value in milliseconds, Default 200. Note: Keep this equal to or larger than np2.TickOffsetBufferTime. If the client gets out of sync with physics ticks and the desync is larger than this value, reset the offset."));
+	int32 TickOffsetCorrectionSizeTimeLimit = 2000;
+	FAutoConsoleVariableRef CVarTickOffsetCorrectionSizeTimeLimit(TEXT("np2.TickOffsetCorrectionSizeTimeLimit"), TickOffsetCorrectionSizeTimeLimit, TEXT("Value in milliseconds. Note: Keep this equal to or larger than np2.TickOffsetBufferTime. If the client gets out of sync with physics ticks and the desync is larger than this value, reset the offset."));
 
-	int32 TickOffsetCorrectionTimeLimit = 1000;
-	FAutoConsoleVariableRef CVarTickOffsetCorrectionTimeLimit(TEXT("np2.TickOffsetCorrectionTimeLimit"), TickOffsetCorrectionTimeLimit, TEXT("Value in milliseconds, Default 1000. If the client gets out of sync with physics ticks and can't get in sync again for this amount of time, reset the offset."));
+	int32 TickOffsetCorrectionTimeLimit = 5000;
+	FAutoConsoleVariableRef CVarTickOffsetCorrectionTimeLimit(TEXT("np2.TickOffsetCorrectionTimeLimit"), TickOffsetCorrectionTimeLimit, TEXT("Value in milliseconds. If the client gets out of sync with physics ticks and can't get in sync again for this amount of time, reset the offset."));
 
-	float TimeDilationAmount = 0.02f;
-	FAutoConsoleVariableRef CVarTimeDilationAmount(TEXT("np2.TimeDilationAmount"), TimeDilationAmount, TEXT("Server-side CVar, Disable TimeDilation by setting to 0 | Default: 0.01 | Value is in percent where 0.01 = 1% dilation. Example: 1.0/0.01 = 100, meaning that over the time it usually takes to tick 100 physics steps we will tick 99 or 101 depending on if we dilate up or down."));
+	float TimeDilationAmount = 0.025f;
+	FAutoConsoleVariableRef CVarTimeDilationAmount(TEXT("np2.TimeDilationAmount"), TimeDilationAmount, TEXT("Server-side CVar, Disable TimeDilation by setting to 0 | Value is in percent where 0.01 = 1% dilation. Example: 1.0/0.01 = 100, meaning that over the time it usually takes to tick 100 physics steps we will tick 99 or 101 depending on if we dilate up or down."));
 
 	bool TimeDilationEscalation = true;
 	FAutoConsoleVariableRef CVarTimeDilationEscalation(TEXT("np2.TimeDilationEscalation"), TimeDilationEscalation, TEXT("Server-side CVar, Dilate the time more depending on how many ticks we need to adjust. When set to false we use the set TimeDilationAmount and wait the amount of time it takes to perform correct the offset. When set to true we multiply the TimeDilationAmount with the buffer offset count which will correct the offset in one TimeDilationAmount cycle."));
