added option to remove all graph node comments from cooked build

diff --git a/Engine/Source/Runtime/Engine/Private/EdGraph/EdGraphNode.cpp b/Engine/Source/Runtime/Engine/Private/EdGraph/EdGraphNode.cpp
--- a/Engine/Source/Runtime/Engine/Private/EdGraph/EdGraphNode.cpp
+++ b/Engine/Source/Runtime/Engine/Private/EdGraph/EdGraphNode.cpp
@@ -24,6 +24,12 @@
 
 #define LOCTEXT_NAMESPACE "EdGraph"
 
+static int32 GEdGraphStripNodeComment = 0;
+static FAutoConsoleVariableRef CVarEdGraphStripNodeComment(
+	TEXT("cook.StripGraphNodeComments"),
+	GEdGraphStripNodeComment,
+	TEXT("1 = Strip graph node comments on cook. 0 = off"));
+
 FEdGraphTerminalType FEdGraphTerminalType::FromPinType(const FEdGraphPinType& PinType)
 {
 	FEdGraphTerminalType TerminalType;
@@ -206,8 +212,20 @@ void UEdGraphNode::Serialize(FArchive& Ar)
 	Ar.UsingCustomVersion(FBlueprintsObjectVersion::GUID);
 #endif
 
+	FString StrippedNodeComment;
+	if (GEdGraphStripNodeComment && Ar.IsCooking())
+	{
+		StrippedNodeComment = NodeComment;
+		NodeComment = FString();
+	}
+
 	Super::Serialize(Ar);
 
+	if (StrippedNodeComment.Len() > 0)
+	{
+		NodeComment = StrippedNodeComment;
+	}
+
 #if WITH_EDITOR
 	if (Ar.IsLoading())
 	{
