[UBA] Enable new breadcrumbs in UBA traces by default.

#rnx

diff --git a/Engine/Source/Runtime/Engine/Private/ShaderCompiler/ShaderCompilerDistributed.cpp b/Engine/Source/Runtime/Engine/Private/ShaderCompiler/ShaderCompilerDistributed.cpp
--- a/Engine/Source/Runtime/Engine/Private/ShaderCompiler/ShaderCompilerDistributed.cpp
+++ b/Engine/Source/Runtime/Engine/Private/ShaderCompiler/ShaderCompilerDistributed.cpp
@@ -30,13 +30,13 @@ namespace DistributedShaderCompilerVariables
 		TEXT("Maximum number of seconds we expect to pass between getting distributed controller complete a task (this is used to detect problems with the distribution controllers).")
 	);
 
-	static int32 GDistributedJobDescriptionLevel = 0;
+	static int32 GDistributedJobDescriptionLevel = 1;
 	static FAutoConsoleVariableRef CVarDistributedJobDescriptionLevel(
 		TEXT("r.ShaderCompiler.DistributedJobDescriptionLevel"),
 		GDistributedJobDescriptionLevel,
 		TEXT("Sets the level of descriptive details for each distributed job batch. The following modes are supported:\n")
 		TEXT(" Mode 0: Disabled.\n")
-		TEXT(" Mode 1: Basic information of the first 20 compile jobs per batch.\n")
+		TEXT(" Mode 1: Basic information of the first 20 compile jobs per batch (Default).\n")
 		TEXT(" Mode 2: Additional information of the shader format per compile job.\n")
 		TEXT("This will show up in the UBA trace files. By default 0.")
 	);
