Iris Factories
* Split creation header's into static and dynamic classes to save memory

diff --git a/Engine/Source/Runtime/Engine/Private/Net/Iris/ReplicationSystem/NetActorFactory.cpp b/Engine/Source/Runtime/Engine/Private/Net/Iris/ReplicationSystem/NetActorFactory.cpp
--- a/Engine/Source/Runtime/Engine/Private/Net/Iris/ReplicationSystem/NetActorFactory.cpp
+++ b/Engine/Source/Runtime/Engine/Private/Net/Iris/ReplicationSystem/NetActorFactory.cpp
@@ -186,16 +186,18 @@ TUniquePtr<UE::Net::FNetObjectCreationHeader> UNetActorFactory::CreateAndFillHea
 
 	UActorReplicationBridge* ActorBridge = CastChecked<UActorReplicationBridge>(Bridge);
 
-	TUniquePtr<FActorNetCreationHeader> Header( new FActorNetCreationHeader );
+	TUniquePtr<FBaseActorNetCreationHeader> BaseHeader;
 
 	const FNetObjectReference ActorReference = ActorBridge->GetOrCreateObjectReference(Actor);	
-	Header->bIsDynamic = ActorReference.GetRefHandle().IsDynamic();
-	//$IRIS TODO: Create two different header classes for dynamic and static actors.
-	if (Header->bIsDynamic)
+		
+	if (ActorReference.GetRefHandle().IsDynamic())
 	{
+		FDynamicActorNetCreationHeader* Header = new FDynamicActorNetCreationHeader;
+		BaseHeader.Reset(Header);
+
 		// This is more or less a straight copy from ClientPackageMap and needs to be updated accordingly
 		UObject* Archetype = nullptr;
-		ULevel* ActorLevel = nullptr;
+		UObject* ActorLevel = nullptr;
 
 		// ChildActor's need to be spawned from the ChildActorTemplate otherwise any non-replicated 
 		// customized properties will be incorrect on the Client.
@@ -252,6 +254,8 @@ TUniquePtr<UE::Net::FNetObjectCreationHeader> UNetActorFactory::CreateAndFillHea
 	}
 	else 
 	{
+		FStaticActorNetCreationHeader* Header = new FStaticActorNetCreationHeader;
+		BaseHeader.Reset(Header);
 		// Refer by path for static actors
 		Header->ObjectReference = ActorReference;
 	}
@@ -260,33 +264,65 @@ TUniquePtr<UE::Net::FNetObjectCreationHeader> UNetActorFactory::CreateAndFillHea
 	{
 		FOutBunch Bunch(UE::Net::Private::SerializeNewActorMaxBunchSize);
 		Actor->OnSerializeNewActor(Bunch);
-		Header->CustomCreationDataBitCount = IntCastChecked<uint16>(Bunch.GetNumBits());
-		if (Header->CustomCreationDataBitCount > 0)
+		BaseHeader->CustomCreationDataBitCount = IntCastChecked<uint16>(Bunch.GetNumBits());
+		if (BaseHeader->CustomCreationDataBitCount > 0)
 		{
-			Header->CustomCreationData.SetNumZeroed(Align(Bunch.GetNumBytes(), 4));
-			FMemory::Memcpy(Header->CustomCreationData.GetData(), Bunch.GetData(), Bunch.GetNumBytes());
+			BaseHeader->CustomCreationData.SetNumZeroed(Align(Bunch.GetNumBytes(), 4));
+			FMemory::Memcpy(BaseHeader->CustomCreationData.GetData(), Bunch.GetData(), Bunch.GetNumBytes());
 		}
 	}
 
-	return Header;
+	return BaseHeader;
 }
 
 bool UNetActorFactory::SerializeHeader(const UE::Net::FCreationHeaderContext& Context, const UE::Net::FNetObjectCreationHeader* Header) 
 {
 	using namespace UE::Net;
 
-	const FActorNetCreationHeader* ActorHeader = static_cast<const FActorNetCreationHeader*>(Header);
+	FNetBitStreamWriter* Writer = Context.Serialization.GetBitStreamWriter();
+
+	const FBaseActorNetCreationHeader* BaseActorHeader = static_cast<const FBaseActorNetCreationHeader*>(Header);
+
+	Writer->WriteBool(BaseActorHeader->IsDynamic());
+	
+	if (BaseActorHeader->IsDynamic())
+	{
+		const FDynamicActorNetCreationHeader* DynamicHeader = static_cast<const FDynamicActorNetCreationHeader*>(Header);
+		return DynamicHeader->Serialize(Context, SpawnInfoFlags, DefaultSpawnInfo);
+	}
+	else
+	{
+		const FStaticActorNetCreationHeader* StaticHeader = static_cast<const FStaticActorNetCreationHeader*>(Header);
+		return StaticHeader->Serialize(Context);
+	}
+	
 	
-	return ActorHeader->Serialize(Context, SpawnInfoFlags, DefaultSpawnInfo);;
 }
 
 TUniquePtr<UE::Net::FNetObjectCreationHeader> UNetActorFactory::CreateAndDeserializeHeader(const UE::Net::FCreationHeaderContext& Context)
 {
 	using namespace UE::Net;
 
-	TUniquePtr<FActorNetCreationHeader> Header(new FActorNetCreationHeader);
+	TUniquePtr<FBaseActorNetCreationHeader> Header;
+
+	FNetBitStreamReader* Writer = Context.Serialization.GetBitStreamReader();
+
+	const bool bIsDynamic = Writer->ReadBool();
+	if (bIsDynamic)
+	{
+		FDynamicActorNetCreationHeader* DynamicHeader = new FDynamicActorNetCreationHeader;
+		Header.Reset(DynamicHeader);
+
+		DynamicHeader->Deserialize(Context, DefaultSpawnInfo);
+	}
+	else
+	{
+		FStaticActorNetCreationHeader* StaticHeader = new FStaticActorNetCreationHeader;
+		Header.Reset(StaticHeader);
 
-	Header->Deserialize(Context, DefaultSpawnInfo);
+		StaticHeader->Deserialize(Context);
+	}
+	
 	return Header;
 }
 
@@ -301,19 +337,21 @@ UNetObjectFactory::FInstantiateResult UNetActorFactory::InstantiateReplicatedObj
 	UActorReplicationBridge* ActorBridge = CastChecked<UActorReplicationBridge>(Bridge);
 	UNetDriver* NetDriver = ActorBridge->GetNetDriver();
 
-	const FActorNetCreationHeader* ActorHeader = static_cast<const FActorNetCreationHeader*>(Header);
+	const FBaseActorNetCreationHeader* BaseHeader = static_cast<const FBaseActorNetCreationHeader*>(Header);
 
 	// For static actors, just find the object using the path
-	if (!ActorHeader->bIsDynamic)
+	if (!BaseHeader->IsDynamic())
 	{
-		AActor* Actor = Cast<AActor>(ActorBridge->ResolveObjectReference(ActorHeader->ObjectReference, Context.ResolveContext));
+		const FStaticActorNetCreationHeader* StaticHeader = static_cast<const FStaticActorNetCreationHeader*>(Header);
+
+		AActor* Actor = Cast<AActor>(ActorBridge->ResolveObjectReference(StaticHeader->ObjectReference, Context.ResolveContext));
 		if (!Actor)
 		{
-			UE_LOG(LogIris, Error, TEXT("UNetActorFactory::InstantiateNetObjectFromHeader Failed to resolve ObjectReference: %s . Could not find static actor."), *ActorBridge->DescribeObjectReference(ActorHeader->ObjectReference, Context.ResolveContext));
+			UE_LOG(LogIris, Error, TEXT("UNetActorFactory::InstantiateNetObjectFromHeader Failed to resolve ObjectReference: %s . Could not find static actor."), *ActorBridge->DescribeObjectReference(StaticHeader->ObjectReference, Context.ResolveContext));
 			return FInstantiateResult();
 		}
 
-		UE_LOG(LogIris, Verbose, TEXT("UNetActorFactory::InstantiateNetObjectFromHeader Found static Actor: %s using ObjectReference: %s"), ToCStr(Actor->GetPathName()), *ActorBridge->DescribeObjectReference(ActorHeader->ObjectReference, Context.ResolveContext));
+		UE_LOG(LogIris, Verbose, TEXT("UNetActorFactory::InstantiateNetObjectFromHeader Found static Actor: %s using ObjectReference: %s"), ToCStr(Actor->GetPathName()), *ActorBridge->DescribeObjectReference(StaticHeader->ObjectReference, Context.ResolveContext));
 
 		FInstantiateResult Result {  .Instance = Actor };
 			
@@ -326,20 +364,21 @@ UNetObjectFactory::FInstantiateResult UNetActorFactory::InstantiateReplicatedObj
 	}
 
 	// For dynamic actors, spawn a new actor using the provided information
+	const FDynamicActorNetCreationHeader* DynamicHeader = static_cast<const FDynamicActorNetCreationHeader*>(Header);
 
 	// Find archetype
-	AActor* Archetype = Cast<AActor>(Bridge->ResolveObjectReference(ActorHeader->ArchetypeReference, Context.ResolveContext));
+	AActor* Archetype = Cast<AActor>(Bridge->ResolveObjectReference(DynamicHeader->ArchetypeReference, Context.ResolveContext));
 	if (!Archetype)
 	{
-		UE_LOG(LogIris, Error, TEXT("UNetActorFactory::InstantiateNetObjectFromHeader Unable to instantiate actor, failed to resolve archetype: %s"), *Bridge->DescribeObjectReference(ActorHeader->ArchetypeReference, Context.ResolveContext));
+		UE_LOG(LogIris, Error, TEXT("UNetActorFactory::InstantiateNetObjectFromHeader Unable to instantiate actor, failed to resolve archetype: %s"), *Bridge->DescribeObjectReference(DynamicHeader->ArchetypeReference, Context.ResolveContext));
 		return FInstantiateResult();
 	}
 
 	// Find level
 	ULevel* Level = nullptr;
-	if (!ActorHeader->bUsePersistentLevel)
+	if (!DynamicHeader->bUsePersistentLevel)
 	{
-		Level = Cast<ULevel>(Bridge->ResolveObjectReference(ActorHeader->LevelReference, Context.ResolveContext));
+		Level = Cast<ULevel>(Bridge->ResolveObjectReference(DynamicHeader->LevelReference, Context.ResolveContext));
 	}
 
 	LLM_SCOPE_DYNAMIC_STAT_OBJECTPATH(Archetype->GetPackage(), ELLMTagSet::Assets);
@@ -361,36 +400,32 @@ UNetObjectFactory::FInstantiateResult UNetActorFactory::InstantiateReplicatedObj
 	SpawnInfo.bNoFail = true;
 
 	UWorld* World = NetDriver->GetWorld();
-	FVector SpawnLocation = FRepMovement::RebaseOntoLocalOrigin(ActorHeader->SpawnInfo.Location, World->OriginLocation);
+	FVector SpawnLocation = FRepMovement::RebaseOntoLocalOrigin(DynamicHeader->SpawnInfo.Location, World->OriginLocation);
 		
-	AActor* Actor = World->SpawnActorAbsolute(Archetype->GetClass(), FTransform(ActorHeader->SpawnInfo.Rotation, SpawnLocation), SpawnInfo);
+	AActor* Actor = World->SpawnActorAbsolute(Archetype->GetClass(), FTransform(DynamicHeader->SpawnInfo.Rotation, SpawnLocation), SpawnInfo);
 
 	// For Iris we expect that we will be able to spawn the actor as streaming always is controlled from server
 	if (!Actor)
 	{
-		ensureMsgf(Actor, TEXT("UNetActorFactory::InstantiateNetObjectFromHeader SpawnActor failed. Used Archetype: %s ObjectReference: %s"), *GetNameSafe(Archetype), *ActorBridge->DescribeObjectReference(ActorHeader->ArchetypeReference, Context.ResolveContext));
+		ensureMsgf(Actor, TEXT("UNetActorFactory::InstantiateNetObjectFromHeader SpawnActor failed. Used Archetype: %s ObjectReference: %s"), *GetNameSafe(Archetype), *ActorBridge->DescribeObjectReference(DynamicHeader->ArchetypeReference, Context.ResolveContext));
 		return FInstantiateResult();
 	}
 
-	FActorNetCreationHeader::FActorNetSpawnInfo DefaultNetSpawnInfo;
 	static constexpr float Epsilon = UE_KINDA_SMALL_NUMBER;
 
 	// Set Velocity if it differs from Default
-	if (!ActorHeader->SpawnInfo.Velocity.Equals(DefaultNetSpawnInfo.Velocity, Epsilon))
+	if (!DynamicHeader->SpawnInfo.Velocity.Equals(DefaultSpawnInfo.Velocity, Epsilon))
 	{
-		Actor->PostNetReceiveVelocity(ActorHeader->SpawnInfo.Velocity);
+		Actor->PostNetReceiveVelocity(DynamicHeader->SpawnInfo.Velocity);
 	}
 					
 	// Set Scale if it differs from Default
-	if (!ActorHeader->SpawnInfo.Scale.Equals(DefaultNetSpawnInfo.Scale, Epsilon))
+	if (!DynamicHeader->SpawnInfo.Scale.Equals(DefaultSpawnInfo.Scale, Epsilon))
 	{
-		Actor->SetActorRelativeScale3D(ActorHeader->SpawnInfo.Scale);
+		Actor->SetActorRelativeScale3D(DynamicHeader->SpawnInfo.Scale);
 	}
 
-	FInstantiateResult Result
-	{ 
-		.Instance = Actor
-	};
+	FInstantiateResult Result { .Instance = Actor };
 
 	if (NetDriver->ShouldClientDestroyActor(Actor))
 	{
@@ -415,12 +450,12 @@ void UNetActorFactory::PostInstantiation(const FPostInstantiationContext& Contex
 	UActorReplicationBridge* ActorBridge = CastChecked<UActorReplicationBridge>(Bridge);
 	UNetDriver* NetDriver = ActorBridge->GetNetDriver();
 
-	const FActorNetCreationHeader* ActorHeader = static_cast<const FActorNetCreationHeader*>(Context.Header);
+	const FBaseActorNetCreationHeader* BaseHeader = static_cast<const FBaseActorNetCreationHeader*>(Context.Header);
 
 	// OnActorChannelOpen
 	{
 		UNetConnection* Connection = NetDriver->GetConnectionById(Context.ConnectionId);
-		FInBunch Bunch(Connection, const_cast<uint8*>(ActorHeader->CustomCreationData.GetData()),ActorHeader->CustomCreationDataBitCount);
+		FInBunch Bunch(Connection, const_cast<uint8*>(BaseHeader->CustomCreationData.GetData()), BaseHeader->CustomCreationDataBitCount);
 		Actor->OnActorChannelOpen(Bunch, Connection);
 
 		if (Bunch.IsError() || Bunch.GetBitsLeft() != 0)
@@ -436,49 +471,91 @@ void UNetActorFactory::PostInstantiation(const FPostInstantiationContext& Contex
 }
 
 
-//------------------------------------------------------------------------
-// FActorNetCreationHeader
-//------------------------------------------------------------------------
  namespace UE::Net
  {
 
-bool FActorNetCreationHeader::Serialize(const FCreationHeaderContext& Context, UE::Net::Private::EActorNetSpawnInfoFlags SpawnFlags, const FActorNetSpawnInfo& DefaultSpawnInfo) const
+//------------------------------------------------------------------------
+// FStaticActorNetCreationHeader
+//------------------------------------------------------------------------
+
+bool FStaticActorNetCreationHeader::Serialize(const FCreationHeaderContext& Context) const
 {
 	using namespace UE::Net::Private;
 
 	FNetBitStreamWriter* Writer = Context.Serialization.GetBitStreamWriter();
 
-	if (Writer->WriteBool(bIsDynamic))
+	WriteFullNetObjectReference(Context.Serialization, ObjectReference);
+
+	if (Writer->WriteBool(CustomCreationDataBitCount > 0))
 	{
-		// Write Archetype and LevelPath
-		WriteFullNetObjectReference(Context.Serialization, ArchetypeReference);
-		
-		// Only write the LevelPath if it differs from the persistent level
-		if (!Writer->WriteBool(bUsePersistentLevel))
-		{
-			WriteFullNetObjectReference(Context.Serialization, LevelReference);
-		}
+		Writer->WriteBits(CustomCreationDataBitCount - 1U, 16U);
+		Writer->WriteBitStream(reinterpret_cast<const uint32*>(CustomCreationData.GetData()), 0U, CustomCreationDataBitCount);
+	}
 
-		const bool bQuantizeLocation = (SpawnFlags&EActorNetSpawnInfoFlags::QuantizeLocation)!=EActorNetSpawnInfoFlags::None;
-		const bool bQuantizeScale = (SpawnFlags&EActorNetSpawnInfoFlags::QuantizeScale)!=EActorNetSpawnInfoFlags::None;
-		const bool bQuantizeVelocity = (SpawnFlags&EActorNetSpawnInfoFlags::QuantizeVelocity)!=EActorNetSpawnInfoFlags::None;
+	return true;
+}
 
-		// Write actor spawn info
-		WriteConditionallyQuantizedVector(Writer, SpawnInfo.Location, DefaultSpawnInfo.Location, bQuantizeLocation);
-		WriteConditionallyQuantizedVector(Writer, SpawnInfo.Scale, DefaultSpawnInfo.Scale, bQuantizeScale);
-		WriteConditionallyQuantizedVector(Writer, SpawnInfo.Velocity, DefaultSpawnInfo.Velocity, bQuantizeVelocity);
+bool FStaticActorNetCreationHeader::Deserialize(const FCreationHeaderContext& Context)
+{
+	using namespace UE::Net::Private;
 
-		// For rotation we use 0.001f for Rotation comparison to keep consistency with old behavior.
-		static constexpr float RotationEpsilon = 0.001f;
-		WriteRotator(Writer, SpawnInfo.Rotation, DefaultSpawnInfo.Rotation, RotationEpsilon);
+	FNetBitStreamReader* Reader = Context.Serialization.GetBitStreamReader();
 
+	ReadFullNetObjectReference(Context.Serialization, ObjectReference);
 
+	if (Reader->ReadBool())
+	{
+		CustomCreationDataBitCount = 1U + Reader->ReadBits(16U);
+		CustomCreationData.SetNumZeroed(((CustomCreationDataBitCount + 31U) & ~31U) >> 3U);
+		Reader->ReadBitStream(reinterpret_cast<uint32*>(CustomCreationData.GetData()), CustomCreationDataBitCount);
 	}
-	else // should be possible to refer by path??
+	
+	return true;
+}
+
+
+FString FStaticActorNetCreationHeader::ToString() const
+{
+	return FString::Printf(TEXT("FStaticActorNetCreationHeader (ProtocolId:0x%x):\n\t"
+								"ObjectReference=%s\n\t"
+								"CustomCreationData=%u bits"),
+						   		GetProtocolId(),
+						   		*ObjectReference.ToString(),
+								CustomCreationDataBitCount);
+}
+
+//------------------------------------------------------------------------
+// FDynamicActorNetCreationHeader
+//------------------------------------------------------------------------
+
+bool FDynamicActorNetCreationHeader::Serialize(const FCreationHeaderContext& Context, UE::Net::Private::EActorNetSpawnInfoFlags SpawnFlags, const FActorNetSpawnInfo& DefaultSpawnInfo) const
+{
+	using namespace UE::Net::Private;
+
+	FNetBitStreamWriter* Writer = Context.Serialization.GetBitStreamWriter();
+
+	// Write Archetype and LevelPath
+	WriteFullNetObjectReference(Context.Serialization, ArchetypeReference);
+		
+	// Only write the LevelPath if it differs from the persistent level
+	if (!Writer->WriteBool(bUsePersistentLevel))
 	{
-		WriteFullNetObjectReference(Context.Serialization, ObjectReference);
+		WriteFullNetObjectReference(Context.Serialization, LevelReference);
 	}
 
+	const bool bQuantizeLocation = (SpawnFlags&EActorNetSpawnInfoFlags::QuantizeLocation)!=EActorNetSpawnInfoFlags::None;
+	const bool bQuantizeScale = (SpawnFlags&EActorNetSpawnInfoFlags::QuantizeScale)!=EActorNetSpawnInfoFlags::None;
+	const bool bQuantizeVelocity = (SpawnFlags&EActorNetSpawnInfoFlags::QuantizeVelocity)!=EActorNetSpawnInfoFlags::None;
+
+	// Write actor spawn info
+	WriteConditionallyQuantizedVector(Writer, SpawnInfo.Location, DefaultSpawnInfo.Location, bQuantizeLocation);
+	WriteConditionallyQuantizedVector(Writer, SpawnInfo.Scale, DefaultSpawnInfo.Scale, bQuantizeScale);
+	WriteConditionallyQuantizedVector(Writer, SpawnInfo.Velocity, DefaultSpawnInfo.Velocity, bQuantizeVelocity);
+
+	// For rotation we use 0.001f for Rotation comparison to keep consistency with old behavior.
+	static constexpr float RotationEpsilon = 0.001f;
+	WriteRotator(Writer, SpawnInfo.Rotation, DefaultSpawnInfo.Rotation, RotationEpsilon);
+
 	if (Writer->WriteBool(CustomCreationDataBitCount > 0))
 	{
 		Writer->WriteBits(CustomCreationDataBitCount - 1U, 16U);
@@ -488,39 +565,28 @@ bool FActorNetCreationHeader::Serialize(const FCreationHeaderContext& Context, U
 	return true;
 }
 
-bool FActorNetCreationHeader::Deserialize(const FCreationHeaderContext& Context, const FActorNetSpawnInfo& DefaultSpawnInfo)
+bool FDynamicActorNetCreationHeader::Deserialize(const FCreationHeaderContext& Context, const FActorNetSpawnInfo& DefaultSpawnInfo)
 {
 	using namespace UE::Net::Private;
 
 	FNetBitStreamReader* Reader = Context.Serialization.GetBitStreamReader();
 
-	// Dynamic actor?
-	bIsDynamic = Reader->ReadBool();
-
-	if (bIsDynamic)
-	{
-		// Read Archetype
-		ReadFullNetObjectReference(Context.Serialization, ArchetypeReference);
+	// Read Archetype
+	ReadFullNetObjectReference(Context.Serialization, ArchetypeReference);
 		
-		bUsePersistentLevel = Reader->ReadBool();
-		if (!bUsePersistentLevel)
-		{
-			ReadFullNetObjectReference(Context.Serialization, LevelReference);
-		}
-
-		// Read actor spawn info
-		ReadConditionallyQuantizedVector(Reader, SpawnInfo.Location, DefaultSpawnInfo.Location);
-		ReadConditionallyQuantizedVector(Reader, SpawnInfo.Scale, DefaultSpawnInfo.Scale);
-		ReadConditionallyQuantizedVector(Reader, SpawnInfo.Velocity, DefaultSpawnInfo.Velocity);
-
-		ReadRotator(Reader, SpawnInfo.Rotation, DefaultSpawnInfo.Rotation);
-	}
-	else
+	bUsePersistentLevel = Reader->ReadBool();
+	if (!bUsePersistentLevel)
 	{
-		ReadFullNetObjectReference(Context.Serialization, ObjectReference);
+		ReadFullNetObjectReference(Context.Serialization, LevelReference);
 	}
 
-	CustomCreationDataBitCount = 0;
+	// Read actor spawn info
+	ReadConditionallyQuantizedVector(Reader, SpawnInfo.Location, DefaultSpawnInfo.Location);
+	ReadConditionallyQuantizedVector(Reader, SpawnInfo.Scale, DefaultSpawnInfo.Scale);
+	ReadConditionallyQuantizedVector(Reader, SpawnInfo.Velocity, DefaultSpawnInfo.Velocity);
+
+	ReadRotator(Reader, SpawnInfo.Rotation, DefaultSpawnInfo.Rotation);
+	
 	if (Reader->ReadBool())
 	{
 		CustomCreationDataBitCount = 1U + Reader->ReadBits(16U);
@@ -532,11 +598,9 @@ bool FActorNetCreationHeader::Deserialize(const FCreationHeaderContext& Context,
 }
 
 
-FString FActorNetCreationHeader::ToString() const
+FString FDynamicActorNetCreationHeader::ToString() const
 {
 	return FString::Printf(TEXT("FDynamicActorNetCreationHeader (ProtocolId:0x%x):\n\t"
-								"bIsDynamic=%u\n\t"
-								"ObjectReference=%s\n\t"
 								"ArchetypeReference=%s\n\t"
 								"SpawnInfo.Location=%s\n\t"
 								"SpawnInfo.Rotation=%s\n\t"
@@ -544,11 +608,8 @@ FString FActorNetCreationHeader::ToString() const
 								"SpawnInfo.Velocity=%s\n\t"		
 								"bUsePersistentLevel=%s\n\t"
 								"LevelReference=%s\n\t"
-								"CustomCreationData=%u bits"
-								),
+								"CustomCreationData=%u bits"),
 						   		GetProtocolId(),
-						   		bIsDynamic,
-						   		*ObjectReference.ToString(),
 								*ArchetypeReference.ToString(),
 								*SpawnInfo.Location.ToCompactString(),
 								*SpawnInfo.Rotation.ToCompactString(),
@@ -556,8 +617,7 @@ FString FActorNetCreationHeader::ToString() const
 								*SpawnInfo.Velocity.ToCompactString(),		
 								bUsePersistentLevel?TEXT("True"):TEXT("False"),
 								*LevelReference.ToString(),
-								CustomCreationDataBitCount
-								);
+								CustomCreationDataBitCount);
 }
 
 
diff --git a/Engine/Source/Runtime/Engine/Private/Net/Iris/ReplicationSystem/NetSubObjectFactory.cpp b/Engine/Source/Runtime/Engine/Private/Net/Iris/ReplicationSystem/NetSubObjectFactory.cpp
--- a/Engine/Source/Runtime/Engine/Private/Net/Iris/ReplicationSystem/NetSubObjectFactory.cpp
+++ b/Engine/Source/Runtime/Engine/Private/Net/Iris/ReplicationSystem/NetSubObjectFactory.cpp
@@ -30,9 +30,6 @@ TUniquePtr<UE::Net::FNetObjectCreationHeader> UNetSubObjectFactory::CreateAndFil
 {
 	using namespace UE::Net;
 
-	TUniquePtr<FNetObjectCreationHeader> Header(new FNetSubObjectCreationHeader);
-	FNetSubObjectCreationHeader* SubObjectHeader = static_cast<FNetSubObjectCreationHeader*>(Header.Get());
-
 	UObject* SubObject = Bridge->GetReplicatedObject(Handle);
 	if (!SubObject)
 	{
@@ -42,16 +39,17 @@ TUniquePtr<UE::Net::FNetObjectCreationHeader> UNetSubObjectFactory::CreateAndFil
 
 	FNetObjectReference ObjectRef = Bridge->GetOrCreateObjectReference(SubObject);
 
-	SubObjectHeader->bIsDynamic = ObjectRef.GetRefHandle().IsDynamic();
-	SubObjectHeader->bIsNameStableForNetworking = SubObject->IsNameStableForNetworking();
-	SubObjectHeader->ObjectReference = ObjectRef;
-
-	if (!SubObjectHeader->bIsDynamic || SubObjectHeader->bIsNameStableForNetworking)
+	if (!ObjectRef.GetRefHandle().IsDynamic() || SubObject->IsNameStableForNetworking())
 	{
 		// No more information needed since we don't need to spawn the object on the remote.
+		TUniquePtr<FNetStaticSubObjectCreationHeader> Header(new FNetStaticSubObjectCreationHeader);
+		Header->ObjectReference = ObjectRef;
 		return Header;
 	}
 
+	TUniquePtr<FNetDynamicSubObjectCreationHeader> Header(new FNetDynamicSubObjectCreationHeader);
+	FNetDynamicSubObjectCreationHeader* SubObjectHeader = static_cast<FNetDynamicSubObjectCreationHeader*>(Header.Get());
+
 	check(SubObject->NeedsLoadForClient() ); // We have no business sending this unless the client can load
 	check(SubObject->GetClass()->NeedsLoadForClient());	// We have no business sending this unless the client can load
 
@@ -91,9 +89,20 @@ TUniquePtr<UE::Net::FNetObjectCreationHeader> UNetSubObjectFactory::CreateAndDes
 {
 	using namespace UE::Net;
 
-	TUniquePtr<FNetSubObjectCreationHeader> Header(new FNetSubObjectCreationHeader);
-	Header->Deserialize(Context);
-	return Header;
+	FNetBitStreamReader* Reader = Context.Serialization.GetBitStreamReader();
+
+	if (Reader->ReadBool())
+	{
+		TUniquePtr<FNetDynamicSubObjectCreationHeader> Header(new FNetDynamicSubObjectCreationHeader);
+		Header->Deserialize(Context);
+		return Header;
+	}
+	else
+	{
+		TUniquePtr<FNetStaticSubObjectCreationHeader> Header(new FNetStaticSubObjectCreationHeader);
+		Header->Deserialize(Context);
+		return Header;
+	}
 }
 
 
@@ -101,34 +110,31 @@ UNetSubObjectFactory::FInstantiateResult UNetSubObjectFactory::InstantiateReplic
 {
 	using namespace UE::Net;
 
-	const FNetSubObjectCreationHeader* SubObjectHeader = static_cast<const FNetSubObjectCreationHeader*>(Header);
+	const FNetBaseSubObjectCreationHeader* BaseHeader = static_cast<const FNetBaseSubObjectCreationHeader*>(Header);
 
-	if (!SubObjectHeader->bIsDynamic || SubObjectHeader->bIsNameStableForNetworking)
+	if (!BaseHeader->IsDynamic())
 	{
+		const FNetStaticSubObjectCreationHeader* SubObjectHeader = static_cast<const FNetStaticSubObjectCreationHeader*>(BaseHeader);
+
 		// Resolve by finding object relative to owner. We do not allow this object to be destroyed.
 		UObject* SubObject = Bridge->ResolveObjectReference(SubObjectHeader->ObjectReference, Context.ResolveContext);
 
 		if (!SubObject)
 		{
-			if (!SubObjectHeader->bIsDynamic)
-			{
-				UE_LOG(LogIris, Error, TEXT("UNetSubObjectFactory::InstantiateNetObjectFromHeader %s: Failed to find static object reference for static SubObject: %s, Owner: %s, RootObject: %s"), *Context.Handle.ToString(), *Bridge->DescribeObjectReference(SubObjectHeader->ObjectReference, Context.ResolveContext), *Bridge->PrintObjectFromNetRefHandle(Context.RootObjectOfSubObject), *GetPathNameSafe(Bridge->GetReplicatedObject(Context.RootObjectOfSubObject)));
-			}
-			else if (SubObjectHeader->bIsNameStableForNetworking)
-			{
-				UE_LOG(LogIris, Error, TEXT("UNetSubObjectFactory::InstantiateNetObjectFromHeader %s: Failed to find stable name reference of dynamic SubObject: %s, Owner: %s, RootObject: %s"), *Context.Handle.ToString(), *Bridge->DescribeObjectReference(SubObjectHeader->ObjectReference, Context.ResolveContext), *Bridge->PrintObjectFromNetRefHandle(Context.RootObjectOfSubObject), *GetPathNameSafe(Bridge->GetReplicatedObject(Context.RootObjectOfSubObject)));
-			}
-
+			UE_LOG(LogIris, Error, TEXT("UNetSubObjectFactory::InstantiateNetObjectFromHeader %s: Failed to find static or stable name object referenced by SubObject: %s, Owner: %s, RootObject: %s"), *Context.Handle.ToString(), *Bridge->DescribeObjectReference(SubObjectHeader->ObjectReference, Context.ResolveContext), *Bridge->PrintObjectFromNetRefHandle(Context.RootObjectOfSubObject), *GetPathNameSafe(Bridge->GetReplicatedObject(Context.RootObjectOfSubObject)));
+			
 			return FInstantiateResult();
 		}
 
-		UE_LOG(LogIris, Verbose, TEXT("UNetSubObjectFactory::InstantiateNetObjectFromHeader %s: Found static SubObject using path %s"), *Context.Handle.ToString(), ToCStr(SubObject->GetPathName()));
+		UE_LOG(LogIris, Verbose, TEXT("UNetSubObjectFactory::InstantiateNetObjectFromHeader %s: Found static or stable name SubObject using path %s"), *Context.Handle.ToString(), ToCStr(SubObject->GetPathName()));
 
 		FInstantiateResult Result { .Instance = SubObject };
 		return Result;
 	}
 
 	// For dynamic objects we have to spawn them
+
+	const FNetDynamicSubObjectCreationHeader* SubObjectHeader = static_cast<const FNetDynamicSubObjectCreationHeader*>(Header);
 		
 	UObject* RootObject = Bridge->GetReplicatedObject(Context.RootObjectOfSubObject);
 	AActor* RootActor = CastChecked<AActor>(RootObject);
@@ -149,7 +155,7 @@ UNetSubObjectFactory::FInstantiateResult UNetSubObjectFactory::InstantiateReplic
 
 		if (!OuterObject)
 		{
-			UE_LOG(LogIris, Error, TEXT("BeginInstantiateFromRemote Failed to find Outer %s for dynamic subobject %s"), *Bridge->DescribeObjectReference(SubObjectHeader->OuterReference, Context.ResolveContext), *Bridge->DescribeObjectReference(SubObjectHeader->ObjectReference, Context.ResolveContext))
+			UE_LOG(LogIris, Error, TEXT("BeginInstantiateFromRemote Failed to find Outer %s for dynamic subobject %s"), *Bridge->DescribeObjectReference(SubObjectHeader->OuterReference, Context.ResolveContext), *Context.Handle.ToString())
 
 			// Fallback to the rootobject instead
 			OuterObject = RootActor;
@@ -182,7 +188,12 @@ UNetSubObjectFactory::FInstantiateResult UNetSubObjectFactory::InstantiateReplic
 bool UNetSubObjectFactory::SerializeHeader(const UE::Net::FCreationHeaderContext& Context, const UE::Net::FNetObjectCreationHeader* Header)
 {
 	using namespace UE::Net;
-	const FNetSubObjectCreationHeader* SubObjectHeader = static_cast<const FNetSubObjectCreationHeader*>(Header);
+	const FNetBaseSubObjectCreationHeader* SubObjectHeader = static_cast<const FNetBaseSubObjectCreationHeader*>(Header);
+	
+	FNetBitStreamWriter* Writer = Context.Serialization.GetBitStreamWriter();
+	
+	Writer->WriteBool(SubObjectHeader->IsDynamic());
+
 	return SubObjectHeader->Serialize(Context);
 }
 
@@ -191,91 +202,82 @@ bool UNetSubObjectFactory::SerializeHeader(const UE::Net::FCreationHeaderContext
 namespace UE::Net
 {
 
-FString FNetSubObjectCreationHeader::ToString() const
+//------------------------------------------------------------------------
+// FNetStaticSubObjectCreationHeader
+//------------------------------------------------------------------------
+FString FNetStaticSubObjectCreationHeader::ToString() const
 {
-	return FString::Printf(TEXT("FNetSubObjectCreationHeader (ProtocolId:0x%x):\n\t"
-								"ObjectReference=%s\n\t"
+	return FString::Printf(TEXT("FNetStaticSubObjectCreationHeader (ProtocolId:0x%x):\n\t"
+								"ObjectReference=%s\n\t"),								
+								GetProtocolId(),
+								*ObjectReference.ToString());
+
+}
+
+bool FNetStaticSubObjectCreationHeader::Serialize(const FCreationHeaderContext& Context) const
+{
+	WriteFullNetObjectReference(Context.Serialization, ObjectReference);
+	return true;
+}
+
+bool FNetStaticSubObjectCreationHeader::Deserialize(const FCreationHeaderContext& Context)
+{
+	ReadFullNetObjectReference(Context.Serialization, ObjectReference);
+	return true;
+}
+
+//------------------------------------------------------------------------
+// FNetDynamicSubObjectCreationHeader
+//------------------------------------------------------------------------
+
+FString FNetDynamicSubObjectCreationHeader::ToString() const
+{
+	return FString::Printf(TEXT("FNetDynamicSubObjectCreationHeader (ProtocolId:0x%x):\n\t"
 								"ObjectClassReference=%s\n\t"
 								"OuterReference=%s\n\t"
-								"bIsDynamic=%u\n\t"
-								"bIsNameStableForNetworking=%u\n\t"
 								"bUsePersistenLevel=%u\n\t"
 								"bOuterIsTransientLevel=%u\n\t"
 								"bOuterIsRootObject=%u\n\t"),
 								GetProtocolId(),
-								*ObjectReference.ToString(),
 								*ObjectClassReference.ToString(),
 								*OuterReference.ToString(),
-								bIsDynamic,
-								bIsNameStableForNetworking,
 								bUsePersistentLevel,
 								bOuterIsTransientLevel,
 								bOuterIsRootObject);
 
 }
 
-bool FNetSubObjectCreationHeader::Serialize(const FCreationHeaderContext& Context) const
+bool FNetDynamicSubObjectCreationHeader::Serialize(const FCreationHeaderContext& Context) const
 {
 	FNetBitStreamWriter* Writer = Context.Serialization.GetBitStreamWriter();
 
-	// Write required references to either find or instantiate the subobject
-	if (Writer->WriteBool(bIsDynamic))
+	WriteFullNetObjectReference(Context.Serialization, ObjectClassReference);
+
+	if (!Writer->WriteBool(bOuterIsTransientLevel))
 	{
-		if (Writer->WriteBool(bIsNameStableForNetworking))
+		if (!Writer->WriteBool(bOuterIsRootObject))
 		{
-			WriteFullNetObjectReference(Context.Serialization, ObjectReference);
-		}
-		else
-		{
-			WriteFullNetObjectReference(Context.Serialization, ObjectClassReference);
-
-			if (!Writer->WriteBool(bOuterIsTransientLevel))
-			{
-				if (!Writer->WriteBool(bOuterIsRootObject))
-				{
-					WriteFullNetObjectReference(Context.Serialization, OuterReference);
-				}
-			}
+			WriteFullNetObjectReference(Context.Serialization, OuterReference);
 		}
 	}
-	else
-	{
-		WriteFullNetObjectReference(Context.Serialization, ObjectReference);
-	}
 
 	return true;
 }
 
-bool FNetSubObjectCreationHeader::Deserialize(const FCreationHeaderContext& Context)
+bool FNetDynamicSubObjectCreationHeader::Deserialize(const FCreationHeaderContext& Context)
 {
 	FNetBitStreamReader* Reader = Context.Serialization.GetBitStreamReader();
 
-	bIsDynamic = Reader->ReadBool();
-	if (bIsDynamic)
+	ReadFullNetObjectReference(Context.Serialization, ObjectClassReference);
+
+	bOuterIsTransientLevel = Reader->ReadBool();
+	if (!bOuterIsTransientLevel)
 	{
-		bIsNameStableForNetworking = Reader->ReadBool();
-		if (bIsNameStableForNetworking)
+		bOuterIsRootObject = Reader->ReadBool();
+		if (!bOuterIsRootObject)
 		{
-			ReadFullNetObjectReference(Context.Serialization, ObjectReference);
+			ReadFullNetObjectReference(Context.Serialization, OuterReference);
 		}
-		else
-		{
-			ReadFullNetObjectReference(Context.Serialization, ObjectClassReference);
-
-			bOuterIsTransientLevel = Reader->ReadBool();
-			if (!bOuterIsTransientLevel)
-			{
-				bOuterIsRootObject = Reader->ReadBool();
-				if (!bOuterIsRootObject)
-				{
-					ReadFullNetObjectReference(Context.Serialization, OuterReference);
-				}
-			}
-		}
-	}
-	else
-	{
-		ReadFullNetObjectReference(Context.Serialization, ObjectReference);
 	}
 
 	return true;
diff --git a/Engine/Source/Runtime/Engine/Public/Net/Iris/ReplicationSystem/NetActorFactory.h b/Engine/Source/Runtime/Engine/Public/Net/Iris/ReplicationSystem/NetActorFactory.h
--- a/Engine/Source/Runtime/Engine/Public/Net/Iris/ReplicationSystem/NetActorFactory.h
+++ b/Engine/Source/Runtime/Engine/Public/Net/Iris/ReplicationSystem/NetActorFactory.h
@@ -22,11 +22,51 @@ namespace UE::Net
 #if UE_WITH_IRIS
 
 /**
-* Header information representing dynamic actors
-*/
-class FActorNetCreationHeader : public FNetObjectCreationHeader
+ * Header information to be able to tell if its a dynamic or static header
+ */
+class FBaseActorNetCreationHeader : public FNetObjectCreationHeader
 {
 public:
+	virtual bool IsDynamic() const = 0;
+
+	TArray<uint8> CustomCreationData;
+	uint16 CustomCreationDataBitCount = 0;
+};
+
+/**
+ * Header information representing static actors
+ */
+class FStaticActorNetCreationHeader : public FBaseActorNetCreationHeader
+{
+public:
+
+	virtual bool IsDynamic() const
+	{
+		return false;
+	}
+
+	FNetObjectReference ObjectReference;
+	
+	bool Serialize(const FCreationHeaderContext& Context) const;
+	bool Deserialize(const FCreationHeaderContext& Context);
+
+	virtual FString ToString() const override;
+};
+
+
+/**
+ * Header information representing dynamic actors
+ */
+class FDynamicActorNetCreationHeader : public FBaseActorNetCreationHeader
+{
+public:
+
+	virtual bool IsDynamic() const override
+	{ 
+		return true;
+	}
+
+	virtual FString ToString() const override;
 
 	struct FActorNetSpawnInfo
 	{
@@ -45,20 +85,13 @@ public:
 
 	FActorNetSpawnInfo SpawnInfo;
 
-	FNetObjectReference ObjectReference;
 	FNetObjectReference ArchetypeReference;
 	FNetObjectReference LevelReference; // Only when bUsePersistentLevel is false
 	
-	bool bIsDynamic = false;
 	bool bUsePersistentLevel = false;
 
-	TArray<uint8> CustomCreationData;
-	uint16 CustomCreationDataBitCount = 0;
-	
 	bool Serialize(const FCreationHeaderContext& Context, UE::Net::Private::EActorNetSpawnInfoFlags SpawnFlags, const FActorNetSpawnInfo& DefaultSpawnInfo) const;
 	bool Deserialize(const FCreationHeaderContext& Context, const FActorNetSpawnInfo& DefaultSpawnInfo);
-
-	virtual FString ToString() const override;
 };
 
 #endif // UE_WITH_IRIS
@@ -77,7 +110,10 @@ class UNetActorFactory : public UNetObjectFactory
 
 public:
 
-	static FName GetFactoryName() { return TEXT("NetActorFactory"); }
+	static FName GetFactoryName()
+	{ 
+		return TEXT("NetActorFactory"); 
+	}
 
 	virtual void OnInit() override;
 
@@ -96,7 +132,7 @@ private:
 
 	UE::Net::Private::EActorNetSpawnInfoFlags SpawnInfoFlags;
 
-	const UE::Net::FActorNetCreationHeader::FActorNetSpawnInfo DefaultSpawnInfo;
+	const UE::Net::FDynamicActorNetCreationHeader::FActorNetSpawnInfo DefaultSpawnInfo;
 
 #endif // UE_WITH_IRIS
 };
diff --git a/Engine/Source/Runtime/Engine/Public/Net/Iris/ReplicationSystem/NetSubObjectFactory.h b/Engine/Source/Runtime/Engine/Public/Net/Iris/ReplicationSystem/NetSubObjectFactory.h
--- a/Engine/Source/Runtime/Engine/Public/Net/Iris/ReplicationSystem/NetSubObjectFactory.h
+++ b/Engine/Source/Runtime/Engine/Public/Net/Iris/ReplicationSystem/NetSubObjectFactory.h
@@ -22,7 +22,10 @@ class UNetSubObjectFactory : public UNetObjectFactory
 
 public:
 
-	static FName GetFactoryName() { return TEXT("NetSubObjectFactory"); }
+	static FName GetFactoryName()
+	{ 
+		return TEXT("NetSubObjectFactory");
+	}
 
 	virtual FInstantiateResult InstantiateReplicatedObjectFromHeader(const FInstantiateContext& Context, const UE::Net::FNetObjectCreationHeader* Header) override;
 
@@ -39,34 +42,65 @@ protected:
 namespace UE::Net
 {
 #if UE_WITH_IRIS
+
+/**
+ * Header information to be able to tell if its a dynamic or static header
+ */
+class FNetBaseSubObjectCreationHeader : public FNetObjectCreationHeader
+{
+public:
+	
+	virtual bool IsDynamic() const = 0;
+
+	virtual bool Serialize(const FCreationHeaderContext& Context) const
+	{ 
+		return false;
+	}
+};
+
+
 /**
-* Header information representing subobjects
-*/
-class FNetSubObjectCreationHeader : public FNetObjectCreationHeader
+ * Header information representing subobjects that can be found via their pathname (aka: static or stable name)
+ */
+class FNetStaticSubObjectCreationHeader : public FNetBaseSubObjectCreationHeader
 {
 public:
 
-	bool Serialize(const FCreationHeaderContext& Context) const;
+	virtual bool IsDynamic() const override
+	{
+		return false;
+	}
+
+	virtual bool Serialize(const FCreationHeaderContext& Context) const override;
 	bool Deserialize(const FCreationHeaderContext& Context);
 
 	virtual FString ToString() const override;
 
-	FNetSubObjectCreationHeader()
-		: bIsDynamic(false)
-		, bIsNameStableForNetworking(false)
-		, bUsePersistentLevel(false)
-		, bOuterIsTransientLevel(false)
-		, bOuterIsRootObject(false)
-	{}
-
 	FNetObjectReference ObjectReference; // Only for static objects
-	FNetObjectReference ObjectClassReference; // Only for dynamic objects
-	FNetObjectReference OuterReference; // Optional: Outer ref sent only for dynamic subobjects.
-	uint8 bIsDynamic : 1;
-	uint8 bIsNameStableForNetworking : 1;
-	uint8 bUsePersistentLevel : 1;
-	uint8 bOuterIsTransientLevel : 1;	// When set the OuterReference was not sent because the Outer is the default transient level.
-	uint8 bOuterIsRootObject : 1;		// When set the OuterReference was not sent because the Outer is the known RootObject.
+};
+
+/**
+ * Header information representing subobjects that must be instantiated
+ */
+class FNetDynamicSubObjectCreationHeader : public FNetBaseSubObjectCreationHeader
+{
+public:
+
+	virtual bool Serialize(const FCreationHeaderContext& Context) const override;
+	bool Deserialize(const FCreationHeaderContext& Context);
+
+	virtual bool IsDynamic() const override
+	{ 
+		return true;
+	}
+
+	virtual FString ToString() const override;
+
+	FNetObjectReference ObjectClassReference;
+	FNetObjectReference OuterReference;
+	uint8 bUsePersistentLevel : 1 = false;
+	uint8 bOuterIsTransientLevel : 1 = false;	// When set the OuterReference was not sent because the Outer is the default transient level.
+	uint8 bOuterIsRootObject : 1 = false;		// When set the OuterReference was not sent because the Outer is the known RootObject.
 
 };
 #endif // UE_WITH_IRIS
