Ensure all ContentBundles that are registered after WorldPartition has initialized are also calling InjectContentBundle, to match logic for ContentBundles registered before WorldPartition has initialized (via InitializeContentBundlesForegisteredClients)

#rnx

diff --git a/Engine/Source/Runtime/Engine/Private/WorldPartition/ContentBundle/ContentBundleContainer.cpp b/Engine/Source/Runtime/Engine/Private/WorldPartition/ContentBundle/ContentBundleContainer.cpp
--- a/Engine/Source/Runtime/Engine/Private/WorldPartition/ContentBundle/ContentBundleContainer.cpp
+++ b/Engine/Source/Runtime/Engine/Private/WorldPartition/ContentBundle/ContentBundleContainer.cpp
@@ -339,7 +339,8 @@ bool FContentBundleContainer::RemoveContentBundle(FContentBundleBase& ContentBun
 
 void FContentBundleContainer::OnContentBundleClientRegistered(TSharedPtr<FContentBundleClient>& ContentBundleClient)
 {
-	InitializeContentBundle(ContentBundleClient);
+	FContentBundleBase& ContentBundle = InitializeContentBundle(ContentBundleClient);
+	InjectContentBundle(ContentBundle);
 }
 
 void FContentBundleContainer::OnContentBundleClientUnregistered(FContentBundleClient& ContentBundleClient)
@@ -389,8 +390,7 @@ void FContentBundleContainer::InitializeContentBundlesForegisteredClients()
 
 		for (TSharedPtr<FContentBundleClient>& ContentBundleClient : ContentBundleClients)
 		{
-			FContentBundleBase& ContentBundle = InitializeContentBundle(ContentBundleClient);
-			InjectContentBundle(ContentBundle);
+			OnContentBundleClientRegistered(ContentBundleClient);
 		}
 
 		UE_LOG(LogContentBundle, Verbose, TEXT("%s End initializing ContentBundles."), *ContentBundle::Log::MakeDebugInfoString(*this));
