Add template overloads for FStreamableHandle::GetLoadedAsset and GetLoadedAssets for convenience

#rnx

diff --git a/Engine/Source/Runtime/Engine/Classes/Engine/StreamableManager.h b/Engine/Source/Runtime/Engine/Classes/Engine/StreamableManager.h
--- a/Engine/Source/Runtime/Engine/Classes/Engine/StreamableManager.h
+++ b/Engine/Source/Runtime/Engine/Classes/Engine/StreamableManager.h
@@ -219,9 +219,27 @@ struct FStreamableHandle : public TSharedFromThis<FStreamableHandle>
 	/** Adds all loaded assets if load has succeeded. Some entries will be null if loading failed */
 	ENGINE_API void GetLoadedAssets(TArray<UObject *>& LoadedAssets) const;
 
+	/** Templated version of above. Entries that fail to cast will also be null. */
+	template<class T>
+	void GetLoadedAssets(TArray<T*>& LoadedAssets) const
+	{
+		ForEachLoadedAsset([&LoadedAssets](UObject* LoadedAsset)
+		{
+			LoadedAssets.Add(Cast<T>(LoadedAsset));
+		});
+	}
+
 	/** Returns first asset in requested asset list, if it's been successfully loaded. This will fail if the asset failed to load */
 	ENGINE_API UObject* GetLoadedAsset() const;
 
+	/** Templated version of above */
+	template<class T>
+	T* GetLoadedAsset() const
+	{
+		UObject* LoadedAsset = GetLoadedAsset();
+		return Cast<T>(LoadedAsset);
+	}
+
 	/** Returns number of assets that have completed loading out of initial list, failed loads will count as loaded */
 	ENGINE_API void GetLoadedCount(int32& LoadedCount, int32& RequestedCount) const;
 
