compilation fix: missing header

diff --git a/Engine/Source/Runtime/Engine/Private/EdGraph/EdGraphNode.cpp b/Engine/Source/Runtime/Engine/Private/EdGraph/EdGraphNode.cpp
--- a/Engine/Source/Runtime/Engine/Private/EdGraph/EdGraphNode.cpp
+++ b/Engine/Source/Runtime/Engine/Private/EdGraph/EdGraphNode.cpp
@@ -10,6 +10,7 @@
 #include "UObject/UE5ReleaseStreamObjectVersion.h"
 #include "Textures/SlateIcon.h"
 #include "EngineLogs.h"
+#include "HAL/IConsoleManager.h"
 #if WITH_EDITOR
 #include "CookerSettings.h"
 #include "Kismet2/BlueprintEditorUtils.h"
