Initial implementation of asynchronous UAnimBank DDC2 data building

diff --git a/Engine/Source/Runtime/Engine/Classes/Animation/AnimBank.h b/Engine/Source/Runtime/Engine/Classes/Animation/AnimBank.h
--- a/Engine/Source/Runtime/Engine/Classes/Animation/AnimBank.h
+++ b/Engine/Source/Runtime/Engine/Classes/Animation/AnimBank.h
@@ -2,10 +2,95 @@
 
 #pragma once
 
-#include "CoreMinimal.h"
 #include "AnimationAsset.h"
+#include "Interfaces/Interface_AsyncCompilation.h"
+#include "IO/IoHash.h"
+#include "Logging/LogMacros.h"
+#include "RenderCommandFence.h"
+
 #include "AnimBank.generated.h"
 
+ENGINE_API DECLARE_LOG_CATEGORY_EXTERN(LogAnimBank, Log, All);
+
+class FQueuedThreadPool;
+enum class EQueuedWorkPriority : uint8;
+struct FPropertyChangedEvent;
+
+struct FReferenceSkeleton;
+class USkinnedAsset;
+class FAnimBankBuildAsyncCacheTask;
+
+struct FSkinnedAssetMapping
+{
+	// Bone transforms in global pose.
+	TArray<FTransform> MeshGlobalRefPose;
+	TArray<FTransform> AnimGlobalRefPose;
+
+	// A map to go from the mesh skeleton bone index to anim skeleton bone index.
+	TArray<int32> MeshToAnimIndexMap;
+
+	// Retargeting table to go from the anim skeleton to the mesh skeleton.
+	TArray<TTuple<FQuat, FQuat>> RetargetingTable;
+
+	// Inverse global space transforms
+	TArray<FVector3f> PositionKeys;
+	TArray<FQuat4f> RotationKeys;
+
+	uint32 BoneCount = 0;
+};
+
+inline FArchive& operator<<(FArchive& Ar, FSkinnedAssetMapping& AssetMapping)
+{
+	Ar << AssetMapping.MeshGlobalRefPose;
+	Ar << AssetMapping.AnimGlobalRefPose;
+	Ar << AssetMapping.MeshToAnimIndexMap;
+	Ar << AssetMapping.RetargetingTable;
+	Ar << AssetMapping.PositionKeys;
+	Ar << AssetMapping.RotationKeys;
+	Ar << AssetMapping.BoneCount;
+	return Ar;
+}
+
+struct FAnimBankEntry
+{
+	TArray<FVector3f>	PositionKeys;
+	TArray<FQuat4f>		RotationKeys;
+	TArray<FVector3f>	ScalingKeys;
+
+	float Position		= 0.0f;
+	float PlayRate		= 1.0f;
+
+	uint32 FrameCount	= 0u;
+	uint32 KeyCount		= 0u;
+	uint32 Flags		= 0u;
+};
+
+inline FArchive& operator<<(FArchive& Ar, FAnimBankEntry& BankEntry)
+{
+	Ar << BankEntry.PositionKeys;
+	Ar << BankEntry.RotationKeys;
+	Ar << BankEntry.ScalingKeys;
+	Ar << BankEntry.Position;
+	Ar << BankEntry.PlayRate;
+	Ar << BankEntry.FrameCount;
+	Ar << BankEntry.KeyCount;
+	Ar << BankEntry.Flags;
+	return Ar;
+}
+
+struct FAnimBankData
+{
+	FSkinnedAssetMapping Mapping;
+	TArray<FAnimBankEntry> Entries;
+};
+
+inline FArchive& operator<<(FArchive& Ar, FAnimBankData& BankData)
+{
+	Ar << BankData.Mapping;
+	Ar << BankData.Entries;
+	return Ar;
+}
+
 USTRUCT(BlueprintType)
 struct FAnimBankItem
 {
@@ -39,26 +124,120 @@ struct FAnimBankItem
 };
 
 UCLASS(BlueprintType, hidecategories=Object, editinlinenew, MinimalAPI)
-class UAnimBank : public UAnimationAsset
+class UAnimBank : public UAnimationAsset, public IInterface_AsyncCompilation
 {
-	GENERATED_UCLASS_BODY()
+	GENERATED_BODY()
 
+private:
+#if WITH_EDITOR
+	DECLARE_MULTICAST_DELEGATE(FOnGPUDataChanged);
+#endif
+
+public:
+#if WITH_EDITORONLY_DATA
 	UPROPERTY(EditAnywhere, BlueprintReadWrite, Category = Sequences, meta = (ShowOnlyInnerProperties))
 	TArray<struct FAnimBankItem> Items;
 
+	UPROPERTY(EditAnywhere, BlueprintReadWrite, Category = Mapping, meta = (ShowOnlyInnerProperties))
+	TObjectPtr<USkinnedAsset> Asset;
+#endif
+
 public:
-	//~ Begin UObject Interface
-	ENGINE_API virtual void Serialize(FArchive& Ar) override;
-	ENGINE_API virtual void PostLoad() override;
+	virtual void Serialize(FArchive& Ar) override;
+	virtual void PostLoad() override;
 	virtual void PreSave(FObjectPreSaveContext ObjectSaveContext) override;
-	//~ End UObject Interface
+	virtual void BeginDestroy() override;
+	virtual bool IsReadyForFinishDestroy() override;
+	virtual bool NeedsLoadForTargetPlatform(const ITargetPlatform* TargetPlatform) const override;
+
+	void InitResources();
+	void ReleaseResources();
+
+	inline const FAnimBankData& GetData() const
+	{
+		// TODO: Temp
+		return Data;
+	}
 
 #if WITH_EDITOR
-	ENGINE_API virtual void PostEditChangeProperty(FPropertyChangedEvent& PropertyChangedEvent) override;
+	virtual void PostEditChangeProperty(FPropertyChangedEvent& PropertyChangedEvent) override;
+	virtual void BeginCacheForCookedPlatformData(const ITargetPlatform* TargetPlatform) override;
+	virtual bool IsCachedCookedPlatformDataLoaded(const ITargetPlatform* TargetPlatform) override;
+	virtual void ClearAllCachedCookedPlatformData() override;
 
-	//~ Begin AnimationAsset interface
 	virtual bool GetAllAnimationSequencesReferred(TArray<UAnimationAsset*>& AnimationAssets, bool bRecursive = true) override;
 	virtual void ReplaceReferredAnimations(const TMap<UAnimationAsset*, UAnimationAsset*>& ReplacementMap) override;
-	//~ End AnimationAsset interface
+
+	/** Returns whether or not the asset is currently being compiled */
+	bool IsCompiling() const override;
+
+	/** Try to cancel any pending async tasks.
+	*  Returns true if there is no more async tasks pending, false otherwise.
+	*/
+	bool TryCancelAsyncTasks();
+
+	/** Returns false if there is currently an async task running */
+	bool IsAsyncTaskComplete() const;
+
+	/**
+	* Wait until all async tasks are complete, up to a time limit
+	* Returns true if all tasks are completed
+	**/
+	bool WaitForAsyncTasks(float TimeLimitSeconds);
+
+	/** Make sure all async tasks are completed before returning */
+	void FinishAsyncTasks();
+
+	typedef FOnGPUDataChanged::FDelegate FOnRebuild;
+	FDelegateHandle RegisterOnGPUDataChanged(const FOnRebuild& Delegate);
+	void UnregisterOnGPUDataChanged(void* Unregister);
+	void UnregisterOnGPUDataChanged(FDelegateHandle Handle);
+	void NotifyOnGPUDataChanged();
+#endif
+
+private:
+#if WITH_EDITOR
+	friend class FAnimBankCompilingManager;
+	void Reschedule(FQueuedThreadPool* InThreadPool, EQueuedWorkPriority InPriority);
+
+	FIoHash CreateDerivedDataKeyHash(const ITargetPlatform* TargetPlatform);
+	FIoHash BeginCacheDerivedData(const ITargetPlatform* TargetPlatform);
+	bool PollCacheDerivedData(const FIoHash& KeyHash) const;
+	void EndCacheDerivedData(const FIoHash& KeyHash);
+
+	/** Synchronously cache and return derived data for the target platform. */
+	FAnimBankData& CacheDerivedData(const ITargetPlatform* TargetPlatform);
+#endif
+
+private:
+	bool bIsInitialized = false;
+
+	FAnimBankData Data;
+	FRenderCommandFence ReleaseResourcesFence;
+
+#if WITH_EDITOR
+	FIoHash DataKeyHash;
+	TMap<FIoHash, TUniquePtr<FAnimBankData>> DataByPlatformKeyHash;
+	TMap<FIoHash, TPimplPtr<FAnimBankBuildAsyncCacheTask>> CacheTasksByKeyHash;
+
+	FOnGPUDataChanged OnGPUDataChanged;
+
+	DECLARE_EVENT_OneParam(UAnimBank, FOnDependenciesChanged, UAnimBank*);
+	static FOnDependenciesChanged OnDependenciesChanged;
 #endif
 };
+
+namespace UE::AnimBank
+{
+
+// Convert a list of transforms from bone/local space to mesh/global space by walking through the
+// hierarchy of a reference skeleton.
+ENGINE_API void ConvertLocalToGlobalSpaceTransforms(
+	const FReferenceSkeleton& InRefSkeleton,
+	const TArray<FTransform>& InLocalSpaceTransforms,
+	TArray<FTransform>& OutGlobalSpaceTransforms
+);
+
+ENGINE_API void BuildSkinnedAssetMapping(const USkinnedAsset& Asset, FSkinnedAssetMapping& Mapping);
+
+}
\ No newline at end of file
diff --git a/Engine/Source/Runtime/Engine/Classes/Animation/AnimSequence.h b/Engine/Source/Runtime/Engine/Classes/Animation/AnimSequence.h
--- a/Engine/Source/Runtime/Engine/Classes/Animation/AnimSequence.h
+++ b/Engine/Source/Runtime/Engine/Classes/Animation/AnimSequence.h
@@ -948,4 +948,6 @@ public:
 	friend class UE::Anim::FAnimSequenceCompilingManager;
 	friend UE::AnimNext::FDecompressionTools;
 	friend class FAnimSequenceDetails;
+	friend class UAnimBank;
+	friend class FAnimBankBuildAsyncCacheTask;
 };
diff --git a/Engine/Source/Runtime/Engine/Private/Animation/AnimBank.cpp b/Engine/Source/Runtime/Engine/Private/Animation/AnimBank.cpp
--- a/Engine/Source/Runtime/Engine/Private/Animation/AnimBank.cpp
+++ b/Engine/Source/Runtime/Engine/Private/Animation/AnimBank.cpp
@@ -1,28 +1,581 @@
 // Copyright Epic Games, Inc. All Rights Reserved.
 
 #include "Animation/AnimBank.h"
-#include "UObject/FrameworkObjectVersion.h"
-#include "UObject/UE5MainStreamObjectVersion.h"
-#include "UObject/ObjectSaveContext.h"
-#include "UObject/UObjectThreadContext.h"
-#include "EngineUtils.h"
+#include "Animation/AnimBank.h"
+#include "Animation/AnimBoneCompressionCodec.h"
+#include "Animation/AnimSequence.h"
+#include "Animation/AnimSequenceDecompressionContext.h"
+#include "AnimationRuntime.h"
 #include "AnimationUtils.h"
-#include "Logging/MessageLog.h"
+#include "AnimEncoding.h"
+#include "Async/ParallelFor.h"
 #include "DerivedDataCacheInterface.h"
+#include "EngineUtils.h"
 #include "Interfaces/ITargetPlatform.h"
 #include "Interfaces/ITargetPlatformManagerModule.h"
+#include "Logging/MessageLog.h"
+#include "Misc/App.h"
+#include "ReferenceSkeleton.h"
+#include "RenderUtils.h"
+#include "SkeletalRenderPublic.h"
+#include "UObject/FrameworkObjectVersion.h"
+#include "UObject/ObjectSaveContext.h"
+#include "UObject/UE5MainStreamObjectVersion.h"
+#include "UObject/UObjectThreadContext.h"
 
 #include UE_INLINE_GENERATED_CPP_BY_NAME(AnimBank)
 
+#if WITH_EDITOR
+#include "Animation/AnimationSequenceCompiler.h"
+#include "Animation/AnimBankCompiler.h"
+#include "Async/AsyncWork.h"
+#include "DerivedDataCache.h"
+#include "DerivedDataRequestOwner.h"
+#include "Experimental/Misc/ExecutionResource.h"
+#include "Serialization/MemoryHasher.h"
+#include "Serialization/MemoryReader.h"
+#include "Serialization/MemoryWriter.h"
+#endif
+
 #define LOCTEXT_NAMESPACE "AnimBank"
 
 CSV_DECLARE_CATEGORY_MODULE_EXTERN(ENGINE_API, Animation);
 
-UAnimBank::UAnimBank(const FObjectInitializer& ObjectInitializer)
-: Super(ObjectInitializer)
+DEFINE_LOG_CATEGORY(LogAnimBank);
+
+#if WITH_EDITOR
+
+UAnimBank::FOnDependenciesChanged UAnimBank::OnDependenciesChanged;
+
+static inline void DecomposeTransform(const FTransform& Transform, FVector3f& OutTranslation, FQuat4f& OutRotation)
+{
+	// Get Translation
+	OutTranslation = (FVector3f)Transform.GetTranslation();
+
+	// Get Rotation 
+	OutRotation = (FQuat4f)Transform.GetRotation();
+}
+
+static inline void DecomposeTransforms(
+	const TArray<FTransform>& Transforms,
+	TArray<FVector3f>& OutTranslations,
+	TArray<FQuat4f>& OutRotations)
+{
+	const int32 NumTransforms = Transforms.Num();
+	for (int32 Index = 0; Index < NumTransforms; ++Index)
+	{
+		FVector3f Translation;
+		FQuat4f Rotation;
+		DecomposeTransform(Transforms[Index], Translation, Rotation);
+		OutTranslations.Emplace(MoveTemp(Translation));
+		OutRotations.Emplace(MoveTemp(Rotation));
+	}
+}
+
+class FAnimBankAsyncBuildWorker : public FNonAbandonableTask
+{
+	FAnimBankBuildAsyncCacheTask* Owner;
+	FIoHash IoHash;
+
+public:
+	FAnimBankAsyncBuildWorker(FAnimBankBuildAsyncCacheTask* InOwner, const FIoHash& InIoHash)
+	: Owner(InOwner)
+	, IoHash(InIoHash)
+	{
+	}
+
+	FORCEINLINE TStatId GetStatId() const
+	{
+		RETURN_QUICK_DECLARE_CYCLE_STAT(FAnimBankAsyncBuildWorker, STATGROUP_ThreadPoolAsyncTasks);
+	}
+
+	void DoWork();
+};
+
+struct FAnimBankAsyncBuildTask : public FAsyncTask<FAnimBankAsyncBuildWorker>
+{
+	FAnimBankAsyncBuildTask(FAnimBankBuildAsyncCacheTask* InOwner, const FIoHash& InIoHash)
+	: FAsyncTask<FAnimBankAsyncBuildWorker>(InOwner, InIoHash)
+	{
+	}
+};
+
+class FAnimBankBuildAsyncCacheTask
+{
+public:
+	FAnimBankBuildAsyncCacheTask(
+		const FIoHash& InKeyHash,
+		FAnimBankData* InData,
+		UAnimBank& InDisplacedMesh,
+		const ITargetPlatform* TargetPlatform
+	);
+
+	inline void Wait()
+	{
+		if (bIsWaitingOnCompilation)
+		{
+			WaitForDependenciesAndBeginCache();
+		}
+
+		if (BuildTask != nullptr)
+		{
+			BuildTask->EnsureCompletion();
+		}
+
+		Owner.Wait();
+	}
+
+	inline bool WaitWithTimeout(float TimeLimitSeconds)
+	{
+		if (bIsWaitingOnCompilation)
+		{
+			if (!WaitForDependenciesAndBeginCacheWithTimeout(TimeLimitSeconds))
+			{
+				return false;
+			}
+		}
+
+		if (BuildTask != nullptr && !BuildTask->WaitCompletionWithTimeout(TimeLimitSeconds))
+		{
+			return false;
+		}
+
+		return Owner.Poll();
+	}
+
+	inline bool Poll()
+	{
+		if (bIsWaitingOnCompilation)
+		{
+			BeginCacheIfDependenciesAreFree();
+			return false;
+		}
+
+		if (BuildTask && !BuildTask->IsDone())
+		{
+			return false;
+		}
+
+		return Owner.Poll();
+	}
+
+	inline void Cancel()
+	{
+		// Cancel the waiting on the build
+		bIsWaitingOnCompilation = false;
+
+		if (BuildTask)
+		{
+			BuildTask->Cancel();
+		}
+
+		Owner.Cancel();
+	}
+
+	void Reschedule(FQueuedThreadPool* InThreadPool, EQueuedWorkPriority InPriority);
+
+private:
+	bool ShouldWaitForCompilation() const;
+
+	void BeginCacheIfDependenciesAreFree();
+	void WaitForDependenciesAndBeginCache();
+	bool WaitForDependenciesAndBeginCacheWithTimeout(float TimeLimitSeconds);
+
+	void BeginCache(const FIoHash& KeyHash);
+	void EndCache(UE::DerivedData::FCacheGetValueResponse&& Response);
+	bool BuildData(const UE::FSharedString& Name, const UE::DerivedData::FCacheKey& Key);
+
+private:
+	friend class FAnimBankAsyncBuildWorker;
+	TUniquePtr<FAnimBankAsyncBuildTask> BuildTask;
+	FAnimBankData* Data;
+	TWeakObjectPtr<UAnimBank> WeakAnimBank;
+
+	TArray<FAnimBankItem> BankItems;
+
+	UE::DerivedData::FRequestOwner Owner;
+	TRefCountPtr<IExecutionResource> ExecutionResource;
+
+	bool bIsWaitingOnCompilation;
+	FIoHash KeyHash;
+};
+
+void FAnimBankAsyncBuildWorker::DoWork()
+{
+	using namespace UE::DerivedData;
+	if (UAnimBank* Bank = Owner->WeakAnimBank.Get())
+	{
+		// Grab any execution resources currently assigned to this worker so that we maintain
+		// concurrency limit and memory pressure until the whole multi-step task is done.
+		Owner->ExecutionResource = FExecutionResourceContext::Get();
+
+		static const FCacheBucket Bucket("AnimBank");
+		GetCache().GetValue({ {{Bank->GetPathName()}, {Bucket, IoHash}} }, Owner->Owner,
+			[Task = Owner](FCacheGetValueResponse&& Response)
+			{
+				Task->EndCache(MoveTemp(Response));
+			}
+		);
+	}
+}
+
+FAnimBankBuildAsyncCacheTask::FAnimBankBuildAsyncCacheTask(
+	const FIoHash& InKeyHash,
+	FAnimBankData* InData,
+	UAnimBank& InBank,
+	const ITargetPlatform* TargetPlatform
+)
+: Data(InData)
+, WeakAnimBank(&InBank)
+, BankItems(InBank.Items)
+// Once we pass the BeginCache throttling gate, we want to finish as fast as possible
+// to avoid holding on to memory for a long time. We use the high priority since it will go fast,
+// but also it will avoid starving the critical threads in the subsequent task.
+, Owner(UE::DerivedData::EPriority::High)
+, bIsWaitingOnCompilation(ShouldWaitForCompilation())
+, KeyHash(InKeyHash)
+{
+	/**
+	 * Unfortunately our async builds are not made to handle the assets that use data from other assets
+	 * This will delay the start of the actual cache until the build of the sequences is done
+	 * This will fix a race condition with the sequence build without blocking the game thread by default.
+	 * Note: This is not a perfect solution since it also delays the DDC data pull.
+	 */
+	if (!bIsWaitingOnCompilation)
+	{
+		BeginCache(InKeyHash);
+	}
+}
+
+bool FAnimBankBuildAsyncCacheTask::ShouldWaitForCompilation() const 
+{
+	if (UAnimBank* Bank = WeakAnimBank.Get())
+	{
+		for (FAnimBankItem& BankItem : Bank->Items)
+		{
+			if (!BankItem.Sequence)
+			{
+				continue;
+			}
+
+			// If the sequence is still waiting for a post load call, let it build its stuff first to avoid blocking the Game Thread
+			if (BankItem.Sequence->HasAnyFlags(RF_NeedPostLoad) || BankItem.Sequence->IsCompiling())
+			{
+				return true;
+			}
+		}
+	}
+
+	return false;
+}
+
+void FAnimBankBuildAsyncCacheTask::BeginCacheIfDependenciesAreFree()
+{
+	if (UAnimBank* Bank = WeakAnimBank.Get())
+	{
+		if (!ShouldWaitForCompilation())
+		{
+			bIsWaitingOnCompilation = false;
+			BeginCache(KeyHash);
+		}
+	}
+	else
+	{
+		bIsWaitingOnCompilation = false;
+	}
+}
+
+void FAnimBankBuildAsyncCacheTask::WaitForDependenciesAndBeginCache()
+{
+	if (UAnimBank* Bank = WeakAnimBank.Get())
+	{
+		for (FAnimBankItem& BankItem : Bank->Items)
+		{
+			if (!BankItem.Sequence)
+			{
+				continue;
+			}
+
+			if (BankItem.Sequence->HasAnyFlags(RF_NeedPostLoad))
+			{
+				BankItem.Sequence->ConditionalPostLoad();
+			}
+
+			UE::Anim::FAnimSequenceCompilingManager::Get().FinishCompilation({ BankItem.Sequence.Get() });
+		}
+
+		bIsWaitingOnCompilation = false;
+		BeginCache(KeyHash);
+	}
+	else
+	{
+		bIsWaitingOnCompilation = false;
+	}
+}
+
+bool FAnimBankBuildAsyncCacheTask::WaitForDependenciesAndBeginCacheWithTimeout(float TimeLimitSeconds)
+{
+	if (UAnimBank* Bank = WeakAnimBank.Get())
+	{
+		for (FAnimBankItem& BankItem : Bank->Items)
+		{
+			if (!BankItem.Sequence || !BankItem.Sequence->IsCompiling())
+			{
+				continue;
+			}
+
+			if (!BankItem.Sequence->WaitForAsyncTasks(TimeLimitSeconds))
+			{
+				return false;
+			}
+		}
+	}
+
+	// Performs any necessary cleanup now that the async task (if any) is complete
+	WaitForDependenciesAndBeginCache();
+
+	return true;
+}
+
+void FAnimBankBuildAsyncCacheTask::BeginCache(const FIoHash& InKeyHash)
 {
+	using namespace UE::DerivedData;
+
+	if (UAnimBank* Bank = WeakAnimBank.Get())
+	{
+		// Queue this launch through the thread pool so that we benefit from fair scheduling and memory throttling
+		FQueuedThreadPool* ThreadPool = FAnimBankCompilingManager::Get().GetThreadPool();
+		EQueuedWorkPriority BasePriority = FAnimBankCompilingManager::Get().GetBasePriority(Bank);
+
+		int64 RequiredMemory = -1; // @todo RequiredMemory
+
+		check(BuildTask == nullptr);
+		BuildTask = MakeUnique<FAnimBankAsyncBuildTask>(this, InKeyHash);
+		BuildTask->StartBackgroundTask(ThreadPool, BasePriority, EQueuedWorkFlags::DoNotRunInsideBusyWait, RequiredMemory, TEXT("AnimBank"));
+	}
 }
 
+void FAnimBankBuildAsyncCacheTask::EndCache(UE::DerivedData::FCacheGetValueResponse&& Response)
+{
+	using namespace UE::DerivedData;
+
+	if (Response.Status == EStatus::Ok)
+	{
+		Owner.LaunchTask(TEXT("AnimBankSerialize"), [this, Value = MoveTemp(Response.Value)]
+		{
+			// Release execution resource as soon as the task is done
+			ON_SCOPE_EXIT { ExecutionResource = nullptr; };
+
+			if (UAnimBank* Bank = WeakAnimBank.Get())
+			{
+				FSharedBuffer RecordData = Value.GetData().Decompress();
+				FMemoryReaderView Ar(RecordData, /*bIsPersistent*/ true);
+				Ar << *Data;
+
+				// The initialization of the resources is done by FAnimBankCompilingManager to avoid race conditions
+			}
+		});
+	}
+	else if (Response.Status == EStatus::Error)
+	{
+		Owner.LaunchTask(TEXT("AnimBankBuild"), [this, Name = Response.Name, Key = Response.Key]
+		{
+			// Release execution resource as soon as the task is done
+			ON_SCOPE_EXIT { ExecutionResource = nullptr; };
+
+			if (!BuildData(Name, Key))
+			{
+				return;
+			}
+
+			if (UAnimBank* Bank = WeakAnimBank.Get())
+			{
+				TArray64<uint8> RecordData;
+				FMemoryWriter64 Ar(RecordData, /*bIsPersistent*/ true);
+				Ar << *Data;
+
+				GetCache().PutValue({ {Name, Key, FValue::Compress(MakeSharedBufferFromArray(MoveTemp(RecordData)))} }, Owner);
+
+				// The initialization of the resources is done by FAnimBankCompilingManager to avoid race conditions
+			}
+		});
+	}
+	else
+	{
+		// Release execution resource as soon as the task is done
+		ExecutionResource = nullptr;
+	}
+}
+
+bool FAnimBankBuildAsyncCacheTask::BuildData(const UE::FSharedString& Name, const UE::DerivedData::FCacheKey& Key)
+{
+	using namespace UE::DerivedData;
+
+	UAnimBank* AnimBank = WeakAnimBank.Get();
+	if (!AnimBank)
+	{
+		return false;
+	}
+
+	TRACE_CPUPROFILER_EVENT_SCOPE(FAnimBankBuildAsyncCacheTask::BuildData);
+
+	*Data = {};
+
+	if (!IsValid(AnimBank->Asset))
+	{
+		UE_LOG(LogAnimBank, Error, TEXT("Cannot find a valid skinned asset to build the animation bank asset."));
+		return false;
+	}
+
+	UE::AnimBank::BuildSkinnedAssetMapping(*AnimBank->Asset.Get(), Data->Mapping);
+
+	const FReferenceSkeleton& MeshRefSkeleton = AnimBank->Asset->GetRefSkeleton();
+	const FReferenceSkeleton& AnimRefSkeleton = AnimBank->Asset->GetSkeleton()->GetReferenceSkeleton();
+
+	// Get the skeleton reference pose in local space
+	const TArray<FTransform>& MeshLocalRefPose = MeshRefSkeleton.GetRawRefBonePose();
+	const TArray<FTransform>& AnimLocalRefPose = AnimRefSkeleton.GetRawRefBonePose();
+
+	const int32 NumMeshBones = MeshRefSkeleton.GetRawBoneNum();
+	const int32 NumAnimBones = AnimRefSkeleton.GetRawBoneNum();
+
+	// Scratch memory
+	BoneTrackArray TrackToBoneIndexMap;
+
+	Data->Entries.SetNum(AnimBank->Items.Num());
+	for (int32 ItemIndex = 0; ItemIndex < AnimBank->Items.Num(); ++ItemIndex)
+	{
+		FAnimBankEntry& BankEntry = Data->Entries[ItemIndex];
+
+		UAnimSequence* Sequence = AnimBank->Items[ItemIndex].Sequence.Get();
+		if (Sequence == nullptr)
+		{
+			continue;
+		}
+
+		Sequence->FinishAsyncTasks();
+
+		check(!Sequence->DoesNeedRecompress());
+		check(Sequence->IsCompressedDataValid() && Sequence->CompressedData.CompressedTrackToSkeletonMapTable.Num() > 0);
+
+		const FCompressedAnimSequence& SampledSequence = Sequence->CompressedData;
+
+		BankEntry.KeyCount = Sequence->GetNumberOfSampledKeys();
+
+		// Set up mapping tables for the decompressor to map internal tracks to the pose array (which is in bone order).
+		const int32 NumTracks = Sequence->CompressedData.CompressedTrackToSkeletonMapTable.Num();
+
+		TrackToBoneIndexMap.Reset(NumTracks);
+		for (int32 TrackIndex = 0; TrackIndex < NumTracks; ++TrackIndex)
+		{
+			const int32 BoneIndex = SampledSequence.GetSkeletonIndexFromTrackIndex(TrackIndex);
+
+			// We only care about raw bones.
+			if (BoneIndex < NumAnimBones)
+			{
+				TrackToBoneIndexMap.Add(BoneTrackPair(BoneIndex, TrackIndex));
+			}
+		}
+
+		// Prep to uncompress a non-additive animation.
+		FAnimSequenceDecompressionContext DecompressionContext(
+			Sequence->GetSamplingFrameRate(),
+			Sequence->GetSamplingFrameRate().AsFrameTime(Sequence->GetPlayLength()).RoundToFrame().Value,
+			EAnimInterpolationType::Linear,
+			Sequence->GetRetargetTransformsSourceName(),
+			*SampledSequence.CompressedDataStructure,
+			AnimRefSkeleton.GetRefBonePose(),
+			SampledSequence.CompressedTrackToSkeletonMapTable,
+			nullptr,
+			false,
+			AAT_None
+		);
+
+		BankEntry.FrameCount = Sequence->GetSamplingFrameRate().AsFrameTime(Sequence->GetPlayLength()).RoundToFrame().Value;
+		check(BankEntry.FrameCount > 0);
+
+		// The sampled pose transforms from the sequence are in bone (parent) space. We need to convert them
+		// to local space (component space) to retarget to the mesh and create ref -> anim pose transform.  
+		TArray<FTransform> SampledLocalAnimPose = AnimRefSkeleton.GetRawRefBonePose();
+		TArray<FTransform> SampledLocalMeshPose;
+		TArray<FTransform> SampledGlobalMeshPose;
+
+		SampledLocalMeshPose.SetNumUninitialized(NumMeshBones);
+		SampledGlobalMeshPose.SetNumUninitialized(NumMeshBones);
+
+		BankEntry.PositionKeys.SetNumUninitialized(BankEntry.FrameCount * NumMeshBones);
+		BankEntry.RotationKeys.SetNumUninitialized(BankEntry.FrameCount * NumMeshBones);
+
+		int32 KeyIndex = 0;
+
+		for (uint32 Frame = 0; Frame < BankEntry.FrameCount; ++Frame)
+		{
+			// Some paths in the decompression code use mem stack, so make sure we put a mark here.
+			FMemMark Mark(FMemStack::Get());
+
+			const double SeekTime = Sequence->GetSamplingFrameRate().AsSeconds(FFrameTime(int32(Frame)));
+			DecompressionContext.Seek(SeekTime);
+
+			TArrayView<FTransform> SampledLocalPoseView(SampledLocalAnimPose);
+			SampledSequence.BoneCompressionCodec->DecompressPose(DecompressionContext, TrackToBoneIndexMap, TrackToBoneIndexMap, TrackToBoneIndexMap, SampledLocalPoseView);
+
+			// Retarget from the anim skeleton to the mesh skeleton.
+			for (int32 MeshBoneIndex = 0; MeshBoneIndex < NumMeshBones; ++MeshBoneIndex)
+			{
+				const int32 AnimBoneIndex = Data->Mapping.MeshToAnimIndexMap[MeshBoneIndex];
+				if (AnimBoneIndex != INDEX_NONE)
+				{
+					const TTuple<FQuat, FQuat>& RetargetingItem = Data->Mapping.RetargetingTable[MeshBoneIndex];
+					FTransform SourceToTargetTransform(
+						RetargetingItem.Get<0>() * AnimLocalRefPose[AnimBoneIndex].GetRotation() * RetargetingItem.Get<1>(),
+						RetargetingItem.Get<0>().RotateVector(AnimLocalRefPose[AnimBoneIndex].GetTranslation()),
+						FVector::OneVector
+					);
+			
+					SampledLocalMeshPose[MeshBoneIndex].SetRotation(SampledLocalAnimPose[AnimBoneIndex].GetRotation() * SourceToTargetTransform.GetRotation().Inverse() * MeshLocalRefPose[MeshBoneIndex].GetRotation());
+					SampledLocalMeshPose[MeshBoneIndex].SetTranslation(SampledLocalAnimPose[AnimBoneIndex].GetTranslation() + (MeshLocalRefPose[MeshBoneIndex].GetTranslation() - SourceToTargetTransform.GetTranslation()));
+					SampledLocalMeshPose[MeshBoneIndex].SetScale3D(FVector::OneVector);
+				}
+				else
+				{
+					SampledLocalMeshPose[MeshBoneIndex] = FTransform::Identity;
+				}
+			}
+		
+			// Convert local pose from the sequence to global pose (in the animation skeleton's space -- we retarget below to the mesh skeleton's space, as needed).
+			UE::AnimBank::ConvertLocalToGlobalSpaceTransforms(MeshRefSkeleton, SampledLocalMeshPose, SampledGlobalMeshPose);
+
+			for (int32 MeshBoneIndex = 0; MeshBoneIndex < NumMeshBones; ++MeshBoneIndex)
+			{
+				BankEntry.PositionKeys[KeyIndex] = FVector3f(SampledGlobalMeshPose[MeshBoneIndex].GetTranslation());
+				BankEntry.RotationKeys[KeyIndex] = FQuat4f(SampledGlobalMeshPose[MeshBoneIndex].GetRotation());
+				++KeyIndex;
+			}
+		}
+
+		check(BankEntry.PositionKeys.Num() > 0);
+		check(BankEntry.RotationKeys.Num() > 0);
+		check(BankEntry.FrameCount > 0);
+	}
+
+	if (Owner.IsCanceled())
+	{
+		return false;
+	}
+
+	return true;
+}
+
+void FAnimBankBuildAsyncCacheTask::Reschedule(FQueuedThreadPool* InThreadPool, EQueuedWorkPriority InPriority)
+{
+	if (BuildTask)
+	{
+		BuildTask->Reschedule(InThreadPool, InPriority);
+	}
+}
+
+#endif
+
 void UAnimBank::Serialize(FArchive& Ar)
 {
 	LLM_SCOPE(ELLMTag::Animation);
@@ -31,11 +584,50 @@ void UAnimBank::Serialize(FArchive& Ar)
 	Ar.UsingCustomVersion(FUE5MainStreamObjectVersion::GUID);
 
 	Super::Serialize(Ar);
+
+	if (Ar.IsFilterEditorOnly() && !Ar.IsObjectReferenceCollector() && !Ar.ShouldSkipBulkData())
+	{
+	#if WITH_EDITOR
+		if (Ar.IsCooking())
+		{
+			if (IsCompiling())
+			{
+				FAnimBankCompilingManager::Get().FinishCompilation({this});
+			}
+
+			FAnimBankData& CookedData = CacheDerivedData(Ar.CookingTarget());
+			Ar << CookedData;
+		}
+		else
+	#endif
+		{
+			Ar << Data;
+		}
+	}
 }
 
 void UAnimBank::PostLoad()
 {
 	Super::PostLoad();
+
+	if (FApp::CanEverRender())
+	{
+		// Only valid for cooked builds
+		if (Data.Entries.Num() > 0)
+		{
+			InitResources();
+		}
+	#if WITH_EDITOR
+		else if (ITargetPlatform* RunningPlatform = GetTargetPlatformManagerRef().GetRunningTargetPlatform())
+		{
+			BeginCacheDerivedData(RunningPlatform);
+		}
+	#endif
+	}
+
+#if WITH_EDITOR
+	OnDependenciesChanged.Broadcast(this);
+#endif
 }
 
 void UAnimBank::PreSave(FObjectPreSaveContext ObjectSaveContext)
@@ -43,6 +635,72 @@ void UAnimBank::PreSave(FObjectPreSaveContext ObjectSaveContext)
 	Super::PreSave(ObjectSaveContext);
 }
 
+void UAnimBank::BeginDestroy()
+{
+	ReleaseResources();
+
+#if WITH_EDITOR
+	OnDependenciesChanged.Broadcast(this);
+#endif
+
+	Super::BeginDestroy();
+}
+
+bool UAnimBank::IsReadyForFinishDestroy()
+{
+	if (!Super::IsReadyForFinishDestroy())
+	{
+		return false;
+	}
+
+#if WITH_EDITOR
+	if (!TryCancelAsyncTasks())
+	{
+		return false;
+	}
+#endif
+
+	return ReleaseResourcesFence.IsFenceComplete();
+}
+
+bool UAnimBank::NeedsLoadForTargetPlatform(const ITargetPlatform* TargetPlatform) const
+{
+	return DoesTargetPlatformSupportNanite(TargetPlatform);
+}
+
+void UAnimBank::InitResources()
+{
+	if (!FApp::CanEverRender())
+	{
+		return;
+	}
+
+	if (!bIsInitialized)
+	{
+		// TODO:
+		//Data.ResourcesPtr->InitResources(this);
+		bIsInitialized = true;
+	}
+}
+
+void UAnimBank::ReleaseResources()
+{
+	if (!bIsInitialized)
+	{
+		return;
+	}
+
+	// TODO:
+	//if (Data.ResourcesPtr->ReleaseResources())
+	{
+		// Make sure the renderer is done processing the command,
+		// and done using the GPU resources before we overwrite the data.
+		ReleaseResourcesFence.BeginFence();
+	}
+
+	bIsInitialized = false;
+}
+
 #if WITH_EDITOR
 
 void UAnimBank::PostEditChangeProperty(FPropertyChangedEvent& PropertyChangedEvent)
@@ -61,6 +719,54 @@ void UAnimBank::PostEditChangeProperty(FPropertyChangedEvent& PropertyChangedEve
 			// ...
 		}
 	}
+
+	OnDependenciesChanged.Broadcast(this);
+
+	// Synchronously build the new data. This calls InitResources.
+	ITargetPlatform* RunningPlatform = GetTargetPlatformManagerRef().GetRunningTargetPlatform();
+	BeginCacheDerivedData(RunningPlatform);
+}
+
+void UAnimBank::BeginCacheForCookedPlatformData(const ITargetPlatform* TargetPlatform)
+{
+	Super::BeginCacheForCookedPlatformData(TargetPlatform);
+	BeginCacheDerivedData(TargetPlatform);
+}
+
+bool UAnimBank::IsCachedCookedPlatformDataLoaded(const ITargetPlatform* TargetPlatform)
+{
+	const FIoHash KeyHash = CreateDerivedDataKeyHash(TargetPlatform);
+	if (KeyHash.IsZero())
+	{
+		return true;
+	}
+
+	if (PollCacheDerivedData(KeyHash))
+	{
+		EndCacheDerivedData(KeyHash);
+		return true;
+	}
+
+	return false;
+}
+
+void UAnimBank::ClearAllCachedCookedPlatformData()
+{
+	TRACE_CPUPROFILER_EVENT_SCOPE(UAnimBank::ClearAllCachedCookedPlatformData);
+
+	if (!TryCancelAsyncTasks())
+	{
+		FinishAsyncTasks();
+	}
+
+	/**
+	 * TryCancelAsyncTasks or FinishAsyncTasks should have been able to clear all tasks.If any tasks remain
+	 * then they must still be running, and we would crash when attempting to delete them.
+	*/
+	check(CacheTasksByKeyHash.IsEmpty());
+
+	DataByPlatformKeyHash.Empty();
+	Super::ClearAllCachedCookedPlatformData();
 }
 
 bool UAnimBank::GetAllAnimationSequencesReferred(TArray<UAnimationAsset*>& AnimationAssets, bool bRecursive)
@@ -106,6 +812,368 @@ void UAnimBank::ReplaceReferredAnimations(const TMap<UAnimationAsset*, UAnimatio
 	}
 }
 
+bool UAnimBank::IsCompiling() const
+{
+	return CacheTasksByKeyHash.Num() > 0;
+}
+
+bool UAnimBank::TryCancelAsyncTasks()
+{
+	bool bHadCachedTaskForRunningPlatform = CacheTasksByKeyHash.Contains(DataKeyHash);
+	
+	for (auto It = CacheTasksByKeyHash.CreateIterator(); It; ++It)
+	{
+		if (It->Value->Poll())
+		{
+			It.RemoveCurrent();
+		}
+		else
+		{
+			It->Value->Cancel();
+
+			// Try to see if we can remove the task now that it might have been canceled
+			if (It->Value->Poll())
+			{
+				It.RemoveCurrent();
+			}
+		}
+	}
+	
+
+	if (bHadCachedTaskForRunningPlatform && !CacheTasksByKeyHash.Contains(DataKeyHash))
+	{
+		// Reset the cached Key for the running platform since we won't have any GPU data
+		DataKeyHash = FIoHash();
+	}
+
+	return CacheTasksByKeyHash.IsEmpty();
+}
+
+bool UAnimBank::IsAsyncTaskComplete() const
+{
+	for (auto& Pair : CacheTasksByKeyHash)
+	{
+		if (!Pair.Value->Poll())
+		{
+			return false;
+		}
+	}
+
+	return true;
+}
+
+bool UAnimBank::WaitForAsyncTasks(float TimeLimitSeconds)
+{
+	double StartTimeSeconds = FPlatformTime::Seconds();
+	for (auto& Pair : CacheTasksByKeyHash)
+	{
+		// Clamp to 0 as it implies polling
+		const float TimeLimit = FMath::Min(0.0f, TimeLimitSeconds - (FPlatformTime::Seconds() - StartTimeSeconds));
+		if (!Pair.Value->WaitWithTimeout(TimeLimit))
+		{
+			return false;
+		}
+	}
+
+	return true;
+}
+
+void UAnimBank::FinishAsyncTasks()
+{
+	for (auto It = CacheTasksByKeyHash.CreateIterator(); It; ++It)
+	{
+		It->Value->Wait();
+		It.RemoveCurrent();
+	}
+}
+
+void UAnimBank::Reschedule(FQueuedThreadPool* InThreadPool, EQueuedWorkPriority InPriority)
+{
+	for (auto& Pair : CacheTasksByKeyHash)
+	{
+		Pair.Value->Reschedule(InThreadPool, InPriority);
+	}
+}
+
+FIoHash UAnimBank::CreateDerivedDataKeyHash(const ITargetPlatform* TargetPlatform)
+{
+	if (!DoesTargetPlatformSupportNanite(TargetPlatform))
+	{
+		return FIoHash::Zero;
+	}
+
+	FMemoryHasherBlake3 Writer;
+
+	FGuid AnimBankVersionGuid(0x2FA68926, 0x4EB94719, 0x9B928A29, 0xAE522B0F);
+	Writer << AnimBankVersionGuid;
+
+	if (IsValid(Asset))
+	{
+		FString AssetHash = Asset->BuildDerivedDataKey(TargetPlatform);
+		Writer << AssetHash;
+	}
+
+	for (int32 BankIndex = 0; BankIndex < Items.Num(); ++BankIndex)
+	{
+		Writer << BankIndex;
+
+		FAnimBankItem& BankItem = Items[BankIndex];
+		if (IsValid(BankItem.Sequence))
+		{
+			FIoHash SequenceHash = BankItem.Sequence->CreateDerivedDataKeyHash(TargetPlatform);
+			Writer << SequenceHash;
+		}
+
+		uint8 bLooping		= BankItem.bLooping ? 1 : 0;
+		uint8 bAutoStart	= BankItem.bAutoStart ? 1 : 0;
+
+		Writer << bLooping;
+		Writer << bAutoStart;
+		Writer << BankItem.Position;
+		Writer << BankItem.PlayRate;
+	}
+
+#if PLATFORM_CPU_ARM_FAMILY
+	// Separate out arm keys as x64 and arm64 clang do not generate the same data for a given
+	// input. Add the arm specifically so that a) we avoid rebuilding the current DDC and
+	// b) we can remove it once we get arm64 to be consistent.
+	FString ArmSuffix(TEXT("_arm64"));
+	Writer << ArmSuffix;
+#endif
+
+	return Writer.Finalize();
+}
+
+FIoHash UAnimBank::BeginCacheDerivedData(const ITargetPlatform* TargetPlatform)
+{
+	const FIoHash KeyHash = CreateDerivedDataKeyHash(TargetPlatform);
+
+	if (KeyHash.IsZero() || DataKeyHash == KeyHash || DataByPlatformKeyHash.Contains(KeyHash))
+	{
+		return KeyHash;
+	}
+
+	// Make sure we finish the previous build before starting another one
+	FAnimBankCompilingManager::Get().FinishCompilation({ this });
+
+	// Make sure the GPU is no longer referencing the current GPU resource data.
+	ReleaseResources();
+	ReleaseResourcesFence.Wait();
+	Data = {};
+
+	NotifyOnGPUDataChanged();
+
+	FAnimBankData* TargetData = nullptr;
+	if (TargetPlatform->IsRunningPlatform())
+	{
+		DataKeyHash = KeyHash;
+		TargetData = &Data;
+	}
+	else
+	{
+		TargetData = DataByPlatformKeyHash.Emplace(KeyHash, MakeUnique<FAnimBankData>()).Get();
+	}
+
+	for (FAnimBankItem& BankItem : Items)
+	{
+		if (!IsValid(BankItem.Sequence))
+		{
+			continue;
+		}
+
+		if (BankItem.Sequence->DoesNeedRecompress() || !BankItem.Sequence->IsCompressedDataValid())
+		{
+			BankItem.Sequence->BeginCacheDerivedData(TargetPlatform);
+		}
+	}
+
+	CacheTasksByKeyHash.Emplace(KeyHash, MakePimpl<FAnimBankBuildAsyncCacheTask>(KeyHash, TargetData, *this, TargetPlatform));
+
+	// The compiling manager provides throttling, notification manager, etc... for the asset being built.
+	FAnimBankCompilingManager::Get().AddAnimBanks({this});
+
+	return KeyHash;
+}
+
+bool UAnimBank::PollCacheDerivedData(const FIoHash& KeyHash) const
+{
+	if (KeyHash.IsZero())
+	{
+		return true;
+	}
+
+	if (const TPimplPtr<FAnimBankBuildAsyncCacheTask>* Task = CacheTasksByKeyHash.Find(KeyHash))
+	{
+		return (*Task)->Poll();
+	}
+
+	return true;
+}
+
+void UAnimBank::EndCacheDerivedData(const FIoHash& KeyHash)
+{
+	if (KeyHash.IsZero())
+	{
+		return;
+	}
+
+	TPimplPtr<FAnimBankBuildAsyncCacheTask> Task;
+	if (CacheTasksByKeyHash.RemoveAndCopyValue(KeyHash, Task))
+	{
+		Task->Wait();
+	}
+}
+
+FAnimBankData& UAnimBank::CacheDerivedData(const ITargetPlatform* TargetPlatform)
+{
+	const FIoHash KeyHash = BeginCacheDerivedData(TargetPlatform);
+	EndCacheDerivedData(KeyHash);
+	FAnimBankData& AnimBankData = (DataKeyHash == KeyHash) ? Data : *DataByPlatformKeyHash[KeyHash];
+	return AnimBankData;
+}
+
+FDelegateHandle UAnimBank::RegisterOnGPUDataChanged(const FOnRebuild& Delegate)
+{
+	return OnGPUDataChanged.Add(Delegate);
+}
+
+void UAnimBank::UnregisterOnGPUDataChanged(void* Unregister)
+{
+	OnGPUDataChanged.RemoveAll(Unregister);
+}
+
+void UAnimBank::UnregisterOnGPUDataChanged(FDelegateHandle Handle)
+{
+	OnGPUDataChanged.Remove(Handle);
+}
+
+void UAnimBank::NotifyOnGPUDataChanged()
+{
+	OnGPUDataChanged.Broadcast();
+}
+
 #endif
 
-#undef LOCTEXT_NAMESPACE 
\ No newline at end of file
+void FAnimBankItem::ValidatePosition()
+{
+	float Min = 0.0f, Max = 0.0f;
+
+	if (Sequence)
+	{
+		Max = Sequence->GetPlayLength();
+	}
+
+	Position = FMath::Clamp<float>(Position, Min, Max);
+}
+
+namespace UE::AnimBank
+{
+
+void ConvertLocalToGlobalSpaceTransforms(
+	const FReferenceSkeleton& InRefSkeleton,
+	const TArray<FTransform>& InLocalSpaceTransforms,
+	TArray<FTransform>& OutGlobalSpaceTransforms
+)
+{
+	OutGlobalSpaceTransforms.Reset();
+
+	if (!ensure(InRefSkeleton.GetRawBoneNum() == InLocalSpaceTransforms.Num()))
+	{
+		return;
+	}
+
+	const int32 NumTransforms = InLocalSpaceTransforms.Num();
+	OutGlobalSpaceTransforms.SetNumUninitialized(NumTransforms);
+
+	for (int32 BoneIndex = 0; BoneIndex < NumTransforms; ++BoneIndex)
+	{
+		// Initialize to identity since some of them don't have tracks
+		int32 IterBoneIndex = BoneIndex;
+		FTransform ComponentSpaceTransform = InLocalSpaceTransforms[BoneIndex];
+
+		do
+		{
+			const int32 ParentIndex = InRefSkeleton.GetRawParentIndex(IterBoneIndex); // Get only raw bones (no virtual)
+			if (ParentIndex != INDEX_NONE)
+			{
+				ComponentSpaceTransform *= InLocalSpaceTransforms[ParentIndex];
+			}
+
+			IterBoneIndex = ParentIndex;
+		}
+		while (InRefSkeleton.IsValidIndex(IterBoneIndex));
+
+		OutGlobalSpaceTransforms[BoneIndex] = ComponentSpaceTransform;
+	}
+}
+
+void BuildSkinnedAssetMapping(const USkinnedAsset& Asset, FSkinnedAssetMapping& Mapping)
+{
+	const FReferenceSkeleton& MeshRefSkeleton = Asset.GetRefSkeleton();
+	const FReferenceSkeleton& AnimRefSkeleton = Asset.GetSkeleton()->GetReferenceSkeleton();
+
+	const int32 NumMeshBones = MeshRefSkeleton.GetRawBoneNum();
+	const int32 NumAnimBones = AnimRefSkeleton.GetRawBoneNum();
+
+	// Get Number of RawBones (no virtual)
+	Mapping.BoneCount = AnimRefSkeleton.GetRawBoneNum();
+
+	// Get the skeleton reference pose in local space;
+	const TArray<FTransform>& MeshLocalRefPose = MeshRefSkeleton.GetRawRefBonePose();
+	const TArray<FTransform>& AnimLocalRefPose = AnimRefSkeleton.GetRawRefBonePose();
+
+	// Get the bone transforms in global pose.
+	ConvertLocalToGlobalSpaceTransforms(MeshRefSkeleton, MeshRefSkeleton.GetRawRefBonePose(), Mapping.MeshGlobalRefPose);
+	ConvertLocalToGlobalSpaceTransforms(AnimRefSkeleton, AnimRefSkeleton.GetRawRefBonePose(), Mapping.AnimGlobalRefPose);
+
+	// A map to go from the mesh skeleton bone index to anim skeleton bone index.
+	Mapping.MeshToAnimIndexMap.SetNumUninitialized(NumMeshBones);
+	for (int32 MeshBoneIndex = 0; MeshBoneIndex < NumMeshBones; ++MeshBoneIndex)
+	{
+		Mapping.MeshToAnimIndexMap[MeshBoneIndex] = AnimRefSkeleton.FindRawBoneIndex(MeshRefSkeleton.GetBoneName(MeshBoneIndex));
+	}
+
+	// Construct a retargeting table to go from the anim skeleton to the mesh skeleton.
+	Mapping.RetargetingTable.SetNumUninitialized(NumMeshBones);
+	Mapping.RetargetingTable[0] = MakeTuple(FQuat::Identity, AnimRefSkeleton.GetRawRefBonePose()[0].GetRotation().Inverse() * MeshRefSkeleton.GetRawRefBonePose()[0].GetRotation());
+	for (int32 MeshBoneIndex = 1; MeshBoneIndex < NumMeshBones; ++MeshBoneIndex)
+	{
+		const int32 AnimBoneIndex = Mapping.MeshToAnimIndexMap[MeshBoneIndex];
+		if (AnimBoneIndex != INDEX_NONE)
+		{
+			const int32 AnimParentIndex = AnimRefSkeleton.GetParentIndex(AnimBoneIndex);
+			const int32 MeshParentIndex = MeshRefSkeleton.GetParentIndex(MeshBoneIndex);
+			check(AnimParentIndex != INDEX_NONE);
+			check(MeshParentIndex != INDEX_NONE);
+
+			const FQuat Pa = Mapping.AnimGlobalRefPose[AnimParentIndex].GetRotation();
+			const FQuat Pm = Mapping.MeshGlobalRefPose[MeshParentIndex].GetRotation();
+
+			const FQuat Ra = AnimLocalRefPose[AnimBoneIndex].GetRotation();
+			const FQuat Rm = MeshLocalRefPose[MeshBoneIndex].GetRotation();
+
+			const FQuat Q0 = Pm.Inverse() * Pa;
+			const FQuat Q1 = Ra.Inverse() * Pa.Inverse() * Pm * Rm;
+
+			Mapping.RetargetingTable[MeshBoneIndex] = MakeTuple(Q0, Q1);
+		}
+		else
+		{
+			Mapping.RetargetingTable[MeshBoneIndex] = MakeTuple(FQuat::Identity, FQuat::Identity);
+		}
+	}
+
+	Mapping.PositionKeys.AddUninitialized(NumMeshBones);
+	Mapping.RotationKeys.AddUninitialized(NumMeshBones);
+	for (int32 MeshBoneIndex = 0; MeshBoneIndex < NumMeshBones; ++MeshBoneIndex)
+	{
+		const FTransform InvMeshGlobalRefPoseXform = Mapping.MeshGlobalRefPose[MeshBoneIndex].Inverse();
+		Mapping.PositionKeys[MeshBoneIndex] = FVector3f(InvMeshGlobalRefPoseXform.GetTranslation());
+		Mapping.RotationKeys[MeshBoneIndex] = FQuat4f(InvMeshGlobalRefPoseXform.GetRotation());
+	}
+}
+
+} // UE::AnimBank
+
+#undef LOCTEXT_NAMESPACE
+
diff --git a/Engine/Source/Runtime/Engine/Private/Animation/AnimBankCompiler.cpp b/Engine/Source/Runtime/Engine/Private/Animation/AnimBankCompiler.cpp
new file mode 100644
--- /dev/null
+++ b/Engine/Source/Runtime/Engine/Private/Animation/AnimBankCompiler.cpp
@@ -0,0 +1,421 @@
+// Copyright Epic Games, Inc. All Rights Reserved.
+
+#include "AnimBankCompiler.h"
+#include "UObject/Package.h"
+
+#if WITH_EDITOR
+
+#include "AsyncCompilationHelpers.h"
+#include "AssetCompilingManager.h"
+#include "Algo/NoneOf.h"
+#include "Misc/QueuedThreadPool.h"
+#include "ObjectCacheContext.h"
+#include "Misc/QueuedThreadPoolWrapper.h"
+#include "Misc/App.h"
+#include "Animation/AnimBank.h"
+#include "ProfilingDebugging/CountersTrace.h"
+#include "AnimationSequenceCompiler.h"
+
+#define LOCTEXT_NAMESPACE "AnimBankCompiler"
+
+static AsyncCompilationHelpers::FAsyncCompilationStandardCVars CVarAsyncAnimBankStandard(
+	TEXT("AnimBank"),
+	TEXT("animation banks"),
+	FConsoleCommandDelegate::CreateLambda(
+	[]()
+	{
+		FAnimBankCompilingManager::Get().FinishAllCompilation();
+	}
+));
+
+namespace AnimBankCompilingManagerImpl
+{
+	static void EnsureInitializedCVars()
+	{
+		static bool bIsInitialized = false;
+
+		if (!bIsInitialized)
+		{
+			bIsInitialized = true;
+
+			AsyncCompilationHelpers::EnsureInitializedCVars(
+				TEXT("AnimBank"),
+				CVarAsyncAnimBankStandard.AsyncCompilation,
+				CVarAsyncAnimBankStandard.AsyncCompilationMaxConcurrency);
+		}
+	}
+}
+
+FAnimBankCompilingManager::FAnimBankCompilingManager()
+: Notification(MakeUnique<FAsyncCompilationNotification>(GetAssetNameFormat()))
+{
+	AnimBankCompilingManagerImpl::EnsureInitializedCVars();
+	PostReachabilityAnalysisHandle = FCoreUObjectDelegates::PostReachabilityAnalysis.AddRaw(this, &FAnimBankCompilingManager::OnPostReachabilityAnalysis);
+}
+
+void FAnimBankCompilingManager::OnPostReachabilityAnalysis()
+{
+	if (GetNumRemainingAssets())
+	{
+		TRACE_CPUPROFILER_EVENT_SCOPE(FAnimBankCompilingManager::CancelUnreachableMeshes);
+
+		TArray<UAnimBank*> PendingAnimBanks;
+		PendingAnimBanks.Reserve(GetNumRemainingAssets());
+
+		for (auto Iterator = RegisteredAnimBanks.CreateIterator(); Iterator; ++Iterator)
+		{
+			UAnimBank* AnimBank = Iterator->GetEvenIfUnreachable();
+			if (AnimBank && AnimBank->IsUnreachable())
+			{
+				UE_LOG(LogAnimBank, Verbose, TEXT("Cancelling animation bank %s compilation because it's being garbage collected"), *AnimBank->GetName());
+
+				if (AnimBank->TryCancelAsyncTasks())
+				{
+					Iterator.RemoveCurrent();
+				}
+				else
+				{
+					PendingAnimBanks.Add(AnimBank);
+				}
+			}
+		}
+
+		FinishCompilation(PendingAnimBanks);
+	}
+}
+
+FName FAnimBankCompilingManager::GetAssetTypeName() const
+{
+	return TEXT("UE-AnimBank");
+}
+
+FTextFormat FAnimBankCompilingManager::GetAssetNameFormat() const
+{
+	return LOCTEXT("AnimBankNameFormat", "{0}|plural(one=Animation Bank,other=Animation Banks)");
+}
+
+TArrayView<FName> FAnimBankCompilingManager::GetDependentTypeNames() const
+{
+	static FName DependentTypeNames[] =	
+	{ 
+		// AnimBank can wait on AnimSequence to finish their own compilation before compiling itself
+		// so they need to be processed before us. This is especially important when FinishAllCompilation is issued
+		// so that we know once we're called that all anim sequences have finished compiling.
+		UE::Anim::FAnimSequenceCompilingManager::GetStaticAssetTypeName() 
+	};
+	return TArrayView<FName>(DependentTypeNames);
+}
+
+EQueuedWorkPriority FAnimBankCompilingManager::GetBasePriority(UAnimBank* InAnimBank) const
+{
+	return EQueuedWorkPriority::Low;
+}
+
+FQueuedThreadPool* FAnimBankCompilingManager::GetThreadPool() const
+{
+	static FQueuedThreadPoolDynamicWrapper* GAnimBankThreadPool = nullptr;
+	if (GAnimBankThreadPool == nullptr && FAssetCompilingManager::Get().GetThreadPool() != nullptr)
+	{
+		// Animation banks will be scheduled on the asset thread pool, where concurrency limits might by dynamically adjusted depending on memory constraints.
+		GAnimBankThreadPool = new FQueuedThreadPoolDynamicWrapper(FAssetCompilingManager::Get().GetThreadPool(), -1, [](EQueuedWorkPriority) { return EQueuedWorkPriority::Low; });
+
+		AsyncCompilationHelpers::BindThreadPoolToCVar(
+			GAnimBankThreadPool,
+			CVarAsyncAnimBankStandard.AsyncCompilation,
+			CVarAsyncAnimBankStandard.AsyncCompilationResume,
+			CVarAsyncAnimBankStandard.AsyncCompilationMaxConcurrency
+		);
+	}
+
+	return GAnimBankThreadPool;
+}
+
+void FAnimBankCompilingManager::Shutdown()
+{
+	bHasShutdown = true;
+	if (GetNumRemainingAssets())
+	{
+		check(IsInGameThread());
+		TRACE_CPUPROFILER_EVENT_SCOPE(FAnimBankCompilingManager::Shutdown);
+
+		TArray<UAnimBank*> PendingAnimBanks;
+		PendingAnimBanks.Reserve(GetNumRemainingAssets());
+
+		for (TWeakObjectPtr<UAnimBank>& WeakAnimBank : RegisteredAnimBanks)
+		{
+			if (WeakAnimBank.IsValid())
+			{
+				UAnimBank* AnimBank = WeakAnimBank.Get();
+				if (!AnimBank->TryCancelAsyncTasks())
+				{
+					PendingAnimBanks.Add(AnimBank);
+				}
+			}
+		}
+
+		FinishCompilation(PendingAnimBanks);
+	}
+
+	FCoreUObjectDelegates::PostReachabilityAnalysis.Remove(PostReachabilityAnalysisHandle);
+}
+
+TRACE_DECLARE_INT_COUNTER(QueuedAnimBankCompilation, TEXT("AsyncCompilation/QueuedAnimBank"));
+void FAnimBankCompilingManager::UpdateCompilationNotification()
+{
+	TRACE_COUNTER_SET(QueuedAnimBankCompilation, GetNumRemainingAssets());
+	Notification->Update(GetNumRemainingAssets());
+}
+
+void FAnimBankCompilingManager::PostCompilation(TArrayView<UAnimBank* const> InAnimBanks)
+{
+	if (InAnimBanks.Num())
+	{
+		TRACE_CPUPROFILER_EVENT_SCOPE(OnAssetPostCompileEvent);
+
+		TArray<FAssetCompileData> AssetsData;
+		AssetsData.Reserve(InAnimBanks.Num());
+
+		for (UAnimBank* AnimBank : InAnimBanks)
+		{
+			// Do not broadcast an event for unreachable objects
+			if (!AnimBank->IsUnreachable())
+			{
+				AssetsData.Emplace(AnimBank);
+
+				if (FApp::CanEverRender())
+				{
+					AnimBank->InitResources();
+					AnimBank->NotifyOnGPUDataChanged();
+				}
+			}
+		}
+
+		if (AssetsData.Num())
+		{
+			FAssetCompilingManager::Get().OnAssetPostCompileEvent().Broadcast(AssetsData);
+		}
+	}
+}
+
+void FAnimBankCompilingManager::PostCompilation(UAnimBank* AnimBank)
+{
+	using namespace AnimBankCompilingManagerImpl;
+
+	// If AsyncTask is null here, the task got canceled so we don't need to do anything
+	if (!IsEngineExitRequested())
+	{
+		check(IsInGameThread());
+		TRACE_CPUPROFILER_EVENT_SCOPE(PostCompilation);
+
+		AnimBank->FinishAsyncTasks();
+
+		// Do not do anything else if the AnimBank is being garbage collected
+		if (AnimBank->IsUnreachable())
+		{
+			return;
+		}
+
+		AnimBank->InitResources();
+
+		// Calling this delegate during app exit might be quite dangerous and lead to crash
+		// if the content browser wants to refresh a thumbnail it might try to load a package
+		// which will then fail due to various reasons related to the editor shutting down.
+		// Triggering this callback while garbage collecting can also result in listeners trying to look up objects
+		if (!GExitPurge && !IsGarbageCollecting())
+		{
+			// Generate an empty property changed event, to force the asset registry tag
+			// to be refreshed now that RenderData is available.
+			FPropertyChangedEvent EmptyPropertyChangedEvent(nullptr);
+			FCoreUObjectDelegates::OnObjectPropertyChanged.Broadcast(AnimBank, EmptyPropertyChangedEvent);
+		}
+	}
+}
+
+FAnimBankCompilingManager& FAnimBankCompilingManager::Get()
+{
+	static FAnimBankCompilingManager Singleton;
+	return Singleton;
+}
+
+int32 FAnimBankCompilingManager::GetNumRemainingAssets() const
+{
+	return RegisteredAnimBanks.Num();
+}
+
+void FAnimBankCompilingManager::AddAnimBanks(TArrayView<UAnimBank* const> InAnimBanks)
+{
+	TRACE_CPUPROFILER_EVENT_SCOPE(FAnimBankCompilingManager::AddAnimBanks);
+	check(IsInGameThread());
+
+	for (UAnimBank* AnimBank : InAnimBanks)
+	{
+		RegisteredAnimBanks.Emplace(AnimBank);
+	}
+
+	TRACE_COUNTER_SET(QueuedAnimBankCompilation, GetNumRemainingAssets());
+}
+
+void FAnimBankCompilingManager::FinishCompilation(TArrayView<UAnimBank* const> InAnimBanks)
+{
+	TRACE_CPUPROFILER_EVENT_SCOPE(FAnimBankCompilingManager::FinishCompilation);
+
+	// Allow calls from any thread if the meshes are already finished compiling.
+	if (Algo::NoneOf(InAnimBanks, &UAnimBank::IsCompiling))
+	{
+		return;
+	}
+
+	check(IsInGameThread());
+
+	TArray<UAnimBank*> PendingAnimBanks;
+	PendingAnimBanks.Reserve(InAnimBanks.Num());
+
+	int32 AnimBankIndex = 0;
+	for (UAnimBank* AnimBank : InAnimBanks)
+	{
+		if (RegisteredAnimBanks.Contains(AnimBank))
+		{
+			PendingAnimBanks.Emplace(AnimBank);
+		}
+	}
+
+	if (PendingAnimBanks.Num())
+	{
+		class FCompilableAnimBank : public AsyncCompilationHelpers::ICompilable
+		{
+		public:
+			FCompilableAnimBank(UAnimBank* InAnimBank)
+			: AnimBank(InAnimBank)
+			{
+			}
+
+			void Reschedule(FQueuedThreadPool* InThreadPool, EQueuedWorkPriority InPriority)
+			{
+				AnimBank->Reschedule(InThreadPool, InPriority);
+			}
+
+			bool WaitCompletionWithTimeout(float TimeLimitSeconds) override
+			{
+				return AnimBank->WaitForAsyncTasks(TimeLimitSeconds);
+			}
+
+			UAnimBank* AnimBank;
+			FName GetName() override { return AnimBank->GetOutermost()->GetFName(); }
+		};
+
+		TArray<FCompilableAnimBank> CompilableAnimBanks(PendingAnimBanks);
+		FObjectCacheContextScope ObjectCacheScope;
+		AsyncCompilationHelpers::FinishCompilation(
+			[&CompilableAnimBanks](int32 Index)	-> AsyncCompilationHelpers::ICompilable& { return CompilableAnimBanks[Index]; },
+			CompilableAnimBanks.Num(),
+			LOCTEXT("AnimBanks", "Animation Banks"),
+			LogAnimBank,
+			[this](AsyncCompilationHelpers::ICompilable* Object)
+			{
+				UAnimBank* AnimBank = static_cast<FCompilableAnimBank*>(Object)->AnimBank;
+				PostCompilation(AnimBank);
+				RegisteredAnimBanks.Remove(AnimBank);
+			}
+		);
+
+		PostCompilation(PendingAnimBanks);
+	}
+}
+
+void FAnimBankCompilingManager::FinishCompilationsForGame()
+{
+	// Nothing special to do when we PIE for now.
+}
+
+void FAnimBankCompilingManager::FinishAllCompilation()
+{
+	check(IsInGameThread());
+	TRACE_CPUPROFILER_EVENT_SCOPE(FAnimBankCompilingManager::FinishAllCompilation)
+
+	if (GetNumRemainingAssets())
+	{
+		TArray<UAnimBank*> PendingAnimBanks;
+		PendingAnimBanks.Reserve(GetNumRemainingAssets());
+
+		for (TWeakObjectPtr<UAnimBank>& AnimBank : RegisteredAnimBanks)
+		{
+			if (AnimBank.IsValid())
+			{
+				PendingAnimBanks.Add(AnimBank.Get());
+			}
+		}
+
+		FinishCompilation(PendingAnimBanks);
+	}
+}
+
+void FAnimBankCompilingManager::Reschedule()
+{
+	// TODO Prioritize animation banks that are nearest to the viewport
+}
+
+void FAnimBankCompilingManager::ProcessAnimBanks(bool bLimitExecutionTime, int32 MinBatchSize)
+{
+	using namespace AnimBankCompilingManagerImpl;
+	TRACE_CPUPROFILER_EVENT_SCOPE(FAnimBankCompilingManager::ProcessAnimBanks);
+
+	const int32 NumRemainingBanks = GetNumRemainingAssets();
+	
+	// Spread out the load over multiple frames but if too many banks, convergence is more important than frame time
+	const int32 MaxBankUpdatesPerFrame = bLimitExecutionTime ? FMath::Max(64, NumRemainingBanks / 10) : INT32_MAX;
+
+	FObjectCacheContextScope ObjectCacheScope;
+	if (NumRemainingBanks && NumRemainingBanks >= MinBatchSize)
+	{
+		TSet<UAnimBank*> AnimBanksToProcess;
+		for (TWeakObjectPtr<UAnimBank>& AnimBank : RegisteredAnimBanks)
+		{
+			if (AnimBank.IsValid())
+			{
+				AnimBanksToProcess.Add(AnimBank.Get());
+			}
+		}
+
+		{
+			TRACE_CPUPROFILER_EVENT_SCOPE(ProcessFinishedAnimBanks);
+
+			const double TickStartTime = FPlatformTime::Seconds();
+
+			TSet<TWeakObjectPtr<UAnimBank>> AnimBanksToPostpone;
+			TArray<UAnimBank*> ProcessedAnimBanks;
+			
+			if (AnimBanksToProcess.Num())
+			{
+				for (UAnimBank* AnimBank : AnimBanksToProcess)
+				{
+					const bool bHasBankUpdateLeft = ProcessedAnimBanks.Num() <= MaxBankUpdatesPerFrame;
+					if (bHasBankUpdateLeft && AnimBank->IsAsyncTaskComplete())
+					{
+						PostCompilation(AnimBank);
+						ProcessedAnimBanks.Add(AnimBank);
+					}
+					else
+					{
+						AnimBanksToPostpone.Emplace(AnimBank);
+					}
+				}
+			}
+
+			RegisteredAnimBanks = MoveTemp(AnimBanksToPostpone);
+			PostCompilation(ProcessedAnimBanks);
+		}
+	}
+}
+
+void FAnimBankCompilingManager::ProcessAsyncTasks(bool bLimitExecutionTime)
+{
+	FObjectCacheContextScope ObjectCacheScope;
+	FinishCompilationsForGame();
+	Reschedule();
+	ProcessAnimBanks(bLimitExecutionTime);
+	UpdateCompilationNotification();
+}
+
+#undef LOCTEXT_NAMESPACE
+
+#endif // #if WITH_EDITOR
diff --git a/Engine/Source/Runtime/Engine/Private/Animation/AnimBankCompiler.h b/Engine/Source/Runtime/Engine/Private/Animation/AnimBankCompiler.h
new file mode 100644
--- /dev/null
+++ b/Engine/Source/Runtime/Engine/Private/Animation/AnimBankCompiler.h
@@ -0,0 +1,85 @@
+// Copyright Epic Games, Inc. All Rights Reserved.
+#pragma once
+
+#if WITH_EDITOR
+
+#include "IAssetCompilingManager.h"
+#include "UObject/GCObject.h"
+#include "UObject/WeakObjectPtr.h"
+
+class FAsyncCompilationNotification;
+class UAnimBank;
+class FQueuedThreadPool;
+struct FAssetCompileContext;
+enum class EQueuedWorkPriority : uint8;
+
+class FAnimBankCompilingManager : public IAssetCompilingManager
+{
+public:
+	static FAnimBankCompilingManager& Get();
+
+	/**
+	* Returns the number of outstanding compilations.
+	*/
+	int32 GetNumRemainingAssets() const override;
+
+	/**
+	* Queue anim banks to be compiled asynchronously so they are monitored.
+	*/
+	void AddAnimBanks(TArrayView<UAnimBank* const> InAnimBanks);
+
+	/**
+	* Blocks until completion of the requested anim banks.
+	*/
+	void FinishCompilation(TArrayView<UAnimBank* const> InAnimBanks);
+
+	/**
+	* Blocks until completion of all async anim bank compilation.
+	*/
+	void FinishAllCompilation() override;
+
+	/**
+	* Returns the priority at which the given anim bank should be scheduled.
+	*/
+	EQueuedWorkPriority GetBasePriority(UAnimBank* InAnimBank) const;
+
+	/**
+	* Returns the thread pool where anim bank compilation should be scheduled.
+	*/
+	FQueuedThreadPool* GetThreadPool() const;
+
+	/**
+	* Cancel any pending work and blocks until it is safe to shut down.
+	*/
+	void Shutdown() override;
+
+private:
+	friend class FAssetCompilingManager;
+
+	FAnimBankCompilingManager();
+
+	FName GetAssetTypeName() const override;
+	FTextFormat GetAssetNameFormat() const override;
+	TArrayView<FName> GetDependentTypeNames() const override;
+
+	/** Called once per frame, fetches completed tasks and applies them to the scene. */
+	void ProcessAsyncTasks(bool bLimitExecutionTime = false) override;
+
+	bool bHasShutdown = false;
+	TSet<TWeakObjectPtr<UAnimBank>> RegisteredAnimBanks;
+
+	TUniquePtr<FAsyncCompilationNotification> Notification;
+
+	void FinishCompilationsForGame();
+	void Reschedule();
+	void ProcessAnimBanks(bool bLimitExecutionTime, int32 MinBatchSize = 1);
+	void UpdateCompilationNotification();
+
+	void PostCompilation(TArrayView<UAnimBank* const> InAnimBanks);
+	void PostCompilation(UAnimBank* InAnimBank);
+
+	void OnPostReachabilityAnalysis();
+	FDelegateHandle PostReachabilityAnalysisHandle;
+};
+
+#endif // WITH_EDITOR
diff --git a/Engine/Source/Runtime/Engine/Private/Animation/AnimationSequenceCompiler.cpp b/Engine/Source/Runtime/Engine/Private/Animation/AnimationSequenceCompiler.cpp
--- a/Engine/Source/Runtime/Engine/Private/Animation/AnimationSequenceCompiler.cpp
+++ b/Engine/Source/Runtime/Engine/Private/Animation/AnimationSequenceCompiler.cpp
@@ -84,11 +84,16 @@ namespace UE::Anim
 		FCoreUObjectDelegates::PostReachabilityAnalysis.Remove(PostReachabilityAnalysisHandle);
 	}
 
-	FName FAnimSequenceCompilingManager::GetAssetTypeName() const
+	FName FAnimSequenceCompilingManager::GetStaticAssetTypeName()
 	{
 		return TEXT("UE-AnimationSequence");
 	}
 
+	FName FAnimSequenceCompilingManager::GetAssetTypeName() const
+	{
+		return GetStaticAssetTypeName();
+	}
+
 	FTextFormat FAnimSequenceCompilingManager::GetAssetNameFormat() const
 	{
 		return LOCTEXT("AnimationSequenceNameFormat", "{0}|plural(one=Animation Sequence,other=Animation Sequences)");
@@ -321,7 +326,7 @@ namespace UE::Anim
 	{
 		if (GetNumRemainingAssets())
 		{
-			TRACE_CPUPROFILER_EVENT_SCOPE(FNaniteDisplacedMeshCompilingManager::CancelUnreachableMeshes);
+			TRACE_CPUPROFILER_EVENT_SCOPE(FAnimSequenceCompilingManager::CancelUnreachableMeshes);
 
 			TArray<UAnimSequence*> PendingAnimationSequences;
 			PendingAnimationSequences.Reserve(GetNumRemainingAssets());
diff --git a/Engine/Source/Runtime/Engine/Private/Animation/AnimationSequenceCompiler.h b/Engine/Source/Runtime/Engine/Private/Animation/AnimationSequenceCompiler.h
--- a/Engine/Source/Runtime/Engine/Private/Animation/AnimationSequenceCompiler.h
+++ b/Engine/Source/Runtime/Engine/Private/Animation/AnimationSequenceCompiler.h
@@ -22,6 +22,8 @@ namespace UE::Anim
 		static FAnimSequenceCompilingManager& Get();
 
 		FAnimSequenceCompilingManager();
+
+		static FName GetStaticAssetTypeName();
 	
 		virtual void FinishAllCompilation() override;
 		virtual void Shutdown() override;
diff --git a/Engine/Source/Runtime/Engine/Private/AssetCompilingManager.cpp b/Engine/Source/Runtime/Engine/Private/AssetCompilingManager.cpp
--- a/Engine/Source/Runtime/Engine/Private/AssetCompilingManager.cpp
+++ b/Engine/Source/Runtime/Engine/Private/AssetCompilingManager.cpp
@@ -23,6 +23,7 @@ LLM_DEFINE_TAG(AssetCompilation, NAME_None, NAME_None, GET_STATFNAME(STAT_AssetC
 #include "Algo/Find.h"
 #include "ProfilingDebugging/CountersTrace.h"
 #include "Animation/AnimationSequenceCompiler.h"
+#include "Animation/AnimBankCompiler.h"
 
 #if WITH_EDITOR
 #include "DerivedDataBuildSchedulerQueue.h"
@@ -402,6 +403,7 @@ FAssetCompilingManager::FAssetCompilingManager()
 	RegisterManager(&FActorDeferredScriptManager::Get());
 	RegisterManager(&FSoundWaveCompilingManager::Get());
 	RegisterManager(&UE::Anim::FAnimSequenceCompilingManager::Get());
+	RegisterManager(&FAnimBankCompilingManager::Get());
 
 	// Ensure that the thread pool is constructed before the memory queue is registered because
 	// the memory queue can call GetThreadPool() from a worker thread, and could lead to a race
