[World Partition] Follow-up for change 36260716: only include non-PIE editor-only actors in PIE, not for other editor streaming generation processes.

diff --git a/Engine/Source/Runtime/Engine/Private/WorldPartition/WorldPartitionStreamingGeneration.cpp b/Engine/Source/Runtime/Engine/Private/WorldPartition/WorldPartitionStreamingGeneration.cpp
--- a/Engine/Source/Runtime/Engine/Private/WorldPartition/WorldPartitionStreamingGeneration.cpp
+++ b/Engine/Source/Runtime/Engine/Private/WorldPartition/WorldPartitionStreamingGeneration.cpp
@@ -803,7 +803,7 @@ class FWorldPartitionStreamingGenerator
 			const bool bActorIsEditorOnly = InActorDescView.GetActorIsEditorOnly() || (InActor && !InActor->IsEditorOnly());
 			const bool bActorIsNonPIEEditorOnly = InActorDescView.GetActorIsNonPIEEditorOnly();
 
-			if (!bActorIsEditorOnly || bActorIsNonPIEEditorOnly)
+			if (!bActorIsEditorOnly || (bIsPIE && bActorIsNonPIEEditorOnly))
 			{
 				const FGuid ActorGuid = InActorDescView.GetGuid();
 				OutActorDescViewMap.Emplace(ActorGuid, MoveTemp(InActorDescView));
@@ -1690,6 +1690,7 @@ public:
 			, bEnableStreaming(false)
 			, bCreateContainerResolver(false)
 			, bHandleUnsavedActors(false)
+			, bIsPIE(false)
 		{}
 
 		const UWorldPartition* WorldPartitionContext;
@@ -1697,6 +1698,7 @@ public:
 		bool bEnableStreaming;
 		bool bCreateContainerResolver;
 		bool bHandleUnsavedActors;
+		bool bIsPIE;
 		TMap<FGuid, const UActorDescContainerInstance*> ActorGuidsToContainerInstanceMap;
 		TArray<TSubclassOf<AActor>> FilteredClasses;
 		TFunction<bool(FName, const UClass*)> IsValidGrid;
@@ -1707,6 +1709,7 @@ public:
 		FWorldPartitionStreamingGeneratorParams& SetEnableStreaming(bool bInEnableStreaming) { bEnableStreaming = bInEnableStreaming; return *this; }
 		FWorldPartitionStreamingGeneratorParams& SetCreateContainerResolver(bool bInCreateContainerResolver) { bCreateContainerResolver = bInCreateContainerResolver; return *this; }
 		FWorldPartitionStreamingGeneratorParams& SetHandleUnsavedActors(bool bInHandleUnsavedActors) { bHandleUnsavedActors = bInHandleUnsavedActors; return *this; }
+		FWorldPartitionStreamingGeneratorParams& SetIsPIE(bool bInIsPIE) { bIsPIE = bInIsPIE; return *this; }
 		FWorldPartitionStreamingGeneratorParams& SetActorGuidsToContainerInstanceMap(const TMap<FGuid, const UActorDescContainerInstance*>& InActorGuidsToContainerInstanceMap) { ActorGuidsToContainerInstanceMap = InActorGuidsToContainerInstanceMap; return *this; }
 		FWorldPartitionStreamingGeneratorParams& SetFilteredClasses(const TArray<TSubclassOf<AActor>>& InFilteredClasses) { FilteredClasses = InFilteredClasses; return *this; }
 		FWorldPartitionStreamingGeneratorParams& SetIsValidGrid(TFunction<bool(FName, const UClass*)> InIsValidGrid) { IsValidGrid = InIsValidGrid; return *this; }
@@ -1720,6 +1723,7 @@ public:
 		, bEnableStreaming(Params.bEnableStreaming)
 		, bCreateContainerResolver(Params.bCreateContainerResolver)
 		, bHandleUnsavedActors(Params.bHandleUnsavedActors)
+		, bIsPIE(Params.bIsPIE)
 		, FilteredClasses(Params.FilteredClasses)
 		, IsValidGrid(Params.IsValidGrid)
 		, IsValidHLODLayer(Params.IsValidHLODLayer)
@@ -1888,6 +1892,7 @@ private:
 	bool bEnableStreaming;
 	bool bCreateContainerResolver;
 	bool bHandleUnsavedActors;
+	bool bIsPIE;
 	TArray<TSubclassOf<AActor>> FilteredClasses;
 	TFunction<bool(FName, const UClass*)> IsValidGrid;
 	TFunction<bool(FName, const FSoftObjectPath&)> IsValidHLODLayer;
@@ -1961,6 +1966,7 @@ bool UWorldPartition::GenerateContainerStreaming(const FGenerateStreamingParams&
 	FWorldPartitionStreamingGenerator::FWorldPartitionStreamingGeneratorParams StreamingGeneratorParams = FWorldPartitionStreamingGenerator::FWorldPartitionStreamingGeneratorParams()
 		.SetWorldPartitionContext(this)
 		.SetHandleUnsavedActors(bIsPIE && !GetTypedOuter<UWorld>()->IsGameWorld())
+		.SetIsPIE(bIsPIE)
 		.SetIsValidGrid([this](FName GridName, const UClass* ActorClass) { return RuntimeHash->IsValidGrid(GridName, ActorClass); })
 		.SetIsValidHLODLayer([this](FName GridName, const FSoftObjectPath& HLODLayerPath) { return RuntimeHash->IsValidHLODLayer(GridName, HLODLayerPath); })
 		.SetErrorHandler(ErrorHandlerSelector.Get())
@@ -2363,6 +2369,7 @@ void UWorldPartition::CheckForErrors(const FCheckForErrorsParams& InParams)
 		FWorldPartitionStreamingGenerator::FWorldPartitionStreamingGeneratorParams StreamingGeneratorParams = FWorldPartitionStreamingGenerator::FWorldPartitionStreamingGeneratorParams()
 			.SetWorldPartitionContext(WorldPartition)
 			.SetHandleUnsavedActors(!!WorldPartition)
+			.SetIsPIE(WorldPartition ? WorldPartition->bIsPIE : false)
 			.SetErrorHandler(ErrorHandlerSelector.Get())
 			.SetIsValidGrid([WorldPartitionRuntimeHash](FName GridName, const UClass* ActorClass) { return WorldPartitionRuntimeHash ? WorldPartitionRuntimeHash->IsValidGrid(GridName, ActorClass) : true; })
 			.SetIsValidHLODLayer([WorldPartitionRuntimeHash](FName GridName, const FSoftObjectPath& HLODLayerPath) { return WorldPartitionRuntimeHash ? WorldPartitionRuntimeHash->IsValidHLODLayer(GridName, HLODLayerPath) : true; })
