Make sure to update level streaming even when screen is 0 size.
Fixes fallout of CL: 36254756 which prevents the screensize from being set to a valid value resulting in level streaming never updating

diff --git a/Engine/Source/Runtime/Engine/Private/UnrealClient.cpp b/Engine/Source/Runtime/Engine/Private/UnrealClient.cpp
--- a/Engine/Source/Runtime/Engine/Private/UnrealClient.cpp
+++ b/Engine/Source/Runtime/Engine/Private/UnrealClient.cpp
@@ -1815,6 +1815,11 @@ void FViewport::Draw( bool bShouldPresent /*= true */)
 
 				GInputLatencyTimer.GameThreadTrigger = false;
 			}
+			else if (World && World->IsGameWorld())
+			{
+				// since we aren't drawing the viewport, we still need to update streaming
+				World->UpdateLevelStreaming();
+			}
 		}
 
 		// Reset the camera cut flags if we are in a viewport that has a world
