[HWRT] Replaced checks in FRayTracingGeometryManager destructor with ensures.
- Failures should only happen during engine shutdown and need further investigation to indentify the issue, however can downgrade to an ensure to avoid crashing.

diff --git a/Engine/Source/Runtime/Engine/Private/Rendering/RayTracingGeometryManager.cpp b/Engine/Source/Runtime/Engine/Private/Rendering/RayTracingGeometryManager.cpp
--- a/Engine/Source/Runtime/Engine/Private/Rendering/RayTracingGeometryManager.cpp
+++ b/Engine/Source/Runtime/Engine/Private/Rendering/RayTracingGeometryManager.cpp
@@ -144,10 +144,10 @@ FRayTracingGeometryManager::FRayTracingGeometryManager()
 
 FRayTracingGeometryManager::~FRayTracingGeometryManager()
 {
-	check(GeometryBuildRequests.IsEmpty());
-	check(RegisteredGeometries.IsEmpty());
+	ensure(GeometryBuildRequests.IsEmpty());
+	ensure(RegisteredGeometries.IsEmpty());
 
-	check(RegisteredGroups.IsEmpty());
+	ensure(RegisteredGroups.IsEmpty());
 }
 
 static float GetInitialBuildPriority(ERTAccelerationStructureBuildPriority InBuildPriority)
