Reimport FBX does not correctly update frame-rate
#fix add RAII Reimport scope for scope-guarding bPopulated flag during reimport process

diff --git a/Engine/Plugins/Animation/AnimationData/Source/AnimationData/Private/AnimSequencerDataModel.h b/Engine/Plugins/Animation/AnimationData/Source/AnimationData/Private/AnimSequencerDataModel.h
--- a/Engine/Plugins/Animation/AnimationData/Source/AnimationData/Private/AnimSequencerDataModel.h
+++ b/Engine/Plugins/Animation/AnimationData/Source/AnimationData/Private/AnimSequencerDataModel.h
@@ -134,6 +134,11 @@ protected:
 	{
 		EvaluationLock.Unlock();
 	}
+
+	virtual bool& GetPopulationFlag() override final
+	{
+		return bPopulated;
+	}
 	/** End IAnimationDataModel overrides */
 	
 	/** Controller helper functionality */
diff --git a/Engine/Plugins/Importers/AlembicImporter/Source/AlembicLibrary/Private/AbcImporter.cpp b/Engine/Plugins/Importers/AlembicImporter/Source/AlembicLibrary/Private/AbcImporter.cpp
--- a/Engine/Plugins/Importers/AlembicImporter/Source/AlembicLibrary/Private/AbcImporter.cpp
+++ b/Engine/Plugins/Importers/AlembicImporter/Source/AlembicLibrary/Private/AbcImporter.cpp
@@ -703,6 +703,8 @@ TArray<UObject*> FAbcImporter::ImportAsSkeletalMesh(UObject* InParent, EObjectFl
 		Controller.OpenBracket(LOCTEXT("ImportAsSkeletalMesh", "Importing Alembic Animation"), bShouldTransact);
 		Controller.InitializeModel();
 
+		IAnimationDataModel::FReimportScope ReimportScope(Sequence->GetDataModel());
+
 		const FFrameRate FrameRate(FMath::RoundToInt(AbcFile->GetFramerate()), 1);
 		Controller.SetFrameRate(FrameRate, bShouldTransact);	
 		const FFrameNumber FrameNumber = FrameRate.AsFrameNumber(AbcFile->GetImportLength());
diff --git a/Engine/Plugins/Interchange/Runtime/Source/Import/Private/Animation/InterchangeAnimSequenceFactory.cpp b/Engine/Plugins/Interchange/Runtime/Source/Import/Private/Animation/InterchangeAnimSequenceFactory.cpp
--- a/Engine/Plugins/Interchange/Runtime/Source/Import/Private/Animation/InterchangeAnimSequenceFactory.cpp
+++ b/Engine/Plugins/Interchange/Runtime/Source/Import/Private/Animation/InterchangeAnimSequenceFactory.cpp
@@ -1421,6 +1421,8 @@ UInterchangeFactoryBase::FImportAssetResult UInterchangeAnimSequenceFactory::End
 		IAnimationDataController& Controller = AnimSequence->GetController();
 		Controller.OpenBracket(NSLOCTEXT("InterchangeAnimSequenceFactory", "ImportAnimationInterchange_Bracket", "Importing Animation (Interchange)"), bShouldTransact);
 		Controller.InitializeModel();
+		
+		IAnimationDataModel::FReimportScope ReimportScope(AnimSequence->GetDataModel());	
 		AnimSequence->ImportFileFramerate = SampleRate;
 		AnimSequence->ImportResampleFramerate = SampleRate;
 		Controller.SetFrameRate(FrameRate, bShouldTransact);
diff --git a/Engine/Source/Editor/UnrealEd/Private/SkeletalMeshEdit.cpp b/Engine/Source/Editor/UnrealEd/Private/SkeletalMeshEdit.cpp
--- a/Engine/Source/Editor/UnrealEd/Private/SkeletalMeshEdit.cpp
+++ b/Engine/Source/Editor/UnrealEd/Private/SkeletalMeshEdit.cpp
@@ -1646,6 +1646,7 @@ bool UnFbx::FFbxImporter::ImportAnimation(USkeleton* Skeleton, UAnimSequence * D
 	float PreviousSequenceLength = DestSeq->GetPlayLength();
 
 	const bool bShouldTransact = bReimport;
+	IAnimationDataModel::FReimportScope ReimportScope(DestSeq->GetDataModel());
 	IAnimationDataController& Controller = DestSeq->GetController();
 	Controller.OpenBracket(LOCTEXT("ImportAnimation_Bracket", "Importing Animation"), bShouldTransact);
 
diff --git a/Engine/Source/Runtime/Engine/Classes/Animation/AnimData/AnimDataModel.h b/Engine/Source/Runtime/Engine/Classes/Animation/AnimData/AnimDataModel.h
--- a/Engine/Source/Runtime/Engine/Classes/Animation/AnimData/AnimDataModel.h
+++ b/Engine/Source/Runtime/Engine/Classes/Animation/AnimData/AnimDataModel.h
@@ -83,7 +83,11 @@ protected:
 	virtual void OnNotify(const EAnimDataModelNotifyType& NotifyType, const FAnimDataModelNotifPayload& Payload) override {}
 	virtual void LockEvaluationAndModification() const override final {}
 	virtual bool TryLockEvaluationAndModification() const override final { return true; }
-	virtual void UnlockEvaluationAndModification() const override final {}
+	virtual void UnlockEvaluationAndModification() const override final {}	
+	virtual bool& GetPopulationFlag() override final
+	{
+		return bPopulated;
+	}
 	/** End IAnimationDataModel overrides */
 
 private:
diff --git a/Engine/Source/Runtime/Engine/Classes/Animation/AnimData/IAnimationDataModel.h b/Engine/Source/Runtime/Engine/Classes/Animation/AnimData/IAnimationDataModel.h
--- a/Engine/Source/Runtime/Engine/Classes/Animation/AnimData/IAnimationDataModel.h
+++ b/Engine/Source/Runtime/Engine/Classes/Animation/AnimData/IAnimationDataModel.h
@@ -128,6 +128,18 @@ class IAnimationDataModel
 public:
 	GENERATED_BODY()
 
+#if WITH_EDITOR
+	// RAII helper to allow repopulating data model without warnings during reimporting
+	struct FReimportScope
+	{
+		FReimportScope(IAnimationDataModel* InModel) : PopulationFlagGuard(InModel->GetPopulationFlag(), false)
+		{
+		}
+	private:
+		TGuardValue<bool> PopulationFlagGuard;
+	};
+#endif // WITH_EDITOR
+
 	/**
 	* @return	Total length of play-able animation data 
 	*/
@@ -383,7 +395,8 @@ protected:
 
 	virtual void LockEvaluationAndModification() const = 0;
 	virtual bool TryLockEvaluationAndModification() const = 0;
-	virtual void UnlockEvaluationAndModification() const = 0;	
+	virtual void UnlockEvaluationAndModification() const = 0;
+	virtual bool& GetPopulationFlag() = 0;
 	
 	struct FModelNotifier;
 	virtual IAnimationDataModel::FModelNotifier& GetNotifier() = 0;
