Expose SkyAtmosphereAmbientContributionColorScale to Blueprints as a UFUNCTION

#tests Tested in editor

diff --git a/Engine/Source/Runtime/Engine/Classes/Components/ExponentialHeightFogComponent.h b/Engine/Source/Runtime/Engine/Classes/Components/ExponentialHeightFogComponent.h
--- a/Engine/Source/Runtime/Engine/Classes/Components/ExponentialHeightFogComponent.h
+++ b/Engine/Source/Runtime/Engine/Classes/Components/ExponentialHeightFogComponent.h
@@ -310,6 +310,9 @@ public:
 	UFUNCTION(BlueprintCallable, Category = "Rendering")
 	ENGINE_API void SetRenderInMainPass(bool bValue);
 
+	UFUNCTION(BlueprintCallable, Category="Rendering|DirectionalInscattering")
+	ENGINE_API void SetSkyAtmosphereAmbientContributionColorScale(FLinearColor NewValue);
+
 protected:
 	//~ Begin UActorComponent Interface.
 	ENGINE_API virtual void CreateRenderState_Concurrent(FRegisterComponentContext* Context) override;
diff --git a/Engine/Source/Runtime/Engine/Private/Components/HeightFogComponent.cpp b/Engine/Source/Runtime/Engine/Private/Components/HeightFogComponent.cpp
--- a/Engine/Source/Runtime/Engine/Private/Components/HeightFogComponent.cpp
+++ b/Engine/Source/Runtime/Engine/Private/Components/HeightFogComponent.cpp
@@ -412,6 +412,15 @@ void UExponentialHeightFogComponent::SetRenderInMainPass(bool bValue)
 	}
 }
 
+void UExponentialHeightFogComponent::SetSkyAtmosphereAmbientContributionColorScale(FLinearColor NewValue)
+{
+	if (SkyAtmosphereAmbientContributionColorScale != NewValue)
+	{
+		SkyAtmosphereAmbientContributionColorScale = NewValue;
+		MarkRenderStateDirty();
+	}
+}
+
 void UExponentialHeightFogComponent::Serialize(FArchive& Ar)
 {
 	Super::Serialize(Ar);
