Expose additional properties to blueprints on SkyAtmosphereComponent. AerialPerspectiveStartDepth and TransmittanceMinLightElevationAngle.

#tests Tested in editor

diff --git a/Engine/Source/Runtime/Engine/Classes/Components/SkyAtmosphereComponent.h b/Engine/Source/Runtime/Engine/Classes/Components/SkyAtmosphereComponent.h
--- a/Engine/Source/Runtime/Engine/Classes/Components/SkyAtmosphereComponent.h
+++ b/Engine/Source/Runtime/Engine/Classes/Components/SkyAtmosphereComponent.h
@@ -218,8 +218,13 @@ class USkyAtmosphereComponent : public USceneComponent
 	UFUNCTION(BlueprintCallable, Category = "Rendering")
 	ENGINE_API void SetAerialPespectiveViewDistanceScale(float NewValue);
 	UFUNCTION(BlueprintCallable, Category = "Rendering")
+	ENGINE_API void SetAerialPerspectiveStartDepth(float NewValue);
+	UFUNCTION(BlueprintCallable, Category = "Rendering")
 	ENGINE_API void SetHeightFogContribution(float NewValue);
 
+	UFUNCTION(BlueprintCallable, Category = "Rendering")
+	ENGINE_API void SetTransmittanceMinLightElevationAngle(float NewValue);
+
 	UFUNCTION(BlueprintCallable, Category = "Rendering")
 	ENGINE_API void SetHoldout(bool bNewHoldout);
 
diff --git a/Engine/Source/Runtime/Engine/Private/Components/SkyAtmosphereComponent.cpp b/Engine/Source/Runtime/Engine/Private/Components/SkyAtmosphereComponent.cpp
--- a/Engine/Source/Runtime/Engine/Private/Components/SkyAtmosphereComponent.cpp
+++ b/Engine/Source/Runtime/Engine/Private/Components/SkyAtmosphereComponent.cpp
@@ -395,7 +395,9 @@ SKY_DECLARE_BLUEPRINT_SETFUNCTION_LINEARCOEFFICIENT(OtherAbsorption);
 
 SKY_DECLARE_BLUEPRINT_SETFUNCTION_LINEARCOEFFICIENT(SkyLuminanceFactor);
 SKY_DECLARE_BLUEPRINT_SETFUNCTION(float, AerialPespectiveViewDistanceScale);
+SKY_DECLARE_BLUEPRINT_SETFUNCTION(float, AerialPerspectiveStartDepth);
 SKY_DECLARE_BLUEPRINT_SETFUNCTION(float, HeightFogContribution);
+SKY_DECLARE_BLUEPRINT_SETFUNCTION(float, TransmittanceMinLightElevationAngle);
 
 void USkyAtmosphereComponent::SetHoldout(bool bNewHoldout)
 {
