Added trace region for RemoveFromWorld to better see the actual time it takes for a streaming level to be removed from the world (similar to AddToWorld).

diff --git a/Engine/Source/Runtime/Engine/Private/World.cpp b/Engine/Source/Runtime/Engine/Private/World.cpp
--- a/Engine/Source/Runtime/Engine/Private/World.cpp
+++ b/Engine/Source/Runtime/Engine/Private/World.cpp
@@ -3603,6 +3603,8 @@ void UWorld::RemoveFromWorld( ULevel* Level, bool bAllowIncrementalRemoval, FNet
 			bFinishRemovingLevel = false;
 			if (CurrentLevelPendingInvisibility == nullptr)
 			{
+				TRACE_BEGIN_REGION(*WriteToString<256>(TEXT("RemoveFromWorld: "), Level->GetPackage()->GetName()));
+
 				double StartTime = FPlatformTime::Seconds();
 				// Mark level as being the one in process of being made invisible. 
 				// This will prevent this level from being unloaded or made visible in the meantime
@@ -3746,6 +3748,8 @@ void UWorld::RemoveFromWorld( ULevel* Level, bool bAllowIncrementalRemoval, FNet
 				double DeltaTime = (FPlatformTime::Seconds() - StartTime) * 1000;
 				GRemoveFromWorldUnregisterComponentTimeCumul += DeltaTime;
 			}
+
+			TRACE_END_REGION(*WriteToString<256>(TEXT("RemoveFromWorld: "), Level->GetPackage()->GetName()));
 		} // if ( bFinishRemovingLevel )
 
 		Level->bIsDisassociatingLevel = false;
