[World Partition] Fix actor descriptor references for actors referenced by both editor-only and runtime properties, depending on the discovery order.
Note: this requires a resave of affected actors to take effect.

#prelight 66e471eded5a5f3b5ef0f6d1

diff --git a/Engine/Source/Runtime/Engine/Private/ActorReferencesUtils.cpp b/Engine/Source/Runtime/Engine/Private/ActorReferencesUtils.cpp
--- a/Engine/Source/Runtime/Engine/Private/ActorReferencesUtils.cpp
+++ b/Engine/Source/Runtime/Engine/Private/ActorReferencesUtils.cpp
@@ -45,7 +45,7 @@ public:
 		ArIsObjectReferenceCollector = true;
 		ArShouldSkipBulkData = true;
 
-		SubObjects.Add(Root);
+		SubObjectsState.Add(Root, EReferenceState::Runtime);
 
 		Root->Serialize(*this);
 	}
@@ -54,11 +54,13 @@ public:
 	{
 		if (Obj && !Obj->IsTemplate() && !Obj->HasAnyFlags(RF_Transient))
 		{
-			bool bWasAlreadyInSet;
-			SubObjects.Add(Obj, &bWasAlreadyInSet);
+			EReferenceState& ExistingReferenceState = SubObjectsState.FindOrAdd(Obj, EReferenceState::Unset);
+			const EReferenceState ReferenceState = EditorOnlyStack > 0 ? EReferenceState::EditorOnly : EReferenceState::Runtime;
 
-			if (!bWasAlreadyInSet)
+			if (ExistingReferenceState < ReferenceState)
 			{
+				ExistingReferenceState = ReferenceState;
+
 				HandleObjectReference(Obj);
 
 				if (bRecursive || Obj->IsInOuter(Root))
@@ -132,9 +134,16 @@ private:
 		}
 	}
 
+	enum class EReferenceState : uint8
+	{
+		Unset = 0,
+		EditorOnly = 1,
+		Runtime = 2
+	};
+
 	UObject* Root;
 	FActorReferenceSet& ActorReferences;
-	TSet<UObject*> SubObjects;
+	TMap<UObject*, EReferenceState> SubObjectsState;
 	int32 EditorOnlyStack;
 	EObjectFlags RequiredFlags;
 	bool bRecursive;
