Pipe IsSeekable to the Wave from the decoders AND Early AudioDevice Init
Early Init is required to startup the decoders ahead of the first CDO loads
#tests Smoked all shippable targets
[RN] Adding Early Audio Device Init and fixing SoundWave::IsSeekable to obey the Decoder

diff --git a/Engine/Source/Runtime/Engine/Classes/Sound/SoundSourceBus.h b/Engine/Source/Runtime/Engine/Classes/Sound/SoundSourceBus.h
--- a/Engine/Source/Runtime/Engine/Classes/Sound/SoundSourceBus.h
+++ b/Engine/Source/Runtime/Engine/Classes/Sound/SoundSourceBus.h
@@ -70,6 +70,7 @@ public:
 	virtual void FinishCachePlatformData() override {};	
 #endif	//WITH_EDITOR_ONLY_DATA
 	virtual void SerializeCookedPlatformData(FArchive& Ar) override {}
+	virtual bool IsSeekable() const override { return false; }
 	//~ End USoundWave Interface
 
 	//~ Begin USoundBase Interface.
diff --git a/Engine/Source/Runtime/Engine/Classes/Sound/SoundWave.h b/Engine/Source/Runtime/Engine/Classes/Sound/SoundWave.h
--- a/Engine/Source/Runtime/Engine/Classes/Sound/SoundWave.h
+++ b/Engine/Source/Runtime/Engine/Classes/Sound/SoundWave.h
@@ -1377,13 +1377,15 @@ public:
 	ENGINE_API bool IsStreaming(const TCHAR* PlatformName = nullptr) const;
 	ENGINE_API bool IsStreaming(const FPlatformAudioCookOverrides& Overrides) const;
 
-	/** Returns whether the sound is seekable. */
+	/**
+	 * Queries if Seeking is supported by the Decoder for this Wave.
+	 * @return true if Seeking is supported, false otherwise.
+	 */
 	ENGINE_API virtual bool IsSeekable() const;
 
 	/**
 	 * Checks whether we should use the load on demand cache.
 	 */
-
 	ENGINE_API bool ShouldUseStreamCaching() const;
 
 	/**
@@ -1632,6 +1634,13 @@ private:
 
 	ENGINE_API FName FindRuntimeFormat(const USoundWave&) const;
 
+	/**
+	 * Cache state that depends on the RuntimeFormat, like things that are dependent on
+	 * the decoder implementation on this platform.
+	 * @param InFormatName Factory Name of the Formats Decoder "ADPCM" etc.
+	 **/
+	ENGINE_API void CacheRuntimeFormatDependentState(const FName InFormatName);
+
 	/** Zeroth Chunk of audio for sources that use Load On Demand. */
 	FBulkDataBuffer<uint8> ZerothChunkData;
 #if WITH_EDITOR
diff --git a/Engine/Source/Runtime/Engine/Classes/Sound/SoundWaveProcedural.h b/Engine/Source/Runtime/Engine/Classes/Sound/SoundWaveProcedural.h
--- a/Engine/Source/Runtime/Engine/Classes/Sound/SoundWaveProcedural.h
+++ b/Engine/Source/Runtime/Engine/Classes/Sound/SoundWaveProcedural.h
@@ -73,6 +73,7 @@ public:
 	ENGINE_API virtual void InitAudioResource( FByteBulkData& CompressedData ) override;
 	ENGINE_API virtual bool InitAudioResource(FName Format) override;
 	ENGINE_API virtual int32 GetResourceSizeForFormat(FName Format) override;
+	ENGINE_API virtual bool IsSeekable() const override { return false; }
 	//~ End USoundWave Interface.
 
 	// Virtual function to generate PCM audio from the audio render thread. 
diff --git a/Engine/Source/Runtime/Engine/Private/AudioDeviceManager.cpp b/Engine/Source/Runtime/Engine/Private/AudioDeviceManager.cpp
--- a/Engine/Source/Runtime/Engine/Private/AudioDeviceManager.cpp
+++ b/Engine/Source/Runtime/Engine/Private/AudioDeviceManager.cpp
@@ -28,6 +28,12 @@
 #include "Settings/LevelEditorMiscSettings.h"
 #endif
 
+// ENTRYPOINT
+// AudioDeviceManager PreInit Callback, fired from Engine Startup Phase.
+// This allows us to partially initialize early in the flow before assets start loading etc.
+static FDelayedAutoRegisterHelper GAudioDeviceManagerPreInit(
+	EDelayedRegisterRunPhase::IniSystemReady,
+	&FAudioDeviceManager::PreInitialize);
 
 static int32 GCVarEnableAudioThreadWait = 1;
 TAutoConsoleVariable<int32> CVarEnableAudioThreadWait(
@@ -71,6 +77,7 @@ static FAutoConsoleCommand GReportAudioDevicesCommand(
 	})
 );
 
+
 namespace AudioDeviceManagerUtils
 {
 	FString PrintDeviceInfo(const Audio::FDeviceId InDeviceId, const EAudioDeviceScope InScope, const bool bInIsNonRealtime, const int32* NumHandles = nullptr, const TMap<uint32, FString>* InStackWalk = nullptr)
@@ -120,8 +127,6 @@ namespace AudioDeviceManagerUtils
 		return DeviceInfo;
 	}
 }
-
-
 FAudioDeviceManager* FAudioDeviceManager::Singleton = nullptr;
 
 // Some stress tests:
@@ -215,7 +220,7 @@ FAudioDeviceManager::~FAudioDeviceManager()
 
 	TMap<Audio::FDeviceId, FAudioDeviceContainer> DevicesToShutdown;
 	{
-		FScopeLock ScopeLock(&DeviceMapCriticalSection);		
+		FScopeLock ScopeLock(&DeviceMapCriticalSection);
 		DevicesToShutdown = MoveTemp(Devices);
 	}
 
@@ -228,6 +233,8 @@ FAudioDeviceManager::~FAudioDeviceManager()
 	MainAudioDeviceHandle.Reset();
 
 	FCoreDelegates::ApplicationWillEnterBackgroundDelegate.RemoveAll(this);
+
+	InitPhase = EInitPhase::Constructed;
 }
 PRAGMA_ENABLE_DEPRECATION_WARNINGS
 
@@ -371,38 +378,57 @@ void FAudioDeviceManager::RegisterAudioInfoFactories()
 	checkf(NumFailedFormats == 0, TEXT("Failed to find these required AudioFormats: [ %s]"), *FailedFormatsString);
 }
 
+bool FAudioDeviceManager::PreInitializeManager()
+{
+	if (InitPhase == EInitPhase::Constructed)
+	{
+		// Register all formats
+		AudioFormatSettings = MakePimpl<Audio::FAudioFormatSettings>(GConfig, GEngineIni, FPlatformProperties::IniPlatformName());
+		RegisterAudioInfoFactories();
+		InitPhase = EInitPhase::PreInitialized;
+	}
+	return InitPhase >= EInitPhase::PreInitialized;
+}
 
 bool FAudioDeviceManager::InitializeManager()
 {
-	if (LoadDefaultAudioDeviceModule())
+	// Do we also need to pre-init?
+	if (InitPhase < EInitPhase::PreInitialized)
 	{
-		check(AudioDeviceModule);
+		if (!PreInitializeManager())
+		{
+			return false;
+		}
+	}
 
-		UAudioSettings* AudioSettings = GetMutableDefault<UAudioSettings>();
-		check(AudioSettings);
+	// Initialize if we need to...
+	if (InitPhase == EInitPhase::PreInitialized )
+	{	
+		if (LoadDefaultAudioDeviceModule())
+		{
+			check(AudioDeviceModule);
 
-		AudioSettings->LoadDefaultObjects();
-		AudioSettings->RegisterParameterInterfaces();
+			UAudioSettings* AudioSettings = GetMutableDefault<UAudioSettings>();
+			check(AudioSettings);
 
-		// Register all formats
-		AudioFormatSettings = MakePimpl<Audio::FAudioFormatSettings>(GConfig, GEngineIni, FPlatformProperties::IniPlatformName());
-		RegisterAudioInfoFactories();
+			AudioSettings->LoadDefaultObjects();
+			AudioSettings->RegisterParameterInterfaces();
 
-		FModuleManager::Get().LoadModuleChecked(TEXT("AudioMixer"));
+			FModuleManager::Get().LoadModuleChecked(TEXT("AudioMixer"));
 
 #if WITH_EDITOR
-		IAudioEditorModule* AudioEditorModule = &FModuleManager::LoadModuleChecked<IAudioEditorModule>("AudioEditor");
-		AudioEditorModule->RegisterAudioMixerAssetActions();
-		AudioEditorModule->RegisterEffectPresetAssetActions();
+			IAudioEditorModule* AudioEditorModule = &FModuleManager::LoadModuleChecked<IAudioEditorModule>("AudioEditor");
+			AudioEditorModule->RegisterAudioMixerAssetActions();
+			AudioEditorModule->RegisterEffectPresetAssetActions();
 #endif
 
-		FCoreDelegates::ApplicationWillEnterBackgroundDelegate.AddRaw(this, &FAudioDeviceManager::AppWillEnterBackground);
+			FCoreDelegates::ApplicationWillEnterBackgroundDelegate.AddRaw(this, &FAudioDeviceManager::AppWillEnterBackground);
 
-		return true;
+			InitPhase = EInitPhase::Initialized;
+		}
 	}
 
-	// Failed to initialize
-	return false;
+	return IsInitialized();
 }
 
 bool FAudioDeviceManager::CreateMainAudioDevice()
@@ -642,49 +668,83 @@ void FAudioDeviceManager::SetAudioDevice(UWorld& InWorld, Audio::FDeviceId InDev
 	}
 }
 
-bool FAudioDeviceManager::Initialize()
+bool FAudioDeviceManager::PreInitialize()
 {
-	if (!Singleton)
-	{
-		bool bUseSound = FApp::CanEverRenderAudio();
+	// (Optionally) Pre-Initialize the AudioDeviceManager.
+	// By pre-initialing the Audio Device Manager we can start up some low level services needed for IO ahead of the main init.
+	// NOTE: Calling Get() will still return null until the device is initialized fully.
 
-		if (bUseSound)
+	if (FAudioDeviceManager* Adm = GetOrCreate())
+	{
+		UE_LOG(LogAudio, Display, TEXT("Pre-Initializing Audio Device Manager..."));
+		if (Adm->PreInitializeManager())
 		{
-			UE_LOG(LogAudio, Display, TEXT("Initializing Audio Device Manager..."));
-			Singleton = new FAudioDeviceManager();
-
-			if (Singleton->InitializeManager())
-			{
-				UE_LOG(LogAudio, Display, TEXT("Audio Device Manager Initialized"));
-			}
-			else
-			{
-				UE_LOG(LogAudio, Warning, TEXT("Audio Device Manager Initialization Failed!"));
-				delete Singleton;
-				Singleton = nullptr;
-			}
+			UE_LOG(LogAudio, Display, TEXT("Audio Device Manager Pre-Initialized"));
 		}
 		else
 		{
-			Audio::Analytics::RecordEvent_Usage(TEXT("AllAudioDisabled"));
+			UE_LOG(LogAudio, Warning, TEXT("Audio Device Manager Pre-Initialization Failed!"));
+			delete Adm;
+			Singleton = nullptr;
 		}
 	}
+	return Singleton && Singleton->InitPhase >= EInitPhase::PreInitialized;
+}
 
-	return Singleton != nullptr;
+bool FAudioDeviceManager::Initialize()
+{
+	if (FAudioDeviceManager* Adm = GetOrCreate())
+	{
+		UE_LOG(LogAudio, Display, TEXT("Initializing Audio Device Manager..."));
+		if (Adm->InitializeManager())
+		{
+			UE_LOG(LogAudio, Display, TEXT("Audio Device Manager Initialized"));
+		}
+		else
+		{
+			UE_LOG(LogAudio, Warning, TEXT("Audio Device Manager Initialization Failed!"));			
+			delete Adm;
+			Singleton = nullptr;
+		}
+	}
+	return Singleton && Singleton->IsInitialized();
 }
 
 FAudioDeviceManager* FAudioDeviceManager::Get()
 {
-	return Singleton;
+	if (Singleton && Singleton->IsInitialized())
+	{
+		return Singleton;
+	}
+	return nullptr;
 }
 
+FAudioDeviceManager* FAudioDeviceManager::GetOrCreate()
+{
+	if (!Singleton)
+	{
+		if (FApp::CanEverRenderAudio())
+		{
+			Singleton = new FAudioDeviceManager();
+		}
+		else
+		{
+			static bool bDoOnce = false;
+			if (!bDoOnce)
+			{
+				Audio::Analytics::RecordEvent_Usage(TEXT("AllAudioDisabled"));
+				bDoOnce = true;
+			}
+		}
+	}
+	return Singleton;
+}
 void FAudioDeviceManager::Shutdown()
 {
 	if (Singleton)
 	{
 		delete Singleton;
 		Singleton = nullptr;
-
 		UE_LOG(LogAudio, Display, TEXT("Audio Device Manager Shutdown"));
 	}
 }
@@ -1019,12 +1079,12 @@ uint32 FAudioDeviceManager::GetNewDeviceID()
 	return ++DeviceIDCounter;
 }
 
-// deprecated
+// (deprecated)
 void FAudioDeviceManager::StopSourcesUsingBuffer(FSoundBuffer*)
 {
 }
 
-// deprecated
+// (deprecated)
 void FAudioDeviceManager::TrackResource(USoundWave* SoundWave, FSoundBuffer* Buffer)
 {
 	// Allocate new resource ID and assign to USoundWave. A value of 0 (default) means not yet registered.
@@ -1050,28 +1110,30 @@ void FAudioDeviceManager::FreeResource(USoundWave* SoundWave)
 	}
 }
 
-// deprecated
+// (deprecated)
 void FAudioDeviceManager::FreeBufferResource(FSoundBuffer* SoundBuffer)
 {
 	if (SoundBuffer)
 	{
 		// Make sure any realtime tasks are finished that are using this buffer
 		SoundBuffer->EnsureRealtimeTaskCompletion();
-
 		delete SoundBuffer;
 		SoundBuffer = nullptr;
 	}
 }
 
-// deprecated
+// (deprecated)
 FSoundBuffer* FAudioDeviceManager::GetSoundBufferForResourceID(uint32 ResourceID)
 {
+	// maxtodo: warn
 	return {};
 }
 
 // deprecated
 void FAudioDeviceManager::RemoveSoundBufferForResourceID(uint32)
 {
+	// maxtodo: warn
+	// WaveBufferMap.Remove(ResourceID);
 }
 
 void FAudioDeviceManager::RemoveSoundMix(USoundMix* SoundMix)
diff --git a/Engine/Source/Runtime/Engine/Private/SoundWave.cpp b/Engine/Source/Runtime/Engine/Private/SoundWave.cpp
--- a/Engine/Source/Runtime/Engine/Private/SoundWave.cpp
+++ b/Engine/Source/Runtime/Engine/Private/SoundWave.cpp
@@ -99,6 +99,7 @@ inline FArchive& operator<<(FArchive& Ar, FSoundWaveCuePoint& CuePoint)
 	return Ar;
 }
 
+
 namespace SoundWave_Private
 {
 	static const TCHAR* GetDefaultLoadingBehaviorCVarName()
@@ -178,6 +179,24 @@ FSoundWaveData::~FSoundWaveData()
 	ResourceData.Empty();
 }
 
+void FSoundWaveData::CacheRuntimeFormatDependentState(const FName InFormatName)
+{
+	check(!InFormatName.IsNone());
+	if (!InFormatName.IsNone())
+	{
+		// Save format.
+		RuntimeFormat = InFormatName;
+
+		// Query the factory for the decoder.
+		const IAudioInfoFactory* Factory = IAudioInfoFactoryRegistry::Get().Find(InFormatName);
+		ensureMsgf(Factory, TEXT("Decoder for AudioFormat '%s' not found"), *InFormatName.ToString());
+		if (Factory)
+		{
+			bIsSeekable = Factory->IsSeekingSupported(bIsStreaming);	// bStreaming should be set before calling this.
+		}
+	}
+}
+
 void FSoundWaveData::InitializeDataFromSoundWave(USoundWave& InWave)
 {
 	LLM_SCOPE(ELLMTag::AudioSoundWaves);
@@ -189,14 +208,6 @@ void FSoundWaveData::InitializeDataFromSoundWave(USoundWave& InWave)
 	{
 		PackageNameCached = Package->GetFName();
 	}
-
-	// cache the runtime format for the wave
-	// note if this fails, it will ensure, and keep the "FSoundWaveProxy_InvalidFormat" as its value.
-	const FName FoundFormat = FindRuntimeFormat(InWave);
-	if (ensure(!FoundFormat.IsNone()))
-	{
-		RuntimeFormat = FoundFormat;
-	}
 	
 	SoundWaveKeyCached = FObjectKey(&InWave);
 	SampleRate = InWave.GetSampleRateForCurrentPlatform();
@@ -207,13 +218,17 @@ void FSoundWaveData::InitializeDataFromSoundWave(USoundWave& InWave)
 	WaveGuid = InWave.CompressedDataGuid;
 
 	// update shared flags
+	bIsStreaming = InWave.IsStreaming();
 	bIsLooping = InWave.IsLooping();
 	bIsTemplate = InWave.IsTemplate();
-	bIsStreaming = InWave.IsStreaming(nullptr);
-	bIsSeekable = InWave.IsSeekable();
+	bIsSeekable = false; // This is determined by the decoder in CacheRuntimeFormatDependentState.
 	SoundAssetCompressionType = InWave.GetSoundAssetCompressionType();
 	bShouldUseStreamCaching = InWave.ShouldUseStreamCaching();
 
+	// Cache the decoder and any dependent state.
+	const FName FoundFormat = FindRuntimeFormat(InWave);
+	CacheRuntimeFormatDependentState(FoundFormat);
+
 #if WITH_EDITOR
 	bLoadedFromCookedData = InWave.IsLoadedFromCookedData();
 
@@ -227,7 +242,7 @@ void FSoundWaveData::InitializeDataFromSoundWave(USoundWave& InWave)
 
 void FSoundWaveData::OverrideRuntimeFormat(const FName& InRuntimeFormat)
 {
-	RuntimeFormat = InRuntimeFormat;
+	CacheRuntimeFormatDependentState(InRuntimeFormat);
 }
 
 static FName GetBaseFormatFromSuffixedFormat(const FName& InFormat)
@@ -1407,8 +1422,9 @@ void USoundWave::SetSoundAssetCompressionType(ESoundAssetCompressionType InSound
 	{
 		SoundAssetCompressionType = InSoundAssetCompressionType;
 	}
-	SoundWaveDataPtr->bIsSeekable = IsSeekable();
-	SoundWaveDataPtr->RuntimeFormat = SoundWaveDataPtr->FindRuntimeFormat(*this);
+
+	SoundWaveDataPtr->CacheRuntimeFormatDependentState(SoundWaveDataPtr->FindRuntimeFormat(*this));
+
 	UpdateAsset(bMarkDirty);
 #endif // #if WITH_EDITOR
 }
@@ -3792,16 +3808,14 @@ TArrayView<const uint8> USoundWave::GetZerothChunk(bool bForImmediatePlayback)
 
 bool USoundWave::IsSeekable() const
 {
-	bool Result = true;
-
-	// Non-streaming codecs are seekable
-	// Streaming platform-specific codecs do not seek. Neither do source buses or procedural sources.
-	if (bIsSourceBus || bProcedural || (IsStreaming() && SoundAssetCompressionType == ESoundAssetCompressionType::PlatformSpecific))
+	// Note derived forms of SoundWave. (i.e. USoundWaveProcedural and USoundSourceBus)
+	// override this call to return false.
+	
+	if (SoundWaveDataPtr)
 	{
-		Result = false;
+		return SoundWaveDataPtr->bIsSeekable;
 	}
-
-	return Result;
+	return false;
 }
 
 bool USoundWave::GetSoundWavesWithCookedAnalysisData(TArray<USoundWave*>& OutSoundWaves)
diff --git a/Engine/Source/Runtime/Engine/Public/AudioDecompress.h b/Engine/Source/Runtime/Engine/Public/AudioDecompress.h
--- a/Engine/Source/Runtime/Engine/Public/AudioDecompress.h
+++ b/Engine/Source/Runtime/Engine/Public/AudioDecompress.h
@@ -354,7 +354,24 @@ protected:
 struct IAudioInfoFactory 
 {	
 	virtual ~IAudioInfoFactory() = default;
-	virtual ICompressedAudioInfo* Create() = 0;	
+	virtual ICompressedAudioInfo* Create() = 0;
+
+	/**
+	 * Simple decode capabilities struct, filled out by the decoder factories.
+	 * NOTE: By Default support is expected unless we opt out.
+	 */
+	struct FCapabilities
+	{
+		bool bSupportsSeeking = true;
+		bool bSupportsSeekableStreaming = true;
+	};
+	/**
+	* 
+	 * Queries if the factory's decoder implementation supports seeking
+	 * @param bIsStreaming If the media is streaming or not.
+	 * @return true if seeking is supported, false otherwise.
+	 */
+	virtual bool IsSeekingSupported(const bool bIsStreaming) const = 0;
 };
 struct IAudioInfoFactoryRegistry
 {
@@ -365,7 +382,7 @@ struct IAudioInfoFactoryRegistry
 	virtual IAudioInfoFactory* Find(FName InName) const = 0;
 	
 	// Convenience helper.
-	ICompressedAudioInfo* Create(FName InName)
+	ICompressedAudioInfo* Create(FName InName) const
 	{
 		if (IAudioInfoFactory* Factory = Find(InName))
 		{
@@ -378,19 +395,39 @@ struct IAudioInfoFactoryRegistry
 class FSimpleAudioInfoFactory : public IAudioInfoFactory
 {
 public:
-	FSimpleAudioInfoFactory(TFunction<ICompressedAudioInfo*(void)> InLambda, FName InFormatName)
-	:  CreateLamda(InLambda), FormatName(InFormatName)
+	/**
+	 * Constructs a simple AudioInfo factory.
+	 * @param InLambda Lambda to Create a ICompressionAudioInfo object
+	 * @param InFormatName The factory name of the Format. "ADPCM" etc.
+	 * @param InCaps The Capabilities of the decoder implementation.
+	 */
+	FSimpleAudioInfoFactory(const TFunction<ICompressedAudioInfo*()>& InLambda, const FName& InFormatName, const FCapabilities InCaps = {})
+	:  CreateLambda(InLambda), FormatName(InFormatName), Capabilities(InCaps)
 	{
 		IAudioInfoFactoryRegistry::Get().Register(this, FormatName);
 	}
-	virtual ~FSimpleAudioInfoFactory() 
+	virtual ~FSimpleAudioInfoFactory() override 
 	{
 		IAudioInfoFactoryRegistry::Get().Unregister(this, FormatName);
 	}
-	virtual ICompressedAudioInfo* Create() override { return CreateLamda(); }
+	virtual ICompressedAudioInfo* Create() override { return CreateLambda(); }
+
+	/**
+	 * Quries if seeking is supported by the decoder.
+	 * @param bIsStreaming If the media is to be streamed or not.
+	 * @return True for supported, false otherwise.
+	 */
+	virtual bool IsSeekingSupported(const bool bIsStreaming) const override
+	{
+		return bIsStreaming ?
+			Capabilities.bSupportsSeekableStreaming :
+			Capabilities.bSupportsSeeking;
+	}
+	
 private:
-	TFunction<ICompressedAudioInfo* ()> CreateLamda;
+	TFunction<ICompressedAudioInfo* ()> CreateLambda;
 	FName FormatName;
+	FCapabilities Capabilities;
 };
 
 /**
diff --git a/Engine/Source/Runtime/Engine/Public/AudioDeviceManager.h b/Engine/Source/Runtime/Engine/Public/AudioDeviceManager.h
--- a/Engine/Source/Runtime/Engine/Public/AudioDeviceManager.h
+++ b/Engine/Source/Runtime/Engine/Public/AudioDeviceManager.h
@@ -145,11 +145,29 @@ public:
 	ENGINE_API void SetAudioDevice(UWorld& InWorld, Audio::FDeviceId InDeviceID);
 
 	/**
-	* Initialize the audio device manager.
-	* Return true if successfully initialized.
-	**/
+	 * Pre-initializes the Manager.
+	 * This starts up and initializes the manager early without need for the Engine to be available, so that 
+	 * we can load essential services ahead of Engine Startup.
+	 * @return True if success, false otherwise
+	 */
+	static ENGINE_API bool PreInitialize();
+
+	/**
+	 * Initializes the Manager.
+	 * NOTE: If the manager hasn't been PreInitialized, that stage will be performed first.
+	 * @return True if success, false otherwise
+	 */
 	static ENGINE_API bool Initialize();
+
+	/**
+	 * Returns the AudioDeviceManager Singleton
+	 * @return The Instance instance or nullptr if the device has not been *fully* initialized.
+	 */
 	static ENGINE_API FAudioDeviceManager* Get();
+
+	/**
+	 * Shuts down the AudioDeviceManager and all AudioDevices.
+	 */
 	static ENGINE_API void Shutdown();
 
 	/** Creates the main audio device. */
@@ -220,22 +238,17 @@ public:
 	/** Sets an audio device to be solo'd */
 	ENGINE_API void SetSoloDevice(Audio::FDeviceId InAudioDeviceHandle);
 
-	/** Links up the resource data indices for looking up and cleaning up. */
 	UE_DEPRECATED(5.5, "This is no longer used.  Existing calls may be deleted")
 	ENGINE_API void TrackResource(USoundWave* SoundWave, FSoundBuffer* Buffer);
-
 	/** Frees the given sound wave resource from the device manager */
 	ENGINE_API void FreeResource(USoundWave* SoundWave);
 
-	/** Frees the sound buffer from the device manager. */
 	UE_DEPRECATED(5.5, "This is no longer used.  Existing calls may be deleted")
 	ENGINE_API void FreeBufferResource(FSoundBuffer* SoundBuffer);
 
-	/** Stops using the given sound buffer. Called before freeing the buffer */
 	UE_DEPRECATED(5.5, "This is no longer used.  Existing calls may be deleted")
 	ENGINE_API void StopSourcesUsingBuffer(FSoundBuffer* Buffer);
 
-	/** Retrieves the sound buffer for the given resource id */
 	UE_DEPRECATED(5.5, "This is no longer used.  Existing calls may be deleted")
 	ENGINE_API FSoundBuffer* GetSoundBufferForResourceID(uint32 ResourceID);
 
@@ -302,14 +315,18 @@ public:
 private:
 
 #if ENABLE_AUDIO_DEBUG
-	/** Instance of audio debugger shared across audio devices */
+	/** Instance of audio debugger shared across audio devices */ 
 	TUniquePtr<Audio::FAudioDebugger> AudioDebugger;
 #endif // ENABLE_AUDIO_DEBUG
 
 	TPimplPtr<Audio::FAudioFormatSettings> AudioFormatSettings;
 	TArray<TPimplPtr<FSimpleAudioInfoFactory>> EngineFormats;
 
-	ENGINE_API bool InitializeManager();
+	/**
+	 * Conditionally Creates the AudioDeviceManager
+	 * @return New Instance or Existing Singleton, or nullptr on failure (which can happen if audio is disabled).
+	 */
+	static FAudioDeviceManager* GetOrCreate();
 
 	/** Creates a handle given the index and a generation value. */
 	ENGINE_API uint32 GetNewDeviceID();
@@ -330,6 +347,28 @@ private:
 	
 	ENGINE_API void RegisterAudioInfoFactories();
 
+	/**
+	 * Queries if the device manager is initalized.
+	 * @return true if initialized, false otherwise.
+	 */
+	ENGINE_API bool IsInitialized() const { return InitPhase == EInitPhase::Initialized; }	
+
+	/**
+	 * Pre-initializes the Manager
+	 * This starts up and intializes the manager early without need for the Engine to be available, so that 
+	 * we can load essential services ahead of Engine Startup.
+	 * This is optional in the flow, as InitializeManager will also call this if required.
+	 * @return True if success, false otherwise
+	 */
+	ENGINE_API bool PreInitializeManager();
+
+	/**
+	 * Initalizes the Manager.
+	 * NOTE: If PreInitializeManager has not been called already, a PreInitializeManager call is made first.
+	 * @return True if success, false otherwise
+	 */
+	ENGINE_API bool InitializeManager();
+
 	/** Audio device module which creates audio devices. */
 	IAudioDeviceModule* AudioDeviceModule;
 
@@ -418,7 +457,17 @@ private:
 	/** Audio Fence to ensure that we don't allow the audio thread to drift never endingly behind. */
 	FAudioCommandFence SyncFence;
 
+	/**
+	 * Simple state machine of the init state of the manager.
+	 */
+	enum class EInitPhase : uint32
+	{
+		Constructed,
+		PreInitialized,
+		Initialized
+	};
+	EInitPhase InitPhase = EInitPhase::Constructed;
+	
 	friend class FAudioDeviceHandle;
-
 	static ENGINE_API FAudioDeviceManager* Singleton;
 };
diff --git a/Engine/Source/Runtime/OpusAudioDecoder/Module/Private/OpusAudioInfo.cpp b/Engine/Source/Runtime/OpusAudioDecoder/Module/Private/OpusAudioInfo.cpp
--- a/Engine/Source/Runtime/OpusAudioDecoder/Module/Private/OpusAudioInfo.cpp
+++ b/Engine/Source/Runtime/OpusAudioDecoder/Module/Private/OpusAudioInfo.cpp
@@ -387,7 +387,12 @@ public:
 
 	virtual void StartupModule() override
 	{
-		Factory = MakeUnique<FSimpleAudioInfoFactory>([] { return new FOpusAudioInfo(); }, Audio::NAME_OPUS);
+		constexpr FSimpleAudioInfoFactory::FCapabilities Caps = 
+		{
+			/*bSupportsSeeking*/ true,				
+			/*bSupportsSeekableStreaming*/ false,	
+		};
+		Factory = MakeUnique<FSimpleAudioInfoFactory>([] { return new FOpusAudioInfo(); }, Audio::NAME_OPUS, Caps);
 	}
 
 	virtual void ShutdownModule() override {}
diff --git a/Engine/Source/Runtime/VorbisAudioDecoder/Module/Private/VorbisAudioInfo.cpp b/Engine/Source/Runtime/VorbisAudioDecoder/Module/Private/VorbisAudioInfo.cpp
--- a/Engine/Source/Runtime/VorbisAudioDecoder/Module/Private/VorbisAudioInfo.cpp
+++ b/Engine/Source/Runtime/VorbisAudioDecoder/Module/Private/VorbisAudioInfo.cpp
@@ -863,7 +863,13 @@ public:
 	virtual void StartupModule() override
 	{
 		LoadVorbisLibraries();
-		Factory = MakeUnique<FSimpleAudioInfoFactory>([] { return new FVorbisAudioInfo(); }, Audio::NAME_OGG);
+
+		constexpr FSimpleAudioInfoFactory::FCapabilities Caps = 
+		{
+			/*bSupportsSeeking*/ true,				
+			/*bSupportsSeekableStreaming*/ false,	
+		};
+		Factory = MakeUnique<FSimpleAudioInfoFactory>([] { return new FVorbisAudioInfo(); }, Audio::NAME_OGG, Caps);
 	}
 
 	virtual void ShutdownModule() override {}
