Support Overridable for T3D
* Disabled by default using CVar: "OverridableSerializationLogic.Capabilities.T3D"

diff --git a/Engine/Source/Editor/UnrealEd/Private/EditorObject.cpp b/Engine/Source/Editor/UnrealEd/Private/EditorObject.cpp
--- a/Engine/Source/Editor/UnrealEd/Private/EditorObject.cpp
+++ b/Engine/Source/Editor/UnrealEd/Private/EditorObject.cpp
@@ -34,6 +34,8 @@
 #include "InstancedFoliage.h"
 #include "Components/BrushComponent.h"
 #include "Algo/Transform.h"
+#include "UObject/OverridableManager.h"
+#include "UObject/OverriddenPropertySet.h"
 
 DEFINE_LOG_CATEGORY_STATIC(LogEditorObject, Log, All);
 
@@ -259,7 +261,16 @@ static const TCHAR* ImportProperties(
 	check(DestData!=NULL);
 
 	if ( SourceText == NULL )
+	{
 		return NULL;
+	}
+	
+	FOverriddenPropertySet* OverriddenProperties = nullptr;
+	if (SubobjectOuter && FOverridableSerializationLogic::HasCapabilities(FOverridableSerializationLogic::ECapabilities::T3DSerialization))
+	{
+		OverriddenProperties = FOverridableManager::Get().GetOverriddenProperties(*SubobjectOuter);
+	}
+	FEnableOverridableSerializationScope Scope(OverriddenProperties != nullptr, OverriddenProperties);
 
 	// Cannot create subobjects when importing struct defaults, or if SubobjectOuter (used as the Outer for any subobject declarations encountered) is NULL
 	bool bSubObjectsAllowed = !ObjectStruct->IsA(UScriptStruct::StaticClass()) && SubobjectOuter != NULL;
@@ -509,6 +520,11 @@ static const TCHAR* ImportProperties(
 			}
 			else 
 			{
+				FName OverriddenOperationName = NAME_None;
+				const bool bParsed = FParse::Value(Str,TEXT("OverriddenOperation="), OverriddenOperationName);
+				const bool bOverrideT3DSupported = FOverridableSerializationLogic::HasCapabilities(FOverridableSerializationLogic::ECapabilities::T3DSerialization);
+				TOptional<EOverriddenPropertyOperation> OverriddenOperation = bParsed && bOverrideT3DSupported ? GetOverriddenOperationFromName(OverriddenOperationName) : TOptional<EOverriddenPropertyOperation>();
+
 				UObject* Archetype = NULL;
 				UObject* ComponentTemplate = NULL;
 
@@ -690,6 +706,18 @@ static const TCHAR* ImportProperties(
 				}
 				FArchiveReplaceObjectRef<UObject> ReplaceAr(SubobjectOuter, ReplacementMap, EArchiveReplaceObjectFlags::IgnoreArchetypeRef);
 
+				if (bOverrideT3DSupported)
+				{
+					if (OverriddenOperation.IsSet())
+					{
+						FOverridableManager::Get().SetOverriddenProperties(*ComponentTemplate, OverriddenOperation.GetValue(), /*bNeedsSubObjectTemplateInstantiation*/true);
+					}
+					else
+					{
+						FOverridableManager::Get().Disable(*ComponentTemplate);
+					}
+				}
+
 				// import the properties for the subobject
 				SourceText = ImportObjectProperties(
 					(uint8*)ComponentTemplate, 
@@ -1054,6 +1082,11 @@ static const TCHAR* ImportCreateSubObjectsStep(
 			FString ExportedObjectFullName;
 			FParse::Value(Str, TEXT("ExportPath="), ExportedObjectFullName);
 
+			FName OverriddenOperationName = NAME_None;
+			const bool bParsed = FParse::Value(Str,TEXT("OverriddenOperation="), OverriddenOperationName);
+			const bool bOverrideT3DSupported = FOverridableSerializationLogic::HasCapabilities(FOverridableSerializationLogic::ECapabilities::T3DSerialization);
+			TOptional<EOverriddenPropertyOperation> OverriddenOperation = bParsed && bOverrideT3DSupported ? GetOverriddenOperationFromName(OverriddenOperationName) : TOptional<EOverriddenPropertyOperation>();
+
 			// points to the parent class's template subobject/component, if we are overriding a subobject/component declared in our parent class
 			UObject* BaseTemplate = nullptr;
 			bool bRedefiningSubobject = false;
@@ -1283,6 +1316,18 @@ static const TCHAR* ImportCreateSubObjectsStep(
 					ObjectRemapper->Add(MoveTemp(ExportedObjectFullName), ComponentTemplate);
 				}
 
+				if (bOverrideT3DSupported)
+				{
+					if (OverriddenOperation.IsSet())
+					{
+						FOverridableManager::Get().SetOverriddenProperties(*ComponentTemplate, OverriddenOperation.GetValue(), /*bNeedsSubObjectTemplateInstantiation*/true);
+					}
+					else
+					{
+						FOverridableManager::Get().Disable(*ComponentTemplate);
+					}
+				}
+
 				// Create the subobjects for the subobject
 				CurrentSourceText = ImportCreateObjectsStep(Params);
 			}
@@ -1335,6 +1380,13 @@ static const TCHAR* ImportPropertiesStep(
 		return nullptr;
 	}
 
+	FOverriddenPropertySet* OverriddenProperties = nullptr;
+	if (SubobjectOuter && FOverridableSerializationLogic::HasCapabilities(FOverridableSerializationLogic::ECapabilities::T3DSerialization))
+	{
+		OverriddenProperties = FOverridableManager::Get().GetOverriddenProperties(*SubobjectOuter);
+	}
+	FEnableOverridableSerializationScope Scope(OverriddenProperties != nullptr, OverriddenProperties);
+
 	// Cannot create subobjects when importing struct defaults, or if SubobjectOuter (used as the Outer for any subobject declarations encountered) is NULL
 	bool bSubObjectsAllowed = !ObjectStruct->IsA(UScriptStruct::StaticClass()) && SubobjectOuter != NULL;
 
diff --git a/Engine/Source/Editor/UnrealEd/Private/Factories/EditorFactories.cpp b/Engine/Source/Editor/UnrealEd/Private/Factories/EditorFactories.cpp
--- a/Engine/Source/Editor/UnrealEd/Private/Factories/EditorFactories.cpp
+++ b/Engine/Source/Editor/UnrealEd/Private/Factories/EditorFactories.cpp
@@ -284,6 +284,7 @@
 #include "UDIMUtilities.h"
 #include "FileHelpers.h"
 #include "StaticMeshOperations.h"
+#include "UObject/OverridableManager.h"
 
 
 DEFINE_LOG_CATEGORY(LogEditorFactories);
@@ -5307,6 +5308,8 @@ void FCustomizableTextObjectFactory::ProcessBuffer(UObject* InParent, EObjectFla
 	int32 NestedDepth     = 0;
 	int32 OmittedOuterObj = 0; // zero signifies "nothing omitted"
 
+	const bool bT3DOverrideEnabled = FOverridableSerializationLogic::HasCapabilities(FOverridableSerializationLogic::ECapabilities::T3DSerialization);
+
 	FString StrLine;
 	while( FParse::Line(&Buffer,StrLine) )
 	{
@@ -5349,6 +5352,14 @@ void FCustomizableTextObjectFactory::ProcessBuffer(UObject* InParent, EObjectFla
 					ObjArchetype = LoadObject<UObject>(nullptr, *ObjArchetypeName, nullptr, LOAD_None, nullptr);
 				}
 
+				FName OverriddenOperationName = NAME_None;
+				TOptional<EOverriddenPropertyOperation> OverriddenOperation;
+				const bool bParsed = FParse::Value(Str,TEXT("OverriddenOperation="),OverriddenOperationName);
+				if (bParsed && bT3DOverrideEnabled)
+				{
+					OverriddenOperation = GetOverriddenOperationFromName(OverriddenOperationName);
+				}
+
 				UObject* ObjectParent = InParent ? InParent : GetParentForNewObject(ObjClass);
 
 				// Make sure we are allowed to create this object
@@ -5364,6 +5375,19 @@ void FCustomizableTextObjectFactory::ProcessBuffer(UObject* InParent, EObjectFla
 				// Spawn the object and reset it's archetype
 				UObject* CreatedObject = NewObject<UObject>(ObjectParent, ObjClass, ObjName, Flags, ObjArchetype, !!ObjectParent, &InstanceGraph);
 
+				// Overridable handling
+				if (bT3DOverrideEnabled)
+				{
+					if (OverriddenOperation.IsSet())
+					{
+						FOverridableManager::Get().SetOverriddenProperties(*CreatedObject, OverriddenOperation.GetValue(), /*bNeedsSubObjectTemplateInstantiation*/true);
+					}
+					else
+					{
+						FOverridableManager::Get().Disable(*CreatedObject);
+					}
+				}
+
 				// Get property text for the new object.
 				FString PropText, PropLine;
 				FString SubObjText;
diff --git a/Engine/Source/Runtime/Core/Public/Misc/AutomationTest.h b/Engine/Source/Runtime/Core/Public/Misc/AutomationTest.h
--- a/Engine/Source/Runtime/Core/Public/Misc/AutomationTest.h
+++ b/Engine/Source/Runtime/Core/Public/Misc/AutomationTest.h
@@ -3207,26 +3207,35 @@ public:
 		//disabled
 	}
 
-	void Describe(const FString& InDescription, TFunction<void()> DoWork)
+	void Describe(const TArray<FString> InDescriptions, TFunction<void(int32 /*Iteration*/)> DoWork)
 	{
 		LLM_SCOPE_BYNAME(TEXT("AutomationTest/Framework"));
-		const TSharedRef<FSpecDefinitionScope> ParentScope = DefinitionScopeStack.Last();
-		const TSharedRef<FSpecDefinitionScope> NewScope = MakeShareable(new FSpecDefinitionScope());
-		NewScope->Description = InDescription;
-		ParentScope->Children.Push(NewScope);
-
-		DefinitionScopeStack.Push(NewScope);
-		PushDescription(InDescription);
-		DoWork();
-		PopDescription(InDescription);
-		DefinitionScopeStack.Pop();
 
-		if (NewScope->It.Num() == 0 && NewScope->Children.Num() == 0)
+		for (int32 It = 0; It < InDescriptions.Num(); ++It)
 		{
-			ParentScope->Children.Remove(NewScope);
+			const TSharedRef<FSpecDefinitionScope> ParentScope = DefinitionScopeStack.Last();
+			const TSharedRef<FSpecDefinitionScope> NewScope = MakeShareable(new FSpecDefinitionScope());
+			NewScope->Description = InDescriptions[It];
+			ParentScope->Children.Push(NewScope);
+
+			DefinitionScopeStack.Push(NewScope);
+			PushDescription(InDescriptions[It]);
+			DoWork(It);
+			PopDescription(InDescriptions[It]);
+			DefinitionScopeStack.Pop();
+
+			if (NewScope->It.Num() == 0 && NewScope->Children.Num() == 0)
+			{
+				ParentScope->Children.Remove(NewScope);
+			}
 		}
 	}
 
+	void Describe(const FString& InDescription, TFunction<void()> DoWork)
+	{
+		Describe({InDescription}, [DoWork](int32){DoWork();});
+	}
+
 	void xIt(const FString& InDescription, TFunction<void()> DoWork)
 	{
 		//disabled
diff --git a/Engine/Source/Runtime/CoreUObject/Private/UObject/Class.cpp b/Engine/Source/Runtime/CoreUObject/Private/UObject/Class.cpp
--- a/Engine/Source/Runtime/CoreUObject/Private/UObject/Class.cpp
+++ b/Engine/Source/Runtime/CoreUObject/Private/UObject/Class.cpp
@@ -68,6 +68,7 @@
 #include "String/Find.h"
 #include "String/ParseTokens.h"
 #include "AutoRTFM/AutoRTFM.h"
+#include "Serialization/ArchiveSerializedPropertyChain.h"
 #include "Serialization/TestUndeclaredScriptStructObjectReferences.h"
 #include "UObject/OverridableManager.h"
 #include "UObject/OverriddenPropertySet.h"
@@ -3403,12 +3404,17 @@ void UScriptStruct::ExportText(FString& ValueStr, const void* Value, const void*
 		Defaults = Value;
 	}
 
+	const bool bT3DOverrideEnabled = FOverridableSerializationLogic::HasCapabilities(FOverridableSerializationLogic::ECapabilities::T3DSerialization);
+
 	for (TFieldIterator<FProperty> It(this); It; ++It)
 	{
 		if (It->ShouldPort(PortFlags))
 		{
 			for (int32 Index = 0; Index < It->ArrayDim; Index++)
 			{
+				const bool bStaticArray = It->ArrayDim > 1;
+				FOverridableTextPortPropertyPathScope ScopePath(*It, bStaticArray ? Index : INDEX_NONE, bStaticArray ? EPropertyVisitorInfoType::StaticArrayIndex : EPropertyVisitorInfoType::None);
+
 				FString InnerValue;
 				if (It->ExportText_InContainer(Index, InnerValue, Value, Defaults, OwnerObject, PPF_Delimited | PortFlags, ExportRootScope))
 				{
@@ -3428,13 +3434,30 @@ void UScriptStruct::ExportText(FString& ValueStr, const void* Value, const void*
 
 					const FString PropertyName = (PortFlags & (PPF_ExternalEditor | PPF_BlueprintDebugView)) != 0 ? It->GetAuthoredName() : It->GetName();
 
-					if (It->ArrayDim == 1)
+					FString OverridableOperation;
+					if (bT3DOverrideEnabled)
+					{
+						if(const FOverriddenPropertySet* OverriddenProperties = FOverridableSerializationLogic::GetOverriddenProperties())
+						{
+							const FPropertyVisitorPath* Path = FOverridableSerializationLogic::GetOverriddenPortTextPropertyPath();
+							checkf(Path, TEXT("Expecting a path"));
+							FArchiveSerializedPropertyChain Chain = Path->ToSerializedPropertyChain();
+						
+							const EOverriddenPropertyOperation Operation = OverriddenProperties->GetOverriddenPropertyOperation(&Chain, /*Property*/nullptr);
+							if (Operation != EOverriddenPropertyOperation::None)
+							{
+								OverridableOperation = FString::Printf(TEXT("<%s>"), *GetOverriddenOperationString(Operation));
+							}
+						}
+					}
+
+					if (!bStaticArray)
 					{
-						ValueStr += FString::Printf(TEXT("%s="), *PropertyName);
+						ValueStr += FString::Printf(TEXT("%s%s="), *PropertyName, *OverridableOperation);
 					}
 					else
 					{
-						ValueStr += FString::Printf(TEXT("%s[%i]="), *PropertyName, Index);
+						ValueStr += FString::Printf(TEXT("%s[%i]%s="), *PropertyName, Index, *OverridableOperation);
 					}
 					ValueStr += InnerValue;
 				}
diff --git a/Engine/Source/Runtime/CoreUObject/Private/UObject/OverridableManager.cpp b/Engine/Source/Runtime/CoreUObject/Private/UObject/OverridableManager.cpp
--- a/Engine/Source/Runtime/CoreUObject/Private/UObject/OverridableManager.cpp
+++ b/Engine/Source/Runtime/CoreUObject/Private/UObject/OverridableManager.cpp
@@ -54,12 +54,12 @@ void FOverridableManager::InheritEnabledFrom(UObject& Object, const UObject* Def
 #endif // WITH_EDITORONLY_DATA
 }
 
-void FOverridableManager::CacheArchetype(UObject& Object)
+void FOverridableManager::CacheArchetype(UObject& Object, UObject* Archetype /*= nullptr*/)
 {
 #if WITH_EDITORONLY_DATA
 	if (FOverriddenPropertySet* OverriddenProperties = OverriddenObjectAnnotations.Find(Object))
 	{
-		OverriddenProperties->CacheArchetype();
+		OverriddenProperties->CacheArchetype(Archetype);
 	}
 #endif // WITH_EDITORONLY_DATA
 }
diff --git a/Engine/Source/Runtime/CoreUObject/Private/UObject/OverriddenPropertySet.cpp b/Engine/Source/Runtime/CoreUObject/Private/UObject/OverriddenPropertySet.cpp
--- a/Engine/Source/Runtime/CoreUObject/Private/UObject/OverriddenPropertySet.cpp
+++ b/Engine/Source/Runtime/CoreUObject/Private/UObject/OverriddenPropertySet.cpp
@@ -2,6 +2,7 @@
 
 #include "UObject/OverriddenPropertySet.h"
 
+#include "HAL/IConsoleManager.h"
 #include "UObject/OverridableManager.h"
 #include "Serialization/ArchiveSerializedPropertyChain.h"
 #include "UObject/PropertyOptional.h"
@@ -19,14 +20,60 @@ DEFINE_LOG_CATEGORY(LogOverridableObject);
 //----------------------------------------------------------------------//
 // FOverridableSerializationLogic
 //----------------------------------------------------------------------//
+
+struct FOverridableSerializationLogicInternalAdapter
+{
+	static void SetCapability(FOverridableSerializationLogic::ECapabilities InCapability, bool bEnable)
+	{
+		if (bEnable)
+		{
+			FOverridableSerializationLogic::Capabilities = FOverridableSerializationLogic::Capabilities | InCapability;
+		}
+		else
+		{
+			FOverridableSerializationLogic::Capabilities = FOverridableSerializationLogic::Capabilities & ~InCapability;
+		}
+	}
+};
+
+namespace Private
+{
+	bool bEnableT3D = false;
+	FAutoConsoleVariableRef CVar_bT3DOverrideSerializationEnabled(
+		TEXT("OverridableSerializationLogic.Capabilities.T3D"),
+		bEnableT3D,
+		TEXT("Enables serialization of override state into/from T3D"),
+		FConsoleVariableDelegate::CreateLambda([](IConsoleVariable* Variable)
+		{
+			FOverridableManager& Manager = FOverridableManager::Get();
+			FOverridableSerializationLogicInternalAdapter::SetCapability(FOverridableSerializationLogic::ECapabilities::T3DSerialization, bEnableT3D);
+		}));
+}
+
+ENUM_CLASS_FLAGS(FOverridableSerializationLogic::ECapabilities);
+
+FOverridableSerializationLogic::ECapabilities FOverridableSerializationLogic::Capabilities = FOverridableSerializationLogic::ECapabilities::None;
 thread_local bool FOverridableSerializationLogic::bUseOverridableSerialization = false;
 thread_local FOverriddenPropertySet* FOverridableSerializationLogic::OverriddenProperties = nullptr;
+thread_local FPropertyVisitorPath* FOverridableSerializationLogic::OverriddenPortTextPropertyPath = nullptr;
+
+namespace Private
+{
+	struct FCapabilitiesAutoInitializer
+	{
+		FCapabilitiesAutoInitializer()
+		{
+			FOverridableSerializationLogicInternalAdapter::SetCapability(FOverridableSerializationLogic::ECapabilities::T3DSerialization, bEnableT3D);
+		}
+	};
+
+	static FCapabilitiesAutoInitializer CapabilitiesAutoInitializer;
+}
 
-EOverriddenPropertyOperation FOverridableSerializationLogic::GetOverriddenPropertyOperation(const FArchive& Ar, FProperty* Property /*= nullptr*/, uint8* DataPtr /*= nullptr*/, uint8* DefaultValue)
+EOverriddenPropertyOperation FOverridableSerializationLogic::GetOverriddenPropertyOperation(const int32 PortFlags, const FArchiveSerializedPropertyChain* CurrentPropertyChain, FProperty* Property, const void* DataPtr, const void* DefaultValue)
 {
 	checkf(bUseOverridableSerialization, TEXT("Nobody should use this method if it is not setup to use overridable serialization"));
 
-	const FArchiveSerializedPropertyChain* CurrentPropertyChain = Ar.GetSerializedPropertyChain();
 	const FProperty* CurrentProperty = Property ? Property : (CurrentPropertyChain ? CurrentPropertyChain->GetPropertyFromStack(0) : nullptr);
 	if (CurrentProperty && CurrentProperty->HasAnyPropertyFlags(CPF_ExperimentalNeverOverriden))
 	{
@@ -51,7 +98,7 @@ EOverriddenPropertyOperation FOverridableSerializationLogic::GetOverriddenProper
 		if (OverriddenProperties && OverriddenProperties->IsCDOOwningProperty(*CurrentProperty))
 		{
 			// Only need serialize this value if it is different from the default property value
-			if (!CurrentProperty->Identical(DataPtr, DefaultValue, Ar.GetPortFlags()))
+			if (!CurrentProperty->Identical(DataPtr, DefaultValue, PortFlags))
 			{
 				return 	EOverriddenPropertyOperation::Replace;
 			}
@@ -61,6 +108,41 @@ EOverriddenPropertyOperation FOverridableSerializationLogic::GetOverriddenProper
 	return EOverriddenPropertyOperation::None;
 }
 
+bool FOverridableSerializationLogic::HasCapabilities(ECapabilities InCapabilities)
+{
+	return (Capabilities & InCapabilities) == InCapabilities;
+}
+
+EOverriddenPropertyOperation FOverridableSerializationLogic::GetOverriddenPropertyOperation(const FArchive& Ar, FProperty* Property /*= nullptr*/, uint8* DataPtr /*= nullptr*/, uint8* DefaultValue /*= nullptr*/)
+{
+	const FArchiveSerializedPropertyChain* CurrentPropertyChain = Ar.GetSerializedPropertyChain();
+	return GetOverriddenPropertyOperation(Ar.GetPortFlags(), CurrentPropertyChain, Property, DataPtr, DefaultValue);
+}
+
+EOverriddenPropertyOperation FOverridableSerializationLogic::GetOverriddenPropertyOperationForPortText(const void* DataPtr, const void* DefaultValue, int32 PortFlags)
+{
+	checkf(OverriddenPortTextPropertyPath, TEXT("Expecting an overridden port text path"));
+
+	const FArchiveSerializedPropertyChain CurrentPropertyChain = OverriddenPortTextPropertyPath->ToSerializedPropertyChain();
+	return GetOverriddenPropertyOperation(PortFlags,  &CurrentPropertyChain, nullptr, DataPtr, DefaultValue);
+}
+
+FPropertyVisitorPath* FOverridableSerializationLogic::GetOverriddenPortTextPropertyPath()
+{
+	return OverriddenPortTextPropertyPath;
+}
+
+void FOverridableSerializationLogic::SetOverriddenPortTextPropertyPath(FPropertyVisitorPath& Path)
+{
+	checkf(OverriddenPortTextPropertyPath == nullptr, TEXT("Should not set a path on top of an existing one"));
+	OverriddenPortTextPropertyPath = &Path;
+}
+
+void FOverridableSerializationLogic::ResetOverriddenPortTextPropertyPath()
+{
+	OverriddenPortTextPropertyPath = nullptr;
+}
+
 //----------------------------------------------------------------------//
 // FOverridableSerializationScope
 //----------------------------------------------------------------------//
@@ -91,6 +173,49 @@ FEnableOverridableSerializationScope::~FEnableOverridableSerializationScope()
 	}
 }
 
+//----------------------------------------------------------------------//
+// FOverridableTextImportPropertyPathScope
+//----------------------------------------------------------------------//
+FOverridableTextPortPropertyPathScope::FOverridableTextPortPropertyPathScope(const FProperty* InProperty, int32 InIndex/* = INDEX_NONE*/, EPropertyVisitorInfoType InPropertyInfo /*= EPropertyVisitorInfoType::None*/)
+{
+	if (!FOverridableSerializationLogic::IsEnabled())
+	{
+		return;
+	}
+
+	checkf(InProperty, TEXT("Expecting a valid property ptr"));
+
+	// Save property for comparison in the destructor
+	Property = InProperty;
+
+	FPropertyVisitorPath* Path = FOverridableSerializationLogic::GetOverriddenPortTextPropertyPath();
+	if (!Path)
+	{
+		FOverridableSerializationLogic::SetOverriddenPortTextPropertyPath(DefaultPath);
+		Path = &DefaultPath;
+	}
+
+	Path->Push(FPropertyVisitorInfo(InProperty, InIndex, InPropertyInfo));
+}
+
+FOverridableTextPortPropertyPathScope::~FOverridableTextPortPropertyPathScope()
+{
+	if (Property)
+	{
+		FPropertyVisitorPath* Path = FOverridableSerializationLogic::GetOverriddenPortTextPropertyPath();
+		checkf(Path, TEXT("Expecting a valid path "));
+		checkf(Path->Num(), TEXT("Expecting at least one property in the path"));
+		verifyf(Path->Pop().Property == Property, TEXT("Expecting at the top property to match the one we pushed in the constructor"));
+		if (!Path->Num())
+		{
+			FOverridableSerializationLogic::ResetOverriddenPortTextPropertyPath();
+		}
+	}
+}
+
+//----------------------------------------------------------------------//
+// FOverriddenPropertyNodeID
+//----------------------------------------------------------------------//
 FOverriddenPropertyNodeID::FOverriddenPropertyNodeID(const FProperty* Property)
 	: Object(nullptr)
 {
@@ -178,9 +303,6 @@ int32 FOverriddenPropertyNodeID::ToMapInternalIndex(FScriptMapHelper& MapHelper)
 	return INDEX_NONE;
 }
 
-//----------------------------------------------------------------------//
-// FOverriddenPropertyNodeID
-//----------------------------------------------------------------------//
 void FOverriddenPropertyNodeID::HandleObjectsReInstantiated(const TMap<UObject*, UObject*>& Map)
 {
 	if (!Object)
@@ -1099,11 +1221,11 @@ void FOverriddenPropertySet::RemoveOverriddenSubProperties(FOverriddenPropertyNo
 	PropertyNode.SubPropertyNodeKeys.Empty();
 }
 
-void FOverriddenPropertySet::CacheArchetype()
+void FOverriddenPropertySet::CacheArchetype(UObject* Archetype /*= nullptr*/)
 {
 	if (!CachedArchetype)
 	{
-		CachedArchetype = Owner->GetArchetype();
+		CachedArchetype = Archetype ? Archetype : Owner->GetArchetype();
 	}
 }
 
diff --git a/Engine/Source/Runtime/CoreUObject/Private/UObject/Property.cpp b/Engine/Source/Runtime/CoreUObject/Private/UObject/Property.cpp
--- a/Engine/Source/Runtime/CoreUObject/Private/UObject/Property.cpp
+++ b/Engine/Source/Runtime/CoreUObject/Private/UObject/Property.cpp
@@ -11,17 +11,20 @@
 #include "Math/Sphere.h"
 #include "Misc/AsciiSet.h"
 #include "Misc/Guid.h"
+#include "Misc/ScopeExit.h"
 #include "Misc/StringBuilder.h"
 #include "Misc/FrameRate.h"
 #include "Misc/FrameTime.h"
 #include "Misc/QualifiedFrameTime.h"
 #include "Misc/FrameNumber.h"
 #include "Misc/Timecode.h"
+#include "Serialization/ArchiveSerializedPropertyChain.h"
 #include "Serialization/TestUndeclaredScriptStructObjectReferences.h"
 #include "Templates/Casts.h"
 #include "UObject/Class.h"
 #include "UObject/CoreNetTypes.h"
 #include "UObject/CoreRedirects.h"
+#include "UObject/OverridableManager.h"
 #include "UObject/Package.h"
 #include "UObject/PropertyHelper.h"
 #include "UObject/PropertyTypeName.h"
@@ -1213,7 +1216,10 @@ bool FProperty::ExportText_Direct
 	UObject*	ExportRootScope
 	) const
 {
-	if( Data==Delta || !Identical(Data,Delta,PortFlags) )
+	const bool bExportOverride = FOverridableSerializationLogic::IsEnabled() && FOverridableSerializationLogic::HasCapabilities(FOverridableSerializationLogic::ECapabilities::T3DSerialization);
+	if( Data==Delta || 
+		(bExportOverride && FOverridableSerializationLogic::GetOverriddenPropertyOperationForPortText(Data, Delta, PortFlags) != EOverriddenPropertyOperation::None)  ||
+		(!bExportOverride && !Identical(Data,Delta,PortFlags)) )
 	{
 		ExportText_Internal
 			(
@@ -1626,7 +1632,8 @@ static const int32 ReadArrayIndex(const UStruct* ObjectStruct, const TCHAR*& Str
 		FString IndexText(TEXT(""));
 		while ( *Str && *Str != ')' && *Str != ']' )
 		{
-			if ( *Str == TCHAR('=') )
+			// Stop at either we reach the = or the beginning of the overridable operation ex: <replace>
+			if ( *Str == TCHAR('=') || *Str == TCHAR('<'))
 			{
 				// we've encountered an equals sign before the closing bracket
 				Warn->Logf(ELogVerbosity::Warning, TEXT("Missing ')' in default properties subscript: %s"), Start);
@@ -1688,6 +1695,65 @@ static const int32 ReadArrayIndex(const UStruct* ObjectStruct, const TCHAR*& Str
 	return Index;
 }
 
+/**
+ * Attempts to read an array index (xxx) sequence.  Handles const/enum replacements, etc.
+ * @param	ObjectStruct	the scope of the object/struct containing the property we're currently importing
+ * @param	Str				[out] pointer to the the buffer containing the property value to import
+ * @param	Warn			the output device to send errors/warnings to
+ * @return	the array index for this defaultproperties line.  INDEX_NONE if this line doesn't contains an array specifier, or 0 if there was an error parsing the specifier.
+ */
+static TOptional<EOverriddenPropertyOperation> ReadOverriddenOperation(const TCHAR*& Str, FOutputDevice* Warn)
+{
+	const TCHAR* Start = Str;
+	TOptional<EOverriddenPropertyOperation> Operation;
+	SkipWhitespace(Str);
+
+	if (*Str == '<')
+	{
+		Str++;
+		FString IndexText(TEXT(""));
+		while (*Str && *Str != '>')
+		{
+			if (*Str == TCHAR('='))
+			{
+				// we've encountered an equals sign before the closing bracket
+				Warn->Logf(ELogVerbosity::Warning, TEXT("Missing '>' in default properties subscript: %s"), Start);
+				return Operation;
+			}
+
+			IndexText += *Str++;
+		}
+
+		if (*Str++)
+		{
+			if (IndexText.Len() > 0)
+			{
+				Operation = GetOverriddenOperationFromString(IndexText);
+				if (!Operation.IsSet())
+				{
+					// unknown or invalid identifier specified for array subscript
+					Warn->Logf(ELogVerbosity::Warning, TEXT("Invalid overridable in default properties: %s"), Start);
+				}
+			}
+			else
+			{
+				// nothing was specified between the opening and closing parenthesis
+				Warn->Logf(ELogVerbosity::Warning, TEXT("Empty overridable operation: %s"), Start);
+			}
+		}
+		else
+		{
+			Warn->Logf(ELogVerbosity::Warning, TEXT("Missing '>' in default properties subscript: %s"), Start );
+		}
+	}
+	if (FOverridableSerializationLogic::HasCapabilities(FOverridableSerializationLogic::ECapabilities::T3DSerialization))
+	{
+		return Operation;
+	}
+	return TOptional<EOverriddenPropertyOperation>();
+}
+
+
 /** 
  * Do not attempt to import this property if there is no value for it - i.e. (Prop1=,Prop2=)
  * This normally only happens for empty strings or empty dynamic arrays, and the alternative
@@ -1713,7 +1779,7 @@ const TCHAR* FProperty::ImportSingleProperty( const TCHAR* Str, void* DestData,
 	check(ObjectStruct);
 
 	constexpr FAsciiSet Whitespaces(" \t");
-	constexpr FAsciiSet Delimiters("=([.");
+	constexpr FAsciiSet Delimiters("=([.<");
 
 	// strip leading whitespace
 	const TCHAR* Start = FAsciiSet::Skip(Str, Whitespaces);
@@ -1778,6 +1844,8 @@ const TCHAR* FProperty::ImportSingleProperty( const TCHAR* Str, void* DestData,
 			return Str;
 		}
 
+		FOverridableTextPortPropertyPathScope Scope(Property);
+
 		// Parse an array operation, if present.
 		enum EArrayOp
 		{
@@ -1811,6 +1879,27 @@ const TCHAR* FProperty::ImportSingleProperty( const TCHAR* Str, void* DestData,
 		}
 
 		FArrayProperty* const ArrayProperty = ExactCastField<FArrayProperty>(Property);
+
+		// Parse overridable info
+		TOptional<EOverriddenPropertyOperation> Operation = ReadOverriddenOperation(Str, Warn);
+		FOverriddenPropertySet* OverriddenProperties = nullptr;
+		FPropertyVisitorPath* Path = nullptr;
+		// Skip add and remove operations as they should be handled independently
+		if (Operation.IsSet())
+		{
+			OverriddenProperties = FOverridableSerializationLogic::GetOverriddenProperties();
+			Path = FOverridableSerializationLogic::GetOverriddenPortTextPropertyPath();
+			if (OverriddenProperties &&
+				!Property->HasAllPropertyFlags(CPF_ExperimentalAlwaysOverriden) &&
+				Operation.GetValue() != EOverriddenPropertyOperation::Add && 
+				Operation.GetValue() != EOverriddenPropertyOperation::Remove)
+			{
+				checkf(Path, TEXT("Expecting a path"));
+				FArchiveSerializedPropertyChain Chain = Path->ToSerializedPropertyChain();
+				OverriddenProperties->SetOverriddenPropertyOperation(Operation.GetValue(), &Chain, /*Property*/nullptr);
+			}
+		}
+
 		FMulticastDelegateProperty* const MulticastDelegateProperty = CastField<FMulticastDelegateProperty>(Property);
 		if( MulticastDelegateProperty != NULL && ArrayOp != ADO_None )
 		{
@@ -2007,8 +2096,270 @@ const TCHAR* FProperty::ImportSingleProperty( const TCHAR* Str, void* DestData,
 			// try to read an array index
 			int32 Index = ReadArrayIndex(ObjectStruct, Str, Warn);
 
+			// strip whitespace before =
+			SkipWhitespace(Str);
+			if (*Str++ != '=')
+			{
+				Warn->Logf(ELogVerbosity::Warning, TEXT("Missing '=' in default properties assignment: %s"), Start );
+				return Str;
+			}
+			// strip whitespace after =
+			SkipWhitespace(Str);
+
+			auto ImportText = [SubobjectOuter, PortFlags, Warn, Start, ObjectStruct](const TCHAR* Buffer, const FProperty* Property, void* PropertyPtr)
+			{
+				FStringOutputDevice ImportError;
+				const TCHAR* Result = Property->ImportText_Direct(Buffer, PropertyPtr, SubobjectOuter, PortFlags, &ImportError);
+
+				// Spit any error we had while importing property
+				if (ImportError.Len() > 0)
+				{
+					TArray<FString> ImportErrors;
+					ImportError.ParseIntoArray(ImportErrors,LINE_TERMINATOR,true);
+
+					Warn->Logf(ELogVerbosity::Warning, TEXT("While importing text for property '%s' in '%s':"), *Property->GetName(), *ObjectStruct->GetName());
+					for (int32 ErrorIndex = 0; ErrorIndex < ImportErrors.Num(); ErrorIndex++)
+					{
+						Warn->Logf(ELogVerbosity::Warning, TEXT("%s"), *ImportErrors[ErrorIndex]);
+					}
+				}
+				else if (Result == Buffer && *Buffer == TCHAR('\0'))
+				{
+					Warn->Logf(ELogVerbosity::Warning, TEXT("Invalid property value in defaults: %s"), Start);
+				}
+				return Result;
+			};
+
+			// Handle modifying overridable operations on arrays
+			if (Operation.IsSet() && Operation.GetValue() != EOverriddenPropertyOperation::Replace)
+			{
+				if (ArrayProperty)
+				{
+					FScriptArrayHelper_InContainer ArrayHelper(ArrayProperty, DestData);
+					switch(Operation.GetValue())
+					{
+						case EOverriddenPropertyOperation::Remove:
+						{
+							checkf(ArrayProperty->Inner->HasAnyPropertyFlags(CPF_PersistentInstance), TEXT("Only instanced sub object is supporting remove operation"));
+							const FObjectPropertyBase* InnerObjectProperty = CastFieldChecked<FObjectPropertyBase>(ArrayProperty->Inner);
+
+							uint8* TempValueStorage = (uint8*)FMemory_Alloca(InnerObjectProperty->GetElementSize());
+							InnerObjectProperty->InitializeValue(TempValueStorage);
+
+							const TCHAR* Result = ImportText(Str, InnerObjectProperty, TempValueStorage);
+							if (UObject* RemovedSubObject = InnerObjectProperty->GetObjectPropertyValue(TempValueStorage))
+							{
+								const int32 ArrayNum = ArrayHelper.Num();
+								for (int i = 0; i < ArrayNum; ++i)
+								{
+									UObject* CurrentObject = InnerObjectProperty->GetObjectPropertyValue(ArrayHelper.GetElementPtr(i));
+									if (CurrentObject->GetArchetype() == RemovedSubObject)
+									{
+										ArrayHelper.RemoveValues(i);
+										break;
+									}
+								}
+
+								if (Path)
+								{
+									// Need to fetch the ArrayOverriddenPropertyNode every loop as the previous iteration might have reallocated the node.
+									FArchiveSerializedPropertyChain Chain = Path->ToSerializedPropertyChain();
+									if (FOverriddenPropertyNode* ArrayOverriddenPropertyNode = OverriddenProperties ? OverriddenProperties->SetOverriddenPropertyOperation(EOverriddenPropertyOperation::Modified, &Chain, /*Property*/nullptr) : nullptr)
+									{
+										// Rebuild the overridden info
+										const FOverriddenPropertyNodeID RemovedSubObjectID(*RemovedSubObject);
+										OverriddenProperties->SetSubPropertyOperation(EOverriddenPropertyOperation::Remove, *ArrayOverriddenPropertyNode, RemovedSubObjectID);
+									}
+								}
+							}
+
+							InnerObjectProperty->DestroyValue(TempValueStorage);
+							return Result;
+						}
+						case EOverriddenPropertyOperation::Add:
+						{
+							// Special case for instanced sub objects
+							if (const FObjectPropertyBase* InnerObjectProperty =  ArrayProperty->Inner->HasAnyPropertyFlags(CPF_PersistentInstance) ? CastField<FObjectPropertyBase>(ArrayProperty->Inner) : nullptr)
+							{
+								uint8* TempValueStorage = (uint8*)FMemory_Alloca(InnerObjectProperty->GetElementSize());
+								InnerObjectProperty->InitializeValue(TempValueStorage);
+
+								const TCHAR* Result = ImportText(Str, InnerObjectProperty, TempValueStorage);
+								if (UObject* AddedSubObject = InnerObjectProperty->GetObjectPropertyValue(TempValueStorage))
+								{
+									UObject* AddedSubObjectArchetype = AddedSubObject->GetArchetype();
+									const int32 ArrayNum = ArrayHelper.Num();
+									for (int i = 0; i < ArrayNum; ++i)
+									{
+										UObject* CurrentObject = InnerObjectProperty->GetObjectPropertyValue(ArrayHelper.GetElementPtr(i));
+										if (CurrentObject == AddedSubObject || (CurrentObject == AddedSubObjectArchetype))
+										{
+											Index = i;
+											break;
+										}
+									}
+
+									if (Index == INDEX_NONE)
+									{
+										Index = ArrayHelper.Num();
+										ArrayHelper.ExpandForIndex(Index);
+									}
+
+									InnerObjectProperty->SetObjectPropertyValue(ArrayHelper.GetRawPtr(Index), AddedSubObject);
+
+									if (Path)
+									{
+										// Need to fetch the ArrayOverriddenPropertyNode every loop as the previous iteration might have reallocated the node.
+										FArchiveSerializedPropertyChain Chain = Path->ToSerializedPropertyChain();
+										if (FOverriddenPropertyNode* ArrayOverriddenPropertyNode = OverriddenProperties ? OverriddenProperties->SetOverriddenPropertyOperation(EOverriddenPropertyOperation::Modified, &Chain, /*Property*/nullptr) : nullptr)
+										{
+											// Rebuild the overridden info
+											const FOverriddenPropertyNodeID AddedSubObjectID(*AddedSubObject);
+											OverriddenProperties->SetSubPropertyOperation(EOverriddenPropertyOperation::Add, *ArrayOverriddenPropertyNode, AddedSubObjectID);
+										}
+									}
+								}
+
+								InnerObjectProperty->DestroyValue(TempValueStorage);
+								return Result;
+							}
+
+							Index = ArrayHelper.Num();
+							ArrayHelper.ExpandForIndex(Index);
+							break;
+						}
+
+						default:
+							checkf(false, TEXT("Unsupported array operation while importing text"));
+							return nullptr;
+					}
+				}
+				else if (FMapProperty* const MapProperty = ExactCastField<FMapProperty>(Property))
+				{
+					FScriptMapHelper_InContainer MapHelper(MapProperty, DestData);
+
+					uint8* TempKeyValueStorage = (uint8*)FMemory_Alloca(MapProperty->MapLayout.SetLayout.Size);
+					MapProperty->KeyProp->InitializeValue(TempKeyValueStorage);
+					MapProperty->ValueProp->InitializeValue(TempKeyValueStorage + MapProperty->MapLayout.ValueOffset);
+					ON_SCOPE_EXIT
+					{
+						MapProperty->KeyProp->DestroyValue(TempKeyValueStorage);
+						MapProperty->ValueProp->DestroyValue(TempKeyValueStorage + MapProperty->MapLayout.ValueOffset);
+					};
+
+					switch(Operation.GetValue())
+					{
+						case EOverriddenPropertyOperation::Remove:
+						{
+							const TCHAR* Result = ImportText(Str, MapHelper.KeyProp, TempKeyValueStorage);
+							MapHelper.RemovePair(TempKeyValueStorage);
+
+							if (Path)
+							{
+								// Need to fetch the MapOverriddenPropertyNode every loop as the previous might have reallocated the node.
+								FArchiveSerializedPropertyChain Chain = Path->ToSerializedPropertyChain();
+								if (FOverriddenPropertyNode* MapOverriddenPropertyNode = OverriddenProperties ? OverriddenProperties->SetOverriddenPropertyOperation(EOverriddenPropertyOperation::Modified, &Chain, /*Property*/nullptr) : nullptr)
+								{
+									// Rebuild the overridden info
+									FOverriddenPropertyNodeID RemovedKeyID = FOverriddenPropertyNodeID::FromMapKey(MapProperty->KeyProp, TempKeyValueStorage);
+									OverriddenProperties->SetSubPropertyOperation(EOverriddenPropertyOperation::Remove, *MapOverriddenPropertyNode, RemovedKeyID);
+								}
+							}
+
+							return Result;
+						}
+						case EOverriddenPropertyOperation::Modified:
+						{
+							const TCHAR* Result = Str;
+
+							SkipWhitespace(Result);
+							if (*Result++ != TCHAR('('))
+							{
+								return nullptr;
+							}
+							SkipWhitespace(Result);
+
+							Result = ImportText(Result, MapHelper.KeyProp, TempKeyValueStorage);
+							const int32 InternalIndex = MapHelper.FindMapPairIndexFromHash(TempKeyValueStorage);
+
+							SkipWhitespace(Result);
+							if (*Result++ != TCHAR(','))
+							{
+								return nullptr;
+							}
+
+							uint8* ValuePtr = InternalIndex != INDEX_NONE ? MapHelper.GetValuePtr(InternalIndex) : TempKeyValueStorage + MapHelper.MapLayout.ValueOffset;
+							Result = ImportText(Result, MapHelper.ValueProp, ValuePtr);
+
+							SkipWhitespace(Result);
+							if (*Result++ != TCHAR(')'))
+							{
+								return nullptr;
+							}
+
+							if (Path)
+							{
+								// Need to fetch the MapOverriddenPropertyNode every loop as the previous might have reallocated the node.
+								FArchiveSerializedPropertyChain Chain = Path->ToSerializedPropertyChain();
+								if (FOverriddenPropertyNode* MapOverriddenPropertyNode = OverriddenProperties ? OverriddenProperties->SetOverriddenPropertyOperation(EOverriddenPropertyOperation::Modified, &Chain, /*Property*/nullptr) : nullptr)
+								{
+									// Rebuild the overridden info
+									FOverriddenPropertyNodeID ModifiedKeyID = FOverriddenPropertyNodeID::FromMapKey(MapProperty->KeyProp, TempKeyValueStorage);
+									OverriddenProperties->SetSubPropertyOperation(EOverriddenPropertyOperation::Modified, *MapOverriddenPropertyNode, ModifiedKeyID);
+								}
+							}
+							return Result;
+						}
+						case EOverriddenPropertyOperation::Add:
+						{
+							const TCHAR* Result = Str;
+							SkipWhitespace(Result);
+							if (*Result++ != TCHAR('('))
+							{
+								return nullptr;
+							}
+							SkipWhitespace(Result);
+
+							Result = ImportText(Result, MapHelper.KeyProp, TempKeyValueStorage);
+
+							SkipWhitespace(Result);
+							if (*Result++ != TCHAR(','))
+							{
+								return nullptr;
+							}
+							SkipWhitespace(Result);
+
+							void* ValuePtr =  MapHelper.FindOrAdd(TempKeyValueStorage);;
+							Result = ImportText(Result, MapHelper.ValueProp, ValuePtr);
+
+							SkipWhitespace(Result);
+							if (*Result++ != TCHAR(')'))
+							{
+								return nullptr;
+							}
+
+							if (Path)
+							{
+								// Need to fetch the MapOverriddenPropertyNode every loop as the previous might have reallocated the node.
+								FArchiveSerializedPropertyChain Chain = Path->ToSerializedPropertyChain();
+								if (FOverriddenPropertyNode* MapOverriddenPropertyNode = OverriddenProperties ? OverriddenProperties->SetOverriddenPropertyOperation(EOverriddenPropertyOperation::Modified, &Chain, /*Property*/nullptr) : nullptr)
+								{
+									// Rebuild the overridden info
+									FOverriddenPropertyNodeID AddedKeyID = FOverriddenPropertyNodeID::FromMapKey(MapProperty->KeyProp, TempKeyValueStorage);
+									OverriddenProperties->SetSubPropertyOperation(EOverriddenPropertyOperation::Add, *MapOverriddenPropertyNode, AddedKeyID);
+								}
+							}
+							return Result;
+						}
+						default:
+							checkf(false, TEXT("Unsupported map operation while importing text"));
+							return nullptr;
+					}
+				}
+			}
+
 			// check for out of bounds on static arrays
-			if (ArrayProperty == NULL && Index >= Property->ArrayDim)
+			if (ArrayProperty == nullptr && Index >= Property->ArrayDim)
 			{
 				Warn->Logf(ELogVerbosity::Warning, TEXT("Out of bound array default property (%i/%i): %s"), Index, Property->ArrayDim, Start);
 				return Str;
@@ -2025,16 +2376,6 @@ const TCHAR* FProperty::ImportSingleProperty( const TCHAR* Str, void* DestData,
 			}
 			DefinedProperties.Add(D);
 
-			// strip whitespace before =
-			SkipWhitespace(Str);
-			if (*Str++ != '=')
-			{
-				Warn->Logf(ELogVerbosity::Warning, TEXT("Missing '=' in default properties assignment: %s"), Start );
-				return Str;
-			}
-			// strip whitespace after =
-			SkipWhitespace(Str);
-
 			if (!IsPropertyValueSpecified(Str) && ArrayProperty == nullptr)
 			{
 				// if we're not importing default properties for classes (i.e. we're pasting something in the editor or something)
@@ -2047,31 +2388,15 @@ const TCHAR* FProperty::ImportSingleProperty( const TCHAR* Str, void* DestData,
 			// not done above with ShouldPort() check because this is intentionally exported so we don't want it to cause errors on import
 			if (Property->GetFName() != NAME_Name || !Property->GetOwnerVariant().IsUObject() || Property->GetOwner<UObject>()->GetFName() != NAME_Object)
 			{
-				if (Index > -1 && ArrayProperty != NULL) //set single dynamic array element
+				if (Index > -1 && ArrayProperty != nullptr) //set single dynamic array element
 				{
 					FScriptArrayHelper_InContainer ArrayHelper(ArrayProperty, DestData);
 
 					ArrayHelper.ExpandForIndex(Index);
 
-					FStringOutputDevice ImportError;
-					const TCHAR* Result = ArrayProperty->Inner->ImportText_Direct(Str, ArrayHelper.GetRawPtr(Index), SubobjectOuter, PortFlags, &ImportError);
-					// Spit any error we had while importing property
-					if (ImportError.Len() > 0)
-					{
-						TArray<FString> ImportErrors;
-						ImportError.ParseIntoArray(ImportErrors,LINE_TERMINATOR,true);
-
-						for ( int32 ErrorIndex = 0; ErrorIndex < ImportErrors.Num(); ErrorIndex++ )
-						{
-							Warn->Logf(ELogVerbosity::Warning, TEXT("%s"), *ImportErrors[ErrorIndex]);
-						}
-					}
-					else if (Result == Str)
-					{
-						Warn->Logf(ELogVerbosity::Warning, TEXT("Invalid property value in defaults: %s"), Start);
-					}
+					const TCHAR* Result = ImportText(Str, ArrayProperty->Inner, ArrayHelper.GetRawPtr(Index));
 					// in the failure case, don't return NULL so the caller can potentially skip less and get values further in the string
-					if (Result != NULL)
+					if (Result != nullptr)
 					{
 						Str = Result;
 					}
@@ -2083,28 +2408,10 @@ const TCHAR* FProperty::ImportSingleProperty( const TCHAR* Str, void* DestData,
 						Index = 0;
 					}
 
-					FStringOutputDevice ImportError;
 
-					const TCHAR* Result = Property->ImportText_Direct(Str, Property->ContainerPtrToValuePtr<void>(DestData, Index), SubobjectOuter, PortFlags, &ImportError);
-					
-					// Spit any error we had while importing property
-					if (ImportError.Len() > 0)
-					{
-						TArray<FString> ImportErrors;
-						ImportError.ParseIntoArray(ImportErrors, LINE_TERMINATOR, true);
-
-						Warn->Logf(ELogVerbosity::Warning, TEXT("While importing text for property '%s' in '%s':"), *Property->GetName(), *ObjectStruct->GetName());
-						for ( int32 ErrorIndex = 0; ErrorIndex < ImportErrors.Num(); ErrorIndex++ )
-						{
-							Warn->Logf(ELogVerbosity::Warning, TEXT("%s"), *ImportErrors[ErrorIndex]);
-						}
-					}
-					else if ((Result == NULL && ArrayProperty == nullptr) || Result == Str)
-					{
-						UE_SUPPRESS(LogExec, Verbose, Warn->Logf(TEXT("Unknown property in %s: %s "), *ObjectStruct->GetName(), Start));
-					}
+					const TCHAR* Result = ImportText(Str, Property, Property->ContainerPtrToValuePtr<void>(DestData, Index));
 					// in the failure case, don't return NULL so the caller can potentially skip less and get values further in the string
-					if (Result != NULL)
+					if (Result != nullptr)
 					{
 						Str = Result;
 					}
diff --git a/Engine/Source/Runtime/CoreUObject/Private/UObject/PropertyOptional.cpp b/Engine/Source/Runtime/CoreUObject/Private/UObject/PropertyOptional.cpp
--- a/Engine/Source/Runtime/CoreUObject/Private/UObject/PropertyOptional.cpp
+++ b/Engine/Source/Runtime/CoreUObject/Private/UObject/PropertyOptional.cpp
@@ -703,8 +703,13 @@ EPropertyVisitorControlFlow FOptionalProperty::Visit(FPropertyVisitorPath& Path,
 
 	EPropertyVisitorControlFlow RetVal = Super::Visit(Path, Data, InFunc);
 
-	if (RetVal == EPropertyVisitorControlFlow::StepInto && IsSet(Data))
+	if (RetVal == EPropertyVisitorControlFlow::StepInto)
 	{
+		if(!IsSet(Data))
+		{
+			// There is nothing to StepInto, so continue to next property
+			return EPropertyVisitorControlFlow::StepOver;
+		}
 		checkf(ValueProperty, TEXT("Expecting a valid property value"));
 
 		FPropertyVisitorScope Scope(Path, FPropertyVisitorInfo(ValueProperty));
diff --git a/Engine/Source/Runtime/CoreUObject/Public/UObject/OverridableManager.h b/Engine/Source/Runtime/CoreUObject/Public/UObject/OverridableManager.h
--- a/Engine/Source/Runtime/CoreUObject/Public/UObject/OverridableManager.h
+++ b/Engine/Source/Runtime/CoreUObject/Public/UObject/OverridableManager.h
@@ -92,6 +92,7 @@ enum class EOverriddenState : uint8
 	Added, // This object was added
 };
 
+struct FOverrideManagerCapabilityInterface;
 
 /*
  *************************************************************************************
@@ -132,8 +133,9 @@ public:
 	 * NOTE: ONLY CACHE THE ARCHETYPE POINTER ON OBJECTS THAT ARE ABOUT TO BE REINSTANTIATED.
 	 * This is only a utility method to ensure the archetype is not lost during reinstantiation
 	 * as it renames and changes outers which makes impossible to retrieve it during the process.
-	 * @param Object top cache its archetype */
-	COREUOBJECT_API void CacheArchetype(UObject& Object);
+	 * @param Object top cache its archetype
+	 * @param Archetype to use as a cache, if null it will use GetArchetype on the owner of the specified object */
+	COREUOBJECT_API void CacheArchetype(UObject& Object, UObject* Archetype = nullptr);
 
 	/**
 	 * Return true if this object needs subobject template instantiation
diff --git a/Engine/Source/Runtime/CoreUObject/Public/UObject/OverriddenPropertySet.h b/Engine/Source/Runtime/CoreUObject/Public/UObject/OverriddenPropertySet.h
--- a/Engine/Source/Runtime/CoreUObject/Public/UObject/OverriddenPropertySet.h
+++ b/Engine/Source/Runtime/CoreUObject/Public/UObject/OverriddenPropertySet.h
@@ -10,6 +10,7 @@ COREUOBJECT_API DECLARE_LOG_CATEGORY_EXTERN(LogOverridableObject, Warning, All);
 
 struct FArchiveSerializedPropertyChain;
 
+struct FOverridableSerializationLogicInternalAdapter;
 /*
  *************************************************************************************
  * Overridable serialization is experimental, not supported and use at your own risk *
@@ -19,12 +20,22 @@ struct FArchiveSerializedPropertyChain;
 struct FOverridableSerializationLogic
 {
 public:
+	
+	enum class ECapabilities : uint32
+	{
+		None = 0,
+		// Enables serialization of override state into and from T3D
+		T3DSerialization = 1 << 0,
+		All = T3DSerialization
+	};
+	
+	FRIEND_ENUM_CLASS_FLAGS(FOverridableSerializationLogic::ECapabilities)
 
 	/**
 	 * Call to enable overridable serialization and to set the overridden properties of the current serialized object
 	 * Note this is not re-entrant and it stores information in a thread local storage
 	 * @param InOverriddenProperties of the current serializing object */
-	FORCEINLINE static void Enable(FOverriddenPropertySet* InOverriddenProperties)
+	inline static void Enable(FOverriddenPropertySet* InOverriddenProperties)
 	{
 		checkf(!bUseOverridableSerialization, TEXT("Nobody should use this method if overridable serialization is already enabled"));
 		bUseOverridableSerialization = true;
@@ -34,7 +45,7 @@ public:
 	/**
 	 * Call to disable overridable serialization
 	 * Note this is not re-entrant and it stores information in a thread local storage */
-	FORCEINLINE static void Disable()
+	inline static void Disable()
 	{
 		checkf(bUseOverridableSerialization, TEXT("Expecting overridable serialization to be already enabled"));
 		OverriddenProperties = nullptr;
@@ -44,17 +55,24 @@ public:
 	/**
 	 * Called during the serialization of an object to know to know if it should do overridden serialization logic
 	 * @return true if the overridable serialization is enabled on the current serializing object */
-	FORCEINLINE static bool IsEnabled()
+	inline static bool IsEnabled()
 	{
 		return bUseOverridableSerialization;
 	}
 
+	
+	/**
+	 * Used to enable override features selectively during development
+	 * Capabilities driven by CVars
+	 */
+	COREUOBJECT_API static bool HasCapabilities(ECapabilities InCapabilities);
+
 	/**
 	 * Call during the serialization of an object to get its overriden properties
 	 * Note: Expects the current serialized object to use overridable serialization
 	 * Note this is not re-entrant and it stores information in a thread local storage
 	 * @return the overridden properties of the current object being serialized */
-	FORCEINLINE static FOverriddenPropertySet* GetOverriddenProperties()
+	inline static FOverriddenPropertySet* GetOverriddenProperties()
 	{
 		return OverriddenProperties;
 	}
@@ -66,11 +84,39 @@ public:
 	 * @param DataPtr to the memory of that property
 	 * @param DefaultValue memory pointer of that property
 	 * @return the overridden property operation */
-	static EOverriddenPropertyOperation GetOverriddenPropertyOperation(const FArchive& Ar, FProperty* Property = nullptr, uint8* DataPtr = nullptr, uint8* DefaultValue = nullptr);
+	COREUOBJECT_API static EOverriddenPropertyOperation GetOverriddenPropertyOperation(const FArchive& Ar, FProperty* Property = nullptr, uint8* DataPtr = nullptr, uint8* DefaultValue = nullptr);
+
+	/**
+	 * Use the port text path to retrieve the current overridden property operation to know if it has to be serialized or not
+	 * @param DataPtr to the memory of that property
+	 * @param DefaultValue memory pointer of that property
+	 * @param PortFlags for the import/export text operation
+	 * @return the overridden property operation */
+	COREUOBJECT_API static EOverriddenPropertyOperation GetOverriddenPropertyOperationForPortText(const void* DataPtr, const void* DefaultValue, int32 PortFlags);
+
+	/**
+	 * Call during the import text
+	 * @return the current text import property path */
+	COREUOBJECT_API static FPropertyVisitorPath* GetOverriddenPortTextPropertyPath();
+
+	/**
+	 * Call during the import text to set the property path
+	 * @param Path to start tracking the property path*/
+	COREUOBJECT_API static void SetOverriddenPortTextPropertyPath(FPropertyVisitorPath& Path);
+
+	/** Call during the import text to reset property path */
+	COREUOBJECT_API static void ResetOverriddenPortTextPropertyPath();
 
 private:
+
+	static EOverriddenPropertyOperation GetOverriddenPropertyOperation(const int32 PortFlags, const FArchiveSerializedPropertyChain* CurrentPropertyChain, FProperty* Property, const void* DataPtr, const void* DefaultValue);
+
+	friend struct FOverridableSerializationLogicInternalAdapter;
+	
+	static ECapabilities Capabilities;
 	static thread_local bool bUseOverridableSerialization;
 	static thread_local FOverriddenPropertySet* OverriddenProperties;
+	static thread_local FPropertyVisitorPath* OverriddenPortTextPropertyPath;
 };
 
 /*
@@ -81,8 +127,8 @@ private:
 */
 struct FEnableOverridableSerializationScope
 {
-	FEnableOverridableSerializationScope(bool bEnableOverridableSerialization, FOverriddenPropertySet* OverriddenProperties);
-	~FEnableOverridableSerializationScope();
+	COREUOBJECT_API FEnableOverridableSerializationScope(bool bEnableOverridableSerialization, FOverriddenPropertySet* OverriddenProperties);
+	COREUOBJECT_API ~FEnableOverridableSerializationScope();
 
 protected:
 	bool bOverridableSerializationEnabled = false;
@@ -90,6 +136,24 @@ protected:
 	FOverriddenPropertySet* SavedOverriddenProperties = nullptr;
 };
 
+
+/*
+ *************************************************************************************
+ * Overridable serialization is experimental, not supported and use at your own risk *
+ *************************************************************************************
+ * Scope responsible for tracking current property path for text importing
+*/
+struct FOverridableTextPortPropertyPathScope
+{
+	COREUOBJECT_API FOverridableTextPortPropertyPathScope(const FProperty* InProperty, int32 InIndex = INDEX_NONE, EPropertyVisitorInfoType InPropertyInfo = EPropertyVisitorInfoType::None);
+	COREUOBJECT_API ~FOverridableTextPortPropertyPathScope();
+
+protected:
+
+	const FProperty* Property = nullptr;
+	FPropertyVisitorPath DefaultPath;
+};
+
 /*
  * Override operation type for each property node
  */
@@ -103,6 +167,27 @@ enum class EOverriddenPropertyOperation : uint8
 	Remove,		/* this element was removed from the container */
 };
 
+inline TOptional<EOverriddenPropertyOperation> GetOverriddenOperationFromString(const FString& OverriddenOperationString)
+{
+	int64 EnumValue = StaticEnum<EOverriddenPropertyOperation>()->GetValueByNameString(StaticEnum<EOverriddenPropertyOperation>()->GetName() + FString(TEXT("::")) + OverriddenOperationString);
+	if (EnumValue == INDEX_NONE)
+	{
+		TOptional<EOverriddenPropertyOperation>();
+	}
+
+	return TOptional((EOverriddenPropertyOperation)EnumValue);
+}
+
+inline TOptional<EOverriddenPropertyOperation> GetOverriddenOperationFromName(const FName OverriddenOperationName)
+{
+	return GetOverriddenOperationFromString(OverriddenOperationName.ToString());
+}
+
+inline FString GetOverriddenOperationString(EOverriddenPropertyOperation Operation)
+{
+	return UEnum::GetValueAsString(Operation).RightChop(StaticEnum<EOverriddenPropertyOperation>()->GetName().Len() + /*::*/2);
+}
+
 USTRUCT()
 struct FOverriddenPropertyNodeID
 {
@@ -123,7 +208,7 @@ struct FOverriddenPropertyNodeID
 
 	// Overridden property node map helpers
 	static COREUOBJECT_API FOverriddenPropertyNodeID FromMapKey(const FProperty* KeyProperty, const void* KeyData);
-	int32 ToMapInternalIndex(FScriptMapHelper& MapHelper) const;
+	COREUOBJECT_API int32 ToMapInternalIndex(FScriptMapHelper& MapHelper) const;
 
 	bool operator==(const FOverriddenPropertyNodeID& Other) const
 	{
@@ -295,8 +380,10 @@ public:
 	void HandleObjectsReInstantiated(const TMap<UObject*, UObject*>& Map);
 
 	/**
-	 * Caches the archetype of this object for future reference */
-	void CacheArchetype();
+	 * Caches the archetype of this object for future reference
+	 * @param Archetype to use as a cache, if null it will use GetArchetype on the owner of this object
+	 */
+	void CacheArchetype(UObject* Archetype = nullptr);
 
 	/**
 	 * Gets the archetype that was cached when the CacheArchetype method was called
diff --git a/Engine/Source/Runtime/Engine/Private/UnrealExporter.cpp b/Engine/Source/Runtime/Engine/Private/UnrealExporter.cpp
--- a/Engine/Source/Runtime/Engine/Private/UnrealExporter.cpp
+++ b/Engine/Source/Runtime/Engine/Private/UnrealExporter.cpp
@@ -16,8 +16,12 @@
 #include "Model.h"
 #include "Misc/FeedbackContext.h"
 #include "AssetExportTask.h"
+#include "DataTableUtils.h"
 #include "Misc/AsciiSet.h"
+#include "Serialization/ArchiveSerializedPropertyChain.h"
 #include "UObject/GCObjectScopeGuard.h"
+#include "UObject/OverridableManager.h"
+#include "UObject/OverriddenPropertySet.h"
 
 #if WITH_EDITOR
 #include "Editor.h"
@@ -28,6 +32,18 @@
 
 DEFINE_LOG_CATEGORY_STATIC(LogExporter, Log, All);
 
+namespace Private
+{
+	FOverriddenPropertySet* GetOverriddenProperties(UObject& Object)
+	{
+		if (FOverridableSerializationLogic::HasCapabilities(FOverridableSerializationLogic::ECapabilities::T3DSerialization))
+		{
+			return FOverridableManager::Get().GetOverriddenProperties(Object);
+		}
+		return nullptr;
+	}
+}
+
 FString UExporter::CurrentFilename(TEXT(""));
 
 TSet< TWeakObjectPtr<UExporter> > UExporter::RegisteredExporters;
@@ -471,6 +487,12 @@ void UExporter::EmitBeginObject( FOutputDevice& Ar, UObject* Obj, uint32 PortFla
 			// we need the fully qualified path for the archetype (so we don't get confused when unpacking this)
 			Ar.Logf(TEXT(" Archetype=%s"), *FObjectPropertyBase::GetExportPath(Archetype, nullptr, /*ExportRootScope =*/nullptr, (PortFlags | PPF_Delimited) & ~PPF_ExportsNotFullyQualified));
 		}
+
+		if (FOverriddenPropertySet* ObjectOverriddenProperties = Private::GetOverriddenProperties(*Obj))
+		{
+			const EOverriddenPropertyOperation Operation = ObjectOverriddenProperties->GetOverriddenPropertyOperation((FArchiveSerializedPropertyChain*)nullptr, (FProperty*)nullptr);
+			Ar.Logf(TEXT(" OverriddenOperation=%s"), *GetOverriddenOperationString(Operation));
+		}
 	}
 
 	// When exporting for diffs, export paths can cause false positives. since diff files don't get imported, we can
@@ -672,7 +694,7 @@ void ExportProperties
 	FOutputDevice&	Out,
 	UClass*			ObjectClass,
 	uint8*			Object,
-	int32				Indent,
+	int32			Indent,
 	UClass*			DiffClass,
 	uint8*			Diff,
 	UObject*		Parent,
@@ -682,10 +704,15 @@ void ExportProperties
 {
 	check(ObjectClass != NULL);
 
-	for( FProperty* Property = ObjectClass->PropertyLink; Property; Property = Property->PropertyLinkNext )
+	FOverriddenPropertySet* OverriddenProperties = Object ? Private::GetOverriddenProperties(*(UObject*)Object) : nullptr;
+	FEnableOverridableSerializationScope Scope(OverriddenProperties!= nullptr, OverriddenProperties);
+
+	for (FProperty* Property = ObjectClass->PropertyLink; Property; Property = Property->PropertyLinkNext)
 	{
 		if (!Property->ShouldPort(PortFlags))
+		{
 			continue;
+		}
 
 		FString SanitizedPropertyName = Property->GetName();
 		constexpr FAsciiSet Whitespace("\t ");
@@ -697,48 +724,110 @@ void ExportProperties
 		}
 
 		FArrayProperty* ArrayProperty = CastField<FArrayProperty>(Property);
+		FMapProperty* MapProperty = CastField<FMapProperty>(Property);
+
 		FObjectPropertyBase* ExportObjectProp = (Property->PropertyFlags & CPF_ExportObject) != 0 ? CastField<FObjectPropertyBase>(Property) : NULL;
 		const uint32 ExportFlags = PortFlags | PPF_Delimited;
-
-		if ( ArrayProperty != NULL )
+		auto HandleExportObject = [Context, &Out, Indent, Parent, PortFlags, ExportObjectProp](const FProperty* Property, uint8* Data)
 		{
-			// Export dynamic array.
-			FProperty* InnerProp = ArrayProperty->Inner;
-			ExportObjectProp = (Property->PropertyFlags & CPF_ExportObject) != 0 ? CastField<FObjectPropertyBase>(InnerProp) : NULL;
-			// This is used as the default value in the case of an array property that has
-			// fewer elements than the exported object.
-			uint8* StructDefaults = NULL;
-			FStructProperty* StructProperty = CastField<FStructProperty>(InnerProp);
-			if ( StructProperty != NULL )
+			if (ExportObjectProp)
 			{
-				checkSlow(StructProperty->Struct);
-				StructDefaults = (uint8*)FMemory::Malloc(StructProperty->Struct->GetStructureSize());
-				StructProperty->InitializeValue(StructDefaults);
+				UObject* Obj = ExportObjectProp->GetObjectPropertyValue(Data);
+				check(!Obj || Obj->IsValidLowLevel());
+				if (Obj && !Obj->HasAnyMarks(OBJECTMARK_TagImp))
+				{
+					// only export the BEGIN OBJECT block for a component if Parent is the component's Outer....when importing subobject definitions,
+					// (i.e. BEGIN OBJECT), whichever BEGIN OBJECT block a component's BEGIN OBJECT block is located within is the object that will be
+					// used as the Outer to create the component
+
+					// Is this an array of components?
+					if (Property->HasAnyPropertyFlags(CPF_InstancedReference))
+					{
+						if (Obj->GetOuter() == Parent)
+						{
+							// Don't export more than once.
+							Obj->Mark(OBJECTMARK_TagImp);
+							UExporter::ExportToOutputDevice(Context, Obj, NULL, Out, TEXT("T3D"), Indent, PortFlags);
+						}
+						else
+						{
+							// set the OBJECTMARK_TagExp flag so that the calling code knows we wanted to export this object
+							Obj->Mark(OBJECTMARK_TagExp);
+						}
+					}
+					else
+					{
+						// Don't export more than once.
+						Obj->Mark(OBJECTMARK_TagImp);
+						UExporter::ExportToOutputDevice(Context, Obj, NULL, Out, TEXT("T3D"), Indent, PortFlags);
+					}
+				}
 			}
+		};
 
-			for( int32 PropertyArrayIndex=0; PropertyArrayIndex<Property->ArrayDim; PropertyArrayIndex++ )
+		for (int32 PropertyArrayIndex=0; PropertyArrayIndex<Property->ArrayDim; PropertyArrayIndex++)
+		{
+			// Array special case
+			if (ArrayProperty != nullptr)
 			{
+				// Export dynamic array.
+				FProperty* InnerProp = ArrayProperty->Inner;
+				ExportObjectProp = (Property->PropertyFlags & CPF_ExportObject) != 0 ? CastField<FObjectPropertyBase>(InnerProp) : NULL;
+				// This is used as the default value in the case of an array property that has
+				// fewer elements than the exported object.
+				uint8* StructDefaults = nullptr;
+				FStructProperty* StructProperty = CastField<FStructProperty>(InnerProp);
+				if (StructProperty != nullptr)
+				{
+					checkSlow(StructProperty->Struct);
+					StructDefaults = (uint8*)FMemory::Malloc(StructProperty->Struct->GetStructureSize());
+					StructProperty->InitializeValue(StructDefaults);
+				}
+				ON_SCOPE_EXIT
+				{
+					if (StructDefaults)
+					{
+						StructProperty->DestroyValue(StructDefaults);
+						FMemory::Free(StructDefaults);
+					}
+				};
+
 				void* Arr = Property->ContainerPtrToValuePtr<void>(Object, PropertyArrayIndex);
 				FScriptArrayHelper ArrayHelper(ArrayProperty, Arr);
 
-				void*	DiffArr = NULL;
-				if( DiffClass )
+				void* DiffArr = nullptr;
+				if (DiffClass)
 				{
 					DiffArr = Property->ContainerPtrToValuePtrForDefaults<void>(DiffClass, Diff, PropertyArrayIndex);
 				}
 				// we won't use this if DiffArr is NULL, but we have to set it up to something
 				FScriptArrayHelper DiffArrayHelper(ArrayProperty, DiffArr);
 
-				// If the current size of the array is 0 and the default one is not, add in an empty item so on import it will be empty
-				if( ArrayHelper.Num() == 0 && DiffArrayHelper.Num() != 0 )
+				EOverriddenPropertyOperation Operation = EOverriddenPropertyOperation::None;
+				FArchiveSerializedPropertyChain Chain;
+				if (OverriddenProperties)
 				{
-					Out.Logf(TEXT("%s%s=\r\n"), FCString::Spc(Indent), *SanitizedPropertyName);
+					FOverridableTextPortPropertyPathScope ScopePath(Property);
+					Operation = FOverridableSerializationLogic::GetOverriddenPropertyOperationForPortText(Arr, DiffArr, PortFlags);
+					FPropertyVisitorPath* Path = FOverridableSerializationLogic::GetOverriddenPortTextPropertyPath();
+					checkf(Path, TEXT("Expecting a path"));
+					Chain = Path->ToSerializedPropertyChain();
 				}
-				else
+
+				if(!OverriddenProperties && !Property->HasAnyPropertyFlags(CPF_ExperimentalOverridableLogic))
 				{
+					// If the current size of the array is 0 and the default one is not, add in an empty item so on import it will be empty
+					if (ArrayHelper.Num() == 0 && DiffArrayHelper.Num() != 0)
+					{
+						Out.Logf(TEXT("%s%s=\r\n"), FCString::Spc(Indent), *SanitizedPropertyName);
+						continue;
+					}
+
 					// If the array sizes are different, we will need to export each index so on import we maintain the size
 					for (int32 DynamicArrayIndex = 0; DynamicArrayIndex < ArrayHelper.Num(); DynamicArrayIndex++)
 					{
+						FOverridableTextPortPropertyPathScope ScopePath(Property, DynamicArrayIndex, EPropertyVisitorInfoType::ContainerIndex);
+
 						FString	Value;
 
 						// compare each element's value manually so that elements which match the NULL value for the array's inner property type
@@ -753,105 +842,246 @@ void ExportProperties
 						if (bExportItem)
 						{
 							InnerProp->ExportTextItem_Direct(Value, SourceData, DiffData, Parent, ExportFlags, ExportRootScope);
-							if (ExportObjectProp)
+							HandleExportObject(InnerProp, ArrayHelper.GetRawPtr(DynamicArrayIndex));
+
+							Out.Logf(TEXT("%s%s(%i)=%s\r\n"), FCString::Spc(Indent), *SanitizedPropertyName, DynamicArrayIndex, *Value);
+						}
+					}
+					for (int32 DynamicArrayIndex = DiffArrayHelper.Num()-1; DynamicArrayIndex >= ArrayHelper.Num(); --DynamicArrayIndex)
+					{
+						Out.Logf(TEXT("%s%s.RemoveIndex(%d)\r\n"), FCString::Spc(Indent), *SanitizedPropertyName, DynamicArrayIndex);
+					}
+					continue;
+				}
+
+				// Replaced operation goes through the generic property export code
+				if (Operation == EOverriddenPropertyOperation::Modified)
+				{
+					if (const FOverriddenPropertyNode* ArrayOverriddenPropertyNode = OverriddenProperties->GetOverriddenPropertyNode(&Chain))
+					{
+						checkf(Arr && DiffArr, TEXT("Expecting a memory ptr to Array and its default"));
+
+						const FObjectProperty* InnerObjectProperty = CastField<FObjectProperty>(InnerProp);
+
+						// Only array of instanced subobjects are handled here as sort of a set where the matching key is done using the archetype
+						checkf(InnerObjectProperty&& InnerObjectProperty->HasAnyPropertyFlags(CPF_PersistentInstance), TEXT("Expecting only arrays of instanced subobjects"));
+
+						auto FindObject = [InnerObjectProperty](const FOverriddenPropertyNodeID ObjectToFind, FScriptArrayHelper& ArrayHelper) -> int32
+						{
+							const int32 ArrayNum = ArrayHelper.Num();
+							for (int i = 0; i < ArrayNum; ++i)
 							{
-								UObject* Obj = ExportObjectProp->GetObjectPropertyValue(ArrayHelper.GetRawPtr(DynamicArrayIndex));
-								check(!Obj || Obj->IsValidLowLevel());
-								if (Obj && !Obj->HasAnyMarks(OBJECTMARK_TagImp))
+								if (UObject* CurrentObject = InnerObjectProperty->GetObjectPropertyValue(ArrayHelper.GetElementPtr(i)))
 								{
-									// only export the BEGIN OBJECT block for a component if Parent is the component's Outer....when importing subobject definitions,
-									// (i.e. BEGIN OBJECT), whichever BEGIN OBJECT block a component's BEGIN OBJECT block is located within is the object that will be
-									// used as the Outer to create the component
+									if (ObjectToFind == FOverriddenPropertyNodeID(*CurrentObject))
+									{
+										return i;
+									}
+								}
+							}
+							return INDEX_NONE;
+						};
+
+						FScriptArrayHelper DefaultsArrayHelper(ArrayProperty, DiffArr);
+
+						TArray<int32> RemovedIndices;
+						TArray<int32> AddedIndices;
 
-									// Is this an array of components?
-									if (InnerProp->HasAnyPropertyFlags(CPF_InstancedReference))
+						for (const auto& Pair : ArrayOverriddenPropertyNode->SubPropertyNodeKeys)
+						{
+							const EOverriddenPropertyOperation OverrideOp = OverriddenProperties->GetSubPropertyOperation(Pair.Value);
+							switch (OverrideOp)
+							{
+							case EOverriddenPropertyOperation::Remove:
+								{
+									const int32 DefaultIndex = FindObject(Pair.Key, DefaultsArrayHelper);
+									if (DefaultIndex != INDEX_NONE)
 									{
-										if (Obj->GetOuter() == Parent)
-										{
-											// Don't export more than once.
-											Obj->Mark(OBJECTMARK_TagImp);
-											UExporter::ExportToOutputDevice(Context, Obj, NULL, Out, TEXT("T3D"), Indent, PortFlags);
-										}
-										else
-										{
-											// set the OBJECTMARK_TagExp flag so that the calling code knows we wanted to export this object
-											Obj->Mark(OBJECTMARK_TagExp);
-										}
+										RemovedIndices.Add(DefaultIndex);
 									}
-									else
+									break;
+								}
+							case EOverriddenPropertyOperation::Add:
+								{
+									const int32 Index = FindObject(Pair.Key, ArrayHelper);
+									if (Index != INDEX_NONE)
 									{
-										// Don't export more than once.
-										Obj->Mark(OBJECTMARK_TagImp);
-										UExporter::ExportToOutputDevice(Context, Obj, NULL, Out, TEXT("T3D"), Indent, PortFlags);
+										AddedIndices.Add(Index);
 									}
+									break;
 								}
+							default:
+								checkf(false, TEXT("Unsupported operation type"));
+								break;
 							}
+						}
 
-							Out.Logf(TEXT("%s%s(%i)=%s\r\n"), FCString::Spc(Indent), *SanitizedPropertyName, DynamicArrayIndex, *Value);
+						for (int32 i : RemovedIndices)
+						{
+							FString Value;
+							InnerProp->ExportTextItem_Direct(Value, DiffArrayHelper.GetRawPtr(i), nullptr, Parent, ExportFlags, ExportRootScope);
+
+							Out.Logf(TEXT("%s%s<%s>=%s\r\n"), FCString::Spc(Indent), *SanitizedPropertyName, *GetOverriddenOperationString(EOverriddenPropertyOperation::Remove), *Value);
+						}
+
+						for (int32 i : AddedIndices)
+						{
+							FOverridableTextPortPropertyPathScope ScopePath(Property, i, EPropertyVisitorInfoType::ContainerIndex);
+
+							FString Value;
+							InnerProp->ExportTextItem_Direct(Value, ArrayHelper.GetRawPtr(i), nullptr, Parent, ExportFlags, ExportRootScope);
+							HandleExportObject(InnerProp, ArrayHelper.GetRawPtr(i));
+
+							Out.Logf(TEXT("%s%s<%s>=%s\r\n"), FCString::Spc(Indent), *SanitizedPropertyName, *GetOverriddenOperationString(EOverriddenPropertyOperation::Add), *Value);
 						}
 					}
-					for (int32 DynamicArrayIndex = DiffArrayHelper.Num()-1; DynamicArrayIndex >= ArrayHelper.Num(); --DynamicArrayIndex)
-					{
-						Out.Logf(TEXT("%s%s.RemoveIndex(%d)\r\n"), FCString::Spc(Indent), *SanitizedPropertyName, DynamicArrayIndex);
-					}
+					continue;
 				}
 			}
-			if (StructDefaults)
-			{
-				StructProperty->DestroyValue(StructDefaults);
-				FMemory::Free(StructDefaults);
-			}
-		}
-		else
-		{
-			for( int32 PropertyArrayIndex=0; PropertyArrayIndex<Property->ArrayDim; PropertyArrayIndex++ )
+
+			// Map special case
+			if (MapProperty != nullptr)
 			{
-				FString	Value;
-				// Export single element.
 
-				uint8* DiffData = (DiffClass && Property->IsInContainer(DiffClass->GetPropertiesSize())) ? Diff : NULL;
-				if( Property->ExportText_InContainer( PropertyArrayIndex, Value, Object, DiffData, Parent, ExportFlags, ExportRootScope ) )
+				void* Map = Property->ContainerPtrToValuePtr<void>(Object, PropertyArrayIndex);
+				void* DiffMap = nullptr;
+				if (DiffClass)
 				{
-					if ( ExportObjectProp )
+					DiffMap = Property->ContainerPtrToValuePtrForDefaults<void>(DiffClass, Diff, PropertyArrayIndex);
+				}
+
+				EOverriddenPropertyOperation Operation = EOverriddenPropertyOperation::None;
+				FArchiveSerializedPropertyChain Chain;
+				if (OverriddenProperties)
+				{
+					FOverridableTextPortPropertyPathScope ScopePath(Property);
+					Operation = FOverridableSerializationLogic::GetOverriddenPropertyOperationForPortText(Map, DiffMap, PortFlags);
+					FPropertyVisitorPath* Path = FOverridableSerializationLogic::GetOverriddenPortTextPropertyPath();
+					checkf(Path, TEXT("Expecting a path"));
+					Chain = Path->ToSerializedPropertyChain();
+				}
+
+				// Replaced operation goes through the generic property export code
+				if (Operation == EOverriddenPropertyOperation::Modified)
+				{
+					if (const FOverriddenPropertyNode* MapOverriddenPropertyNode = OverriddenProperties->GetOverriddenPropertyNode(&Chain))
 					{
-						UObject* Obj = ExportObjectProp->GetObjectPropertyValue(Property->ContainerPtrToValuePtr<void>(Object, PropertyArrayIndex));
-						if( Obj && !Obj->HasAnyMarks(OBJECTMARK_TagImp) )
+						checkf(Map && DiffMap, TEXT("Expecting memory ptr to the map and its defaults"));
+
+						FScriptMapHelper MapHelper(MapProperty, Map);
+						FScriptMapHelper DiffMapHelper(MapProperty, DiffMap);
+
+						TArray<int32> RemovedIndices;
+						TArray<int32> ModifiedIndices;
+						TArray<int32> AddedIndices;
+
+
+						// Figure out the modifications of the map
+						for (const auto& Pair : MapOverriddenPropertyNode->SubPropertyNodeKeys)
 						{
-							// only export the BEGIN OBJECT block for a component if Parent is the component's Outer....when importing subobject definitions,
-							// (i.e. BEGIN OBJECT), whichever BEGIN OBJECT block a component's BEGIN OBJECT block is located within is the object that will be
-							// used as the Outer to create the component
-							if ( Property->HasAnyPropertyFlags(CPF_InstancedReference) )
+							const EOverriddenPropertyOperation OverrideOp = OverriddenProperties->GetSubPropertyOperation(Pair.Value);
+							switch (OverrideOp)
 							{
-								if ( Obj->GetOuter() == Parent )
+							case EOverriddenPropertyOperation::Remove:
 								{
-									// Don't export more than once.
-									Obj->Mark(OBJECTMARK_TagImp);
-									UExporter::ExportToOutputDevice( Context, Obj, NULL, Out, TEXT("T3D"), Indent, PortFlags );
+									const int32 InternalIndex = Pair.Key.ToMapInternalIndex(DiffMapHelper);
+									if (InternalIndex != INDEX_NONE)
+									{
+										RemovedIndices.Add(InternalIndex);
+									}
+									break;
 								}
-								else
+							case EOverriddenPropertyOperation::Add:
 								{
-									// set the OBJECTMARK_TagExp flag so that the calling code knows we wanted to export this object
-									Obj->Mark(OBJECTMARK_TagExp);
+									const int32 InternalIndex = Pair.Key.ToMapInternalIndex(MapHelper);
+									if (InternalIndex != INDEX_NONE)
+									{
+										AddedIndices.Add(InternalIndex);
+									}
+									break;
 								}
-							}
-							else
-							{
-								// Don't export more than once.
-								Obj->Mark(OBJECTMARK_TagImp);
-								UExporter::ExportToOutputDevice( Context, Obj, NULL, Out, TEXT("T3D"), Indent, PortFlags );
+							case EOverriddenPropertyOperation::Modified:
+								{
+									const int32 InternalIndex = Pair.Key.ToMapInternalIndex(MapHelper);
+									if (InternalIndex != INDEX_NONE)
+									{
+										ModifiedIndices.Add(InternalIndex);
+									}
+									break;
+								}
+							default:
+								checkf(false, TEXT("Unsupported map operation"));
+								break;
 							}
 						}
-					}
 
-					if( Property->ArrayDim == 1 )
-					{
-						Out.Logf( TEXT("%s%s=%s\r\n"), FCString::Spc(Indent), *SanitizedPropertyName, *Value );
+						auto ExportItem = [&MapHelper, &DiffMapHelper, Parent, ExportRootScope, Indent, SanitizedPropertyName, ExportFlags, &Out](int32 i, EOverriddenPropertyOperation Operation)
+						{
+							FString Key;
+							MapHelper.KeyProp->ExportTextItem_Direct(Key, MapHelper.GetKeyPtr(i), Operation == EOverriddenPropertyOperation::Modified ? DiffMapHelper.GetKeyPtr(i) : nullptr, Parent, ExportFlags, ExportRootScope);
+
+							FString Value;
+							MapHelper.ValueProp->ExportTextItem_Direct(Value, MapHelper.GetValuePtr(i), Operation == EOverriddenPropertyOperation::Modified ? DiffMapHelper.GetValuePtr(i) : nullptr, Parent, ExportFlags, ExportRootScope);
+
+							Out.Logf(TEXT("%s%s<%s>=(%s,%s)\r\n"), FCString::Spc(Indent), *SanitizedPropertyName, *GetOverriddenOperationString(Operation), *Key, *Value);
+						};
+
+						for (int32 i : RemovedIndices)
+						{
+							FString Value;
+							DiffMapHelper.KeyProp->ExportTextItem_Direct(Value, DiffMapHelper.GetKeyPtr(i), nullptr, Parent, ExportFlags, ExportRootScope);
+
+							Out.Logf(TEXT("%s%s<%s>=%s\r\n"), FCString::Spc(Indent), *SanitizedPropertyName, *GetOverriddenOperationString(EOverriddenPropertyOperation::Remove), *Value);
+						}
+
+						for (int32 i : ModifiedIndices)
+						{
+							FOverridableTextPortPropertyPathScope ScopePath(Property, i, EPropertyVisitorInfoType::ContainerIndex);
+
+							ExportItem(i, EOverriddenPropertyOperation::Modified);
+						}
+
+						for (int32 i : AddedIndices)
+						{
+							FOverridableTextPortPropertyPathScope ScopePath(Property, i, EPropertyVisitorInfoType::ContainerIndex);
+
+							ExportItem(i, EOverriddenPropertyOperation::Add);
+						}
 					}
-					else
+
+					continue;
+				}
+			}
+
+
+			// Generic property export code
+			FString	Value;
+			const bool bStaticArray = Property->ArrayDim > 1;
+			FOverridableTextPortPropertyPathScope ScopePath(Property, bStaticArray ? PropertyArrayIndex : INDEX_NONE, bStaticArray ? EPropertyVisitorInfoType::StaticArrayIndex : EPropertyVisitorInfoType::None);
+
+			uint8* DiffData = (DiffClass && Property->IsInContainer(DiffClass->GetPropertiesSize())) ? Diff : NULL;
+			if (Property->ExportText_InContainer(PropertyArrayIndex, Value, Object, DiffData, Parent, ExportFlags, ExportRootScope))
+			{
+				HandleExportObject(Property, Property->ContainerPtrToValuePtr<uint8>(Object, PropertyArrayIndex));
+
+				FString OverridableOperation;
+				if (OverriddenProperties)
+				{
+					const EOverriddenPropertyOperation Operation = FOverridableSerializationLogic::GetOverriddenPropertyOperationForPortText(Object, DiffData, PortFlags);
+					if (Operation!=EOverriddenPropertyOperation::None)
 					{
-						Out.Logf( TEXT("%s%s(%i)=%s\r\n"), FCString::Spc(Indent), *SanitizedPropertyName, PropertyArrayIndex, *Value );
+						OverridableOperation = FString::Printf(TEXT("<%s>"), *GetOverriddenOperationString(Operation));
 					}
 				}
+
+				if (!bStaticArray)
+				{
+					Out.Logf( TEXT("%s%s%s=%s\r\n"), FCString::Spc(Indent), *SanitizedPropertyName, *OverridableOperation, *Value );
+				}
+				else
+				{
+					Out.Logf( TEXT("%s%s(%i)%s=%s\r\n"), FCString::Spc(Indent), *SanitizedPropertyName, PropertyArrayIndex, *OverridableOperation, *Value );
+				}
 			}
 		}
 	}
