[NetPhysics] Fix deprecation compiler issue

diff --git a/Engine/Source/Runtime/Engine/Classes/PhysicsEngine/PhysicsSettings.h b/Engine/Source/Runtime/Engine/Classes/PhysicsEngine/PhysicsSettings.h
--- a/Engine/Source/Runtime/Engine/Classes/PhysicsEngine/PhysicsSettings.h
+++ b/Engine/Source/Runtime/Engine/Classes/PhysicsEngine/PhysicsSettings.h
@@ -172,7 +172,9 @@ struct FPhysicsPredictionSettings
 	FPhysicsReplicationResimulationSettings ResimulationSettings;
 
 	FPhysicsPredictionSettings()
-		: bEnablePhysicsPrediction(false)
+		: bEnablePhysicsResimulation_DEPRECATED(false)
+		, ResimulationErrorThreshold_DEPRECATED(10.0f)
+		, bEnablePhysicsPrediction(false)
 		, bEnablePhysicsHistoryCapture(false)
 		, MaxSupportedLatencyPrediction(1000)
 	{ }
diff --git a/Engine/Source/Runtime/Engine/Public/Physics/NetworkPhysicsSettingsComponent.h b/Engine/Source/Runtime/Engine/Public/Physics/NetworkPhysicsSettingsComponent.h
--- a/Engine/Source/Runtime/Engine/Public/Physics/NetworkPhysicsSettingsComponent.h
+++ b/Engine/Source/Runtime/Engine/Public/Physics/NetworkPhysicsSettingsComponent.h
@@ -243,6 +243,11 @@ struct FNetworkPhysicsSettingsPredictiveInterpolation
 USTRUCT()
 struct FNetworkPhysicsSettingsResimulation
 {
+	FNetworkPhysicsSettingsResimulation()
+	: bOverrideResimulationErrorThreshold_DEPRECATED(0)
+	, ResimulationErrorThreshold_DEPRECATED(10.0f)
+	{};
+
 	GENERATED_BODY()
 
 	/** Deprecated UE 5.5 - bOverrideResimulationErrorThreshold has been renamed, please use bOverrideResimulationErrorPositionThreshold*/
