Allow to configure world partition asynchronous UpdateStreamingState task priority using "wp.Runtime.TaskPriorities.AsyncUpdateStreamingStateTask" (only used when wp.Runtime.UpdateStreaming.EnableAsyncUpdate = true).

diff --git a/Engine/Source/Runtime/Engine/Private/WorldPartition/WorldPartitionStreamingPolicy.cpp b/Engine/Source/Runtime/Engine/Private/WorldPartition/WorldPartitionStreamingPolicy.cpp
--- a/Engine/Source/Runtime/Engine/Private/WorldPartition/WorldPartitionStreamingPolicy.cpp
+++ b/Engine/Source/Runtime/Engine/Private/WorldPartition/WorldPartitionStreamingPolicy.cpp
@@ -63,6 +63,17 @@ FAutoConsoleVariableRef UWorldPartitionStreamingPolicy::CVarAsyncUpdateStreaming
 	TEXT("Set to enable asynchronous World Partition UpdateStreamingState."),
 	ECVF_Default);
 
+namespace UE::Private::WorldPartition
+{
+	static FAutoConsoleTaskPriority CPrio_UpdateStreamingStateAsyncTaskPriority(
+		TEXT("wp.Runtime.TaskPriorities.AsyncUpdateStreamingStateTask"),
+		TEXT("Task and thread priority for world partition asynchronous UpdateStreamingState task."),
+		ENamedThreads::BackgroundThreadPriority, // Run on background threads by default...
+		ENamedThreads::HighTaskPriority, // .. at high task priority
+		ENamedThreads::NormalTaskPriority // if we have to run on a normal thread, run at normal task priority
+	);
+}
+
 DECLARE_CYCLE_STAT(TEXT("WorldPartition_AsyncUpdateStreamingState"), WPStreamingPolicy_AsyncUpdateStreamingState, STATGROUP_TaskGraphTasks);
 DECLARE_CYCLE_STAT(TEXT("WorldPartition_AsyncPostUpdateStreamingState"), WPStreamingPolicy_AsyncPostUpdateStreamingState, STATGROUP_TaskGraphTasks);
 
@@ -762,7 +773,7 @@ void UWorldPartitionStreamingPolicy::OnStreamingStateUpdated()
 			{
 				TRACE_CPUPROFILER_EVENT_SCOPE(UWorldPartitionStreamingPolicy::AsyncUpdateStreamingState);
 				UWorldPartitionStreamingPolicy::UpdateStreamingStateInternal(InputParams, AsyncTaskTargetState);
-			}, GET_STATID(WPStreamingPolicy_AsyncUpdateStreamingState), NULL, ENamedThreads::AnyBackgroundHiPriTask);
+			}, GET_STATID(WPStreamingPolicy_AsyncUpdateStreamingState), NULL, UE::Private::WorldPartition::CPrio_UpdateStreamingStateAsyncTaskPriority.Get());
 			AsyncPostUpdateStreamingStateTask = FFunctionGraphTask::CreateAndDispatchWhenReady([this]() { PostUpdateStreamingStateInternal_GameThread(AsyncTaskTargetState); }, GET_STATID(WPStreamingPolicy_AsyncPostUpdateStreamingState), AsyncUpdateStreamingStateTask, ENamedThreads::GameThread);
 			AsyncUpdateTaskState = EAsyncUpdateTaskState::Started;
 		}
