Root Motion preview option LoopAndReset is broken, from one loop and ahead, for animation composites.

#rnx none

diff --git a/Engine/Source/Runtime/Engine/Private/Animation/AnimSequenceHelpers.cpp b/Engine/Source/Runtime/Engine/Private/Animation/AnimSequenceHelpers.cpp
--- a/Engine/Source/Runtime/Engine/Private/Animation/AnimSequenceHelpers.cpp
+++ b/Engine/Source/Runtime/Engine/Private/Animation/AnimSequenceHelpers.cpp
@@ -265,6 +265,10 @@ FTransform ExtractRootTransformFromAnimationAsset(const UAnimationAsset* Animati
 			}	
 		}
 	}
+	else if (const UAnimComposite* AnimComposite = Cast<UAnimComposite>(Animation))
+	{
+		Result = AnimComposite->ExtractRootTrackTransform(Time, nullptr);
+	}
 	else if (const UAnimSequence* AnimSequence = Cast<UAnimSequence>(Animation))
 	{
 		Result = AnimSequence->ExtractRootTrackTransform(Time, nullptr);
