Updated ShowDebug ANIMATION to show sync marker names and when a tick record is kicked from a sync group.

#rnx

diff --git a/Engine/Source/Runtime/Engine/Private/Animation/AnimInstance.cpp b/Engine/Source/Runtime/Engine/Private/Animation/AnimInstance.cpp
--- a/Engine/Source/Runtime/Engine/Private/Animation/AnimInstance.cpp
+++ b/Engine/Source/Runtime/Engine/Private/Animation/AnimInstance.cpp
@@ -965,6 +965,16 @@ void OutputTickRecords(const TArray<FAnimTickRecord>& Records, UCanvas* Canvas,
 	{
 		const FAnimTickRecord& Player = Records[PlayerIndex];
 
+		// Assuming highlight index is the group leader index
+		if (HighlightIndex != INDEX_NONE && Player.bIsExclusiveLeader && PlayerIndex != HighlightIndex)
+		{
+			DisplayDebugManager.SetLinearDrawColor(FLinearColor::Red);
+			FString PlayerEntry = FString::Printf(TEXT("%i) %s (%s) W(%.f%%) - Push to ungrouped as it failed to be leader."), 
+			PlayerIndex, *Player.SourceAsset->GetName(), *Player.SourceAsset->GetClass()->GetName(), Player.EffectiveBlendWeight*100.f);
+			DisplayDebugManager.DrawString(PlayerEntry, Indent);
+			continue;
+		}
+		
 		DisplayDebugManager.SetLinearDrawColor((PlayerIndex == HighlightIndex) ? HighlightColor : TextColor);
 
 		FString PlayerEntry = FString::Printf(TEXT("%i) %s (%s) W(%.f%%)"), 
@@ -985,9 +995,36 @@ void OutputTickRecords(const TArray<FAnimTickRecord>& Records, UCanvas* Canvas,
 		// Part of a sync group
 		if (HighlightIndex != INDEX_NONE)
 		{
-			PlayerEntry += FString::Printf(TEXT(" Prev(i:%d, t:%.3f) Next(i:%d, t:%.3f)"),
-				Player.MarkerTickRecord->PreviousMarker.MarkerIndex, Player.MarkerTickRecord->PreviousMarker.TimeToMarker, 
-				Player.MarkerTickRecord->NextMarker.MarkerIndex, Player.MarkerTickRecord->NextMarker.TimeToMarker);
+			FName PrevMarkerName = NAME_None;
+			FName NextMarkerName = NAME_None;
+
+			TArray<FAnimSyncMarker> * AuthoredMarkerNames = nullptr;
+			
+			if (auto AnimSequence = Cast<UAnimSequence>(Player.SourceAsset))
+			{
+				AuthoredMarkerNames = &AnimSequence->AuthoredSyncMarkers;  
+			}
+			else if (auto AnimMontage = Cast<UAnimMontage>(Player.SourceAsset))
+			{
+				AuthoredMarkerNames = &AnimMontage->MarkerData.AuthoredSyncMarkers;
+			}
+
+			if (AuthoredMarkerNames)
+			{
+				if (Player.MarkerTickRecord->PreviousMarker.MarkerIndex >= 0)
+				{
+					PrevMarkerName = (*AuthoredMarkerNames)[Player.MarkerTickRecord->PreviousMarker.MarkerIndex].MarkerName; 
+				}
+
+				if (Player.MarkerTickRecord->NextMarker.MarkerIndex >= 0)
+				{
+					NextMarkerName = (*AuthoredMarkerNames)[Player.MarkerTickRecord->NextMarker.MarkerIndex].MarkerName; 
+				}
+			}
+
+			PlayerEntry += FString::Printf(TEXT(" Prev(%s, i:%d, t:%.3f) Next(%s, i:%d, t:%.3f)"),
+				*PrevMarkerName.ToString(), Player.MarkerTickRecord->PreviousMarker.MarkerIndex, Player.MarkerTickRecord->PreviousMarker.TimeToMarker, 
+				*NextMarkerName.ToString(), Player.MarkerTickRecord->NextMarker.MarkerIndex, Player.MarkerTickRecord->NextMarker.TimeToMarker);
 		}
 
 		DisplayDebugManager.DrawString(PlayerEntry, Indent);
