Add file watcher to the MaterialTemplate shader template.
This invalidates the cached template if the file is modified.
This allows us to modify MaterialTemplate.usf while in editor and have the changes be reflected in subsequent material compilations.

diff --git a/Engine/Source/Runtime/Engine/Private/Materials/MaterialSourceTemplate.cpp b/Engine/Source/Runtime/Engine/Private/Materials/MaterialSourceTemplate.cpp
--- a/Engine/Source/Runtime/Engine/Private/Materials/MaterialSourceTemplate.cpp
+++ b/Engine/Source/Runtime/Engine/Private/Materials/MaterialSourceTemplate.cpp
@@ -3,9 +3,13 @@
 #if WITH_EDITOR
 
 #include "MaterialSourceTemplate.h"
+
+#include "DirectoryWatcherModule.h"
+#include "IDirectoryWatcher.h"
+#include "MaterialShared.h"
 #include "Misc/ScopeRWLock.h"
+#include "Modules/ModuleManager.h"
 #include "ShaderCore.h"
-#include "MaterialShared.h"
 
 FMaterialSourceTemplate& FMaterialSourceTemplate::Get()
 {
@@ -13,6 +17,37 @@ FMaterialSourceTemplate& FMaterialSourceTemplate::Get()
 	return Instance;
 }
 
+FMaterialSourceTemplate::FMaterialSourceTemplate()
+{
+	for (int32 ShaderPlatformIndex = 0; ShaderPlatformIndex < SP_NumPlatforms; ++ShaderPlatformIndex)
+	{
+		MaterialTemplateLineNumbers[ShaderPlatformIndex] = INDEX_NONE;
+		bFileWatchInvalidation[ShaderPlatformIndex] = false;
+	}
+
+	FString FileWatchDirectory(FPaths::EngineDir() + TEXT("Shaders/Private"));
+
+	FDirectoryWatcherModule& DirectoryWatcherModule = FModuleManager::LoadModuleChecked<FDirectoryWatcherModule>("DirectoryWatcher");
+	DirectoryWatcherModule.Get()->RegisterDirectoryChangedCallback_Handle(
+		FileWatchDirectory,
+		IDirectoryWatcher::FDirectoryChanged::CreateLambda([this](const TArray<struct FFileChangeData>& Changes)
+		{
+			for (FFileChangeData const& Change : Changes)
+			{
+				if ((Change.Action == FFileChangeData::FCA_Modified || Change.Action == FFileChangeData::FCA_Added) && Change.Filename.EndsWith(TEXT("MaterialTemplate.ush")))
+				{
+					for (int32 ShaderPlatformIndex = 0; ShaderPlatformIndex < SP_NumPlatforms; ++ShaderPlatformIndex)
+					{
+						FWriteScopeLock Lock{ RWLocks[ShaderPlatformIndex] };
+						// Only set invalidation flag for this shader platform if we have the template loaded.
+						bFileWatchInvalidation[ShaderPlatformIndex] = !Templates[ShaderPlatformIndex].GetTemplateString().IsEmpty();
+					}
+				}
+			}
+		}),
+		FileWatchHandle);
+}
+
 FStringTemplateResolver FMaterialSourceTemplate::BeginResolve(EShaderPlatform ShaderPlatform, int32* MaterialTemplateLineNumber)
 {
 	Preload(ShaderPlatform);
@@ -42,7 +77,7 @@ bool FMaterialSourceTemplate::Preload(EShaderPlatform ShaderPlatform)
 	// Is the material source template already loaded?
 	{
 		FReadScopeLock Lock{ RWLocks[ShaderPlatform] };
-		if (!Templates[ShaderPlatform].GetTemplateString().IsEmpty())
+		if (!Templates[ShaderPlatform].GetTemplateString().IsEmpty() && !bFileWatchInvalidation[ShaderPlatform])
 		{
 			return true;
 		}
@@ -50,13 +85,21 @@ bool FMaterialSourceTemplate::Preload(EShaderPlatform ShaderPlatform)
 
 	// Material source template not yet loaded. Acquire a write lock an try again.
 	FWriteScopeLock Lock{ RWLocks[ShaderPlatform] };
-	if (!Templates[ShaderPlatform].GetTemplateString().IsEmpty())
+	if (!Templates[ShaderPlatform].GetTemplateString().IsEmpty() && !bFileWatchInvalidation[ShaderPlatform])
 	{
 		return true;
 	}
 
+	static const TCHAR* VirtualFilePath = TEXT("/Engine/Private/MaterialTemplate.ush");
+	
+	if (bFileWatchInvalidation[ShaderPlatform])
+	{
+		InvalidateShaderFileCacheEntry(VirtualFilePath, ShaderPlatform);
+		bFileWatchInvalidation[ShaderPlatform] = false;
+	}
+
 	FString MaterialTemplateString;
-	LoadShaderSourceFileChecked(TEXT("/Engine/Private/MaterialTemplate.ush"), ShaderPlatform, MaterialTemplateString);
+	LoadShaderSourceFileChecked(VirtualFilePath, ShaderPlatform, MaterialTemplateString);
 
 	// Normalize line endings -- preprocessor does this later if necessary, but that can run faster if it's already done, and doing it here
 	// means it only happens once when the template gets loaded, rather than for every Material shader.
diff --git a/Engine/Source/Runtime/Engine/Private/Materials/MaterialSourceTemplate.h b/Engine/Source/Runtime/Engine/Private/Materials/MaterialSourceTemplate.h
--- a/Engine/Source/Runtime/Engine/Private/Materials/MaterialSourceTemplate.h
+++ b/Engine/Source/Runtime/Engine/Private/Materials/MaterialSourceTemplate.h
@@ -14,6 +14,8 @@ struct FMaterialSourceTemplate
 	/** Returns the singleton instance */
 	static FMaterialSourceTemplate& Get();
 
+	FMaterialSourceTemplate();
+
 	/** Creates and returns a string template resolver for the specified shader platform.
 	 * The appropriate template is loaded for the first time if it wasn't loaded before.*/
 	FStringTemplateResolver BeginResolve(EShaderPlatform ShaderPlatform, int32* MaterialTemplateLineNumber = nullptr);
@@ -39,9 +41,14 @@ struct FMaterialSourceTemplate
 	/** Cached material template line numbers per shader platform */
 	int32 MaterialTemplateLineNumbers[SP_NumPlatforms];
 
+	/** Per shader platform flags that are set true if the template is invalidated by a file change. */
+	bool bFileWatchInvalidation[SP_NumPlatforms];
+
+	/** Handle for file watcher. */
+	FDelegateHandle FileWatchHandle;
+
 	/** RW locks for thread safe access */
 	FRWLock RWLocks[SP_NumPlatforms];
-
 };
 
 #endif // WITH_EDITOR
diff --git a/Engine/Source/Runtime/RenderCore/Private/ShaderCore.cpp b/Engine/Source/Runtime/RenderCore/Private/ShaderCore.cpp
--- a/Engine/Source/Runtime/RenderCore/Private/ShaderCore.cpp
+++ b/Engine/Source/Runtime/RenderCore/Private/ShaderCore.cpp
@@ -311,6 +311,11 @@ public:
 		return Platforms[ShaderPlatform].ShaderHashCache.Add(VirtualFilePath, FSHAHash());
 	}
 
+	void RemoveHash(EShaderPlatform ShaderPlatform, const FString& VirtualFilePath)
+	{
+		Platforms[ShaderPlatform].ShaderHashCache.Remove(VirtualFilePath);
+	}
+
 	static bool IsPlatformInclude(const FString& VirtualFilePath)
 	{
 		return (VirtualFilePath.StartsWith(TEXT("/Engine/Private/Platform/"))
@@ -3239,6 +3244,23 @@ void FlushShaderFileCache()
 	UE_LOG(LogShaders, Log, TEXT("FlushShaderFileCache() end"));
 }
 
+void InvalidateShaderFileCacheEntry(const TCHAR* InVirtualFilePath, EShaderPlatform InShaderPlatform, const FName* InShaderPlatformName)
+{
+	FString VirtualFilePath(InVirtualFilePath);
+	FixupShaderFilePath(VirtualFilePath, InShaderPlatform, InShaderPlatformName);
+
+	const uint32 CurrentHash = GetTypeHash(VirtualFilePath);
+	FShaderFileCache& ShaderFileCache = GShaderFileCache[CurrentHash % GSHADERFILECACHE_BUCKETS];
+	{
+		FRWScopeLock ScopeLock(ShaderFileCache.Lock, SLT_Write);
+		ShaderFileCache.Map.RemoveByHash(CurrentHash, VirtualFilePath);
+	}
+	{
+		FRWScopeLock ShaderHashAccessLock(GShaderHashAccessRWLock, SLT_Write);
+		GShaderHashCache.RemoveHash(InShaderPlatform, VirtualFilePath);
+	}
+}
+
 #if WITH_EDITOR
 
 void UpdateReferencedUniformBufferNames(
diff --git a/Engine/Source/Runtime/RenderCore/Public/ShaderCore.h b/Engine/Source/Runtime/RenderCore/Public/ShaderCore.h
--- a/Engine/Source/Runtime/RenderCore/Public/ShaderCore.h
+++ b/Engine/Source/Runtime/RenderCore/Public/ShaderCore.h
@@ -1383,6 +1383,9 @@ extern RENDERCORE_API const class FSHAHash& GetShaderFilesHash(const TArray<FStr
  */
 extern RENDERCORE_API void FlushShaderFileCache();
 
+/** Invalidates a single entry in the shader file and CRC caches. */
+extern RENDERCORE_API void InvalidateShaderFileCacheEntry(const TCHAR* InVirtualFilePath, EShaderPlatform InShaderPlatform, const FName* InShaderPlatformName = nullptr);
+
 extern RENDERCORE_API void VerifyShaderSourceFiles(EShaderPlatform ShaderPlatform);
 
 #if WITH_EDITOR
