[Nanite Skinning] Properly clean up invalid NaniteMeshObject instead of deleting it so FRenderResource will be correctly released.

#rnx

diff --git a/Engine/Source/Runtime/Engine/Private/Components/SkinnedMeshComponent.cpp b/Engine/Source/Runtime/Engine/Private/Components/SkinnedMeshComponent.cpp
--- a/Engine/Source/Runtime/Engine/Private/Components/SkinnedMeshComponent.cpp
+++ b/Engine/Source/Runtime/Engine/Private/Components/SkinnedMeshComponent.cpp
@@ -1020,7 +1020,8 @@ void USkinnedMeshComponent::CreateRenderState_Concurrent(FRegisterComponentConte
 					}
 					else
 					{
-						delete NaniteMeshObject;
+						NaniteMeshObject->ReleaseResources();
+						BeginCleanup(NaniteMeshObject);
 					}
 				}
 
