[Material Editor] Added support for redirectors while loading parameter collections references in FEnvironmentDefines.

diff --git a/Engine/Source/Runtime/Engine/Private/Materials/HLSLMaterialTranslator.cpp b/Engine/Source/Runtime/Engine/Private/Materials/HLSLMaterialTranslator.cpp
--- a/Engine/Source/Runtime/Engine/Private/Materials/HLSLMaterialTranslator.cpp
+++ b/Engine/Source/Runtime/Engine/Private/Materials/HLSLMaterialTranslator.cpp
@@ -433,19 +433,28 @@ struct FHLSLMaterialTranslator::FEnvironmentDefines
 				ParameterCollectionObjects.Add(Object);
 			}
 		}
+
 		Ar << ParameterCollectionObjects;
+
 		if (Ar.IsLoading())
 		{
 			ParameterCollections.Reset(ParameterCollectionObjects.Num());
 			for (UObject* Object : ParameterCollectionObjects)
 			{
+				while (UObjectRedirector* Redirector = Cast<UObjectRedirector>(Object))
+				{
+					Object = Redirector->DestinationObject;
+				}
+
 				UMaterialParameterCollection* Collection = Cast<UMaterialParameterCollection>(Object);
+				
 				if (!Collection && Object)
 				{
 					UE_LOG(LogMaterial, Error,
 						TEXT("FEnvironmentDefines.ParameterCollections had object %s which is not a UMaterialParameterCollection; setting it to null."),
 						*Object->GetFullName());
 				}
+
 				ParameterCollections.Add(Collection);
 			}
 		}
