[WorldPartition] Fix regression in if condition from CL 36260716

Was causing PCGWorldActor to not load anymore in PIE

#rnx

diff --git a/Engine/Source/Runtime/Engine/Private/WorldPartition/WorldPartitionStreamingGeneration.cpp b/Engine/Source/Runtime/Engine/Private/WorldPartition/WorldPartitionStreamingGeneration.cpp
--- a/Engine/Source/Runtime/Engine/Private/WorldPartition/WorldPartitionStreamingGeneration.cpp
+++ b/Engine/Source/Runtime/Engine/Private/WorldPartition/WorldPartitionStreamingGeneration.cpp
@@ -793,7 +793,7 @@ class FWorldPartitionStreamingGenerator
 				OutContainerInstances.Add(InActorDescView);
 			}
 
-			const bool bActorIsEditorOnly = InActorDescView.GetActorIsEditorOnly() || (InActor && !InActor->IsEditorOnly());
+			const bool bActorIsEditorOnly = InActor ? InActor->IsEditorOnly() : InActorDescView.GetActorIsEditorOnly();
 			const bool bActorIsNonPIEEditorOnly = InActorDescView.GetActorIsNonPIEEditorOnly();
 
 			if (!bActorIsEditorOnly || (bIsPIE && bActorIsNonPIEEditorOnly))
