Tooltip change for GlobalGravityZ

diff --git a/Engine/Source/Runtime/Engine/Classes/GameFramework/WorldSettings.h b/Engine/Source/Runtime/Engine/Classes/GameFramework/WorldSettings.h
--- a/Engine/Source/Runtime/Engine/Classes/GameFramework/WorldSettings.h
+++ b/Engine/Source/Runtime/Engine/Classes/GameFramework/WorldSettings.h
@@ -562,7 +562,8 @@ public:
 	UFUNCTION()
 	ENGINE_API virtual void OnRep_WorldGravityZ();
 
-	// optional level specific gravity override set by level designer
+	// Set a custom gravity value for this level, in distance units per second squared.
+	// For example, when the project's distance unit is set to centimeters (the default), normal earth gravity would be -980.0.
 	UPROPERTY(EditAnywhere, BlueprintReadOnly, Category=Physics, meta=(editcondition = "bGlobalGravitySet"))
 	float GlobalGravityZ;
 
