Added IsEditorPropertyOverridden to compare a named object property against its archetype value

diff --git a/Engine/Source/Runtime/Engine/Classes/Kismet/KismetSystemLibrary.h b/Engine/Source/Runtime/Engine/Classes/Kismet/KismetSystemLibrary.h
--- a/Engine/Source/Runtime/Engine/Classes/Kismet/KismetSystemLibrary.h
+++ b/Engine/Source/Runtime/Engine/Classes/Kismet/KismetSystemLibrary.h
@@ -71,6 +71,20 @@ namespace EQuitPreference
 	};
 }
 
+/** States a property value queried by IsEditorPropertyOverridden can be in */
+UENUM()
+enum class EEditorPropertyValueState
+{
+	/** This property is in a default state */
+	Default,
+	/** This property is in an overridden state */
+	Overridden,
+	/** The property was not found on the object or its archetype; its state is unknown */
+	NotFound,
+	/** The property could not be accessed to query its state; its state is unknown */
+	AccessDenied,
+};
+
 USTRUCT(BlueprintInternalUseOnly)
 struct FGenericStruct
 {
@@ -1988,6 +2002,17 @@ class UKismetSystemLibrary : public UBlueprintFunctionLibrary
 	 */
 	UFUNCTION(BlueprintCallable, Category = "Utilities", meta=(ScriptMethod))
 	static ENGINE_API bool ResetEditorProperty(UObject* Object, const FName PropertyName, const EPropertyAccessChangeNotifyMode ChangeNotifyMode = EPropertyAccessChangeNotifyMode::Default);
+
+	/**
+	 * Attempts to query whether the value of a named property on the given object overrides the value of its archetype (ie, would ResetEditorProperty do anything?).
+	 *
+	 * @param Object The object you want to query a property value on.
+	 * @param PropertyName The name of the object property to query the value of.
+	 *
+	 * @return What state the requested property is in.
+	 */
+	UFUNCTION(BlueprintCallable, Category = "Utilities", meta=(ScriptMethod))
+	static ENGINE_API EEditorPropertyValueState IsEditorPropertyOverridden(UObject* Object, const FName PropertyName);
 #endif
 
 	// --- Transactions ------------------------------
diff --git a/Engine/Source/Runtime/Engine/Private/KismetSystemLibrary.cpp b/Engine/Source/Runtime/Engine/Private/KismetSystemLibrary.cpp
--- a/Engine/Source/Runtime/Engine/Private/KismetSystemLibrary.cpp
+++ b/Engine/Source/Runtime/Engine/Private/KismetSystemLibrary.cpp
@@ -3410,18 +3410,13 @@ DEFINE_FUNCTION(UKismetSystemLibrary::execSetEditorProperty)
 	*(bool*)RESULT_PARAM = bResult;
 }
 
-bool UKismetSystemLibrary::ResetEditorProperty(UObject* Object, const FName PropertyName, const EPropertyAccessChangeNotifyMode ChangeNotifyMode)
+namespace UE::Blueprint::Private
 {
-	if (!Object)
-	{
-		LogRuntimeError(NSLOCTEXT("KismetSystemLibrary", "ResetEditorProperty_AccessNone", "Accessed None attempting to call ResetEditorProperty."));
-		return false;
-	}
-
-	auto FindArchetypeValue = [Object, PropertyName](const FProperty*& OutArchetypeProperty, const void*& OutArchetypeValuePtr)
+	void FindArchetypeValue(UObject* Object, const FName PropertyName, const FProperty*& OutArchetypeProperty, const void*& OutArchetypeValuePtr, bool& OutIsSparseData)
 	{
 		OutArchetypeProperty = nullptr;
 		OutArchetypeValuePtr = nullptr;
+		OutIsSparseData = false;
 
 		const FProperty* ObjectProp = PropertyAccessUtil::FindPropertyByName(PropertyName, Object->GetClass());
 		if ((!ObjectProp || ObjectProp->HasAnyPropertyFlags(CPF_Deprecated)) && Object->HasAllFlags(RF_ClassDefaultObject))
@@ -3434,6 +3429,7 @@ bool UKismetSystemLibrary::ResetEditorProperty(UObject* Object, const FName Prop
 				{
 					if (UScriptStruct* SparseDataArchetypeStruct = Object->GetClass()->GetSparseClassDataArchetypeStruct())
 					{
+						OutIsSparseData = true;
 						OutArchetypeProperty = SparseProp;
 						OutArchetypeValuePtr = SparseProp->ContainerPtrToValuePtrForDefaults<const void>(SparseDataArchetypeStruct, Object->GetClass()->GetArchetypeForSparseClassData());
 					}
@@ -3450,11 +3446,21 @@ bool UKismetSystemLibrary::ResetEditorProperty(UObject* Object, const FName Prop
 				OutArchetypeValuePtr = ObjectProp->ContainerPtrToValuePtrForDefaults<const void>(ObjectArchetype->GetClass(), ObjectArchetype);
 			}
 		}
-	};
+	}
+}
+
+bool UKismetSystemLibrary::ResetEditorProperty(UObject* Object, const FName PropertyName, const EPropertyAccessChangeNotifyMode ChangeNotifyMode)
+{
+	if (!Object)
+	{
+		LogRuntimeError(NSLOCTEXT("KismetSystemLibrary", "ResetEditorProperty_AccessNone", "Accessed None attempting to call ResetEditorProperty."));
+		return false;
+	}
 
 	const FProperty* ArchetypeProperty = nullptr;
 	const void* ArchetypeValuePtr = nullptr;
-	FindArchetypeValue(ArchetypeProperty, ArchetypeValuePtr);
+	bool bIsSparseData = false;
+	UE::Blueprint::Private::FindArchetypeValue(Object, PropertyName, ArchetypeProperty, ArchetypeValuePtr, bIsSparseData);
 	if (!ArchetypeValuePtr)
 	{
 		FFrame::KismetExecutionMessage(*FString::Printf(TEXT("Property '%s' on '%s' (%s) had no archetype value to reset to"), *PropertyName.ToString(), *Object->GetPathName(), *Object->GetClass()->GetName()), ELogVerbosity::Warning, UE::Blueprint::Private::PropertySetFailedWarning);
@@ -3464,6 +3470,67 @@ bool UKismetSystemLibrary::ResetEditorProperty(UObject* Object, const FName Prop
 	return Generic_SetEditorProperty(Object, PropertyName, ArchetypeValuePtr, ArchetypeProperty, ChangeNotifyMode);
 }
 
+EEditorPropertyValueState UKismetSystemLibrary::IsEditorPropertyOverridden(UObject* Object, const FName PropertyName)
+{
+	if (!Object)
+	{
+		LogRuntimeError(NSLOCTEXT("KismetSystemLibrary", "IsEditorPropertyOverridden_AccessNone", "Accessed None attempting to call IsEditorPropertyOverridden."));
+		return EEditorPropertyValueState::NotFound;
+	}
+
+	const FProperty* ArchetypeProperty = nullptr;
+	const void* ArchetypeValuePtr = nullptr;
+	bool bIsSparseData = false;
+	UE::Blueprint::Private::FindArchetypeValue(Object, PropertyName, ArchetypeProperty, ArchetypeValuePtr, bIsSparseData);
+	if (!ArchetypeValuePtr)
+	{
+		FFrame::KismetExecutionMessage(*FString::Printf(TEXT("Property '%s' on '%s' (%s) had no archetype value to query against"), *PropertyName.ToString(), *Object->GetPathName(), *Object->GetClass()->GetName()), ELogVerbosity::Warning, UE::Blueprint::Private::PropertyGetFailedWarning);
+		return EEditorPropertyValueState::NotFound;
+	}
+
+	if (EPropertyAccessResultFlags AccessResult = PropertyAccessUtil::CanGetPropertyValue(ArchetypeProperty);
+		EnumHasAnyFlags(AccessResult, EPropertyAccessResultFlags::PermissionDenied))
+	{
+		if (EnumHasAnyFlags(AccessResult, EPropertyAccessResultFlags::AccessProtected))
+		{
+			FFrame::KismetExecutionMessage(*FString::Printf(TEXT("Property '%s' on '%s' (%s) is protected and cannot be read"), *PropertyName.ToString(), *Object->GetPathName(), *Object->GetClass()->GetName()), ELogVerbosity::Warning, UE::Blueprint::Private::PropertyGetFailedWarning);
+			return EEditorPropertyValueState::AccessDenied;
+		}
+
+		FFrame::KismetExecutionMessage(*FString::Printf(TEXT("Property '%s' on '%s' (%s) cannot be read"), *PropertyName.ToString(), *Object->GetPathName(), *Object->GetClass()->GetName()), ELogVerbosity::Warning, UE::Blueprint::Private::PropertyGetFailedWarning);
+		return EEditorPropertyValueState::AccessDenied;
+	}
+
+	auto GetPropertyValueForDiff = [Object, ArchetypeProperty](const void* PropertyValuePtr) -> FString
+	{
+		uint32 PortFlags = PPF_ForDiff;
+		if (!Object->IsTemplate())
+		{
+			PortFlags |= PPF_ForDiffInstanceOnly;
+		}
+		if (ArchetypeProperty->ContainsInstancedObjectProperty())
+		{
+			PortFlags |= PPF_DeepComparison;
+		}
+
+		FString DefaultValue;
+		ArchetypeProperty->ExportTextItem_Direct(DefaultValue, PropertyValuePtr, PropertyValuePtr, nullptr, PortFlags);
+		return DefaultValue;
+	};
+
+	const void* ObjectValuePtr = bIsSparseData
+		? ArchetypeProperty->ContainerPtrToValuePtr<const void>(Object->GetClass()->GetOrCreateSparseClassData())
+		: ArchetypeProperty->ContainerPtrToValuePtr<const void>(Object);
+
+	const FString ObjectValue = GetPropertyValueForDiff(ObjectValuePtr);
+	const FString ArchetypeValue = GetPropertyValueForDiff(ArchetypeValuePtr);
+	const bool bIsIdentical = ObjectValue.Equals(ArchetypeValue, ESearchCase::CaseSensitive);
+
+	return bIsIdentical
+		? EEditorPropertyValueState::Default
+		: EEditorPropertyValueState::Overridden;
+}
+
 #endif	// WITH_EDITOR
 
 int32 UKismetSystemLibrary::BeginTransaction(const FString& Context, FText Description, UObject* PrimaryObject)
