Made GracefulClose virtual to override in Client Connection so that it can close parent connection in a splitscreen game

diff --git a/Engine/Source/Runtime/Engine/Classes/Engine/ChildConnection.h b/Engine/Source/Runtime/Engine/Classes/Engine/ChildConnection.h
--- a/Engine/Source/Runtime/Engine/Classes/Engine/ChildConnection.h
+++ b/Engine/Source/Runtime/Engine/Classes/Engine/ChildConnection.h
@@ -5,6 +5,7 @@
 #include "CoreMinimal.h"
 #include "UObject/ObjectMacros.h"
 #include "Engine/NetConnection.h"
+#include "Net/Core/Connection/NetCloseResult.h"
 #include "ChildConnection.generated.h"
 
 /**
@@ -15,7 +16,7 @@ class UChildConnection
 	: public UNetConnection
 {
 	GENERATED_UCLASS_BODY()
-
+	using FNetCloseResult = UE::Net::FNetCloseResult;
 	UPROPERTY(transient)
 	TObjectPtr<class UNetConnection> Parent;
 
@@ -73,7 +74,7 @@ public:
 	{
 		SetConnectionState(Parent->GetConnectionState());
 	}
-
+	virtual void GracefulClose(FNetCloseResult&& CloseReason) override;
 	virtual void HandleClientPlayer(class APlayerController* PC, class UNetConnection* NetConnection) override;
 	virtual void CleanUp() override;
 
diff --git a/Engine/Source/Runtime/Engine/Classes/Engine/NetConnection.h b/Engine/Source/Runtime/Engine/Classes/Engine/NetConnection.h
--- a/Engine/Source/Runtime/Engine/Classes/Engine/NetConnection.h
+++ b/Engine/Source/Runtime/Engine/Classes/Engine/NetConnection.h
@@ -1056,7 +1056,7 @@ public:
 	 *
 	 * @param CloseReason	Specifies the reason for the Close
 	 */
-	void GracefulClose(FNetCloseResult&& CloseReason)
+	virtual void GracefulClose(FNetCloseResult&& CloseReason)
 	{
 		GracefulClose(static_cast<FNetResult&&>(MoveTemp(CloseReason)));
 	}
diff --git a/Engine/Source/Runtime/Engine/Private/NetConnection.cpp b/Engine/Source/Runtime/Engine/Private/NetConnection.cpp
--- a/Engine/Source/Runtime/Engine/Private/NetConnection.cpp
+++ b/Engine/Source/Runtime/Engine/Private/NetConnection.cpp
@@ -1499,6 +1499,14 @@ void UChildConnection::InitChildConnection(UNetDriver* InDriver, UNetConnection*
 	InitChannelData();
 }
 
+void UChildConnection::GracefulClose(FNetCloseResult&& CloseReason)
+{
+	if (Parent)
+	{
+		Parent->GracefulClose(MoveTemp(CloseReason));
+	}
+}
+
 void UNetConnection::FinishDestroy()
 {
 	if ( !HasAnyFlags(RF_ClassDefaultObject) )
