Iris
* Renamed ActorReplicationBridge to EngineReplicationBridge

diff --git a/Engine/Source/Runtime/Engine/Private/ActorReplication.cpp b/Engine/Source/Runtime/Engine/Private/ActorReplication.cpp
--- a/Engine/Source/Runtime/Engine/Private/ActorReplication.cpp
+++ b/Engine/Source/Runtime/Engine/Private/ActorReplication.cpp
@@ -23,7 +23,7 @@
 #include "Iris/ReplicationSystem/ReplicationSystem.h"
 #include "Iris/ReplicationSystem/Conditionals/ReplicationCondition.h"
 #include "Net/Iris/ReplicationSystem/ReplicationSystemUtil.h"
-#include "Net/Iris/ReplicationSystem/ActorReplicationBridge.h"
+#include "Net/Iris/ReplicationSystem/EngineReplicationBridge.h"
 #endif // UE_WITH_IRIS
 #include "Physics/Experimental/PhysScene_Chaos.h"
 
diff --git a/Engine/Source/Runtime/Engine/Private/Components/CharacterMovementComponent.cpp b/Engine/Source/Runtime/Engine/Private/Components/CharacterMovementComponent.cpp
--- a/Engine/Source/Runtime/Engine/Private/Components/CharacterMovementComponent.cpp
+++ b/Engine/Source/Runtime/Engine/Private/Components/CharacterMovementComponent.cpp
@@ -30,7 +30,7 @@
 #include "Net/PerfCountersHelpers.h"
 #if UE_WITH_IRIS
 #include "Net/Iris/ReplicationSystem/ReplicationSystemUtil.h"
-#include "Net/Iris/ReplicationSystem/ActorReplicationBridge.h"
+#include "Net/Iris/ReplicationSystem/EngineReplicationBridge.h"
 #include "Iris/ReplicationSystem/ReplicationSystem.h"
 #endif
 
@@ -8861,7 +8861,7 @@ static UIrisObjectReferencePackageMap* GetIrisPackageMapToCaptureReferences(UNet
 {
 	using namespace UE::Net;
 
-	if (const UActorReplicationBridge* Bridge = FReplicationSystemUtil::GetActorReplicationBridge(NetConnection))
+	if (const UEngineReplicationBridge* Bridge = FReplicationSystemUtil::GetActorReplicationBridge(NetConnection))
 	{
 		if (UIrisObjectReferencePackageMap* ObjectReferencePackageMap = Bridge->GetObjectReferencePackageMap())
 		{
@@ -8881,7 +8881,7 @@ static UIrisObjectReferencePackageMap* GetIrisPackageMapToCaptureReferences(UNet
 static UIrisObjectReferencePackageMap* GetIrisPackageMapToReadReferences(const UNetConnection* NetConnection, const FCharacterNetworkSerializationPackedBits& PackedBits)
 {
 	using namespace UE::Net;
-	if (const UActorReplicationBridge* Bridge = FReplicationSystemUtil::GetActorReplicationBridge(NetConnection))
+	if (const UEngineReplicationBridge* Bridge = FReplicationSystemUtil::GetActorReplicationBridge(NetConnection))
 	{
 		if (UIrisObjectReferencePackageMap* ObjectReferencePackageMap = Bridge->GetObjectReferencePackageMap())
 		{
diff --git a/Engine/Source/Runtime/Engine/Private/Net/Experimental/Iris/DataStreamChannel.cpp b/Engine/Source/Runtime/Engine/Private/Net/Experimental/Iris/DataStreamChannel.cpp
--- a/Engine/Source/Runtime/Engine/Private/Net/Experimental/Iris/DataStreamChannel.cpp
+++ b/Engine/Source/Runtime/Engine/Private/Net/Experimental/Iris/DataStreamChannel.cpp
@@ -23,7 +23,7 @@
 #include "Net/Core/Trace/NetTrace.h"
 #include "Net/Core/Connection/NetResult.h"
 #include "Net/DataChannel.h"
-#include "Net/Iris/ReplicationSystem/ActorReplicationBridge.h"
+#include "Net/Iris/ReplicationSystem/EngineReplicationBridge.h"
 #include "Net/Iris/ReplicationSystem/ReplicationSystemUtil.h"
 
 
@@ -152,7 +152,7 @@ void UDataStreamChannel::ReceivedBunch(FInBunch& Bunch)
 			TNetResult<ENetCloseResult> NetResult(ENetCloseResult::IrisNetRefHandleError, FString::Printf(TEXT("IrisNetRefHandleError=%s"), *SerializationContext.GetErrorHandleContext().ToString()));
 			AddToChainResultPtr(Bunch.ExtendedError, MoveTemp(NetResult));
 
-			UActorReplicationBridge* ActorBridge = FReplicationSystemUtil::GetActorReplicationBridge(Connection);
+			UEngineReplicationBridge* ActorBridge = FReplicationSystemUtil::GetActorReplicationBridge(Connection);
 			if (ActorBridge)
 			{
 				ActorBridge->SendErrorWithNetRefHandle(UE::Net::ENetRefHandleError::BitstreamCorrupted, SerializationContext.GetErrorHandleContext(), Connection->GetConnectionId());
diff --git a/Engine/Source/Runtime/Engine/Private/Net/Iris/ReplicationSystem/ActorReplicationBridge.cpp b/Engine/Source/Runtime/Engine/Private/Net/Iris/ReplicationSystem/ActorReplicationBridge.cpp
deleted file mode 100644
diff --git a/Engine/Source/Runtime/Engine/Private/Net/Iris/ReplicationSystem/EngineReplicationBridge.cpp b/Engine/Source/Runtime/Engine/Private/Net/Iris/ReplicationSystem/EngineReplicationBridge.cpp
new file mode 100644
--- /dev/null
+++ b/Engine/Source/Runtime/Engine/Private/Net/Iris/ReplicationSystem/EngineReplicationBridge.cpp
@@ -0,0 +1,1000 @@
+// Copyright Epic Games, Inc. All Rights Reserved.
+#include "Net/Iris/ReplicationSystem/EngineReplicationBridge.h"
+
+#include UE_INLINE_GENERATED_CPP_BY_NAME(EngineReplicationBridge)
+
+#if UE_WITH_IRIS
+
+#include "Net/Iris/ReplicationSystem/NetActorFactory.h"
+#include "Net/Iris/ReplicationSystem/NetSubObjectFactory.h"
+
+#include "Iris/IrisConfig.h"
+#include "Iris/IrisConstants.h"
+
+#include "Iris/Core/IrisLog.h"
+#include "Iris/Core/IrisProfiler.h"
+#include "Iris/Core/NetObjectReference.h"
+#include "Iris/Core/IrisMemoryTracker.h"
+
+#include "Iris/ReplicationSystem/NetToken.h"
+#include "Iris/ReplicationSystem/ObjectReplicationBridgeConfig.h"
+#include "Iris/ReplicationSystem/StringTokenStore.h"
+#include "Iris/ReplicationSystem/ReplicationOperations.h"
+#include "Iris/ReplicationSystem/ReplicationSystem.h"
+#include "Iris/ReplicationSystem/Filtering/NetObjectFilter.h"
+
+#include "Iris/Serialization/NetBitStreamReader.h"
+#include "Iris/Serialization/NetBitStreamUtil.h"
+#include "Iris/Serialization/NetBitStreamWriter.h"
+#include "Iris/Serialization/ObjectNetSerializer.h"
+#include "Iris/Serialization/IrisObjectReferencePackageMap.h"
+
+#include "Iris/Metrics/NetMetrics.h"
+
+#include "AnalyticsEventAttribute.h"
+
+#include "Engine/Engine.h"
+#include "Engine/EngineTypes.h"
+#include "Engine/Level.h"
+#include "Engine/NetConnection.h"
+#include "Engine/NetDriver.h"
+#include "Engine/Level.h"
+#include "Engine/World.h"
+
+#include "GameFramework/Actor.h"
+#include "GameFramework/PlayerController.h"
+
+#include "HAL/LowLevelMemStats.h"
+
+#include "Net/DataBunch.h"
+#include "Net/DataChannel.h"
+#include "Net/Core/Connection/NetEnums.h"
+#include "Net/Core/Connection/NetCloseResult.h"
+#include "Net/Core/Misc/NetSubObjectRegistry.h"
+#include "Net/Core/Trace/NetDebugName.h"
+#include "Net/NetSubObjectRegistryGetter.h"
+
+#include "ProfilingDebugging/AssetMetadataTrace.h"
+
+#include "Templates/Casts.h"
+
+#include "UObject/Package.h"
+
+#include <limits>
+
+#define UE_LOG_ACTORREPLICATIONBRIDGE(Category, Format, ...)	UE_LOG(LogIrisBridge, Category, TEXT("ActorReplicationBridge(%u)::") Format, GetReplicationSystem()->GetId(), ##__VA_ARGS__)
+
+extern bool GDefaultUseSubObjectReplicationList;
+
+namespace UE::Net::Private
+{
+
+static bool bEnableActorLevelChanges = true;
+static FAutoConsoleVariableRef CVarEnableActorLevelChanges(
+	TEXT("net.Iris.EnableActorLevelChanges"),
+	bEnableActorLevelChanges,
+	TEXT("When true the ActorReplicationBridge will process actors that change levels by updating the actor's level groups.")
+);
+
+static bool bEnableDynamicNetUpdateFrequency = true;
+static FAutoConsoleVariableRef CVarEnableDynamicNetUpdateFrequency(
+	TEXT("net.Iris.EnableDynamicNetUpdateFrequency"),
+	bEnableDynamicNetUpdateFrequency,
+	TEXT("When true changes to AActor::NetUpdateFrequency will be updated in Iris after being registered for replication.")
+);
+
+bool IsActorValidForIrisReplication(const AActor* Actor)
+{
+	return IsValid(Actor) && !Actor->IsActorBeingDestroyed() && !Actor->IsUnreachable();
+}
+
+void ActorReplicationBridgePreUpdateFunction(TArrayView<UObject*> Instances, const UReplicationBridge* Bridge)
+{
+	UNetDriver* NetDriver = CastChecked<const UEngineReplicationBridge>(Bridge)->GetNetDriver();
+	for (UObject* Instance : Instances)
+	{
+		AActor* Actor = Cast<AActor>(Instance);
+		if (IsActorValidForIrisReplication(Actor))
+		{
+			Actor->CallPreReplication(NetDriver);
+		}
+	}
+}
+
+void ActorReplicationBridgeGetActorWorldObjectInfo(FNetRefHandle /*Handle*/, const UObject* Instance, FVector& OutWorldLocation, float& OutCullDistance)
+{
+	if (const AActor* Actor = Cast<AActor>(Instance))
+	{
+		OutWorldLocation = Actor->GetActorLocation();
+		OutCullDistance = Actor->GetNetCullDistanceSquared() > 0.0f ? FMath::Sqrt(Actor->GetNetCullDistanceSquared()) : 0.0f;
+	}
+}
+
+bool ShouldIncludeActorInLevelGroups(const AActor* Actor)
+{
+	// Never filter out PlayerControllers based on level as they are required for travel.
+	// Preserves the special case for PlayerControllers from UNetDriver::IsLevelInitializedForActor.
+	return !Actor->IsA<APlayerController>();
+}
+
+} // end namespace UE::Net::Private
+
+UEngineReplicationBridge::UEngineReplicationBridge()
+: UObjectReplicationBridge()
+, ActorFactoryId(UE::Net::InvalidNetObjectFactoryId)
+, SubObjectFactoryId(UE::Net::InvalidNetObjectFactoryId)
+{
+	SetInstancePreUpdateFunction(UE::Net::Private::ActorReplicationBridgePreUpdateFunction);
+	SetInstanceGetWorldObjectInfoFunction(UE::Net::Private::ActorReplicationBridgeGetActorWorldObjectInfo);
+}
+
+void UEngineReplicationBridge::Initialize(UReplicationSystem* InReplicationSystem)
+{
+	using namespace UE::Net;
+
+	Super::Initialize(InReplicationSystem);
+
+	ensureMsgf(GDefaultUseSubObjectReplicationList, TEXT("Iris requires replicated actors to use registered subobjectslists. Add \n[SystemSettings]\nnet.SubObjects.DefaultUseSubObjectReplicationList=1\n to your DefaultEngine.ini"));
+
+	ActorFactoryId = FNetObjectFactoryRegistry::GetFactoryIdFromName(UNetActorFactory::GetFactoryName());
+	checkf(ActorFactoryId != InvalidNetObjectFactoryId, TEXT("UNetActorFactory with name %s was not registered"), ToCStr(UNetActorFactory::GetFactoryName().ToString()));
+
+	SubObjectFactoryId = FNetObjectFactoryRegistry::GetFactoryIdFromName(UNetSubObjectFactory::GetFactoryName());
+	checkf(SubObjectFactoryId != InvalidNetObjectFactoryId, TEXT("UNetSubObjectFactory with name %s was not registered"), ToCStr(UNetSubObjectFactory::GetFactoryName().ToString()));
+
+	{
+		auto ShouldSpatialize = [](const UClass* Class)
+		{
+			if (AActor* CDO = Cast<AActor>(Class->GetDefaultObject()))
+			{
+				return !(CDO->bAlwaysRelevant || CDO->bOnlyRelevantToOwner || CDO->bNetUseOwnerRelevancy);
+			}
+
+			return false;
+		};
+
+		SetShouldUseDefaultSpatialFilterFunction(ShouldSpatialize);
+	}
+
+	{
+		auto ClassesAreRelevantEqual = [](const UClass* Class, const UClass* Subclass)
+		{
+			const AActor* CDO = Cast<AActor>(Class->GetDefaultObject());
+			const AActor* SubCDO = Cast<AActor>(Subclass->GetDefaultObject());
+			// Same CDO (nullptr)?
+			if (CDO == SubCDO)
+			{
+				return true;
+			}
+
+			if (CDO == nullptr || SubCDO == nullptr)
+			{
+				return false;
+			}
+
+			{
+				return CDO->bAlwaysRelevant == SubCDO->bAlwaysRelevant
+					&& CDO->bOnlyRelevantToOwner == SubCDO->bOnlyRelevantToOwner
+					&& CDO->bNetUseOwnerRelevancy == SubCDO->bNetUseOwnerRelevancy;
+			}
+		};
+
+		SetShouldSubclassUseSameFilterFunction(ClassesAreRelevantEqual);
+	}
+
+	ObjectReferencePackageMap = NewObject<UIrisObjectReferencePackageMap>();
+}
+
+UEngineReplicationBridge::~UEngineReplicationBridge()
+{
+}
+
+void UEngineReplicationBridge::Deinitialize()
+{
+	if (NetDriver)
+	{
+		NetDriver->OnNetServerMaxTickRateChanged.RemoveAll(this);
+		NetDriver->GetOnNetUpdateFrequencyChanged().RemoveAll(this);
+		NetDriver = nullptr;
+	}
+	Super::Deinitialize();
+	ObjectReferencePackageMap = nullptr;
+}
+
+UE::Net::FNetRefHandle UEngineReplicationBridge::StartReplicatingActor(AActor* Actor, const FActorReplicationParams& Params)
+{
+	using namespace UE::Net;
+
+	if (!ensureMsgf(Actor == nullptr || !(Actor->HasAnyFlags(RF_ClassDefaultObject | RF_ArchetypeObject)), TEXT("Actor %s is a CDO or Archetype and should not be replicated."), ToCStr(GetFullNameSafe(Actor))))
+	{
+		return FNetRefHandle::GetInvalid();
+	}
+
+	const bool bIsNetActor = ULevel::IsNetActor(Actor);
+	if (!bIsNetActor)
+	{
+		UE_LOG_ACTORREPLICATIONBRIDGE(VeryVerbose, TEXT("Actor %s doesn't have a NetRole."), ToCStr(GetFullNameSafe(Actor)));
+		return FNetRefHandle::GetInvalid();
+	}
+
+	if (Actor->GetLocalRole() != ROLE_Authority)
+	{
+		UE_LOG_ACTORREPLICATIONBRIDGE(VeryVerbose, TEXT("Actor %s NetRole isn't Authority."), ToCStr(GetFullNameSafe(Actor)));
+		return FNetRefHandle::GetInvalid();
+	}
+
+	if (Actor->IsActorBeingDestroyed() || !IsValid(Actor) || Actor->IsUnreachable())
+	{
+		UE_LOG_ACTORREPLICATIONBRIDGE(Verbose, TEXT("Actor %s is being destroyed or unreachable and can't be replicated."), ToCStr(GetFullNameSafe(Actor)));
+		return FNetRefHandle::GetInvalid();
+	}
+
+	if (!Actor->GetIsReplicated())
+	{
+		UE_LOG_ACTORREPLICATIONBRIDGE(Verbose, TEXT("Actor %s is not supposed to be replicated."), ToCStr(GetFullNameSafe(Actor)));
+		return FNetRefHandle::GetInvalid();
+	}
+
+	if (Actor->GetTearOff())
+	{
+		UE_LOG_ACTORREPLICATIONBRIDGE(Verbose, TEXT("Actor %s is torn off and should not be replicated."), ToCStr(GetFullNameSafe(Actor)));
+		return FNetRefHandle::GetInvalid();
+	}
+
+	if (!Actor->IsActorInitialized())
+	{
+		UE_LOG_ACTORREPLICATIONBRIDGE(Warning, TEXT("Actor %s is not initialized and won't be replicated."), ToCStr(GetFullNameSafe(Actor)));
+		return FNetRefHandle::GetInvalid();
+	}
+
+	if (!ensureMsgf(Actor->HasActorBegunPlay() || Actor->IsActorBeginningPlay(), TEXT("Actor %s hasn't begun or isn't beginning play in BeginReplication."), ToCStr(GetFullNameSafe(Actor))))
+	{
+		UE_LOG_ACTORREPLICATIONBRIDGE(Error, TEXT("Actor %s hasn't begun or isn't beginning play in BeginReplication."), ToCStr(GetFullNameSafe(Actor)));
+		return FNetRefHandle::GetInvalid();
+	}
+
+	if (!NetDriver)
+	{
+		UE_LOG_ACTORREPLICATIONBRIDGE(VeryVerbose, TEXT("There's no NetDriver so nothing can be replicated."));
+		return FNetRefHandle::GetInvalid();
+	}
+
+
+	if (!NetDriver->ShouldReplicateActor(Actor))
+	{
+		UE_LOG_ACTORREPLICATIONBRIDGE(VeryVerbose, TEXT("Actor %s doesn't want to replicate with NetDriver %s."), ToCStr(GetFullNameSafe(Actor)), ToCStr(NetDriver->GetName()));
+		return FNetRefHandle::GetInvalid();
+	}
+
+	// Initially dormant actors begin replication when their dormancy is flushed
+	const ENetDormancy Dormancy = Actor->NetDormancy;	
+	if (Actor->IsNetStartupActor() && (Dormancy == DORM_Initial))
+	{
+		return FNetRefHandle::GetInvalid();
+	}
+
+	FNetRefHandle ExistingHandle = GetReplicatedRefHandle(Actor);
+	if (ExistingHandle.IsValid())
+	{
+		return ExistingHandle;
+	}
+
+	if (!Actor->IsUsingRegisteredSubObjectList())
+	{
+		// Ensure the first time to get attention!
+		ensureMsgf(false, TEXT("Actor %s does not replicate subobjects using the registered SubObjectsLists, SubObjects will not replicate properly"), ToCStr(GetFullNameSafe(Actor)));
+		UE_LOG_ACTORREPLICATIONBRIDGE(Warning, TEXT("Actor %s does not replicate subobjects using the registered SubObjectsLists, SubObjects will not replicate properly"), ToCStr(GetFullNameSafe(Actor)));
+	}
+
+	// Create handles for the registered fragments
+	Super::FRootObjectReplicationParams RootObjectParams = {
+		.bNeedsPreUpdate = 1U,
+		.bNeedsWorldLocationUpdate = 1U,
+		.bIsDormant = Actor->NetDormancy > DORM_Awake,
+		.StaticPriority = (Actor->bAlwaysRelevant || Actor->bOnlyRelevantToOwner) ? Actor->NetPriority : 0.0f,
+		.PollFrequency = Actor->GetNetUpdateFrequency()
+	};
+
+#if !UE_BUILD_SHIPPING
+	ensureMsgf(!(Actor->bAlwaysRelevant || Actor->bOnlyRelevantToOwner) || RootObjectParams.StaticPriority >= 1.0f, TEXT("Very low NetPriority %.02f for always relevant or owner relevant Actor %s. Set it to 1.0f or higher."), Actor->NetPriority, ToCStr(Actor->GetName()));
+#endif
+
+	FNetRefHandle ActorRefHandle = StartReplicatingRootObject(Actor, RootObjectParams, ActorFactoryId);
+
+	if (!ActorRefHandle.IsValid())
+	{
+		ensureMsgf(false, TEXT("Failed to create NetRefHandle for Actor Named %s"), ToCStr(Actor->GetName()));
+		return FNetRefHandle::GetInvalid();
+	}
+
+    UE_CLOG(Actor->bAlwaysRelevant, LogIrisBridge, Verbose, TEXT("StartReplicatingActor of AlwaysRelevant actor %s"), *PrintObjectFromNetRefHandle(ActorRefHandle));
+
+	// Set owning connection filtering if actor is only relevant to owner
+	{
+		if (Actor->bOnlyRelevantToOwner && !Actor->bAlwaysRelevant)
+		{
+			// Only apply owner filter if we haven't force enabled a dynamic filter.
+			constexpr bool bRequireForceEnabled = true;
+			FName FilterProfile;
+			FNetObjectFilterHandle FilterHandle = GetDynamicFilter(Actor->GetClass(), bRequireForceEnabled, FilterProfile);
+
+			if (FilterHandle == InvalidNetObjectFilterHandle)
+			{
+				GetReplicationSystem()->SetFilter(ActorRefHandle, ToOwnerFilterHandle);
+			}
+		}
+	}
+
+	// Set if this is a NetTemporary
+	{
+		if (Actor->bNetTemporary)
+		{
+			GetReplicationSystem()->SetIsNetTemporary(ActorRefHandle);
+		}
+	}
+
+	// Dormancy, we track all actors that does want to be dormant
+	{
+		if (Dormancy > DORM_Awake)
+		{
+			SetObjectWantsToBeDormant(ActorRefHandle, true);
+		}
+	}
+
+	// Setup Level filtering
+	AddActorToLevelGroup(Actor);
+
+	// If we have registered sub objects we replicate them as well
+	const FSubObjectRegistry& ActorSubObjects = FSubObjectRegistryGetter::GetSubObjects(Actor);
+	const TArray<FReplicatedComponentInfo>& ReplicatedComponents = FSubObjectRegistryGetter::GetReplicatedComponents(Actor);
+
+	if (ActorSubObjects.GetRegistryList().Num() != 0 || ReplicatedComponents.Num() != 0)
+	{
+		const FSubObjectReplicationParams SubObjectParams { .RootObjectHandle = ActorRefHandle };
+		// Start with the Actor's SubObjects (that is SubObjects that are not ActorComponents)
+		for (const FSubObjectRegistry::FEntry& SubObjectInfo : ActorSubObjects.GetRegistryList())
+		{
+			UObject* SubObjectToReplicate = SubObjectInfo.GetSubObject();
+			if (IsValid(SubObjectToReplicate) && SubObjectInfo.NetCondition != ELifetimeCondition::COND_Never)
+			{
+				FNetRefHandle SubObjectRefHandle = Super::StartReplicatingSubObject(SubObjectToReplicate, SubObjectParams, SubObjectFactoryId);
+				if (SubObjectRefHandle.IsValid() && SubObjectInfo.NetCondition != ELifetimeCondition::COND_None)
+				{
+					UObjectReplicationBridge::SetSubObjectNetCondition(SubObjectRefHandle, SubObjectInfo.NetCondition);
+				}
+			}
+		}
+
+		// Now the replicated ActorComponents and their SubObjects
+		for (const FReplicatedComponentInfo& RepComponentInfo : ReplicatedComponents)
+		{
+			if (IsValid(RepComponentInfo.Component) && RepComponentInfo.NetCondition != ELifetimeCondition::COND_Never)
+			{
+				UActorComponent* ReplicatedComponent = RepComponentInfo.Component;
+				ReplicatedComponent->BeginReplication();
+				// NetCondition is set by replicated component
+			}
+		}
+	}
+
+	return ActorRefHandle;
+}
+
+UE::Net::FNetRefHandle UEngineReplicationBridge::StartReplicatingComponent(FNetRefHandle RootObjectHandle, UActorComponent* SubObject)
+{
+	using namespace UE::Net;
+
+	if (!RootObjectHandle.IsValid())
+	{
+		return FNetRefHandle::GetInvalid();
+	}
+
+	AActor* Owner = SubObject->GetOwner();
+
+	FNetRefHandle ReplicatedComponentHandle = GetReplicatedRefHandle(SubObject);
+	const FReplicatedComponentInfo* RepComponentInfo = FSubObjectRegistryGetter::GetReplicatedComponentInfoForComponent(Owner, SubObject);
+
+	if (!ReplicatedComponentHandle.IsValid())
+	{
+		if (!IsValid(SubObject)
+			|| SubObject->IsUnreachable()
+			|| !SubObject->GetIsReplicated()
+			|| SubObject->HasAnyFlags(RF_ArchetypeObject | RF_ClassDefaultObject)
+		)
+		{
+			return FNetRefHandle::GetInvalid();
+		}
+
+		if (!SubObject->IsUsingRegisteredSubObjectList())
+		{
+			UE_LOG_ACTORREPLICATIONBRIDGE(Warning, TEXT("ActorComponent %s does not replicate subobjects using the registered SubObjectsLists, SubObjects might not replicate properly."), ToCStr(GetFullNameSafe(SubObject)));
+		}
+
+		if (RepComponentInfo == nullptr || RepComponentInfo->NetCondition == ELifetimeCondition::COND_Never)
+		{
+			return FNetRefHandle::GetInvalid();
+		}
+
+		// Start replicating the subobject with its owner.
+		const FSubObjectReplicationParams Params{ .RootObjectHandle = RootObjectHandle };
+		ReplicatedComponentHandle = Super::StartReplicatingSubObject(SubObject, Params, SubObjectFactoryId);
+	}
+
+	if (!ReplicatedComponentHandle.IsValid())
+	{
+		ensureMsgf(false, TEXT("Failed to create or find NetRefHandle for ActorComponent Named %s"), ToCStr(SubObject->GetName()));
+		return FNetRefHandle::GetInvalid();
+	}
+
+	// Update or set any conditionals
+	if (RepComponentInfo->NetCondition != ELifetimeCondition::COND_None)
+	{
+		SetSubObjectNetCondition(ReplicatedComponentHandle, RepComponentInfo->NetCondition);
+	}
+
+	// Begin replication for any SubObjects registered by the component
+	{
+		const FSubObjectReplicationParams Params
+		{
+			.RootObjectHandle = RootObjectHandle,
+			.InsertRelativeToSubObjectHandle = ReplicatedComponentHandle,
+			.InsertionOrder = ESubObjectInsertionOrder::ReplicateWith
+		};
+	
+		for (const FSubObjectRegistry::FEntry& SubObjectInfo : RepComponentInfo->SubObjects.GetRegistryList())
+		{
+			UObject* SubObjectToReplicate = SubObjectInfo.GetSubObject();
+			if (IsValid(SubObjectToReplicate) && SubObjectInfo.NetCondition != ELifetimeCondition::COND_Never)
+			{
+				FNetRefHandle SubObjectHandle = Super::StartReplicatingSubObject(SubObjectToReplicate, Params, SubObjectFactoryId);
+				if (SubObjectHandle.IsValid() && SubObjectInfo.NetCondition != ELifetimeCondition::COND_None)
+				{
+					SetSubObjectNetCondition(SubObjectHandle, SubObjectInfo.NetCondition);
+				}
+			}
+		}
+	}
+
+	return ReplicatedComponentHandle;
+}
+
+UE::Net::FNetRefHandle UEngineReplicationBridge::StartReplicatingSubObject(UObject* SubObject, const FSubObjectReplicationParams& Params)
+{
+	return Super::StartReplicatingSubObject(SubObject, Params, SubObjectFactoryId);
+}
+
+void UEngineReplicationBridge::StopReplicatingActor(AActor* Actor, EEndPlayReason::Type EndPlayReason)
+{
+	using namespace UE::Net;
+
+	FNetRefHandle RefHandle = GetReplicatedRefHandle(Actor, EGetRefHandleFlags::EvenIfGarbage);
+	if (!RefHandle.IsValid())
+	{
+		// Already not replicated
+		return;
+	}
+
+	UE_LOG(LogIrisBridge, Verbose, TEXT("StopReplicatingActor %s. Reason %s "), *PrintObjectFromNetRefHandle(RefHandle), *UEnum::GetValueAsString(TEXT("Engine.EEndPlayReason"), EndPlayReason));
+	ensureMsgf(IsValid(Actor), TEXT("StopReplicatingActor called on invalid actor tied to handle: %s."), *PrintObjectFromNetRefHandle(RefHandle));
+	
+	EEndReplicationFlags Flags = EEndReplicationFlags::None;
+	const bool bShouldDestroyObject = EndPlayReason == EEndPlayReason::Destroyed;		
+	if (bShouldDestroyObject)
+	{ 
+		Flags |= EEndReplicationFlags::Destroy | EEndReplicationFlags::DestroyNetHandle | EEndReplicationFlags::ClearNetPushId;
+	}
+
+	// If we are shutting down or the actor is a nettemporary we do not need to validate that we are not detaching remote instances by accident.
+	const bool bIsShuttingDown = (EndPlayReason == EEndPlayReason::EndPlayInEditor) || (EndPlayReason == EEndPlayReason::Quit);
+	// Streaming out a level can happen prior to all actors in it being destroyed/ending replication through replication so we don't want to validate they're detached by accident.
+	const bool bIsStreamingOutLevel = (EndPlayReason == EEndPlayReason::RemovedFromWorld);
+	if (bIsShuttingDown || bIsStreamingOutLevel || Actor->bNetTemporary)
+	{
+		Flags |= EEndReplicationFlags::SkipPendingEndReplicationValidation;
+	}
+					
+	const bool bShouldCreateDestructionInfo = RefHandle.IsStatic() && bShouldDestroyObject && GetReplicationSystem()->IsServer();
+	if (bShouldCreateDestructionInfo)
+	{
+		UObjectReplicationBridge::FEndReplicationParameters EndReplicationParameters
+		{
+			.Location = Actor->GetActorLocation(),
+			.Level = Actor->GetLevel(),
+			.bUseDistanceBasedPrioritization = (Actor->bAlwaysRelevant == false)
+		};
+
+		// Store this static actor destruction info so it can be replicated if a client can know to delete it in the future.
+		StoreDestructionInfo(RefHandle, EndReplicationParameters);
+	}
+
+	StopReplicatingNetRefHandle(RefHandle, Flags);
+}
+
+void UEngineReplicationBridge::StopReplicatingComponent(UActorComponent* ActorComponent, EEndReplicationFlags EndReplicationFlags)
+{
+	using namespace UE::Net;
+
+	FNetRefHandle ComponentHandle = GetReplicatedRefHandle(ActorComponent, EGetRefHandleFlags::EvenIfGarbage);
+	if (ComponentHandle.IsValid())
+	{
+		UE_LOG(LogIrisBridge, Verbose, TEXT("StopReplicatingComponent for %s %s."), *GetNameSafe(ActorComponent), *ComponentHandle.ToString());
+		ensureMsgf(IsValid(ActorComponent), TEXT("StopReplicatingActor called on invalid component tied to handle: %s."), *ComponentHandle.ToString());
+		
+		StopReplicatingNetRefHandle(ComponentHandle, EndReplicationFlags);
+	}
+}
+
+void UEngineReplicationBridge::EndInstantiateFromRemote(FNetRefHandle Handle)
+{
+	if (AActor* Actor = Cast<AActor>(GetReplicatedObject(Handle)))
+	{
+		// Optional
+		//Actor->NetHandle = Handle;
+
+		Actor->PostNetInit();
+	}
+}
+
+//$IRIS todo: move this to be a factory callback
+void UEngineReplicationBridge::OnSubObjectCreatedFromReplication(FNetRefHandle SubObjectHandle)
+{
+	AActor* RootObject = Cast<AActor>(GetReplicatedObject(GetRootObjectOfSubObject(SubObjectHandle)));
+	UObject* SubObject = GetReplicatedObject(SubObjectHandle);
+	if (IsValid(RootObject) && IsValid(SubObject))
+	{
+		RootObject->OnSubobjectCreatedFromReplication(SubObject);
+	}
+}
+
+void UEngineReplicationBridge::DestroyInstanceFromRemote(const FDestroyInstanceParams& Params)
+{
+	if (Params.DestroyReason == EReplicationBridgeDestroyInstanceReason::DoNotDestroy)
+	{
+		return;
+	}
+
+	if (AActor* Actor = Cast<AActor>(Params.Instance))
+	{
+		if ((Params.DestroyReason == EReplicationBridgeDestroyInstanceReason::TearOff) && !NetDriver->ShouldClientDestroyTearOffActors())
+		{
+			NetDriver->ClientSetActorTornOff(Actor);
+		}
+		else if (EnumHasAnyFlags(Params.DestroyFlags, EReplicationBridgeDestroyInstanceFlags::AllowDestroyInstanceFromRemote))
+		{
+			// Any subobjects have already been detached from ReplicationBridge
+			Actor->PreDestroyFromReplication();
+			Actor->Destroy(true);
+		}
+	}
+	else
+	{
+		// If the SubObject is being torn off it is up to owning actor to clean it up properly
+		if (Params.DestroyReason == EReplicationBridgeDestroyInstanceReason::TearOff)
+		{
+			return;
+		}
+
+		if (!EnumHasAnyFlags(Params.DestroyFlags, EReplicationBridgeDestroyInstanceFlags::AllowDestroyInstanceFromRemote))
+		{
+			return;
+		}
+
+		AActor* Owner = Cast<AActor>(Params.RootObject);
+		if (ensureMsgf(IsValid(Owner) && !Owner->IsUnreachable(), TEXT("UEngineReplicationBridge::DestroyInstanceFromRemote Destroyed subobject: %s has an invalid owner: %s"), *GetNameSafe(Params.Instance), *GetPathNameSafe(Params.RootObject)))
+		{
+			Owner->OnSubobjectDestroyFromReplication(Params.Instance);
+		}
+
+		Params.Instance->PreDestroyFromReplication();
+		Params.Instance->MarkAsGarbage();
+	}
+}
+
+void UEngineReplicationBridge::GetInitialDependencies(FNetRefHandle Handle, FNetDependencyInfoArray& OutDependencies) const
+{
+	using namespace UE::Net;
+
+	uint32 DependencyCount = 0U;
+
+	// $TODO: Cache create dependencies and do the lookup based on index
+
+	// Handles with static names does not have have any initial dependencies (other than the reference itself)
+	if (Handle.IsDynamic())
+	{
+		UObject* Object = GetReplicatedObject(Handle);
+		if (!Object)
+		{
+			return;
+		}
+
+		if (AActor* Actor = Cast<AActor>(Object))
+		{
+			UObject* Archetype = nullptr;
+			UObject* ActorLevel = nullptr;
+
+			// ChildActor's need to be spawned from the ChildActorTemplate otherwise any non-replicated 
+			// customized properties will be incorrect on the Client.
+			if (UChildActorComponent* CAC = Actor->GetParentComponent())
+			{
+				Archetype = CAC->GetSpawnableChildActorTemplate();
+			}
+			if (Archetype == nullptr)
+			{
+				Archetype = Actor->GetArchetype();
+			}
+			ActorLevel = Actor->GetLevel();
+
+			check( Archetype != nullptr );
+			check( Actor->NeedsLoadForClient() );			// We have no business sending this unless the client can load
+			check( Archetype->NeedsLoadForClient() );		// We have no business sending this unless the client can load
+
+			// Add initial dependencies
+	
+			// Archetype
+			OutDependencies.Emplace(FNetDependencyInfo(GetOrCreateObjectReference(Archetype)));
+			
+			// Level if it differs from the persistent level
+			const bool bUsePersistentLevel = NetDriver->GetWorld()->PersistentLevel == ActorLevel;
+			if (!bUsePersistentLevel)
+			{
+				OutDependencies.Emplace(FNetDependencyInfo(GetOrCreateObjectReference(ActorLevel)));
+			}
+		}
+		else
+		{
+			// Add initial dependencies
+
+			// SubObjects either have a dependency on their path relative the owner or a reference to their class
+			if (const bool bIsNameStableForNetworking = Object->IsNameStableForNetworking())
+			{
+				OutDependencies.Emplace(FNetDependencyInfo(GetOrCreateObjectReference(Object)));
+			}
+			else
+			{
+				OutDependencies.Emplace(FNetDependencyInfo(GetOrCreateObjectReference(Object->GetClass())));
+			}
+		}
+	}
+}
+
+void UEngineReplicationBridge::SetNetDriver(UNetDriver* const InNetDriver)
+{
+	if (NetDriver)
+	{
+		NetDriver->OnNetServerMaxTickRateChanged.RemoveAll(this);
+		NetDriver->GetOnNetUpdateFrequencyChanged().RemoveAll(this);
+	}
+
+	NetDriver = InNetDriver;
+	if (InNetDriver != nullptr)
+	{
+		SetMaxTickRate(static_cast<float>(FPlatformMath::Max(InNetDriver->GetNetServerMaxTickRate(), 0)));
+
+		InNetDriver->OnNetServerMaxTickRateChanged.AddUObject(this, &UEngineReplicationBridge::OnMaxTickRateChanged);
+		InNetDriver->GetOnNetUpdateFrequencyChanged().AddUObject(this, &UEngineReplicationBridge::OnNetUpdateFrequencyChanged);
+
+		const FName RequiredChannelName = UObjectReplicationBridgeConfig::GetConfig()->GetRequiredNetDriverChannelClassName();
+		
+		if (!RequiredChannelName.IsNone())
+		{
+			const bool bRequiredChannelIsConfigured = InNetDriver->ChannelDefinitions.FindByPredicate([RequiredChannelName](const FChannelDefinition& rhs)
+				{
+					return rhs.ClassName == RequiredChannelName;
+				}) != nullptr;
+
+			checkf(bRequiredChannelIsConfigured, TEXT("ObjectReplication needs the netdriver channel %s to work. Add this channel to the netdriver channel definitions config"), *RequiredChannelName.ToString());
+		}
+	}
+}
+
+void UEngineReplicationBridge::OnMaxTickRateChanged(UNetDriver* InNetDriver, int32 NewMaxTickRate, int32 OldMaxTickRate)
+{
+	SetMaxTickRate(static_cast<float>(FPlatformMath::Max(InNetDriver->GetNetServerMaxTickRate(), 0)));
+
+	ReinitPollFrequency();
+}
+
+bool UEngineReplicationBridge::RemapPathForPIE(uint32 ConnectionId, FString& Str, bool bReading) const
+{
+	if (ConnectionId == UE::Net::InvalidConnectionId)
+	{
+		return GEngine->NetworkRemapPath(static_cast<UNetConnection*>(nullptr), Str, bReading);
+	}
+	else
+	{
+		UObject* UserData = GetReplicationSystem()->GetConnectionUserData(ConnectionId);
+		UNetConnection* NetConnection = Cast<UNetConnection>(UserData);
+		return GEngine->NetworkRemapPath(NetConnection, Str, bReading);
+	}
+}
+
+bool UEngineReplicationBridge::ObjectLevelHasFinishedLoading(UObject* Object) const
+{
+	const UWorld* DriverWorld = Object && NetDriver ? NetDriver->GetWorld() : nullptr;
+	if (DriverWorld)
+	{
+		// get the level for the object
+		const ULevel* Level = Object->GetTypedOuter<ULevel>();
+
+		if (Level != nullptr && Level != DriverWorld->PersistentLevel)
+		{
+			return Level->bIsVisible;
+		}
+	}
+
+	return true;
+}
+
+bool UEngineReplicationBridge::IsAllowedToDestroyInstance(const UObject* Instance) const
+{
+	if (AActor* Actor = const_cast<AActor*>(Cast<AActor>(Instance)))
+	{
+		return !NetDriver || NetDriver->ShouldClientDestroyActor(Actor);
+	}
+
+	return true;
+}
+
+UEngineReplicationBridge* UEngineReplicationBridge::Create(UNetDriver* NetDriver)
+{
+	UEngineReplicationBridge* Bridge = NewObject<UEngineReplicationBridge>(GetTransientPackage(), UEngineReplicationBridge::StaticClass());
+	Bridge->SetNetDriver(NetDriver);
+
+	return Bridge;
+}
+
+float UEngineReplicationBridge::GetPollFrequencyOfRootObject(const UObject* ReplicatedObject) const
+{
+	const AActor* ReplicatedActor = CastChecked<AActor>(ReplicatedObject);
+	float PollFrequency = ReplicatedActor->GetNetUpdateFrequency();
+	GetClassPollFrequency(ReplicatedActor->GetClass(), PollFrequency);
+	return PollFrequency;
+}
+
+void UEngineReplicationBridge::WakeUpObjectInstantiatedFromRemote(AActor* Actor) const
+{
+	// If the actor is already awake or can't be woken up then return immediately.
+	if (Actor->NetDormancy <= DORM_Awake)
+	{
+		return;
+	}
+
+	ENetDormancy OldDormancy = Actor->NetDormancy;
+	Actor->NetDormancy = DORM_Awake;
+
+	if (!NetDriver)
+	{
+		return;
+	}
+
+	if (FWorldContext* WorldContext = GEngine->GetWorldContextFromWorld(NetDriver->GetWorld()))
+	{
+		for (FNamedNetDriver& Driver : WorldContext->ActiveNetDrivers)
+		{
+			if (Driver.NetDriver != nullptr && Driver.NetDriver != NetDriver && Driver.NetDriver->ShouldReplicateActor(Actor))
+			{
+				Driver.NetDriver->NotifyActorClientDormancyChanged(Actor, OldDormancy);
+			}
+		}
+	}
+}
+
+void UEngineReplicationBridge::OnProtocolMismatchDetected(FNetRefHandle ObjectHandle)
+{
+	Super::OnProtocolMismatchDetected(ObjectHandle);
+
+	// As a client tell the server we could not bind this specific NetRefHandle
+	if (NetDriver && NetDriver->ServerConnection)
+	{
+		uint64 RawHandleId = ObjectHandle.GetId();
+		FNetControlMessage<NMT_IrisProtocolMismatch>::Send(NetDriver->ServerConnection, RawHandleId);
+	}
+}
+
+void UEngineReplicationBridge::OnProtocolMismatchReported(FNetRefHandle RefHandle, uint32 ConnectionId)
+{
+	Super::OnProtocolMismatchReported(RefHandle, ConnectionId);
+
+	// If we are the server force the client to disconnect since not replicating a critical class will prevent the game from working.
+	if (NetDriver && NetDriver->IsServer())
+	{
+		const UObject* ReplicatedObject = GetReplicatedObject(RefHandle);
+
+		// If the object instance doesn't exist anymore, pass a null class anyway in case the config wants to disconnect on ALL class types.
+		const UClass* ObjectClass = ReplicatedObject ? ReplicatedObject->GetClass() : nullptr;
+
+		if (IsClassCritical(ObjectClass))
+		{
+			if (UNetConnection* ClientConnection = NetDriver->GetConnectionById(ConnectionId))
+			{
+				FString ErrorMsg = FString::Printf(TEXT("Protocol mismatch: %s:%s. Class: %s"), *RefHandle.ToString(), *GetNameSafe(ReplicatedObject), *GetNameSafe(ObjectClass));
+				UE_LOG(LogIrisBridge, Error, TEXT("%s: Closing connection due to: %s"), ToCStr(ClientConnection->Describe()), ToCStr(ErrorMsg));
+				{
+					UE::Net::FNetCloseResult CloseReason = ENetCloseResult::IrisProtocolMismatch;
+					ClientConnection->SendCloseReason(MoveTemp(CloseReason));
+				}
+
+				FNetControlMessage<NMT_Failure>::Send(ClientConnection, ErrorMsg);
+				ClientConnection->FlushNet(true);
+
+				{
+					UE::Net::FNetCloseResult CloseReason = ENetCloseResult::IrisProtocolMismatch;
+					ClientConnection->Close(MoveTemp(CloseReason));
+				}
+			}
+		}
+	}
+}
+
+void UEngineReplicationBridge::SendErrorWithNetRefHandle(UE::Net::ENetRefHandleError ErrorType, FNetRefHandle RefHandle, uint32 ConnectionId)
+{
+	if (NetDriver)
+	{
+		if (UNetConnection* ClientConnection = NetDriver->GetConnectionById(ConnectionId))
+		{
+			uint64 RawHandleId = RefHandle.GetId();
+			FNetControlMessage<NMT_IrisNetRefHandleError>::Send(ClientConnection, ErrorType, RawHandleId);
+		}
+		else
+		{
+			UE_LOG(LogIrisBridge, Error, TEXT("UEngineReplicationBridge::SendErrorWithNetRefHandle could not find Connection for id:%u"), ConnectionId);
+		}
+	}
+}
+
+void UEngineReplicationBridge::ActorChangedLevel(const AActor* Actor, const ULevel* PreviousLevel)
+{
+	if (!UE::Net::Private::bEnableActorLevelChanges)
+	{
+		return;
+	}
+
+	UE_LOG_ACTORREPLICATIONBRIDGE(Verbose, TEXT("ActorChangedLevel: Actor %s from PreviousLevel %s"), *GetFullNameSafe(Actor), *GetFullNameSafe(PreviousLevel));
+
+	if (!UE::Net::Private::ShouldIncludeActorInLevelGroups(Actor))
+	{
+		return;
+	}
+
+	// Remove from previous level group
+	const bool bPreviousLevelIsPersistentLevel = PreviousLevel && (PreviousLevel->IsPersistentLevel() || (PreviousLevel == NetDriver->GetWorld()->PersistentLevel));
+	
+	if (PreviousLevel && !bPreviousLevelIsPersistentLevel)
+	{
+		const UPackage* const PreviousLevelPackage = PreviousLevel->GetOutermost();
+		const FName PreviousLevelPackageName = PreviousLevelPackage->GetFName();
+
+		const UE::Net::FNetRefHandle ActorRefHandle = GetReplicatedRefHandle(Actor);
+		UE::Net::FNetObjectGroupHandle PreviousLevelGroup = GetLevelGroup(PreviousLevel);
+		if (GetReplicationSystem()->IsInGroup(PreviousLevelGroup, ActorRefHandle))
+		{
+			UE_LOG_ACTORREPLICATIONBRIDGE(Verbose, TEXT("ActorChangedLevel: removing %s from GroupIndex: %u PreviousLevel: %s"), *ActorRefHandle.ToString(), PreviousLevelGroup.GetGroupIndex(), ToCStr(PreviousLevelPackageName.ToString()));
+			GetReplicationSystem()->RemoveFromGroup(PreviousLevelGroup, ActorRefHandle);
+		}
+		else
+		{
+			UE_LOG_ACTORREPLICATIONBRIDGE(Warning, TEXT("ActorChangedLevel: %s not found in GroupIndex: %u PreviousLevel: %s"), *ActorRefHandle.ToString(), PreviousLevelGroup.GetGroupIndex(), ToCStr(PreviousLevelPackageName.ToString()));
+		}
+	}
+
+	AddActorToLevelGroup(Actor);
+}
+
+void UEngineReplicationBridge::OnNetUpdateFrequencyChanged(const AActor* Actor)
+{
+	if (UE::Net::Private::bEnableDynamicNetUpdateFrequency)
+	{
+		const UE::Net::FNetRefHandle ActorHandle = GetReplicatedRefHandle(Actor);
+
+		if (ActorHandle != UE::Net::FNetRefHandle::GetInvalid())
+		{
+			SetPollFrequency(ActorHandle, Actor->GetNetUpdateFrequency());
+		}
+	}
+}
+
+void UEngineReplicationBridge::AddActorToLevelGroup(const AActor* Actor)
+{
+	if (!UE::Net::Private::ShouldIncludeActorInLevelGroups(Actor))
+	{
+		return;
+	}
+	
+	const ULevel* Level = Actor->GetLevel();
+	
+	// Don't filter out actors in the persistent level
+	if (Level && (!Level->IsPersistentLevel() || (Level != NetDriver->GetWorld()->PersistentLevel)))
+	{
+		const UPackage* const LevelPackage = Level->GetOutermost();
+		const FName PackageName = LevelPackage->GetFName();
+
+		UE::Net::FNetObjectGroupHandle LevelGroup = GetLevelGroup(Level);
+		if (!LevelGroup.IsValid())
+		{
+			LevelGroup = CreateLevelGroup(Level, PackageName);
+
+			UE_LOG_ACTORREPLICATIONBRIDGE(Log, TEXT("Created new GroupIndex: %u for Level: %s"), LevelGroup.GetGroupIndex(), ToCStr(PackageName.ToString()));
+
+			// Update the filtering status of the group based on current level visibility for all connections
+			NetDriver->UpdateGroupFilterStatusForLevel(Level, LevelGroup);
+		}
+
+		const UE::Net::FNetRefHandle ActorRefHandle = GetReplicatedRefHandle(Actor);
+
+		// Add object to group
+		UE_LOG_ACTORREPLICATIONBRIDGE(Verbose, TEXT("Added %s to GroupIndex: %u Level: %s"), *ActorRefHandle.ToString(), LevelGroup.GetGroupIndex(), ToCStr(PackageName.ToString()));
+		GetReplicationSystem()->AddToGroup(LevelGroup, ActorRefHandle);
+	}
+}
+
+FString UEngineReplicationBridge::PrintConnectionInfo(uint32 ConnectionId) const
+{
+	if (NetDriver)
+	{
+		if (UNetConnection* ClientConnection = NetDriver->GetConnectionById(ConnectionId))
+		{
+			return FString::Printf(TEXT("ConnectionId:%u ViewTarget: %s Named: %s"), ConnectionId, *GetNameSafe(ClientConnection->ViewTarget), *ClientConnection->Describe());
+		}
+		else
+		{
+			return FString::Printf(TEXT("ConnectionId:%u no NetConnection found"), ConnectionId);
+		}
+	}
+	else
+	{
+		return FString::Printf(TEXT("ConnectionId:%u no NetDriver attached"), ConnectionId);
+	}
+}
+
+void UEngineReplicationBridge::ConsumeNetMetrics(TArray<FAnalyticsEventAttribute>& OutAttrs)
+{
+	using namespace UE::Net;
+
+	FNetMetrics MetricsCollector;
+
+	GetReplicationSystem()->CollectNetMetrics(MetricsCollector);
+
+	//$IRIS TODO: Add ResetNetMetrics call if the RepSystem starts caching cumulative stats
+
+	const TMap<FName, FNetMetric>& Metrics = MetricsCollector.GetMetrics();
+
+	for (auto& It : Metrics)
+	{
+		FName MetricName = It.Key;
+		const FNetMetric& Metric = It.Value;
+
+		switch (Metric.GetDataType())
+		{
+			case FNetMetric::EDataType::Signed:
+			{
+				const int32 Value = Metric.GetSigned();
+				OutAttrs.Add(FAnalyticsEventAttribute(MetricName.ToString(), Value));
+			} break;
+
+			case FNetMetric::EDataType::Unsigned:
+			{
+				const uint32 Value = Metric.GetUnsigned();
+				OutAttrs.Add(FAnalyticsEventAttribute(MetricName.ToString(), Value));
+			} break;
+
+			case FNetMetric::EDataType::Double:
+			{
+				const double Value = Metric.GetDouble();
+				OutAttrs.Add(FAnalyticsEventAttribute(MetricName.ToString(), Value));
+			} break;
+
+			default:
+			{
+				checkNoEntry();
+			} break;
+		}
+	}
+
+}
+
+#else //!UE_WITH_IRIS
+
+UEngineReplicationBridge::UEngineReplicationBridge() = default;
+UEngineReplicationBridge::~UEngineReplicationBridge() = default;
+
+#endif
diff --git a/Engine/Source/Runtime/Engine/Private/Net/Iris/ReplicationSystem/NetActorFactory.cpp b/Engine/Source/Runtime/Engine/Private/Net/Iris/ReplicationSystem/NetActorFactory.cpp
--- a/Engine/Source/Runtime/Engine/Private/Net/Iris/ReplicationSystem/NetActorFactory.cpp
+++ b/Engine/Source/Runtime/Engine/Private/Net/Iris/ReplicationSystem/NetActorFactory.cpp
@@ -25,7 +25,7 @@
 
 #include "Net/DataBunch.h"
 
-#include "Net/Iris/ReplicationSystem/ActorReplicationBridge.h"
+#include "Net/Iris/ReplicationSystem/EngineReplicationBridge.h"
 
 #include "ProfilingDebugging/AssetMetadataTrace.h"
 
@@ -184,7 +184,7 @@ TUniquePtr<UE::Net::FNetObjectCreationHeader> UNetActorFactory::CreateAndFillHea
 		return nullptr;
 	}
 
-	UActorReplicationBridge* ActorBridge = CastChecked<UActorReplicationBridge>(Bridge);
+	UEngineReplicationBridge* ActorBridge = CastChecked<UEngineReplicationBridge>(Bridge);
 
 	TUniquePtr<FBaseActorNetCreationHeader> BaseHeader;
 
@@ -334,7 +334,7 @@ UNetObjectFactory::FInstantiateResult UNetActorFactory::InstantiateReplicatedObj
 	LLM_SCOPE(ELLMTag::EngineMisc);
 	IRIS_PROFILER_SCOPE(NetActorFactory_InstantiateReplicatedObjectFromHeader);
 
-	UActorReplicationBridge* ActorBridge = CastChecked<UActorReplicationBridge>(Bridge);
+	UEngineReplicationBridge* ActorBridge = CastChecked<UEngineReplicationBridge>(Bridge);
 	UNetDriver* NetDriver = ActorBridge->GetNetDriver();
 
 	const FBaseActorNetCreationHeader* BaseHeader = static_cast<const FBaseActorNetCreationHeader*>(Header);
@@ -447,7 +447,7 @@ void UNetActorFactory::PostInstantiation(const FPostInstantiationContext& Contex
 		return;
 	}
 
-	UActorReplicationBridge* ActorBridge = CastChecked<UActorReplicationBridge>(Bridge);
+	UEngineReplicationBridge* ActorBridge = CastChecked<UEngineReplicationBridge>(Bridge);
 	UNetDriver* NetDriver = ActorBridge->GetNetDriver();
 
 	const FBaseActorNetCreationHeader* BaseHeader = static_cast<const FBaseActorNetCreationHeader*>(Context.Header);
diff --git a/Engine/Source/Runtime/Engine/Private/Net/Iris/ReplicationSystem/ReplicationSystemUtil.cpp b/Engine/Source/Runtime/Engine/Private/Net/Iris/ReplicationSystem/ReplicationSystemUtil.cpp
--- a/Engine/Source/Runtime/Engine/Private/Net/Iris/ReplicationSystem/ReplicationSystemUtil.cpp
+++ b/Engine/Source/Runtime/Engine/Private/Net/Iris/ReplicationSystem/ReplicationSystemUtil.cpp
@@ -13,7 +13,7 @@
 #include "Iris/ReplicationSystem/Filtering/NetObjectFilter.h"
 #include "Iris/ReplicationSystem/ReplicationSystem.h"
 #include "Iris/Core/IrisLog.h"
-#include "Net/Iris/ReplicationSystem/ActorReplicationBridge.h"
+#include "Net/Iris/ReplicationSystem/EngineReplicationBridge.h"
 #include "Net/NetSubObjectRegistryGetter.h"
 #include "HAL/IConsoleManager.h"
 #include "Templates/Casts.h"
@@ -73,11 +73,11 @@ UReplicationSystem* FReplicationSystemUtil::GetReplicationSystem(const UNetDrive
 	return NetDriver ? NetDriver->GetReplicationSystem() : nullptr;
 }
 
-UActorReplicationBridge* FReplicationSystemUtil::GetActorReplicationBridge(const AActor* Actor)
+UEngineReplicationBridge* FReplicationSystemUtil::GetActorReplicationBridge(const AActor* Actor)
 {
 	if (UReplicationSystem* ReplicationSystem = GetReplicationSystem(Actor))
 	{
-		return Cast<UActorReplicationBridge>(ReplicationSystem->GetReplicationBridge());
+		return Cast<UEngineReplicationBridge>(ReplicationSystem->GetReplicationBridge());
 	}
 	else
 	{
@@ -85,12 +85,12 @@ UActorReplicationBridge* FReplicationSystemUtil::GetActorReplicationBridge(const
 	}
 }
 
-UActorReplicationBridge* FReplicationSystemUtil::GetActorReplicationBridge(const UNetConnection* NetConnection)
+UEngineReplicationBridge* FReplicationSystemUtil::GetActorReplicationBridge(const UNetConnection* NetConnection)
 {
 	const UNetDriver* Driver = NetConnection ? NetConnection->GetDriver() : nullptr;
 	if (const UReplicationSystem* ReplicationSystem = Driver ? Driver->GetReplicationSystem() : nullptr)
 	{
-		return Cast<UActorReplicationBridge>(ReplicationSystem->GetReplicationBridge());
+		return Cast<UEngineReplicationBridge>(ReplicationSystem->GetReplicationBridge());
 	}
 	else
 	{
@@ -124,7 +124,7 @@ void FReplicationSystemUtil::BeginReplication(AActor* Actor, const FActorReplica
 		{
 			if (ReplicationSystem->IsServer())
 			{
-				if (UActorReplicationBridge* Bridge = Cast<UActorReplicationBridge>(ReplicationSystem->GetReplicationBridge()))
+				if (UEngineReplicationBridge* Bridge = Cast<UEngineReplicationBridge>(ReplicationSystem->GetReplicationBridge()))
 				{
 					Bridge->StartReplicatingActor(Actor, Params);
 				}
@@ -145,7 +145,7 @@ void FReplicationSystemUtil::EndReplication(AActor* Actor, EEndPlayReason::Type
 	// The bridge itself will verify that the actor is replicated by it so there's no reason to check that either.
 	ReplicationSystemUtil::ForEachReplicationSystem([Actor, EndPlayReason](UReplicationSystem* ReplicationSystem)
 	{
-		if (UActorReplicationBridge* Bridge = Cast<UActorReplicationBridge>(ReplicationSystem->GetReplicationBridge()))
+		if (UEngineReplicationBridge* Bridge = Cast<UEngineReplicationBridge>(ReplicationSystem->GetReplicationBridge()))
 		{
 			Bridge->StopReplicatingActor(Actor, EndPlayReason);
 		}
@@ -175,7 +175,7 @@ void FReplicationSystemUtil::BeginReplicationForActorComponent(FNetHandle ActorH
 			{
 				if (ReplicationSystem->IsServer())
 				{
-					if (UActorReplicationBridge* Bridge = Cast<UActorReplicationBridge>(ReplicationSystem->GetReplicationBridge()))
+					if (UEngineReplicationBridge* Bridge = Cast<UEngineReplicationBridge>(ReplicationSystem->GetReplicationBridge()))
 					{
 						const FNetRefHandle OwnerRefHandle = Bridge->GetReplicatedRefHandle(ActorHandle);
 						if (OwnerRefHandle.IsValid())
@@ -204,7 +204,7 @@ void FReplicationSystemUtil::BeginReplicationForActorComponent(const AActor* Act
 		{
 			if (ReplicationSystem->IsServer())
 			{
-				if (UActorReplicationBridge* Bridge = Cast<UActorReplicationBridge>(ReplicationSystem->GetReplicationBridge()))
+				if (UEngineReplicationBridge* Bridge = Cast<UEngineReplicationBridge>(ReplicationSystem->GetReplicationBridge()))
 				{
 					const FNetRefHandle ActorRefHandle = Bridge->GetReplicatedRefHandle(ActorHandle);
 					if (ActorRefHandle.IsValid())
@@ -237,7 +237,7 @@ void FReplicationSystemUtil::BeginReplicationForActorSubObject(const AActor* Act
 		{
 			if (ReplicationSystem->IsServer())
 			{
-				if (UActorReplicationBridge* Bridge = ReplicationSystem->GetReplicationBridgeAs<UActorReplicationBridge>())
+				if (UEngineReplicationBridge* Bridge = ReplicationSystem->GetReplicationBridgeAs<UEngineReplicationBridge>())
 				{
 					const FNetRefHandle ActorRefHandle = Bridge->GetReplicatedRefHandle(ActorHandle);
 					if (ActorRefHandle.IsValid())
@@ -272,7 +272,7 @@ void FReplicationSystemUtil::BeginReplicationForActorComponentSubObject(UActorCo
 			{
 				if (ReplicationSystem->IsServer())
 				{
-					if (UActorReplicationBridge* Bridge = ReplicationSystem->GetReplicationBridgeAs<UActorReplicationBridge>())
+					if (UEngineReplicationBridge* Bridge = ReplicationSystem->GetReplicationBridgeAs<UEngineReplicationBridge>())
 					{
 						const FNetRefHandle ActorRefHandle = Bridge->GetReplicatedRefHandle(ActorHandle);
 						const FNetRefHandle ActorComponentRefHandle = Bridge->GetReplicatedRefHandle(ActorComponent);
@@ -301,7 +301,7 @@ void FReplicationSystemUtil::EndReplicationForActorComponent(UActorComponent* Ac
 {
 	ReplicationSystemUtil::ForEachReplicationSystem([ActorComponent](UReplicationSystem* ReplicationSystem)
 	{
-		if (UActorReplicationBridge* Bridge = Cast<UActorReplicationBridge>(ReplicationSystem->GetReplicationBridge()))
+		if (UEngineReplicationBridge* Bridge = Cast<UEngineReplicationBridge>(ReplicationSystem->GetReplicationBridge()))
 		{
 			constexpr EEndReplicationFlags EndReplicationFlags = EEndReplicationFlags::DestroyNetHandle | EEndReplicationFlags::ClearNetPushId;
 			Bridge->StopReplicatingComponent(ActorComponent, EndReplicationFlags);
@@ -313,7 +313,7 @@ void FReplicationSystemUtil::EndReplicationForActorSubObject(const AActor* Actor
 {
 	ReplicationSystemUtil::ForEachReplicationSystem([SubObject](UReplicationSystem* ReplicationSystem)
 	{
-		if (UActorReplicationBridge* Bridge = Cast<UActorReplicationBridge>(ReplicationSystem->GetReplicationBridge()))
+		if (UEngineReplicationBridge* Bridge = Cast<UEngineReplicationBridge>(ReplicationSystem->GetReplicationBridge()))
 		{				
 			constexpr EEndReplicationFlags EndReplicationFlags = EEndReplicationFlags::Destroy | EEndReplicationFlags::DestroyNetHandle | EEndReplicationFlags::ClearNetPushId;
 			Bridge->StopReplicatingNetObject(SubObject, EndReplicationFlags);
@@ -347,7 +347,7 @@ void FReplicationSystemUtil::AddDependentActor(const AActor* Parent, AActor* Chi
 	{
 		if (ReplicationSystem->IsServer())
 		{
-			if (UActorReplicationBridge* Bridge = ReplicationSystem->GetReplicationBridgeAs<UActorReplicationBridge>())
+			if (UEngineReplicationBridge* Bridge = ReplicationSystem->GetReplicationBridgeAs<UEngineReplicationBridge>())
 			{
 				const FNetRefHandle ParentRefHandle = Bridge->GetReplicatedRefHandle(ParentHandle);
 				if (ParentRefHandle.IsValid())
diff --git a/Engine/Source/Runtime/Engine/Private/NetConnection.cpp b/Engine/Source/Runtime/Engine/Private/NetConnection.cpp
--- a/Engine/Source/Runtime/Engine/Private/NetConnection.cpp
+++ b/Engine/Source/Runtime/Engine/Private/NetConnection.cpp
@@ -48,7 +48,7 @@
 #include "Iris/IrisConfig.h"
 #include "Iris/ReplicationSystem/ReplicationSystem.h"
 #include "Iris/ReplicationSystem/Filtering/NetObjectFilter.h"
-#include "Net/Iris/ReplicationSystem/ActorReplicationBridge.h"
+#include "Net/Iris/ReplicationSystem/EngineReplicationBridge.h"
 #endif // UE_WITH_IRIS
 
 #include UE_INLINE_GENERATED_CPP_BY_NAME(NetConnection)
diff --git a/Engine/Source/Runtime/Engine/Private/NetDriver.cpp b/Engine/Source/Runtime/Engine/Private/NetDriver.cpp
--- a/Engine/Source/Runtime/Engine/Private/NetDriver.cpp
+++ b/Engine/Source/Runtime/Engine/Private/NetDriver.cpp
@@ -60,7 +60,7 @@
 #include "Iris/ReplicationSystem/ReplicationSystem.h"
 #include "Iris/ReplicationSystem/ReplicationView.h"
 #include "Iris/ReplicationSystem/Filtering/NetObjectFilter.h"
-#include "Net/Iris/ReplicationSystem/ActorReplicationBridge.h"
+#include "Net/Iris/ReplicationSystem/EngineReplicationBridge.h"
 #include "Net/Iris/ReplicationSystem/ReplicationSystemUtil.h"
 #endif // UE_WITH_IRIS
 
@@ -3925,7 +3925,7 @@ void UNetDriver::NotifyActorDestroyed(AActor* ThisActor, bool IsSeamlessTravel)
 
 			if (bShouldCreateDestructionInfoForInitiallyDormantActor)
 			{
-				if (UActorReplicationBridge* Bridge = ReplicationSystem->GetReplicationBridgeAs<UActorReplicationBridge>())
+				if (UEngineReplicationBridge* Bridge = ReplicationSystem->GetReplicationBridgeAs<UEngineReplicationBridge>())
 				{
 					const UE::Net::FNetObjectReference Reference = Bridge->GetOrCreateObjectReference(ThisActor);
 					if (Reference.GetRefHandle().IsValid() && Reference.GetRefHandle().IsStatic())
@@ -4048,7 +4048,7 @@ void UNetDriver::DeleteSubObjectOnClients(AActor* Actor, UObject* SubObject)
 #if UE_WITH_IRIS
 	if (ReplicationSystem)
 	{
-		if (UActorReplicationBridge* Bridge = Cast<UActorReplicationBridge>(ReplicationSystem->GetReplicationBridge()))
+		if (UEngineReplicationBridge* Bridge = Cast<UEngineReplicationBridge>(ReplicationSystem->GetReplicationBridge()))
 		{
 			constexpr EEndReplicationFlags EndReplicationFlags = EEndReplicationFlags::Destroy | EEndReplicationFlags::DestroyNetHandle | EEndReplicationFlags::ClearNetPushId;
 			Bridge->StopReplicatingNetObject(SubObject, EndReplicationFlags);
@@ -4068,7 +4068,7 @@ void UNetDriver::TearOffSubObjectOnClients(AActor* Actor, UObject* SubObject)
 #if UE_WITH_IRIS
 	if (ReplicationSystem)
 	{
-		if (UActorReplicationBridge* Bridge = Cast<UActorReplicationBridge>(ReplicationSystem->GetReplicationBridge()))
+		if (UEngineReplicationBridge* Bridge = Cast<UEngineReplicationBridge>(ReplicationSystem->GetReplicationBridge()))
 		{
 			Bridge->StopReplicatingNetObject(SubObject, EEndReplicationFlags::TearOff);
 		}
@@ -4143,7 +4143,7 @@ void UNetDriver::NotifyActorRenamed(AActor* ThisActor, UObject* PreviousOuter, F
 #if UE_WITH_IRIS
 				if (ReplicationSystem)
 				{
-					if (UActorReplicationBridge* Bridge = ReplicationSystem->GetReplicationBridgeAs<UActorReplicationBridge>())
+					if (UEngineReplicationBridge* Bridge = ReplicationSystem->GetReplicationBridgeAs<UEngineReplicationBridge>())
 					{
 						Bridge->ActorChangedLevel(ThisActor, Cast<ULevel>(PreviousOuter));
 					}
@@ -6563,7 +6563,7 @@ void UNetDriver::CreateReplicatedStaticActorDestructionInfo(ULevel* Level, const
 #if UE_WITH_IRIS
 	if (ReplicationSystem)
 	{
-		if (UActorReplicationBridge* Bridge = ReplicationSystem->GetReplicationBridgeAs<UActorReplicationBridge>())
+		if (UEngineReplicationBridge* Bridge = ReplicationSystem->GetReplicationBridgeAs<UEngineReplicationBridge>())
 		{			
 			// Add explicit destruction info for this object
 			UReplicationBridge::FEndReplicationParameters Params;
@@ -6988,7 +6988,7 @@ bool UNetDriver::InitReplicationBridgeClass()
 	else
 	{
 		// Fall back on ActorReplicationBridge
-		ReplicationBridgeClass = UActorReplicationBridge::StaticClass();
+		ReplicationBridgeClass = UEngineReplicationBridge::StaticClass();
 	}
 
 	return ReplicationBridgeClass != nullptr;
@@ -7095,7 +7095,7 @@ void UNetDriver::CreateReplicationSystem(bool bInitAsClient)
 		return;
 	}
 
-	UActorReplicationBridge* ReplicationBridge = NewObject<UActorReplicationBridge>(GetTransientPackage(), ReplicationBridgeClass);
+	UEngineReplicationBridge* ReplicationBridge = NewObject<UEngineReplicationBridge>(GetTransientPackage(), ReplicationBridgeClass);
 	if (ReplicationBridge)
 	{
 		ReplicationBridge->SetNetDriver(this);
diff --git a/Engine/Source/Runtime/Engine/Private/PlayerController.cpp b/Engine/Source/Runtime/Engine/Private/PlayerController.cpp
--- a/Engine/Source/Runtime/Engine/Private/PlayerController.cpp
+++ b/Engine/Source/Runtime/Engine/Private/PlayerController.cpp
@@ -71,7 +71,7 @@
 
 #if UE_WITH_IRIS
 #include "Iris/ReplicationSystem/ReplicationSystem.h"
-#include "Net/Iris/ReplicationSystem/ActorReplicationBridge.h"
+#include "Net/Iris/ReplicationSystem/EngineReplicationBridge.h"
 #include "Net/Iris/ReplicationSystem/ReplicationSystemUtil.h"
 #endif // UE_WITH_IRIS
 
diff --git a/Engine/Source/Runtime/Engine/Private/World.cpp b/Engine/Source/Runtime/Engine/Private/World.cpp
--- a/Engine/Source/Runtime/Engine/Private/World.cpp
+++ b/Engine/Source/Runtime/Engine/Private/World.cpp
@@ -131,7 +131,7 @@
 
 #if UE_WITH_IRIS
 #include "Iris/ReplicationSystem/ReplicationSystem.h"
-#include "Net/Iris/ReplicationSystem/ActorReplicationBridge.h"
+#include "Net/Iris/ReplicationSystem/EngineReplicationBridge.h"
 #endif // UE_WITH_IRIS
 
 #include "ChaosSolversModule.h"
diff --git a/Engine/Source/Runtime/Engine/Public/Net/Iris/ReplicationSystem/ActorReplicationBridge.h b/Engine/Source/Runtime/Engine/Public/Net/Iris/ReplicationSystem/ActorReplicationBridge.h
deleted file mode 100644
diff --git a/Engine/Source/Runtime/Engine/Public/Net/Iris/ReplicationSystem/EngineReplicationBridge.h b/Engine/Source/Runtime/Engine/Public/Net/Iris/ReplicationSystem/EngineReplicationBridge.h
new file mode 100644
--- /dev/null
+++ b/Engine/Source/Runtime/Engine/Public/Net/Iris/ReplicationSystem/EngineReplicationBridge.h
@@ -0,0 +1,138 @@
+// Copyright Epic Games, Inc. All Rights Reserved.
+
+#pragma once
+
+#include "Iris/ReplicationSystem/ObjectReplicationBridge.h"
+#include "Engine/EngineTypes.h"
+#include "UObject/ObjectPtr.h"
+#include "EngineReplicationBridge.generated.h"
+
+class UNetDriver;
+class UIrisObjectReferencePackageMap;
+struct FAnalyticsEventAttribute;
+
+#if UE_WITH_IRIS
+
+class UActorComponent;
+class UWorld;
+
+namespace UE::Net
+{
+	enum class ENetRefHandleError : uint32;
+}
+
+namespace UE::Net
+{
+
+// If actor should be replicated using IRIS or old replication system
+ENGINE_API bool ShouldUseIrisReplication(const UObject* Actor);
+
+}
+
+struct FActorReplicationParams
+{
+	/** When true we ignore the configured dynamic filter for this actor type and use the explicit filter instead */
+	bool bOverrideDynamicFilterConfig = false;
+
+	/** Only used when bOverrideDynamicFilterConfig is true. The name of the dynamic filter to assign to this actor. Can be NAME_None if no filters should be set. */
+	FName ExplicitDynamicFilterName;
+};
+
+#endif // UE_WITH_IRIS
+
+UCLASS(Transient, MinimalAPI)
+class UEngineReplicationBridge final : public UObjectReplicationBridge
+{
+	GENERATED_BODY()
+
+public:
+	ENGINE_API UEngineReplicationBridge();
+	virtual ENGINE_API ~UEngineReplicationBridge() override;
+
+#if UE_WITH_IRIS
+
+	ENGINE_API static UEngineReplicationBridge* Create(UNetDriver* NetDriver);
+
+	/** Sets the net driver for the bridge. */
+	ENGINE_API void SetNetDriver(UNetDriver* const InNetDriver);
+	
+	/** Get net driver used by the bridge .*/
+	inline UNetDriver* GetNetDriver() const { return NetDriver; }
+
+	/** Begin replication of an actor and its registered ActorComponents and SubObjects. */
+	ENGINE_API FNetRefHandle StartReplicatingActor(AActor* Instance, const FActorReplicationParams& Params);
+
+	/** Stop replicating an actor. Will destroy handle for actor and registered subobjects. */
+	ENGINE_API void StopReplicatingActor(AActor* Actor, EEndPlayReason::Type EndPlayReason);
+		
+	/**
+	 * Begin replication of an ActorComponent and its registered SubObjects, 
+	 * if the ActorComponent already is replicated any set NetObjectConditions will be updated.
+	*/
+	ENGINE_API FNetRefHandle StartReplicatingComponent(FNetRefHandle RootObjectHandle, UActorComponent* ActorComponent);
+
+	/** Begin replication of a subobject. */
+	ENGINE_API FNetRefHandle StartReplicatingSubObject(UObject* SubObject, const FSubObjectReplicationParams& Params);
+
+	/** Stop replicating an ActorComponent and its associated SubObjects. */
+	ENGINE_API void StopReplicatingComponent(UActorComponent* ActorComponent, EEndReplicationFlags EndReplicationFlags = EEndReplicationFlags::None);
+
+	/** Get object reference packagemap. Used in special cases where serialization hasn't been converted to use NetSerializers.  */
+	UIrisObjectReferencePackageMap* GetObjectReferencePackageMap() const { return ObjectReferencePackageMap; }
+
+	/** Tell the remote connection that we detected a reading error with a specific replicated object */
+	ENGINE_API virtual void SendErrorWithNetRefHandle(UE::Net::ENetRefHandleError ErrorType, FNetRefHandle RefHandle, uint32 ConnectionId) override;
+	
+	/** Updates the level group for an actor that changed levels */
+	void ActorChangedLevel(const AActor* Actor, const ULevel* PreviousLevel);
+
+	/** Called when NetUpdateFrequency has changed on the Actor. */
+	void OnNetUpdateFrequencyChanged(const AActor* Actor);
+
+	void WakeUpObjectInstantiatedFromRemote(AActor* Actor) const;
+
+	/**
+	 * Add relevant network metrics gathered since the last call to ConsumeNetMetrics.
+	 * Any periodic stat will be reset here too.
+	 * @param OutAttrs A list of Name/Value pairings that will be sent to an AnalyticsProvider
+	 */
+	ENGINE_API void ConsumeNetMetrics(TArray<FAnalyticsEventAttribute>& OutAttrs);
+
+protected:
+
+	// UObjectReplicationBridge
+	virtual void Initialize(UReplicationSystem* ReplicationSystem) override;
+	virtual void Deinitialize() override;
+	virtual void OnSubObjectCreatedFromReplication(FNetRefHandle SubObjectHandle) override;
+	virtual void EndInstantiateFromRemote(FNetRefHandle Handle) override;
+	virtual void DestroyInstanceFromRemote(const FDestroyInstanceParams& Params) override;
+	virtual void GetInitialDependencies(FNetRefHandle Handle, FNetDependencyInfoArray& OutDependencies) const override;
+	virtual bool RemapPathForPIE(uint32 ConnectionId, FString& Path, bool bReading) const override;
+	virtual bool ObjectLevelHasFinishedLoading(UObject* Object) const override;
+	virtual bool IsAllowedToDestroyInstance(const UObject* Instance) const override;
+	virtual void OnProtocolMismatchDetected(FNetRefHandle ObjectHandle) override;
+	virtual void OnProtocolMismatchReported(FNetRefHandle RefHandle, uint32 ConnectionId) override;
+
+	virtual float GetPollFrequencyOfRootObject(const UObject* ReplicatedObject) const override;
+
+	[[nodiscard]] virtual FString PrintConnectionInfo(uint32 ConnectionId) const override;
+
+private:
+	
+	void OnMaxTickRateChanged(UNetDriver* InNetDriver, int32 NewMaxTickRate, int32 OldMaxTickRate);
+
+	void AddActorToLevelGroup(const AActor* Actor);
+
+private:
+
+	UE::Net::FNetObjectFactoryId ActorFactoryId;
+	UE::Net::FNetObjectFactoryId SubObjectFactoryId;
+
+#endif // UE_WITH_IRIS
+
+	UNetDriver* NetDriver = nullptr;
+
+	UPROPERTY(Transient)
+	TObjectPtr<UIrisObjectReferencePackageMap> ObjectReferencePackageMap = nullptr;
+
+};
diff --git a/Engine/Source/Runtime/Engine/Public/Net/Iris/ReplicationSystem/ReplicationSystemUtil.h b/Engine/Source/Runtime/Engine/Public/Net/Iris/ReplicationSystem/ReplicationSystemUtil.h
--- a/Engine/Source/Runtime/Engine/Public/Net/Iris/ReplicationSystem/ReplicationSystemUtil.h
+++ b/Engine/Source/Runtime/Engine/Public/Net/Iris/ReplicationSystem/ReplicationSystemUtil.h
@@ -13,7 +13,7 @@
 class AActor;
 class UActorComponent;
 struct FActorReplicationParams;
-class UActorReplicationBridge;
+class UEngineReplicationBridge;
 class UObject;
 class FRepChangedPropertyTracker;
 class UReplicationSystem;
@@ -38,11 +38,11 @@ struct FReplicationSystemUtil
 	/** Returns the UReplicationSystem for a specific NetDriver. May return null. */
 	ENGINE_API static UReplicationSystem* GetReplicationSystem(const UNetDriver* NetDriver);
 
-	/** Returns the UActorReplicationBridge of the UReplicationSystem belogning to the main NetDriver assigned to the Actor. May return null. */
-	ENGINE_API static UActorReplicationBridge* GetActorReplicationBridge(const AActor* Actor);
+	/** Returns the UEngineReplicationBridge of the UReplicationSystem belogning to the main NetDriver assigned to the Actor. May return null. */
+	ENGINE_API static UEngineReplicationBridge* GetActorReplicationBridge(const AActor* Actor);
 
 	/** Returns the UReplicationSystem of the UNetDriver the UNetConnection belongs to. */
-	ENGINE_API static UActorReplicationBridge* GetActorReplicationBridge(const UNetConnection* NetConnection);
+	ENGINE_API static UEngineReplicationBridge* GetActorReplicationBridge(const UNetConnection* NetConnection);
 
 	/** Returns the NetHandle for the actor. The returned handle may be invalid. */
 	ENGINE_API static FNetHandle GetNetHandle(const AActor* Actor);
