Fix Subsurface profile being not passed to the material proxy when converted through Substrate.

diff --git a/Engine/Source/Runtime/Engine/Private/Materials/Material.cpp b/Engine/Source/Runtime/Engine/Private/Materials/Material.cpp
--- a/Engine/Source/Runtime/Engine/Private/Materials/Material.cpp
+++ b/Engine/Source/Runtime/Engine/Private/Materials/Material.cpp
@@ -3289,7 +3289,7 @@ EBlendMode ConvertLegacyBlendMode(EBlendMode InBlendMode, FMaterialShadingModelF
 
 #if WITH_EDITOR
 
-void UMaterial::ConvertMaterialToSubstrateMaterial()
+bool UMaterial::ConvertMaterialToSubstrateMaterial()
 {
 	/*
 	* The data flow for legacy material conversion node that can be used in isolation is as such:
@@ -3309,7 +3309,7 @@ void UMaterial::ConvertMaterialToSubstrateMaterial()
 	*/
 	if (!Substrate::IsSubstrateEnabled())
 	{
-		return;
+		return false;
 	}
 
 	// Store current node post from the root node.
@@ -3953,6 +3953,8 @@ void UMaterial::ConvertMaterialToSubstrateMaterial()
 		// We might have moved connections above so update the CachedExpressionData from the EditorOnly connection data (ground truth).
 		UpdateCachedExpressionData();
 	}
+
+	return bInvalidateShader;
 }
 #endif // WITH_EDITOR
 
@@ -4262,9 +4264,16 @@ void UMaterial::PostLoad()
 #endif
 
 #if WITH_EDITOR
-	// Substrate materials conversion needs to be done after expressions are cached, otherwise material function won't have 
-	// valid inputs in certain cases
-	ConvertMaterialToSubstrateMaterial();
+	if (Substrate::IsSubstrateEnabled())
+	{
+		// Substrate materials conversion needs to be done after expressions are cached, otherwise material function won't have 
+		// valid inputs in certain cases
+		if (ConvertMaterialToSubstrateMaterial())
+		{
+			// Call PropagateDataToMaterialProxy in order to propagate Subsurface profiles and specular profiles data to the material proxy
+			PropagateDataToMaterialProxy();
+		}
+	}
 #endif // WITH_EDITOR
 
 	checkf(CachedExpressionData, TEXT("Missing cached expression data for material, should have been either serialized or created during PostLoad"));
diff --git a/Engine/Source/Runtime/Engine/Private/Materials/MaterialInterface.cpp b/Engine/Source/Runtime/Engine/Private/Materials/MaterialInterface.cpp
--- a/Engine/Source/Runtime/Engine/Private/Materials/MaterialInterface.cpp
+++ b/Engine/Source/Runtime/Engine/Private/Materials/MaterialInterface.cpp
@@ -1601,6 +1601,7 @@ void UMaterialInterface::UpdateMaterialRenderProxy(FMaterialRenderProxy& Proxy)
 	{
 		FMaterialRenderProxy* InProxy = &Proxy;
 
+		// Specular profiles
 		{
 			struct FEntry
 			{
@@ -1625,21 +1626,25 @@ void UMaterialInterface::UpdateMaterialRenderProxy(FMaterialRenderProxy& Proxy)
 				}
 			}
 
-			ENQUEUE_RENDER_COMMAND(UpdateMaterialRenderProxySpecular)(
-			[InProxy, Entries](FRHICommandListImmediate& RHICmdList)
+			if (Entries.Num() > 0)
 			{
-				for (const FEntry& Entry : Entries)
+				ENQUEUE_RENDER_COMMAND(UpdateMaterialRenderProxySpecular)(
+				[InProxy, Entries](FRHICommandListImmediate& RHICmdList)
 				{
-					if (Entry.Profile)
+					for (const FEntry& Entry : Entries)
 					{
-						const uint32 AllocationId = SpecularProfile::AddOrUpdateProfile(Entry.Profile, Entry.Guid, Entry.Settings, Entry.Texture);
-						check(AllocationId >= 0 && AllocationId < MAX_SPECULAR_PROFILE_COUNT);
+						if (Entry.Profile)
+						{
+							const uint32 AllocationId = SpecularProfile::AddOrUpdateProfile(Entry.Profile, Entry.Guid, Entry.Settings, Entry.Texture);
+							check(AllocationId >= 0 && AllocationId < MAX_SPECULAR_PROFILE_COUNT);
+						}
+						InProxy->AddSpecularProfileRT(Entry.Profile);
 					}
-					InProxy->AddSpecularProfileRT(Entry.Profile);
-				}
-			});
+				});
+			}
 		}
 
+		// Subsurface profiles
 		{
 			struct FEntry
 			{
diff --git a/Engine/Source/Runtime/Engine/Public/Materials/Material.h b/Engine/Source/Runtime/Engine/Public/Materials/Material.h
--- a/Engine/Source/Runtime/Engine/Public/Materials/Material.h
+++ b/Engine/Source/Runtime/Engine/Public/Materials/Material.h
@@ -1399,7 +1399,7 @@ private:
 	void BackwardsCompatibilityInputConversion();
 	void BackwardsCompatibilityVirtualTextureOutputConversion();
 	void BackwardsCompatibilityDecalConversion();
-	void ConvertMaterialToSubstrateMaterial();
+	bool ConvertMaterialToSubstrateMaterial();
 
 	/** Marks the material's package dirty in order to make a material usage change set during map load persistent. 
 	  * This couldn't be done during map load as loading cannot mark packages dirty. Invoked manually by the user 
