FSuggestProjectileVelocityParameters no longer keeps a mutable reference to FCollisionResponseParams::DefaultResponseParam.
Prevents game code from unintentionally modifying DefaultResponseParam when modifying a locally declared FSuggestProjectileVelocityParameters.

diff --git a/Engine/Source/Runtime/Engine/Classes/Kismet/GameplayStatics.h b/Engine/Source/Runtime/Engine/Classes/Kismet/GameplayStatics.h
--- a/Engine/Source/Runtime/Engine/Classes/Kismet/GameplayStatics.h
+++ b/Engine/Source/Runtime/Engine/Classes/Kismet/GameplayStatics.h
@@ -564,7 +564,7 @@ public:
 		float CollisionRadius = 0.f;
 		float OverrideGravityZ = 0;
 		ESuggestProjVelocityTraceOption::Type TraceOption = ESuggestProjVelocityTraceOption::TraceFullPath;
-		FCollisionResponseParams& ResponseParam = FCollisionResponseParams::DefaultResponseParam;
+		FCollisionResponseParams ResponseParam = FCollisionResponseParams::DefaultResponseParam;
 		TArray<AActor*> ActorsToIgnore = TArray<AActor*>();
 		bool bDrawDebug = false;
 		bool bAcceptClosestOnNoSolutions = false;
