Fix assertion when trying to break LevelInstanceActor when its level isn't current

diff --git a/Engine/Source/Runtime/Engine/Private/LevelInstance/LevelInstanceSubsystem.cpp b/Engine/Source/Runtime/Engine/Private/LevelInstance/LevelInstanceSubsystem.cpp
--- a/Engine/Source/Runtime/Engine/Private/LevelInstance/LevelInstanceSubsystem.cpp
+++ b/Engine/Source/Runtime/Engine/Private/LevelInstance/LevelInstanceSubsystem.cpp
@@ -1627,9 +1627,17 @@ bool ULevelInstanceSubsystem::BreakLevelInstance(ILevelInstanceInterface* LevelI
 	ON_SCOPE_EXIT { GetMutableDefault<ULevelEditorMiscSettings>()->bAvoidRelabelOnPasteSelected = bAvoidRelabelOnPasteSelected; };
 	GetMutableDefault<ULevelEditorMiscSettings>()->bAvoidRelabelOnPasteSelected = 1;
 
+	ULevel* OldCurrentLevel = GetWorld()->GetCurrentLevel();
+	if (AActor* LevelInstanceActor = CastChecked<AActor>(LevelInstance))
+	{
+		GetWorld()->SetCurrentLevel(LevelInstanceActor->GetLevel());
+	}
+
 	TArray<AActor*> MovedActors;
 	BreakLevelInstance_Impl(LevelInstance, Levels, MovedActors, Flags);
 
+	GetWorld()->SetCurrentLevel(OldCurrentLevel);
+
 	USelection* ActorSelection = GEditor->GetSelectedActors();
 	ActorSelection->BeginBatchSelectOperation();
 	for (AActor* MovedActor : MovedActors)
