Add editor function to remove row from datatable

diff --git a/Engine/Source/Runtime/Engine/Classes/Kismet/DataTableFunctionLibrary.h b/Engine/Source/Runtime/Engine/Classes/Kismet/DataTableFunctionLibrary.h
--- a/Engine/Source/Runtime/Engine/Classes/Kismet/DataTableFunctionLibrary.h
+++ b/Engine/Source/Runtime/Engine/Classes/Kismet/DataTableFunctionLibrary.h
@@ -154,5 +154,9 @@ class UDataTableFunctionLibrary : public UBlueprintFunctionLibrary
     UFUNCTION(BlueprintCallable, CustomThunk, Category = "Editor Scripting | DataTable", meta=(AutoCreateRefTerm="RowName", CustomStructureParam="RowData"))
 	static ENGINE_API void AddDataTableRow(UDataTable* const DataTable, const FName& RowName, const FTableRowBase& RowData);
     DECLARE_FUNCTION(execAddDataTableRow);
+
+	/** Removes the row with the provided name from a Data Table. */
+	UFUNCTION(BlueprintCallable, Category = "Editor Scripting | DataTable", meta = (AutoCreateRefTerm = "RowName"))
+	static ENGINE_API void RemoveDataTableRow(UDataTable* DataTable, const FName& RowName);
 #endif //WITH_EDITOR
 };
diff --git a/Engine/Source/Runtime/Engine/Private/DataTableFunctionLibrary.cpp b/Engine/Source/Runtime/Engine/Private/DataTableFunctionLibrary.cpp
--- a/Engine/Source/Runtime/Engine/Private/DataTableFunctionLibrary.cpp
+++ b/Engine/Source/Runtime/Engine/Private/DataTableFunctionLibrary.cpp
@@ -434,4 +434,23 @@ DEFINE_FUNCTION(UDataTableFunctionLibrary::execAddDataTableRow)
 		FBlueprintCoreDelegates::ThrowScriptException(P_THIS, Stack, ExceptionInfo);
 	}
 }
+
+void UDataTableFunctionLibrary::RemoveDataTableRow(UDataTable* DataTable, const FName& RowName)
+{
+	if (!DataTable)
+	{
+		UE_LOG(LogDataTable, Error, TEXT("RemoveDataTableRow - The DataTable is invalid."));
+		return;
+	}
+
+	if (DataTable->FindRowUnchecked(RowName) != nullptr)
+	{
+		DataTable->Modify();
+		DataTable->RemoveRow(RowName);
+	}
+	else
+	{
+		UE_LOG(LogDataTable, Log, TEXT("RemoveDataTableRow - Row %s not found is %s"), *RowName.ToString(), *DataTable->GetName());
+	}
+}
 #endif //WITH_EDITOR
