[World Partition Actor Descriptors] Fix actor descriptor support for internal actors not prefetching actor BP classes + also added missing error checks.
This will require a resave of the affected levels in otrder to take effect.

diff --git a/Engine/Source/Runtime/Engine/Private/WorldPartition/ActorDescContainer.cpp b/Engine/Source/Runtime/Engine/Private/WorldPartition/ActorDescContainer.cpp
--- a/Engine/Source/Runtime/Engine/Private/WorldPartition/ActorDescContainer.cpp
+++ b/Engine/Source/Runtime/Engine/Private/WorldPartition/ActorDescContainer.cpp
@@ -135,6 +135,20 @@ void UActorDescContainer::Initialize(const FInitializeParams& InitParams)
 			ClassPaths.Add(Asset.AssetClassPath);
 		}
 
+		for (const FString& InternalAsset : InternalAssets)
+		{
+			FWorldPartitionActorDescUtils::FActorDescInitParams ActorDescInitParams(InternalAsset);
+
+			if (!ActorDescInitParams.BaseClassName.IsNone())
+			{
+				ClassPaths.Emplace(ActorDescInitParams.BaseClassName.ToString());
+			}
+			else
+			{
+				ClassPaths.Emplace(ActorDescInitParams.NativeClassName.ToString());
+			}
+		}
+
 		ClassDescRegistry.PrefetchClassDescs(ClassPaths.Array());
 	}
 
@@ -214,8 +228,46 @@ void UActorDescContainer::Initialize(const FInitializeParams& InitParams)
 
 			TUniquePtr<FWorldPartitionActorDesc> ActorDesc = FWorldPartitionActorDescUtils::GetActorDescriptorFromInitParams(ActorDescInitParams, ContainerPackageName);
 
-			ActorDescsByPackage.Add(ActorDesc->GetActorPackage(), ActorDesc.Get());
-			ValidActorDescs.Add(ActorDesc->GetGuid(), MoveTemp(ActorDesc));
+			if (!ActorDesc.IsValid())
+			{
+				UE_LOG(LogWorldPartition, Warning, TEXT("Invalid actor descriptor for actor '%s' from package '%s'"), *ActorDescInitParams.PathName.ToString(), *InitParams.PackageName.ToString());
+			} 
+			else if (!ActorDesc->GetNativeClass().IsValid())
+			{
+				UE_LOG(LogWorldPartition, Warning, TEXT("Invalid actor native class: Actor: '%s' (guid '%s') from package '%s'"),
+					*ActorDesc->GetActorNameString(),
+					*ActorDesc->GetGuid().ToString(),
+					*ActorDesc->GetActorPackage().ToString());
+			}
+			else if (ActorDesc->GetBaseClass().IsValid() && !ClassDescRegistry.IsRegisteredClass(ActorDesc->GetBaseClass()))
+			{
+				UE_LOG(LogWorldPartition, Warning, TEXT("Unknown actor base class `%s`: Actor: '%s' (guid '%s') from package '%s'"),
+					*ActorDesc->GetBaseClass().ToString(),
+					*ActorDesc->GetActorNameString(),
+					*ActorDesc->GetGuid().ToString(),
+					*ActorDesc->GetActorPackage().ToString());
+			}
+			else if (!InitParams.FilterActorDesc || InitParams.FilterActorDesc(ActorDesc.Get()))
+			{
+				// At this point, the actor descriptor is well formed and valid on its own. We now make validations based on the already registered
+				// actor descriptors, such as duplicated actor GUIDs or multiple actors in the same package, etc.
+				 if (TUniquePtr<FWorldPartitionActorDesc>* ExistingActorDescPtr = ValidActorDescs.Find(ActorDesc->GetGuid()))
+				{
+					const FWorldPartitionActorDesc* ExistingActorDesc = ExistingActorDescPtr->Get();
+					check(ExistingActorDesc->GetGuid() == ActorDesc->GetGuid());
+					UE_LOG(LogWorldPartition, Warning, TEXT("Duplicate actor descriptor guid `%s`: Actor: '%s' from package '%s' -> Existing actor '%s' from package '%s'"), 
+						*ActorDesc->GetGuid().ToString(), 
+						*ActorDesc->GetActorNameString(), 
+						*ActorDesc->GetActorPackage().ToString(),
+						*ExistingActorDesc->GetActorNameString(),
+						*ExistingActorDesc->GetActorPackage().ToString());
+				}
+				else
+				{
+					ActorDescsByPackage.Add(ActorDesc->GetActorPackage(), ActorDesc.Get());
+					ValidActorDescs.Add(ActorDesc->GetGuid(), MoveTemp(ActorDesc));
+				}
+			}
 		}
 	}
 
diff --git a/Engine/Source/Runtime/Engine/Private/WorldPartition/WorldPartitionActorDescUtils.cpp b/Engine/Source/Runtime/Engine/Private/WorldPartition/WorldPartitionActorDescUtils.cpp
--- a/Engine/Source/Runtime/Engine/Private/WorldPartition/WorldPartitionActorDescUtils.cpp
+++ b/Engine/Source/Runtime/Engine/Private/WorldPartition/WorldPartitionActorDescUtils.cpp
@@ -65,6 +65,7 @@ FWorldPartitionActorDescUtils::FActorDescInitParams::FActorDescInitParams(const
 	TUniquePtr<FWorldPartitionActorDesc> ActorDesc(InActor->CreateActorDesc());
 
 	PathName = *InActor->GetPathName();
+	BaseClassName = !InActor->GetClass()->IsNative() ? FName(*InActor->GetClass()->GetPathName()) : NAME_None;
 	NativeClassName = *(GetParentNativeClass(InActor->GetClass()))->GetPathName();
 	AssetData = FWorldPartitionActorDescUtils::GetAssetDataFromActorDescriptor(ActorDesc);
 }
@@ -75,6 +76,7 @@ FWorldPartitionActorDescUtils::FActorDescInitParams::FActorDescInitParams(const
 	if (InString.ParseIntoArray(ParsedAssetData, TEXT(",")))
 	{
 		PathName = *ParsedAssetData[0];
+		BaseClassName = ParsedAssetData.IsValidIndex(3) ? FName(*ParsedAssetData[3]) : NAME_None;
 		NativeClassName = *ParsedAssetData[1];
 		AssetData = *ParsedAssetData[2];
 	}
@@ -82,7 +84,7 @@ FWorldPartitionActorDescUtils::FActorDescInitParams::FActorDescInitParams(const
 
 FString FWorldPartitionActorDescUtils::FActorDescInitParams::ToString() const
 {
-	return FString::Printf(TEXT("%s,%s,%s"), *PathName.ToString(), *NativeClassName.ToString(), *AssetData);
+	return FString::Printf(TEXT("%s,%s,%s,%s"), *PathName.ToString(), *NativeClassName.ToString(), *AssetData, *BaseClassName.ToString());
 }
 
 FName FWorldPartitionActorDescUtils::ActorMetaDataClassTagName()
diff --git a/Engine/Source/Runtime/Engine/Public/WorldPartition/WorldPartitionActorDescUtils.h b/Engine/Source/Runtime/Engine/Public/WorldPartition/WorldPartitionActorDescUtils.h
--- a/Engine/Source/Runtime/Engine/Public/WorldPartition/WorldPartitionActorDescUtils.h
+++ b/Engine/Source/Runtime/Engine/Public/WorldPartition/WorldPartitionActorDescUtils.h
@@ -27,6 +27,7 @@ struct FWorldPartitionActorDescUtils
 		FString ToString() const;
 
 		FName PathName;
+		FName BaseClassName;
 		FName NativeClassName;
 		FString AssetData;
 	};
