Fix GCreateShadersOnLoad  by correctly checking GRHIShaderPlatform as it was done in the initial CL 1581054

diff --git a/Engine/Source/Runtime/Engine/Private/Materials/MaterialShader.cpp b/Engine/Source/Runtime/Engine/Private/Materials/MaterialShader.cpp
--- a/Engine/Source/Runtime/Engine/Private/Materials/MaterialShader.cpp
+++ b/Engine/Source/Runtime/Engine/Private/Materials/MaterialShader.cpp
@@ -3358,7 +3358,7 @@ void FMaterialShaderMap::Register(EShaderPlatform InShaderPlatform)
 	} MaterialShaderMapInnerLazyInitializer;
 
 	extern int32 GCreateShadersOnLoad;
-	if (GCreateShadersOnLoad && GetShaderPlatform() == InShaderPlatform)
+	if (GCreateShadersOnLoad && GMaxRHIShaderPlatform == InShaderPlatform)
 	{
 		FShaderMapResource* ShaderResource = GetResource();
 		if (ShaderResource)
