Iris/Networking
- Enabled use of NetTokens/NetTokenStore and NetTokenDataStores outside of Iris to facilitate exporting data from shared serialization paths
- NetTokenStore and NetTokenDataStores are now owned by NetDriver.
- Added support for CustomExport bunches to carry data assoiciated with NetTokenExports

GameplayTags
- GameplayTagNetSerializer now respects settings from GameplayTags.ini
- Implemented placeholder GameplayTagNetTokenDataStore.

GameplayTags are now using NetTokens to export data if running in the new dynamic mode.

diff --git a/Engine/Plugins/Runtime/GameplayAbilities/Source/GameplayAbilities/Private/GameplayDebuggerCategory_Abilities.cpp b/Engine/Plugins/Runtime/GameplayAbilities/Source/GameplayAbilities/Private/GameplayDebuggerCategory_Abilities.cpp
--- a/Engine/Plugins/Runtime/GameplayAbilities/Source/GameplayAbilities/Private/GameplayDebuggerCategory_Abilities.cpp
+++ b/Engine/Plugins/Runtime/GameplayAbilities/Source/GameplayAbilities/Private/GameplayDebuggerCategory_Abilities.cpp
@@ -84,11 +84,15 @@ void FGameplayDebuggerCategory_Abilities::OnShowGameplayAttributesToggle()
 
 void FGameplayDebuggerCategory_Abilities::FRepData::Serialize(FArchive& Ar)
 {
-	bool bSuccess;
-	OwnedTags.NetSerialize(Ar, nullptr, bSuccess);
-
 	Ar << TagCounts;
 
+	FString StrOwnedTags = OwnedTags.ToString();
+	Ar << StrOwnedTags;
+	if (Ar.IsLoading())
+	{
+		OwnedTags.FromExportString(StrOwnedTags);
+	}
+
 	int32 NumAbilities = Abilities.Num();
 	Ar << NumAbilities;
 	if (Ar.IsLoading())
diff --git a/Engine/Plugins/Runtime/ReplicationGraph/Source/Private/ReplicationGraph.cpp b/Engine/Plugins/Runtime/ReplicationGraph/Source/Private/ReplicationGraph.cpp
--- a/Engine/Plugins/Runtime/ReplicationGraph/Source/Private/ReplicationGraph.cpp
+++ b/Engine/Plugins/Runtime/ReplicationGraph/Source/Private/ReplicationGraph.cpp
@@ -2030,6 +2030,11 @@ int64 UReplicationGraph::ReplicateSingleActor_FastShared(AActor* Actor, FConnect
 
 	ensureAlwaysMsgf(OutBunch.bHasMustBeMappedGUIDs == 0, TEXT("FastShared bHasMustBeMappedGUIDs! %s"), *Actor->GetPathName());
 
+	// Need to reset flags set on bunch send as we are re-using the bunch for more than one connection
+	OutBunch.bPartial = 0;
+	OutBunch.bPartialInitial = 0;
+	OutBunch.bPartialFinal = 0;
+
 	return OutBunch.GetNumBits();
 #else
 	return 0;
@@ -2339,6 +2344,7 @@ bool UReplicationGraph::ProcessRemoteFunction(class AActor* Actor, UFunction* Fu
 		{
 			// Reset the bunch here. It will be reused and we should only reset it right before we actually write to it.
 			FastSharedReplicationBunch->Reset();
+			FastSharedReplicationBunch->NetTokensPendingExport.Reset();
 
 			// It sucks we have to a temp writer like this, but we don't know how big the payload will be until we serialize it
 			FNetBitWriter TempWriter(nullptr, 0);
@@ -2360,6 +2366,9 @@ bool UReplicationGraph::ProcessRemoteFunction(class AActor* Actor, UFunction* Fu
 			SetTraceCollector(TempWriter, Collector);
 #endif // UE_NET_TRACE_ENABLED
 
+			// Create export scope, storing NetTokens pending export in FastSharedReplicationBunch->NetTokensPendingExport
+			UE::Net::FNetTokenExportScope NetTokenExportScope(TempWriter, NetDriver->GetNetTokenStore(), FastSharedReplicationBunch->NetTokensPendingExport, "SendPropertiesForRPC");
+
 			TSharedPtr<FRepLayout> RepLayout = NetDriver->GetFunctionRepLayout( Function );
 			RepLayout->SendPropertiesForRPC(Function, FastSharedReplicationChannel, TempWriter, Parameters);
 
@@ -2416,7 +2425,7 @@ bool UReplicationGraph::ProcessRemoteFunction(class AActor* Actor, UFunction* Fu
 			}
 		}
 
-		RepLayout->BuildSharedSerializationForRPC(Parameters);
+		RepLayout->BuildSharedSerializationForRPC(Parameters, NetDriver->GetNetTokenStore());
 		FGlobalActorReplicationInfo& GlobalInfo = GlobalActorReplicationInfoMap.Get(Actor);
 
 		bool ForceFlushNetDormancy = false;
diff --git a/Engine/Plugins/Runtime/ReplicationSystemTestPlugin/Source/Private/Tests/ReplicationSystem/ReplicationSystemServerClientTestFixture.cpp b/Engine/Plugins/Runtime/ReplicationSystemTestPlugin/Source/Private/Tests/ReplicationSystem/ReplicationSystemServerClientTestFixture.cpp
--- a/Engine/Plugins/Runtime/ReplicationSystemTestPlugin/Source/Private/Tests/ReplicationSystem/ReplicationSystemServerClientTestFixture.cpp
+++ b/Engine/Plugins/Runtime/ReplicationSystemTestPlugin/Source/Private/Tests/ReplicationSystem/ReplicationSystemServerClientTestFixture.cpp
@@ -75,6 +75,21 @@ FReplicationSystemTestNode::FReplicationSystemTestNode(FReplicationSystemTestNod
 
 void FReplicationSystemTestNode::Setup(bool bIsServer, const TCHAR* Name, FReplicationSystemTestNode::FReplicationSystemParamsOverride* ParamsOverride)
 {
+	// Init NetTokenStore
+	{
+		using namespace UE::Net;
+
+		NetTokenStore = MakeUnique<FNetTokenStore>();
+
+		FNetTokenStore::FInitParams NetTokenStoreInitParams;
+		NetTokenStoreInitParams.Authority = bIsServer ? FNetToken::ENetTokenAuthority::Authority : FNetToken::ENetTokenAuthority::None;
+		NetTokenStore->Init(NetTokenStoreInitParams);
+
+		// Register data stores for supported types, $TODO: make this configurable.
+		NetTokenDataStores.Add(MakeUnique<FStringTokenStore>(*NetTokenStore));
+		NetTokenDataStores.Add(MakeUnique<FNameTokenStore>(*NetTokenStore));	
+	}
+
 	ReplicationBridge = NewObject<UReplicatedTestObjectBridge>();
 	check(ReplicationBridge != nullptr);
 
@@ -84,6 +99,7 @@ void FReplicationSystemTestNode::Setup(bool bIsServer, const TCHAR* Name, FRepli
 	Params.ReplicationBridge = ReplicationBridge;
 	Params.bIsServer = bIsServer;
 	Params.bAllowObjectReplication = bIsServer;
+	Params.NetTokenStore = NetTokenStore.Get();
 
 	if (ParamsOverride)
 	{
@@ -226,6 +242,10 @@ uint32 FReplicationSystemTestNode::AddConnection()
 	Connection.DataStreamManager = DataStreamManager;
 	CreatedObjects.Add(TStrongObjectPtr<UObject>(Connection.DataStreamManager));
 
+	// Init and store RemoteNetTokenStoreState
+	ReplicationSystem->GetNetTokenStore()->InitRemoteNetTokenStoreState(Connection.ConnectionId);
+	Connection.RemoteNetTokenStoreState = ReplicationSystem->GetNetTokenStore()->GetRemoteNetTokenStoreState(Connection.ConnectionId);
+
 	// Streams created based on config
 	Connection.DataStreamManager->CreateStream("NetToken");
 	Connection.NetTokenDataStream = StaticCast<UNetTokenDataStream*>(Connection.DataStreamManager->GetStream("NetToken"));
@@ -242,9 +262,6 @@ uint32 FReplicationSystemTestNode::AddConnection()
 	ReplicationSystem->InitDataStreams(Connection.ConnectionId, Connection.DataStreamManager);
 	ReplicationSystem->SetReplicationEnabledForConnection(Connection.ConnectionId, true);
 
-	// Store RemoteNetTokenStoreState
-	Connection.RemoteNetTokenStoreState = &ReplicationSystem->GetReplicationSystemInternal()->GetConnections().GetRemoteNetTokenStoreState(Connection.ConnectionId);
-
 	// Add view
 	FReplicationView View;
 	View.Views.AddDefaulted();
diff --git a/Engine/Plugins/Runtime/ReplicationSystemTestPlugin/Source/Private/Tests/ReplicationSystem/ReplicationSystemServerClientTestFixture.h b/Engine/Plugins/Runtime/ReplicationSystemTestPlugin/Source/Private/Tests/ReplicationSystem/ReplicationSystemServerClientTestFixture.h
--- a/Engine/Plugins/Runtime/ReplicationSystemTestPlugin/Source/Private/Tests/ReplicationSystem/ReplicationSystemServerClientTestFixture.h
+++ b/Engine/Plugins/Runtime/ReplicationSystemTestPlugin/Source/Private/Tests/ReplicationSystem/ReplicationSystemServerClientTestFixture.h
@@ -167,6 +167,8 @@ public:
 	float ConvertPollPeriodIntoFrequency(uint32 PollPeriod) const;
 
 public:
+	TUniquePtr<UE::Net::FNetTokenStore> NetTokenStore;
+	TArray<TUniquePtr<UE::Net::FNetTokenDataStore>> NetTokenDataStores;
 	UReplicationSystem* ReplicationSystem;
 	UReplicatedTestObjectBridge* ReplicationBridge;
 	TArray<TStrongObjectPtr<UObject>> CreatedObjects;
diff --git a/Engine/Plugins/Runtime/ReplicationSystemTestPlugin/Source/Private/Tests/ReplicationSystem/ReplicationSystemTestFixture.h b/Engine/Plugins/Runtime/ReplicationSystemTestPlugin/Source/Private/Tests/ReplicationSystem/ReplicationSystemTestFixture.h
--- a/Engine/Plugins/Runtime/ReplicationSystemTestPlugin/Source/Private/Tests/ReplicationSystem/ReplicationSystemTestFixture.h
+++ b/Engine/Plugins/Runtime/ReplicationSystemTestPlugin/Source/Private/Tests/ReplicationSystem/ReplicationSystemTestFixture.h
@@ -5,9 +5,14 @@
 #include "NetworkAutomationTest.h"
 #include "NetworkAutomationTestMacros.h"
 #include "ReplicatedTestObject.h"
+#include "Containers/Array.h"
+#include "Templates/UniquePtr.h"
 #include "Iris/Core/IrisLog.h"
 #include "Iris/ReplicationSystem/ReplicationProtocolManager.h"
 #include "Iris/ReplicationSystem/ReplicationSystem.h"
+#include "Iris/ReplicationSystem/StringTokenStore.h"
+#include "Iris/ReplicationSystem/NameTokenStore.h"
+#include "Net/Core/NetToken/NetToken.h"
 
 namespace UE::Net
 {
@@ -25,13 +30,30 @@ public:
 protected:
 	virtual void SetUp() override
 	{
+		const bool bIsServer = true;
+		// Init NetTokenStore
+		{
+			using namespace UE::Net;
+
+			NetTokenStore = MakeUnique<FNetTokenStore>();
+
+			FNetTokenStore::FInitParams NetTokenStoreInitParams;
+			NetTokenStoreInitParams.Authority = bIsServer ? FNetToken::ENetTokenAuthority::Authority : FNetToken::ENetTokenAuthority::None;
+			NetTokenStore->Init(NetTokenStoreInitParams);
+
+			// Register data stores for supported types, $TODO: make this configurable.
+			NetTokenDataStores.Add(MakeUnique<FStringTokenStore>(*NetTokenStore));
+			NetTokenDataStores.Add(MakeUnique<FNameTokenStore>(*NetTokenStore));	
+		}
+
 		ReplicationBridge = NewObject<UReplicatedTestObjectBridge>();
 		CreatedObjects.Add(TStrongObjectPtr<UObject>(ReplicationBridge));
 
 		UReplicationSystem::FReplicationSystemParams Params;
 		Params.ReplicationBridge = ReplicationBridge;
-		Params.bIsServer = true;
+		Params.bIsServer = bIsServer;
 		Params.bAllowObjectReplication = true;
+		Params.NetTokenStore =  NetTokenStore.Get();
 
 		// In a testing environment without configs the creation of the ReplicationSystem can be quite spammy
 		ELogVerbosity::Type IrisLogVerbosity = UE_GET_LOG_VERBOSITY(LogIris);
@@ -102,6 +124,8 @@ protected:
 		return CreatedObject;
 	}
 
+	TUniquePtr<UE::Net::FNetTokenStore> NetTokenStore;
+	TArray<TUniquePtr<UE::Net::FNetTokenDataStore>> NetTokenDataStores;
 	UReplicationSystem* ReplicationSystem;
 	UReplicatedTestObjectBridge* ReplicationBridge;
 	
diff --git a/Engine/Plugins/Runtime/ReplicationSystemTestPlugin/Source/Private/Tests/ReplicationSystem/TestReplicationSystem.cpp b/Engine/Plugins/Runtime/ReplicationSystemTestPlugin/Source/Private/Tests/ReplicationSystem/TestReplicationSystem.cpp
--- a/Engine/Plugins/Runtime/ReplicationSystemTestPlugin/Source/Private/Tests/ReplicationSystem/TestReplicationSystem.cpp
+++ b/Engine/Plugins/Runtime/ReplicationSystemTestPlugin/Source/Private/Tests/ReplicationSystem/TestReplicationSystem.cpp
@@ -5,7 +5,7 @@
 #include "Iris/Serialization/NetBitStreamWriter.h"
 #include "Iris/ReplicationSystem/ReplicationSystem.h"
 #include "Iris/ReplicationSystem/ReplicationSystemInternal.h"
-#include "Iris/ReplicationSystem/NetTokenStore.h"
+#include "Net/Core/NetToken/NetToken.h"
 #include "Iris/Core/IrisLog.h"
 #include "Iris/Metrics/NetMetrics.h"
 #include "Misc/ScopeExit.h"
diff --git a/Engine/Plugins/Runtime/ReplicationSystemTestPlugin/Source/Private/Tests/Serialization/TestInstancedStructNetSerializer.cpp b/Engine/Plugins/Runtime/ReplicationSystemTestPlugin/Source/Private/Tests/Serialization/TestInstancedStructNetSerializer.cpp
--- a/Engine/Plugins/Runtime/ReplicationSystemTestPlugin/Source/Private/Tests/Serialization/TestInstancedStructNetSerializer.cpp
+++ b/Engine/Plugins/Runtime/ReplicationSystemTestPlugin/Source/Private/Tests/Serialization/TestInstancedStructNetSerializer.cpp
@@ -421,12 +421,10 @@ void FTestInstancedStructNetSerializerFixture::SetUp()
 	// Init default serialization context
 	InternalNetSerializationContext.ReplicationSystem = Server->ReplicationSystem;
 
-	FReplicationSystemInternal* ReplicationSystemInternal = Server->GetReplicationSystem()->GetReplicationSystemInternal();
-
 	FInternalNetSerializationContext TempInternalNetSerializationContext;
 	FInternalNetSerializationContext::FInitParameters TempInternalNetSerializationContextInitParams;
 	TempInternalNetSerializationContextInitParams.ReplicationSystem = Server->ReplicationSystem;
-	TempInternalNetSerializationContextInitParams.ObjectResolveContext.RemoteNetTokenStoreState = ReplicationSystemInternal->GetNetTokenStore().GetLocalNetTokenStoreState();
+	TempInternalNetSerializationContextInitParams.ObjectResolveContext.RemoteNetTokenStoreState = Server->ReplicationSystem->GetNetTokenStore()->GetLocalNetTokenStoreState();
 	TempInternalNetSerializationContext.Init(TempInternalNetSerializationContextInitParams);
 
 	InternalNetSerializationContext = MoveTemp(TempInternalNetSerializationContext);
diff --git a/Engine/Plugins/Runtime/ReplicationSystemTestPlugin/Source/Private/Tests/Serialization/TestPolymorphicStructNetSerializer.cpp b/Engine/Plugins/Runtime/ReplicationSystemTestPlugin/Source/Private/Tests/Serialization/TestPolymorphicStructNetSerializer.cpp
--- a/Engine/Plugins/Runtime/ReplicationSystemTestPlugin/Source/Private/Tests/Serialization/TestPolymorphicStructNetSerializer.cpp
+++ b/Engine/Plugins/Runtime/ReplicationSystemTestPlugin/Source/Private/Tests/Serialization/TestPolymorphicStructNetSerializer.cpp
@@ -432,12 +432,10 @@ void FTestPolymorphicArrayStructNetSerializerFixture::SetUp()
 	// Init default serialization context
 	InternalNetSerializationContext.ReplicationSystem = Server->ReplicationSystem;
 
-	FReplicationSystemInternal* ReplicationSystemInternal = Server->GetReplicationSystem()->GetReplicationSystemInternal();
-
 	FInternalNetSerializationContext TempInternalNetSerializationContext;
 	FInternalNetSerializationContext::FInitParameters TempInternalNetSerializationContextInitParams;
 	TempInternalNetSerializationContextInitParams.ReplicationSystem = Server->ReplicationSystem;
-	TempInternalNetSerializationContextInitParams.ObjectResolveContext.RemoteNetTokenStoreState = ReplicationSystemInternal->GetNetTokenStore().GetLocalNetTokenStoreState();
+	TempInternalNetSerializationContextInitParams.ObjectResolveContext.RemoteNetTokenStoreState = Server->ReplicationSystem->GetNetTokenStore()->GetLocalNetTokenStoreState();
 	TempInternalNetSerializationContext.Init(TempInternalNetSerializationContextInitParams);
 
 	InternalNetSerializationContext = MoveTemp(TempInternalNetSerializationContext);
@@ -833,12 +831,10 @@ void FTestPolymorphicArrayStructNetSerializerDeltaSerializationFixture::SetUp()
 	Server = new FReplicationSystemTestServer(GetName());
 
 	// Init NetSerializationContext
-	FReplicationSystemInternal* ReplicationSystemInternal = Server->GetReplicationSystem()->GetReplicationSystemInternal();
-
 	FInternalNetSerializationContext TempInternalNetSerializationContext;
 	FInternalNetSerializationContext::FInitParameters TempInternalNetSerializationContextInitParams;
 	TempInternalNetSerializationContextInitParams.ReplicationSystem = Server->ReplicationSystem;
-	TempInternalNetSerializationContextInitParams.ObjectResolveContext.RemoteNetTokenStoreState = ReplicationSystemInternal->GetNetTokenStore().GetLocalNetTokenStoreState();
+	TempInternalNetSerializationContextInitParams.ObjectResolveContext.RemoteNetTokenStoreState = Server->ReplicationSystem->GetNetTokenStore()->GetLocalNetTokenStoreState();
 	TempInternalNetSerializationContext.Init(TempInternalNetSerializationContextInitParams);
 
 	InternalNetSerializationContext = MoveTemp(TempInternalNetSerializationContext);
diff --git a/Engine/Source/Developer/TraceInsights/Private/Insights/NetworkingProfiler/ViewModels/PacketContentViewDrawHelper.cpp b/Engine/Source/Developer/TraceInsights/Private/Insights/NetworkingProfiler/ViewModels/PacketContentViewDrawHelper.cpp
--- a/Engine/Source/Developer/TraceInsights/Private/Insights/NetworkingProfiler/ViewModels/PacketContentViewDrawHelper.cpp
+++ b/Engine/Source/Developer/TraceInsights/Private/Insights/NetworkingProfiler/ViewModels/PacketContentViewDrawHelper.cpp
@@ -233,6 +233,10 @@ void FPacketContentViewDrawStateBuilder::AddEvent(const TraceServices::FNetProfi
 				{
 					Builder.Append(TEXT(" | HasPackageMapExports"));
 				}
+				if (Event.BunchInfo.bPartialCustomExportsFinal)
+				{
+					Builder.Append(TEXT(" | HasCustomExports"));
+				}
 
 				Builder.Append(TEXT(", "));
 			}
diff --git a/Engine/Source/Developer/TraceServices/Public/TraceServices/Model/NetProfiler.h b/Engine/Source/Developer/TraceServices/Public/TraceServices/Model/NetProfiler.h
--- a/Engine/Source/Developer/TraceServices/Public/TraceServices/Model/NetProfiler.h
+++ b/Engine/Source/Developer/TraceServices/Public/TraceServices/Model/NetProfiler.h
@@ -110,6 +110,7 @@ union FNetProfilerBunchInfo
 		uint64 ChannelCloseReason : 4;
 		uint64 bPartial : 1;
 		uint64 bPartialInitial : 1;
+		uint64 bPartialCustomExportsFinal : 1;
 		uint64 bPartialFinal : 1;
 		UE_DEPRECATED(5.3, "Replication pausing is now deprecated.")
 		uint64 bIsReplicationPaused : 1;
diff --git a/Engine/Source/Runtime/Core/Public/Misc/EngineNetworkCustomVersion.h b/Engine/Source/Runtime/Core/Public/Misc/EngineNetworkCustomVersion.h
--- a/Engine/Source/Runtime/Core/Public/Misc/EngineNetworkCustomVersion.h
+++ b/Engine/Source/Runtime/Core/Public/Misc/EngineNetworkCustomVersion.h
@@ -88,6 +88,7 @@ struct FEngineNetworkCustomVersion
 		DynamicMontageSerialization = 33,		// Bump version to support dynamic montage serialization in the Gameplay Ability System
 		PredictionKeyBaseNotReplicated = 34,	// Bump version to stop FPredictionKey::Base from being replicated (it was unused).
 		RepMoveOptionalAcceleration = 35,		// Bump version to support serialization changes to RepMove for optional Acceleration
+		CustomExports = 36,						// Bump version to support CustomExports (such as NetTokens used for supporting "Dynamic Replication" of GameplayTags -- see FGameplayTag::NetSerialize_Packed)		// -----<new versions can be added above this line>-------------------------------------------------
 
 
 		// -----<new versions can be added above this line>-------------------------------------------------
diff --git a/Engine/Source/Runtime/CoreUObject/Public/UObject/CoreNet.h b/Engine/Source/Runtime/CoreUObject/Public/UObject/CoreNet.h
--- a/Engine/Source/Runtime/CoreUObject/Public/UObject/CoreNet.h
+++ b/Engine/Source/Runtime/CoreUObject/Public/UObject/CoreNet.h
@@ -53,6 +53,9 @@ namespace UE::Net
 	struct FReplicationStateDescriptor;
 	typedef FReplicationFragment* (*CreateAndRegisterReplicationFragmentFunc)(UObject* Owner, const FReplicationStateDescriptor* Descriptor, FFragmentRegistrationContext& Context);
 #endif
+	
+	class FNetTokenExportContext;
+	class FNetTokenResolveContext;
 
 	namespace Private
 	{
@@ -232,6 +235,9 @@ class UPackageMap : public UObject
 
 	COREUOBJECT_API virtual void Serialize(FArchive& Ar) override;
 
+	// Provides access to the context required to resolve received NetTokens.
+	virtual const UE::Net::FNetTokenResolveContext* GetNetTokenResolveContext() const { return nullptr; }
+
 	COREUOBJECT_API UPackageMap();
 	COREUOBJECT_API virtual ~UPackageMap();
 
@@ -346,26 +352,27 @@ template <> struct TIsZeroConstructType<FLifetimeProperty> { enum { Value = true
 
 GENERATE_MEMBER_FUNCTION_CHECK(GetLifetimeReplicatedProps, void, const, TArray<FLifetimeProperty>&)
 
-#if UE_TRACE_ENABLED
 /**
- * We pass a NetTraceCollector along with the NetBitWriter in order avoid modifying all API`s where we want to be able to collect Network stats
- * Since the pointer to the collector is temporary we need to avoid copying it around by accident
+ * Used to pass in temporary pointers to NetBitWriter which we absolutely need to avoid copying by acciden.
+ * Used by NetTraceCollector and NetTokenExport scope
  */
-class FNetTraceCollectorDoNotCopyWrapper
+template<typename T>
+class TNetDoNotCopyPtr
 {
 public:
-	FNetTraceCollectorDoNotCopyWrapper() : Collector(nullptr) {}
-	FNetTraceCollectorDoNotCopyWrapper(const FNetTraceCollectorDoNotCopyWrapper&) : Collector(nullptr) {}
-    FNetTraceCollectorDoNotCopyWrapper(FNetTraceCollectorDoNotCopyWrapper&&) { Collector = nullptr; }
-	FNetTraceCollectorDoNotCopyWrapper& operator=(const FNetTraceCollectorDoNotCopyWrapper& Other) { Collector = nullptr; return *this; }
-    FNetTraceCollectorDoNotCopyWrapper& operator=(FNetTraceCollectorDoNotCopyWrapper&&) { Collector = nullptr; return *this; }
-
-	void Set(FNetTraceCollector* InCollector) { Collector = InCollector; }
-	FNetTraceCollector* Get() const { return Collector; }
-
+	TNetDoNotCopyPtr() : Ptr(nullptr) {}
+	TNetDoNotCopyPtr(const TNetDoNotCopyPtr&) : Ptr(nullptr) {}
+    TNetDoNotCopyPtr(TNetDoNotCopyPtr&&) { Ptr = nullptr; }
+	TNetDoNotCopyPtr& operator=(const TNetDoNotCopyPtr& Other) { Ptr = nullptr; return *this; }
+    TNetDoNotCopyPtr& operator=(TNetDoNotCopyPtr&&) { Ptr = nullptr; return *this; }
+	void Set(T* InPtr) { Ptr = InPtr; }
+	T* Get() const { return Ptr; }
 private:
-	FNetTraceCollector* Collector;
+	T* Ptr;
 };
+
+#if UE_TRACE_ENABLED
+using FNetTraceCollectorDoNotCopyWrapper = TNetDoNotCopyPtr<FNetTraceCollector>;
 #endif
 
 /**
@@ -376,8 +383,8 @@ private:
 class FNetBitWriter : public FBitWriter
 {
 public:
-	COREUOBJECT_API FNetBitWriter( UPackageMap * InPackageMap, int64 InMaxBits );
-	COREUOBJECT_API FNetBitWriter( int64 InMaxBits );
+	COREUOBJECT_API FNetBitWriter(UPackageMap* InPackageMap, int64 InMaxBits);
+	COREUOBJECT_API FNetBitWriter(int64 InMaxBits);
 	COREUOBJECT_API FNetBitWriter();
 
 	class UPackageMap * PackageMap;
@@ -394,8 +401,10 @@ public:
 	COREUOBJECT_API virtual FArchive& operator<<(struct FWeakObjectPtr& Value) override;
 
 	COREUOBJECT_API virtual void CountMemory(FArchive& Ar) const override;
-};
 
+	/* Context allowing exports of NetTokens from NetSerialize(), the context will not be copied when copying NetBitWriter */
+	TNetDoNotCopyPtr<UE::Net::FNetTokenExportContext> NetTokenExportContext;
+};
 
 /**
  * FNetBitReader
@@ -429,6 +438,18 @@ public:
 	COREUOBJECT_API virtual void CountMemory(FArchive& Ar) const override;
 };
 
+/** Helper methods to allow us to instrument different type of BitStreams */
+inline uint32 GetBitStreamPositionForNetTrace(const FBitWriter& Stream) { return (uint32)Stream.GetNumBits(); }
+inline uint32 GetBitStreamPositionForNetTrace(const FBitReader& Stream) { return (uint32)Stream.GetPosBits(); }
+
+#if UE_TRACE_ENABLED
+inline FNetTraceCollector* GetTraceCollector(const FNetBitWriter& BitWriter) { return BitWriter.TraceCollector.Get(); }
+inline void SetTraceCollector(FNetBitWriter& BitWriter, FNetTraceCollector* Collector) { BitWriter.TraceCollector.Set(Collector); }
+#else
+inline FNetTraceCollector* GetTraceCollector(const FNetBitWriter& BitWriter) { return nullptr; }
+inline void SetTraceCollector(FNetBitWriter& BitWriter, FNetTraceCollector* Collector) {}
+#endif
+
 bool FORCEINLINE NetworkGuidSetsAreSame( const TSet< FNetworkGUID >& A, const TSet< FNetworkGUID >& B )
 {
 	if ( A.Num() != B.Num() )
diff --git a/Engine/Source/Runtime/Engine/Classes/Engine/ActorChannel.h b/Engine/Source/Runtime/Engine/Classes/Engine/ActorChannel.h
--- a/Engine/Source/Runtime/Engine/Classes/Engine/ActorChannel.h
+++ b/Engine/Source/Runtime/Engine/Classes/Engine/ActorChannel.h
@@ -202,8 +202,12 @@ public:
 	ENGINE_API virtual void NotifyActorChannelOpen(AActor* InActor, FInBunch& InBunch);
 
 	/** Append any export bunches */
+	UE_DEPRECATED(5.5, "Use GetAdditionalRequiredBunches")
 	ENGINE_API virtual void AppendExportBunches( TArray< FOutBunch* >& OutExportBunches ) override;
 
+	/** Returns any additional bunches (such as exports) needed by the passed-in OutgoingBunch */
+	ENGINE_API virtual TArray<FOutBunch*> GetAdditionalRequiredBunches(const FOutBunch& OutgoingBunch, EChannelGetAdditionalRequiredBunchesFlags Flags) override;
+
 	/** Append any "must be mapped" guids to front of bunch. These are guids that the client will wait on before processing this bunch. */
 	ENGINE_API virtual void AppendMustBeMappedGuids( FOutBunch* Bunch ) override;
 
diff --git a/Engine/Source/Runtime/Engine/Classes/Engine/Channel.h b/Engine/Source/Runtime/Engine/Classes/Engine/Channel.h
--- a/Engine/Source/Runtime/Engine/Classes/Engine/Channel.h
+++ b/Engine/Source/Runtime/Engine/Classes/Engine/Channel.h
@@ -42,9 +42,15 @@ enum class EChannelCreateFlags : uint32
 	None			= (1 << 0),
 	OpenedLocally	= (1 << 1)
 };
-
 ENUM_CLASS_FLAGS(EChannelCreateFlags);
 
+enum class EChannelGetAdditionalRequiredBunchesFlags : uint32
+{
+	None				= 0U,
+	SkipNetGUIDExports	= 1U
+};
+ENUM_CLASS_FLAGS(EChannelGetAdditionalRequiredBunchesFlags);
+
 // The channel index to use for voice
 #define VOICE_CHANNEL_INDEX 1
 
@@ -137,9 +143,13 @@ public:
 	 * The bunch is sure not to be discarded.
 	 */
 	ENGINE_API void ReceivedRawBunch( FInBunch & Bunch, bool & bOutSkipAck );
-	
+
 	/** Append any export bunches */
-	ENGINE_API virtual void AppendExportBunches( TArray< FOutBunch* >& OutExportBunches );
+	UE_DEPRECATED(5.5, "Use GetAdditionalRequiredBunches")
+	ENGINE_API virtual void AppendExportBunches(TArray< FOutBunch* >& OutExportBunches);
+
+	/** Returns any additional bunches (such as exports) needed by the passed-in OutgoingBunch */
+	ENGINE_API virtual TArray<FOutBunch*> GetAdditionalRequiredBunches(const FOutBunch& OutgoingBunch, EChannelGetAdditionalRequiredBunchesFlags Flags = EChannelGetAdditionalRequiredBunchesFlags::None);
 
 	/** Append any "must be mapped" guids to front of bunch. These are guids that the client will wait on before processing this bunch. */
 	ENGINE_API virtual void AppendMustBeMappedGuids( FOutBunch* Bunch );
diff --git a/Engine/Source/Runtime/Engine/Classes/Engine/NetDriver.h b/Engine/Source/Runtime/Engine/Classes/Engine/NetDriver.h
--- a/Engine/Source/Runtime/Engine/Classes/Engine/NetDriver.h
+++ b/Engine/Source/Runtime/Engine/Classes/Engine/NetDriver.h
@@ -355,6 +355,8 @@ namespace UE::Net
 
 namespace UE::Net
 {
+	class FNetTokenStore;
+
 	class FScopedIgnoreStaticActorDestruction
 	{
 	public:
@@ -1915,8 +1917,17 @@ public:
 	inline UReplicationSystem* GetReplicationSystem() const { return ReplicationSystem; }
 
 	void UpdateGroupFilterStatusForLevel(const ULevel* Level, UE::Net::FNetObjectGroupHandle LevelGroupHandle);
+
+	/** Returns NetTokenStore that is required to create and serialize NetTokens */
+	const UE::Net::FNetTokenStore* GetNetTokenStore() const { return NetTokenStore.Get(); }
+	UE::Net::FNetTokenStore* GetNetTokenStore() { return NetTokenStore.Get(); }
+#else
+	/** Returns NetTokenStore that is required to create and serialize NetTokens */
+	const UE::Net::FNetTokenStore* GetNetTokenStore() const { return nullptr; }
+	UE::Net::FNetTokenStore* GetNetTokenStore() { return nullptr; }
 #endif // UE_WITH_IRIS
 
+
 	ENGINE_API void RemoveClientConnection(UNetConnection* ClientConnectionToRemove);
 
 	/** Adds (fully initialized, ready to go) client connection to the ClientConnections list + any other game related setup */
@@ -2276,6 +2287,8 @@ private:
 	TObjectPtr<UReplicationDriver> ReplicationDriver;
 
 #if UE_WITH_IRIS
+	TUniquePtr<UE::Net::FNetTokenStore> NetTokenStore;
+
 	UReplicationSystem* ReplicationSystem = nullptr;
 
 	/** When set this will skip registering all the network relevant actors when setting the World */
diff --git a/Engine/Source/Runtime/Engine/Classes/Engine/PackageMapClient.h b/Engine/Source/Runtime/Engine/Classes/Engine/PackageMapClient.h
--- a/Engine/Source/Runtime/Engine/Classes/Engine/PackageMapClient.h
+++ b/Engine/Source/Runtime/Engine/Classes/Engine/PackageMapClient.h
@@ -26,16 +26,19 @@
 #include "Misc/NetworkGuid.h"
 #include "Misc/NetworkVersion.h"
 #include "UObject/CoreNet.h"
+#include "Net/Core/NetToken/NetToken.h"
 #include "Net/DataBunch.h"
 #include "Net/NetAnalyticsTypes.h"
 #include "ProfilingDebugging/CsvProfiler.h"
 #include "HAL/LowLevelMemTracker.h"
+
 #include "PackageMapClient.generated.h"
 
 LLM_DECLARE_TAG_API(GuidCache, ENGINE_API);
 
 class UNetConnection;
 class UNetDriver;
+enum class EChannelGetAdditionalRequiredBunchesFlags : uint32;
 
 class FNetFieldExport
 {
@@ -437,13 +440,7 @@ public:
 
 	ENGINE_API UPackageMapClient(const FObjectInitializer & ObjectInitializer = FObjectInitializer::Get());
 
-	void Initialize(UNetConnection * InConnection, TSharedPtr<FNetGUIDCache> InNetGUIDCache)
-	{
-		Connection = InConnection;
-		GuidCache = InNetGUIDCache;
-		ExportNetGUIDCount = 0;
-		OverrideAckState = &AckState;
-	}
+	ENGINE_API void Initialize(UNetConnection * InConnection, TSharedPtr<FNetGUIDCache> InNetGUIDCache);
 
 	virtual ~UPackageMapClient()
 	{
@@ -453,7 +450,6 @@ public:
 			CurrentExportBunch = NULL;
 		}
 	}
-
 	
 	// UPackageMap Interface
 	ENGINE_API virtual bool SerializeObject( FArchive& Ar, UClass* InClass, UObject*& Obj, FNetworkGUID *OutNetGUID = NULL ) override;
@@ -469,9 +465,16 @@ public:
 	ENGINE_API virtual void ReceivedAck( const int32 AckPacketId ) override;
 	ENGINE_API virtual void NotifyBunchCommit( const int32 OutPacketId, const FOutBunch* OutBunch ) override;
 	ENGINE_API virtual void GetNetGUIDStats(int32 &AckCount, int32 &UnAckCount, int32 &PendingCount) override;
+	ENGINE_API virtual const UE::Net::FNetTokenResolveContext* GetNetTokenResolveContext() const override { return &NetTokenResolveContext; }
 
 	ENGINE_API void ReceiveNetGUIDBunch( FInBunch &InBunch );
+
+	UE_DEPRECATED(5.5, "Use GetAdditionalRequiredBunches")
 	ENGINE_API void AppendExportBunches(TArray<FOutBunch *>& OutgoingBunches);
+
+	/** Get bunches to prepend to the OutgoingBunch. The required bunches usually contain exports or extension data. */
+	ENGINE_API TArray<FOutBunch*> GetAdditionalRequiredBunches(const FOutBunch& OutgoingBunch, EChannelGetAdditionalRequiredBunchesFlags Flags);
+
 	ENGINE_API int32 GetNumExportBunches() const;
 
 	ENGINE_API void AppendExportData(FArchive& Archive);
@@ -522,6 +525,13 @@ public:
 	ENGINE_API void								SerializeNetFieldExportGroupMap( FArchive& Ar, bool bClearPendingExports=true );
 	ENGINE_API void								SerializeNetFieldExportDelta(FArchive& Ar);
 
+	/** This is called during the sending of a network bunch.  The returned bunch is prepended to the outgoing bunches and guaranteed to be processed upon receive (before ack) */
+	ENGINE_API FOutBunch* CreateCustomExportsBunch(const FOutBunch& OutgoingBunch);
+
+	/** Receive a bunch that contains the data for Package Map Extensions */
+	void ReceiveCustomExportsBunch(FInBunch& InBunch);
+
+public:
 	TUniquePtr<TGuardValue<bool>> ScopedIgnoreReceivedExportGUIDs()
 	{
 		return MakeUnique<TGuardValue<bool>>(bIgnoreReceivedExportGUIDs, true);
@@ -543,6 +553,7 @@ public:
 
 	ENGINE_API virtual void AddUnmappedNetGUIDReference(FNetworkGUID UnmappedGUID) override;
 	ENGINE_API virtual void RemoveUnmappedNetGUIDReference(FNetworkGUID MappedGUID) override;
+	ENGINE_API bool IsNetGUIDAuthority() const;
 
 protected:
 
@@ -566,8 +577,6 @@ protected:
 
 	ENGINE_API bool	ShouldSendFullPath(const UObject* Object, const FNetworkGUID &NetGUID);
 	
-	ENGINE_API bool IsNetGUIDAuthority() const;
-
 	class UNetConnection* Connection;
 
 	ENGINE_API bool ObjectLevelHasFinishedLoading(UObject* Obj) const;
@@ -620,4 +629,13 @@ private:
 	 * of a particular property or actor spawn, to help correlate those sync loads with that property or actor.
 	 */
 	TArray<FNetworkGUID> TrackedSyncLoadedGUIDs;
+
+	// Stored information required to resolve NetTokens
+	UE::Net::FNetTokenResolveContext NetTokenResolveContext;
+
+	/** NetTokens acknowledged by remote */
+	TSet<UE::Net::FNetToken> AcknowledgedExportedNetTokens;
+	
+	/** NetTokens still pending acknowledgment by remote */
+	TMultiMap<int32, UE::Net::FNetToken> NetTokenPendingAckMap;
 };
diff --git a/Engine/Source/Runtime/Engine/Classes/GameFramework/CharacterMovementReplication.h b/Engine/Source/Runtime/Engine/Classes/GameFramework/CharacterMovementReplication.h
--- a/Engine/Source/Runtime/Engine/Classes/GameFramework/CharacterMovementReplication.h
+++ b/Engine/Source/Runtime/Engine/Classes/GameFramework/CharacterMovementReplication.h
@@ -7,6 +7,7 @@
 #include "Engine/NetSerialization.h"
 #include "Serialization/BitWriter.h"
 #include "Containers/BitArray.h"
+#include "Net/Core/NetToken/NetTokenExportContext.h"
 #if UE_WITH_IRIS
 #include "Iris/Serialization/IrisObjectReferencePackageMap.h"
 #endif
@@ -64,6 +65,9 @@ struct FCharacterNetworkSerializationPackedBits
 	UE::Net::FIrisPackageMapExports PackageMapExports;
 #endif
 
+	// Since we capturing data outside of the normal serialization path we also need to store exports to inject when actually sending the data.
+	UE::Net::FNetTokenExportContext::FNetTokenExports NetTokensPendingExport;
+
 private:
 	UPackageMap* SavedPackageMap;
 };
diff --git a/Engine/Source/Runtime/Engine/Private/Components/CharacterMovementComponent.cpp b/Engine/Source/Runtime/Engine/Private/Components/CharacterMovementComponent.cpp
--- a/Engine/Source/Runtime/Engine/Private/Components/CharacterMovementComponent.cpp
+++ b/Engine/Source/Runtime/Engine/Private/Components/CharacterMovementComponent.cpp
@@ -8867,10 +8867,6 @@ static UIrisObjectReferencePackageMap* GetIrisPackageMapToCaptureReferences(UNet
 		{
 			ObjectReferencePackageMap->InitForWrite(&PackedBits.PackageMapExports);
 
-			// $IRIS: $TODO: figure out how this should be done
-			//UReplicationSystem* ReplicationSystem = UE::Net::FReplicationSystemUtil::GetReplicationSystem(NetConnection->GetDriver());
-			//ReplicationSystem->InitNetTokenExportContext(*ObjectReferencePackageMap->GetNetTokenExportContext(), NetConnection->GetConnectionId());
-
 			return ObjectReferencePackageMap;
 		}
 	}
@@ -8885,11 +8881,8 @@ static UIrisObjectReferencePackageMap* GetIrisPackageMapToReadReferences(const U
 	{
 		if (UIrisObjectReferencePackageMap* ObjectReferencePackageMap = Bridge->GetObjectReferencePackageMap())
 		{
-			ObjectReferencePackageMap->InitForRead(&PackedBits.PackageMapExports);
-
-			// $IRIS: $TODO: figure out how this should be done
-			//UReplicationSystem* ReplicationSystem = UE::Net::FReplicationSystemUtil::GetReplicationSystem(NetConnection->GetDriver());
-			//ReplicationSystem->InitNetTokenExportContext(*ObjectReferencePackageMap->GetNetTokenExportContext(), NetConnection->GetConnectionId());
+			const UReplicationSystem* ReplicationSystem = UE::Net::FReplicationSystemUtil::GetReplicationSystem(NetConnection->GetDriver());
+			ObjectReferencePackageMap->InitForRead(&PackedBits.PackageMapExports, ReplicationSystem->GetNetTokenResolveContext(NetConnection->GetConnectionId()));
 
 			return ObjectReferencePackageMap;
 		}
@@ -8933,6 +8926,10 @@ void UCharacterMovementComponent::CallServerMovePacked(const FSavedMove_Characte
 		return;
 	}
 
+	// Reset captured exports stored in PackedBits
+	PackedBits.NetTokensPendingExport.Reset();
+	UE::Net::FNetTokenExportScope NetTokenExportScope(ServerMoveBitWriter, NetConnection->GetDriver()->GetNetTokenStore(), PackedBits.NetTokensPendingExport, "CallServerMovePacked");
+
 	// Serialize move struct into a bit stream
 	if (!MoveDataContainer.Serialize(*this, ServerMoveBitWriter, ServerMoveBitWriter.PackageMap) || ServerMoveBitWriter.IsError())
 	{
@@ -9473,6 +9470,15 @@ bool FCharacterNetworkSerializationPackedBits::NetSerialize(FArchive& Ar, class
 	{
 		DataBits.Init(0, NumBits);
 	}
+	else if (Ar.IsSaving() && NetTokensPendingExport.Num())
+	{
+		// As we now support exporting NetTokens from shared serialization and FCharacterNetworkSerializationPackedBits serializes data outside of the normal flow
+		// we explicitly capture exports which we needs to be inject during actual serialization.
+		if (UE::Net::FNetTokenExportContext* ExportContext = UE::Net::FNetTokenExportContext::GetNetTokenExportContext(Ar))
+		{
+			ExportContext->AppendNetTokensPendingExport(NetTokensPendingExport);
+		}		
+	}
 
 	// Array data
 	Ar.SerializeBits(DataBits.GetData(), NumBits);
@@ -10571,6 +10577,10 @@ void UCharacterMovementComponent::ServerSendMoveResponse(const FClientAdjustment
 		return;
 	}
 
+	// Reset NetTokensPendingExport as PackedBits is reused
+	PackedBits.NetTokensPendingExport.Reset();
+	UE::Net::FNetTokenExportScope NetTokenExportsScope(MoveResponseBitWriter, NetConnection->GetDriver()->GetNetTokenStore(), PackedBits.NetTokensPendingExport, "ServerSendMoveResponse");
+
 	// Serialize move struct into a bit stream
 	if (!ResponseDataContainer.Serialize(*this, MoveResponseBitWriter, MoveResponseBitWriter.PackageMap) || MoveResponseBitWriter.IsError())
 	{
diff --git a/Engine/Source/Runtime/Engine/Private/DataBunch.cpp b/Engine/Source/Runtime/Engine/Private/DataBunch.cpp
--- a/Engine/Source/Runtime/Engine/Private/DataBunch.cpp
+++ b/Engine/Source/Runtime/Engine/Private/DataBunch.cpp
@@ -34,6 +34,7 @@ FInBunch::FInBunch( UNetConnection* InConnection, uint8* Src, int64 CountBits )
 ,	bPartialFinal ( 0 )
 ,	bHasPackageMapExports ( 0 )
 ,	bHasMustBeMappedGUIDs ( 0 )
+,	bPartialCustomExportsFinal( 0 )
 ,	bIgnoreRPCs ( 0 )
 ,	CloseReason( EChannelCloseReason::Destroyed )
 {
@@ -60,6 +61,7 @@ FInBunch::FInBunch( FInBunch &InBunch, bool CopyBuffer )
 	bPartialFinal =	InBunch.bPartialFinal;
 	bHasPackageMapExports = InBunch.bHasPackageMapExports;
 	bHasMustBeMappedGUIDs =	InBunch.bHasMustBeMappedGUIDs;
+	bPartialCustomExportsFinal = InBunch.bPartialCustomExportsFinal;
 	bIgnoreRPCs = InBunch.bIgnoreRPCs;
 	CloseReason = InBunch.CloseReason;
 	PackageMap = InBunch.PackageMap;
@@ -135,6 +137,7 @@ FOutBunch::FOutBunch( UChannel* InChannel, bool bInClose )
 ,	bPartialFinal			( 0 )
 ,	bHasPackageMapExports	( 0 )
 ,	bHasMustBeMappedGUIDs	( 0 )
+,	bPartialCustomExportsFinal	( 0 )
 ,	CloseReason( EChannelCloseReason::Destroyed )
 {
 	checkSlow(!Channel->Closing);
@@ -169,6 +172,7 @@ FOutBunch::FOutBunch( UPackageMap *InPackageMap, int64 MaxBits )
 ,	bPartialFinal	( 0 )
 ,	bHasPackageMapExports	( 0 )
 ,	bHasMustBeMappedGUIDs	( 0 )
+,	bPartialCustomExportsFinal	( 0 )
 ,	CloseReason( EChannelCloseReason::Destroyed )
 {
 }
@@ -179,6 +183,12 @@ FOutBunch::FOutBunch(int64 InMaxBits)
 {
 }
 
+void FOutBunch::Reset()
+{
+	FNetBitWriter::Reset();
+	NetTokensPendingExport.Reset();
+}
+
 void FOutBunch::CountMemory(FArchive& Ar) const
 {
 	for (const FOutBunch* Current = this; Current; Current = Current->Next)
diff --git a/Engine/Source/Runtime/Engine/Private/DataChannel.cpp b/Engine/Source/Runtime/Engine/Private/DataChannel.cpp
--- a/Engine/Source/Runtime/Engine/Private/DataChannel.cpp
+++ b/Engine/Source/Runtime/Engine/Private/DataChannel.cpp
@@ -644,6 +644,14 @@ void UChannel::ReceivedRawBunch( FInBunch & Bunch, bool & bOutSkipAck )
 		// Verify that UConnection::ReceivedPacket has passed us a valid bunch.
 		check(Bunch.ChSequence>Connection->InReliable[ChIndex]);
 
+		// Pessimistic NAK so that the other side doesn't think we have processed the export
+		// In practice, this would only happen if we split an export bunch into multiple bunches
+		if (Bunch.bPartialCustomExportsFinal)
+		{
+			UE_LOG(LogNetPartialBunch, Warning, TEXT("New Pessimistic Nak Performed on %s"), *Bunch.ToString());
+			bOutSkipAck = true;
+		}
+
 		// Find the place for this item, sorted in sequence.
 		UE_LOG(LogNetTraffic, Log, TEXT("      Queuing bunch with unreceived dependency: %d / %d"), Bunch.ChSequence, Connection->InReliable[ChIndex]+1 );
 		FInBunch** InPtr;
@@ -812,7 +820,8 @@ bool UChannel::ReceivedNextBunch( FInBunch & Bunch, bool & bOutSkipAck )
 			{		
 				if ( !Bunch.bHasPackageMapExports && Bunch.GetBitsLeft() > 0 )
 				{
-					if ( Bunch.GetBitsLeft() % 8 != 0 )
+					// If we're an extensions bunch, we're not byte aligned
+					if ( !Bunch.bPartialCustomExportsFinal && Bunch.GetBitsLeft() % 8 != 0 )
 					{
 						UE_LOG(LogNetPartialBunch, Warning, TEXT("Corrupt partial bunch. Initial partial bunches are expected to be byte-aligned. BitsLeft = %u. %s"), Bunch.GetBitsLeft(), *Describe());
 
@@ -855,7 +864,8 @@ bool UChannel::ReceivedNextBunch( FInBunch & Bunch, bool & bOutSkipAck )
 			if ( InPartialBunch && !InPartialBunch->bPartialFinal && bSequenceMatches && InPartialBunch->bReliable == Bunch.bReliable )
 			{
 				// Merge.
-				UE_LOG(LogNetPartialBunch, Verbose, TEXT("Merging Partial Bunch: %d Bytes"), Bunch.GetBytesLeft() );
+				UE_LOG(LogNetPartialBunch, Verbose, TEXT("Merging Partial Bunch: %d Bits."), Bunch.GetBitsLeft() );
+				UE_LOG(LogNetPartialBunch, VeryVerbose, TEXT(" %s"), *Bunch.ToString() );
 
 				if ( !Bunch.bHasPackageMapExports && Bunch.GetBitsLeft() > 0 )
 				{
@@ -864,7 +874,7 @@ bool UChannel::ReceivedNextBunch( FInBunch & Bunch, bool & bOutSkipAck )
 
 				// Only the final partial bunch should ever be non byte aligned. This is enforced during partial bunch creation
 				// This is to ensure fast copies/appending of partial bunches. The final partial bunch may be non byte aligned.
-				if (!Bunch.bHasPackageMapExports && !Bunch.bPartialFinal && (Bunch.GetBitsLeft() % 8 != 0))
+				if (!Bunch.bHasPackageMapExports && !Bunch.bPartialCustomExportsFinal && !Bunch.bPartialFinal && (Bunch.GetBitsLeft() % 8 != 0))
 				{
 					UE_LOG(LogNetPartialBunch, Warning, TEXT("Corrupt partial bunch. Non-final partial bunches are expected to be byte-aligned. bHasPackageMapExports = %d, bPartialFinal = %d, BitsLeft = %u. %s"),
 						Bunch.bHasPackageMapExports ? 1 : 0, Bunch.bPartialFinal ? 1 : 0, Bunch.GetBitsLeft(), *Describe());
@@ -964,7 +974,35 @@ bool UChannel::ReceivedNextBunch( FInBunch & Bunch, bool & bOutSkipAck )
 
 			return false;
 		}
-	}
+
+		// This flag denotes that the previous partial data has now fully reassembled the exports data, and we should now reconstruct it
+		// The reading of the exports data is guaranteed to occur in the same frame it's received (and therefore ack'd).  This is so an ack 
+		// guarantees that the export data has been read (so we can change our representation to a fully ack'd representation).
+		if (!bOutSkipAck && Bunch.bPartialCustomExportsFinal)
+		{
+			UE_LOG(LogNetPartialBunch, Verbose, TEXT("Partial CustomExports Bunch (ChSequence %d PacketId %d) is being processed."), Bunch.ChSequence, Bunch.PacketId);
+
+			ensureMsgf(!HandleBunch, TEXT("Code logic error:  CustomExports Bunch is going to be read as if it contains other data"));
+			ensureMsgf(!Bunch.bPartialFinal, TEXT("CustomExports Bunch was also marked as the final bunch (therefore no data depended on this)"));
+			ensureMsgf(InPartialBunch, TEXT("CustomExports Bunch was not properly merged into a Partial Bunch"));
+			if (InPartialBunch && !Connection->IsInternalAck())
+			{
+				InPartialBunch->bPartialCustomExportsFinal = Bunch.bPartialCustomExportsFinal;
+				CastChecked<UPackageMapClient>(Connection->PackageMap)->ReceiveCustomExportsBunch(*InPartialBunch);
+
+				if (InPartialBunch->IsError())
+				{
+					AddToChainResultPtr(InPartialBunch->ExtendedError, ENetCloseResult::ReceivedNetGUIDBunchFail);
+					UE_LOG(LogNetTraffic, Error, TEXT("UChannel::ReceivedRawBunch: Bunch.IsError() after ReceiveCustomExportsBunch. ChIndex: %i"), ChIndex);
+					return false;
+				}
+
+				// Reset the partial bunch so that it only contains data (not the exports, which we just processed)...
+				InPartialBunch->ResetData(*InPartialBunch, 0, InPartialBunch->GetNumBits());
+				InPartialBunch->bPartialCustomExportsFinal = false;
+			}
+		}
+	} // bPartial
 
 	if ( HandleBunch != NULL )
 	{
@@ -1042,6 +1080,7 @@ bool UChannel::ReceivedNextBunch( FInBunch & Bunch, bool & bOutSkipAck )
 	return false;
 }
 
+PRAGMA_DISABLE_DEPRECATION_WARNINGS
 void UChannel::AppendExportBunches( TArray<FOutBunch *>& OutExportBunches )
 {
 	UPackageMapClient * PackageMapClient = CastChecked< UPackageMapClient >( Connection->PackageMap );
@@ -1049,6 +1088,13 @@ void UChannel::AppendExportBunches( TArray<FOutBunch *>& OutExportBunches )
 	// Let the package map add any outgoing bunches it needs to send
 	PackageMapClient->AppendExportBunches( OutExportBunches );
 }
+PRAGMA_ENABLE_DEPRECATION_WARNINGS
+
+TArray<FOutBunch*> UChannel::GetAdditionalRequiredBunches(const FOutBunch& OutgoingBunch, EChannelGetAdditionalRequiredBunchesFlags Flags)
+{
+	UPackageMapClient* PackageMapClient = CastChecked< UPackageMapClient >(Connection->PackageMap);
+	return PackageMapClient->GetAdditionalRequiredBunches(OutgoingBunch, Flags);
+}
 
 PRAGMA_DISABLE_DEPRECATION_WARNINGS
 void UChannel::AppendMustBeMappedGuids( FOutBunch* Bunch )
@@ -1108,12 +1154,16 @@ PRAGMA_ENABLE_DEPRECATION_WARNINGS
 
 void UActorChannel::AppendExportBunches( TArray<FOutBunch *>& OutExportBunches )
 {
+	ensureMsgf(false, TEXT("%hs is deprecated. Use GetAdditionalRequiredBunches"), __func__);
+
 	if (bHoldQueuedExportBunchesAndGUIDs)
 	{
 		return;
 	}
 
+PRAGMA_DISABLE_DEPRECATION_WARNINGS
 	Super::AppendExportBunches( OutExportBunches );
+PRAGMA_ENABLE_DEPRECATION_WARNINGS
 
 	// We don't want to append QueuedExportBunches to these bunches, since these were for queued RPC's, and we don't want to record RPC's during bResendAllDataSinceOpen
 	if ( Connection->ResendAllDataState == EResendAllDataState::None )
@@ -1135,6 +1185,38 @@ void UActorChannel::AppendExportBunches( TArray<FOutBunch *>& OutExportBunches )
 	}
 }
 
+TArray<FOutBunch*> UActorChannel::GetAdditionalRequiredBunches(const FOutBunch& OutgoingBunch, EChannelGetAdditionalRequiredBunchesFlags Flags)
+{
+	TArray<FOutBunch*> AdditionalBunches;
+
+	// Tell super to skip appending NetGUID export bunches. 
+	// This is used from fastpath to disallow NetGUID exports but since we now allow "stable" exports we need to get the additional bunches for that type of exportdata.
+	Flags |= bHoldQueuedExportBunchesAndGUIDs ? EChannelGetAdditionalRequiredBunchesFlags::SkipNetGUIDExports : EChannelGetAdditionalRequiredBunchesFlags::None;
+
+	AdditionalBunches = Super::GetAdditionalRequiredBunches(OutgoingBunch, Flags);
+
+	// We don't want to append QueuedExportBunches to these bunches, since these were for queued RPC's, and we don't want to record RPC's during bResendAllDataSinceOpen
+	if (Connection->ResendAllDataState == EResendAllDataState::None)
+	{
+		// Let the profiler know about exported GUID bunches
+		for (const FOutBunch* ExportBunch : QueuedExportBunches)
+		{
+			if (ExportBunch != nullptr)
+			{
+				NETWORK_PROFILER(GNetworkProfiler.TrackExportBunch(ExportBunch->GetNumBits(), Connection));
+			}
+		}
+
+		if (QueuedExportBunches.Num())
+		{
+			AdditionalBunches.Append(QueuedExportBunches);
+			QueuedExportBunches.Empty();
+		}
+	}
+
+	return AdditionalBunches;
+}
+
 void UActorChannel::AppendMustBeMappedGuids( FOutBunch* Bunch )
 {
 	if (bHoldQueuedExportBunchesAndGUIDs)
@@ -1223,7 +1305,7 @@ FPacketIdRange UChannel::SendBunch( FOutBunch* Bunch, bool Merge )
 	// Replay connections will manage export bunches separately.
 	if (!Connection->IsInternalAck())
 	{
-		AppendExportBunches( OutgoingBunches );
+		OutgoingBunches = GetAdditionalRequiredBunches(*Bunch);
 	}
 
 	if ( OutgoingBunches.Num() )
@@ -1268,7 +1350,6 @@ FPacketIdRange UChannel::SendBunch( FOutBunch* Bunch, bool Merge )
 		Connection->LastOut.SerializeBits( Bunch->GetData(), Bunch->GetNumBits() );
 		Connection->LastOut.bOpen     |= Bunch->bOpen;
 		Connection->LastOut.bClose    |= Bunch->bClose;
-
 #if UE_NET_TRACE_ENABLED		
 		SetTraceCollector(Connection->LastOut, GetTraceCollector(*Bunch));
 		SetTraceCollector(*Bunch, nullptr);
@@ -1378,7 +1459,7 @@ FPacketIdRange UChannel::SendBunch( FOutBunch* Bunch, bool Merge )
 		NextBunch->ChIndex = Bunch->ChIndex;
 		NextBunch->ChName = Bunch->ChName;
 
-		if ( !NextBunch->bHasPackageMapExports )
+		if ( !(NextBunch->bHasPackageMapExports || NextBunch->bPartialCustomExportsFinal))
 		{
 			NextBunch->bHasMustBeMappedGUIDs |= Bunch->bHasMustBeMappedGUIDs;
 		}
@@ -1391,12 +1472,14 @@ FPacketIdRange UChannel::SendBunch( FOutBunch* Bunch, bool Merge )
 			NextBunch->bOpen &= (PartialNum == 0);											// Only the first bunch should have the bOpen bit set
 			NextBunch->bClose = (Bunch->bClose && (OutgoingBunches.Num()-1 == PartialNum)); // Only last bunch should have bClose bit set
 		}
+		ensureMsgf(!NextBunch->bPartialCustomExportsFinal || NextBunch->bPartial, TEXT("Can't have PackageMapExtensions without being part of a partial bunch"));
 
 		FOutBunch *ThisOutBunch = PrepBunch(NextBunch, OutBunch, Merge); // This handles queuing reliable bunches into the ack list
 
 		if (UE_LOG_ACTIVE(LogNetPartialBunch,Verbose) && (OutgoingBunches.Num() > 1)) // Don't want to call appMemcrc unless we need to
 		{
-			UE_LOG(LogNetPartialBunch, Verbose, TEXT("	Bunch[%d]: Bytes: %d Bits: %d ChSequence: %d 0x%X"), PartialNum, ThisOutBunch->GetNumBytes(), ThisOutBunch->GetNumBits(), ThisOutBunch->ChSequence, FCrc::MemCrc_DEPRECATED(ThisOutBunch->GetData(), ThisOutBunch->GetNumBytes()));
+			UE_LOG(LogNetPartialBunch, Verbose, TEXT("	Bunch[%d]: Bytes: %d Bits: %d ChSequence: %d NetGUIDs: %d CustomExports: %d 0x%X"), PartialNum, ThisOutBunch->GetNumBytes(), ThisOutBunch->GetNumBits(), ThisOutBunch->ChSequence, ThisOutBunch->bHasPackageMapExports, ThisOutBunch->bPartialCustomExportsFinal, FCrc::MemCrc_DEPRECATED(ThisOutBunch->GetData(), ThisOutBunch->GetNumBytes()));
+			UE_LOG(LogNetPartialBunch, VeryVerbose, TEXT("	 %s"), *ThisOutBunch->GetDebugString());
 		}
 
 		// Update Packet Range
@@ -3526,6 +3609,9 @@ int64 UActorChannel::ReplicateActor()
 	// Create an outgoing bunch, and skip this actor if the channel is saturated.
 	FOutBunch Bunch( this, 0 );
 
+	// Create export scope to capture NetToken exports and store them in Bunch.NetTokensPendingExport
+	UE::Net::FNetTokenExportScope NetTokenExportScope(Bunch, Connection->GetDriver()->GetNetTokenStore(), Bunch.NetTokensPendingExport, "ReplicateActor");
+	
 	if( Bunch.IsError() )
 	{
 		return 0;
diff --git a/Engine/Source/Runtime/Engine/Private/DataReplication.cpp b/Engine/Source/Runtime/Engine/Private/DataReplication.cpp
--- a/Engine/Source/Runtime/Engine/Private/DataReplication.cpp
+++ b/Engine/Source/Runtime/Engine/Private/DataReplication.cpp
@@ -438,6 +438,8 @@ void FObjectReplicator::InitRecentProperties(uint8* Source)
 		for (uint16 CustomDeltaProperty = 0; CustomDeltaProperty < NumLifetimeCustomDeltaProperties; ++CustomDeltaProperty)
 		{
 			FOutBunch DeltaState(Connection->PackageMap);
+			UE::Net::FNetTokenExportScope NetTokenExportScope(DeltaState, ConnectionDriver->GetNetTokenStore(), DeltaState.NetTokensPendingExport, "SendPropertiesForRPC");
+
 			TSharedPtr<INetDeltaBaseState>& NewState = SendingRepState->RecentCustomDeltaState[CustomDeltaProperty];
 			NewState.Reset();
 
@@ -1609,6 +1611,9 @@ void FObjectReplicator::ReplicateCustomDeltaProperties( FNetBitWriter & Bunch, F
 	
 	const bool bIsConnectionInternalAck = Connection->IsInternalAck();
 
+	// We must forward the current NetToken export context from the bunch to the temporary NetBitWriter into which this data will be written
+	TempBitWriter.NetTokenExportContext.Set(Bunch.NetTokenExportContext.Get());
+	
 	// Replicate those properties.
 	for (uint16 CustomDeltaProperty = 0; CustomDeltaProperty < NumLifetimeCustomDeltaProperties; ++CustomDeltaProperty)
 	{
@@ -1833,11 +1838,11 @@ bool FObjectReplicator::ReplicateProperties(FOutBunch& Bunch, FReplicationFlags
 }
 
 /** Replicates properties to the Bunch. Returns true if it wrote anything */
-bool FObjectReplicator::ReplicateProperties_r( FOutBunch & Bunch, FReplicationFlags RepFlags, FNetBitWriter& Writer)
+bool FObjectReplicator::ReplicateProperties_r(FOutBunch& Bunch, FReplicationFlags RepFlags, FNetBitWriter& Writer)
 {
 	UObject* Object = GetObject();
 
-	if ( Object == nullptr )
+	if (Object == nullptr)
 	{
 		UE_LOG(LogRep, Verbose, TEXT("ReplicateProperties: Object == nullptr"));
 		return false;
@@ -1856,11 +1861,22 @@ bool FObjectReplicator::ReplicateProperties_r( FOutBunch & Bunch, FReplicationFl
 
 	UNetConnection* OwningChannelConnection = OwningChannel->Connection;
 
+	// Pass on NetTokenExportContext
+	Writer.NetTokenExportContext.Set(Bunch.NetTokenExportContext.Get());
+
 #if UE_NET_TRACE_ENABLED
 	// Create trace collector if tracing is enabled for the target bunch
 	SetTraceCollector(Writer, GetTraceCollector(Bunch) ? UE_NET_TRACE_CREATE_COLLECTOR(ENetTraceVerbosity::Trace) : nullptr);
-    ON_SCOPE_EXIT { UE_NET_TRACE_DESTROY_COLLECTOR(GetTraceCollector(Writer)); };
 #endif
+
+    ON_SCOPE_EXIT 
+	{
+#if UE_NET_TRACE_ENABLED		
+		UE_NET_TRACE_DESTROY_COLLECTOR(GetTraceCollector(Writer)); 
+#endif
+		// Just to be safe, restore context on exit as we technically do not own the Writer.
+		Writer.NetTokenExportContext.Set(nullptr);
+	};
  
 	// TODO: Maybe ReplicateProperties could just take the RepState, Changelist Manger, Writer, and OwningChannel
 	//		and all the work could just be done in a single place.
@@ -1983,6 +1999,11 @@ bool FObjectReplicator::ReplicateProperties_r( FOutBunch & Bunch, FReplicationFl
 		}
 
 		Writer.SerializeBits( RemoteFunctions->GetData(), RemoteFunctions->GetNumBits() );
+
+		// Append potential NetToken exports from queued remote functions.
+		Writer.NetTokenExportContext.Get()->AppendNetTokensPendingExport(RemoteFunctions->NetTokensPendingExport);		
+		RemoteFunctions->NetTokensPendingExport.Reset();
+
 		RemoteFunctions->Reset();
 		RemoteFuncInfo.Empty();
 
@@ -2219,6 +2240,10 @@ void FObjectReplicator::QueueRemoteFunctionBunch( UFunction* Func, FOutBunch &Bu
 
 	RemoteFunctions->SerializeBits(Bunch.GetData(), Bunch.GetNumBits());
 
+	// Save potential NetToken exports for export when we actually commit the RemoteFunctions
+	RemoteFunctions->NetTokensPendingExport.Append(Bunch.NetTokensPendingExport);
+	Bunch.NetTokensPendingExport.Reset();
+
 	if (Connection->PackageMap != nullptr)
 	{
 		UPackageMapClient* PackageMapClient = CastChecked<UPackageMapClient>(Connection->PackageMap);
@@ -2232,8 +2257,9 @@ void FObjectReplicator::QueueRemoteFunctionBunch( UFunction* Func, FOutBunch &Bu
 
 		if (!Connection->IsInternalAck())
 		{
-			// Copy over any exported bunches
-			PackageMapClient->AppendExportBunches(OwningChannel->QueuedExportBunches);
+			// Copy over any additionally required bunches
+			TArray<FOutBunch*> AdditionalBunches = PackageMapClient->GetAdditionalRequiredBunches(Bunch, EChannelGetAdditionalRequiredBunchesFlags::None);
+			OwningChannel->QueuedExportBunches.Append(MoveTemp(AdditionalBunches));
 		}
 	}
 }
diff --git a/Engine/Source/Runtime/Engine/Private/GameFramework/CharacterNetworkSerializationPackedBitsNetSerializer.cpp b/Engine/Source/Runtime/Engine/Private/GameFramework/CharacterNetworkSerializationPackedBitsNetSerializer.cpp
--- a/Engine/Source/Runtime/Engine/Private/GameFramework/CharacterNetworkSerializationPackedBitsNetSerializer.cpp
+++ b/Engine/Source/Runtime/Engine/Private/GameFramework/CharacterNetworkSerializationPackedBitsNetSerializer.cpp
@@ -18,6 +18,7 @@
 #include "Iris/Serialization/NetReferenceCollector.h"
 #include "Iris/Serialization/NetSerializerArrayStorage.h"
 #include "Iris/Serialization/IrisPackageMapExportUtil.h"
+#include "Containers/ArrayView.h"
 
 namespace UE::Net::Private
 {
@@ -101,6 +102,12 @@ FCharacterNetworkSerializationPackedBitsNetSerializer::FNetSerializerRegistryDel
 
 void FCharacterNetworkSerializationPackedBitsNetSerializer::Serialize(FNetSerializationContext& Context, const FNetSerializeArgs& Args)
 {
+	// For now we ignore this in default state hash due to complications with asymmetrically serialized state.
+	if (Context.IsInitializingDefaultState())
+	{
+		return;
+	}
+
 	const QuantizedType& Value = *reinterpret_cast<QuantizedType*>(Args.Source);
 	
 	FNetBitStreamWriter* Writer = Context.GetBitStreamWriter();
@@ -128,6 +135,12 @@ void FCharacterNetworkSerializationPackedBitsNetSerializer::FreeDynamicState(FNe
 
 void FCharacterNetworkSerializationPackedBitsNetSerializer::Deserialize(FNetSerializationContext& Context, const FNetDeserializeArgs& Args)
 {
+	// For consistency, we should never get here. For now we ignore this in default state hash due to complications with asymmetrically serialized state.
+	if (Context.IsInitializingDefaultState())
+	{
+		return;
+	}
+
 	const ConfigType* Config = static_cast<const ConfigType*>(Args.NetSerializerConfig);
 	QuantizedType& TargetValue = *reinterpret_cast<QuantizedType*>(Args.Target);
 	
@@ -169,7 +182,7 @@ void FCharacterNetworkSerializationPackedBitsNetSerializer::Quantize(FNetSeriali
 	QuantizedType& TargetValue = *reinterpret_cast<QuantizedType*>(Args.Target);
 
 	// Quantize captured references and exports
-	FIrisPackageMapExportsUtil::Quantize(Context, SourceValue.PackageMapExports, TargetValue.QuantizedExports);
+	FIrisPackageMapExportsUtil::Quantize(Context, SourceValue.PackageMapExports, MakeArrayView(SourceValue.NetTokensPendingExport), TargetValue.QuantizedExports);
 
 	uint32 NumDataBits = SourceValue.DataBits.Num();
 
diff --git a/Engine/Source/Runtime/Engine/Private/Net/Experimental/Iris/DataStreamChannel.cpp b/Engine/Source/Runtime/Engine/Private/Net/Experimental/Iris/DataStreamChannel.cpp
--- a/Engine/Source/Runtime/Engine/Private/Net/Experimental/Iris/DataStreamChannel.cpp
+++ b/Engine/Source/Runtime/Engine/Private/Net/Experimental/Iris/DataStreamChannel.cpp
@@ -500,6 +500,12 @@ void UDataStreamChannel::AppendExportBunches(TArray<FOutBunch*>& OutExportBunche
 {
 }
 
+TArray<FOutBunch*> UDataStreamChannel::GetAdditionalRequiredBunches(const FOutBunch& OutgoingBunch, EChannelGetAdditionalRequiredBunchesFlags Flags)
+{
+	// Base version assumed we have a UPackageMapClient, so override it
+	return TArray<FOutBunch*>{};
+}
+
 void UDataStreamChannel::AppendMustBeMappedGuids(FOutBunch* Bunch)
 {
 }
diff --git a/Engine/Source/Runtime/Engine/Private/Net/Experimental/Iris/DataStreamChannel.h b/Engine/Source/Runtime/Engine/Private/Net/Experimental/Iris/DataStreamChannel.h
--- a/Engine/Source/Runtime/Engine/Private/Net/Experimental/Iris/DataStreamChannel.h
+++ b/Engine/Source/Runtime/Engine/Private/Net/Experimental/Iris/DataStreamChannel.h
@@ -51,6 +51,8 @@ private:
 
 	/** We do not want to append orphaned exportbunches from other channels */
 	ENGINE_API virtual void AppendExportBunches(TArray<FOutBunch *>& OutExportBunches) override;
+	virtual TArray<FOutBunch*> GetAdditionalRequiredBunches(const FOutBunch& OutgoingBunch, EChannelGetAdditionalRequiredBunchesFlags Flags) override;
+
 	ENGINE_API virtual void AppendMustBeMappedGuids(FOutBunch* Bunch) override;
 
 
diff --git a/Engine/Source/Runtime/Engine/Private/Net/Iris/ReplicationSystem/EngineReplicationBridge.cpp b/Engine/Source/Runtime/Engine/Private/Net/Iris/ReplicationSystem/EngineReplicationBridge.cpp
--- a/Engine/Source/Runtime/Engine/Private/Net/Iris/ReplicationSystem/EngineReplicationBridge.cpp
+++ b/Engine/Source/Runtime/Engine/Private/Net/Iris/ReplicationSystem/EngineReplicationBridge.cpp
@@ -15,8 +15,7 @@
 #include "Iris/Core/IrisProfiler.h"
 #include "Iris/Core/NetObjectReference.h"
 #include "Iris/Core/IrisMemoryTracker.h"
-
-#include "Iris/ReplicationSystem/NetToken.h"
+#include "Net/Core/NetToken/NetToken.h"
 #include "Iris/ReplicationSystem/ObjectReplicationBridgeConfig.h"
 #include "Iris/ReplicationSystem/StringTokenStore.h"
 #include "Iris/ReplicationSystem/ReplicationOperations.h"
diff --git a/Engine/Source/Runtime/Engine/Private/Net/NetConnectionIdHandler.cpp b/Engine/Source/Runtime/Engine/Private/Net/NetConnectionIdHandler.cpp
--- a/Engine/Source/Runtime/Engine/Private/Net/NetConnectionIdHandler.cpp
+++ b/Engine/Source/Runtime/Engine/Private/Net/NetConnectionIdHandler.cpp
@@ -34,3 +34,9 @@ void FNetConnectionIdHandler::Free(uint32 Id)
 	UsedIds[Id] = (Id == 0U);
 	IdHint = FMath::Min(IdHint, Id);
 }
+
+uint32 FNetConnectionIdHandler::GetMaxConnectionIdCount() const
+{
+	return UsedIds.Num();
+}
+
diff --git a/Engine/Source/Runtime/Engine/Private/NetConnection.cpp b/Engine/Source/Runtime/Engine/Private/NetConnection.cpp
--- a/Engine/Source/Runtime/Engine/Private/NetConnection.cpp
+++ b/Engine/Source/Runtime/Engine/Private/NetConnection.cpp
@@ -49,6 +49,7 @@
 #include "Iris/ReplicationSystem/ReplicationSystem.h"
 #include "Iris/ReplicationSystem/Filtering/NetObjectFilter.h"
 #include "Net/Iris/ReplicationSystem/EngineReplicationBridge.h"
+#include "Iris/ReplicationSystem/NetTokenStore.h"
 #endif // UE_WITH_IRIS
 
 #include UE_INLINE_GENERATED_CPP_BY_NAME(NetConnection)
@@ -556,6 +557,15 @@ void UNetConnection::InitBase(UNetDriver* InDriver,class FSocket* InSocket, cons
 		CurrentNetSpeed = FMath::Max<int32>(CurrentNetSpeed, 1800);
 	}
 
+	// For now, NetToken are only enabled if we compile with iris.
+#if UE_WITH_IRIS
+	// Init RemoteNetTokenState.
+	if (Driver->GetNetTokenStore())
+	{
+		Driver->GetNetTokenStore()->InitRemoteNetTokenStoreState(GetConnectionId());
+	}
+#endif
+
 	// Create package map.
 	UPackageMapClient* PackageMapClient = NewObject<UPackageMapClient>(this, PackageMapClass);
 
@@ -3292,6 +3302,9 @@ void UNetConnection::DispatchPacket( FBitReader& Reader, int32 PacketId, bool& b
 
 	const FEngineNetworkCustomVersion::Type PacketEngineNetVer = static_cast<FEngineNetworkCustomVersion::Type>(Reader.EngineNetVer());
 
+	// We want to be able to read replays that do not support export extentions.
+	const bool bHasPartialCustomExportsFinalBit = !(IsInternalAck() && PacketEngineNetVer < FEngineNetworkCustomVersion::CustomExports);
+
 	// Disassemble and dispatch all bunches in the packet.
 	while( !Reader.AtEnd() && GetConnectionState()!=USOCK_Closed )
 	{
@@ -3380,6 +3393,10 @@ void UNetConnection::DispatchPacket( FBitReader& Reader, int32 PacketId, bool& b
 							if (bPartial)
 							{
 								Reader.ReadBit(); // bPartialInitial
+								if (bHasPartialCustomExportsFinalBit)
+								{
+									Reader.ReadBit(); // bPartialCustomExportsFinal
+								}
 								Reader.ReadBit(); // bPartialFinal
 							}
 
@@ -3410,9 +3427,9 @@ void UNetConnection::DispatchPacket( FBitReader& Reader, int32 PacketId, bool& b
 			Bunch.bHasMustBeMappedGUIDs	= Reader.ReadBit();
 			Bunch.bPartial				= Reader.ReadBit();
 
-			if ( Bunch.bReliable )
+			if (Bunch.bReliable)
 			{
-				if ( IsInternalAck() )
+				if (IsInternalAck())
 				{
 					// We can derive the sequence for 100% reliable connections
 					Bunch.ChSequence = InReliable[Bunch.ChIndex] + 1;
@@ -3420,10 +3437,10 @@ void UNetConnection::DispatchPacket( FBitReader& Reader, int32 PacketId, bool& b
 				else
 				{
 					// If this is a reliable bunch, use the last processed reliable sequence to read the new reliable sequence
-					Bunch.ChSequence = MakeRelative( Reader.ReadInt( MAX_CHSEQUENCE ), InReliable[Bunch.ChIndex], MAX_CHSEQUENCE );
+					Bunch.ChSequence = MakeRelative(Reader.ReadInt(MAX_CHSEQUENCE), InReliable[Bunch.ChIndex], MAX_CHSEQUENCE);
 				}
 			} 
-			else if ( Bunch.bPartial )
+			else if (Bunch.bPartial)
 			{
 				// If this is an unreliable partial bunch, we simply use packet sequence since we already have it
 				Bunch.ChSequence = PacketId;
@@ -3434,6 +3451,7 @@ void UNetConnection::DispatchPacket( FBitReader& Reader, int32 PacketId, bool& b
 			}
 
 			Bunch.bPartialInitial = Bunch.bPartial ? Reader.ReadBit() : 0;
+			Bunch.bPartialCustomExportsFinal = Bunch.bPartial && bHasPartialCustomExportsFinalBit ? Reader.ReadBit() : 0;
 			Bunch.bPartialFinal = Bunch.bPartial ? Reader.ReadBit() : 0;
 
 			if (PacketEngineNetVer < FEngineNetworkCustomVersion::ChannelNames)
@@ -3502,7 +3520,7 @@ void UNetConnection::DispatchPacket( FBitReader& Reader, int32 PacketId, bool& b
 
 			const int32 HeaderPos = Reader.GetPosBits();
 
-			if( Reader.IsError() )
+			if(Reader.IsError())
 			{
 				UE_LOG(LogNet, Warning, TEXT("Bunch header overflowed"));
 
@@ -4217,6 +4235,7 @@ int32 UNetConnection::SendRawBunch(FOutBunch& Bunch, bool InAllowMerge, const FN
 	if (Bunch.bPartial)
 	{
 		SendBunchHeader.WriteBit(Bunch.bPartialInitial);
+		SendBunchHeader.WriteBit(Bunch.bPartialCustomExportsFinal);
 		SendBunchHeader.WriteBit(Bunch.bPartialFinal);
 	}
 
@@ -4278,7 +4297,7 @@ int32 UNetConnection::SendRawBunch(FOutBunch& Bunch, bool InAllowMerge, const FN
 
 	UE_LOG(LogNetTraffic, Verbose, TEXT("UNetConnection::SendRawBunch. ChIndex: %d. Bits: %d. PacketId: %d"), Bunch.ChIndex, Bunch.GetNumBits(), Bunch.PacketId);
 
-	if (PackageMap && Bunch.bHasPackageMapExports)
+	if (PackageMap && (Bunch.bHasPackageMapExports || Bunch.bPartialCustomExportsFinal))
 	{
 		PackageMap->NotifyBunchCommit(Bunch.PacketId, &Bunch);
 	}
diff --git a/Engine/Source/Runtime/Engine/Private/NetDriver.cpp b/Engine/Source/Runtime/Engine/Private/NetDriver.cpp
--- a/Engine/Source/Runtime/Engine/Private/NetDriver.cpp
+++ b/Engine/Source/Runtime/Engine/Private/NetDriver.cpp
@@ -62,6 +62,10 @@
 #include "Iris/ReplicationSystem/Filtering/NetObjectFilter.h"
 #include "Net/Iris/ReplicationSystem/EngineReplicationBridge.h"
 #include "Net/Iris/ReplicationSystem/ReplicationSystemUtil.h"
+
+#include "Iris/ReplicationSystem/NameTokenStore.h"
+#include "Iris/ReplicationSystem/StringTokenStore.h"
+#include "GameplayTagTokenStore.h"
 #endif // UE_WITH_IRIS
 
 #if USE_SERVER_PERF_COUNTERS
@@ -1573,6 +1577,29 @@ bool UNetDriver::InitBase(bool bInitAsClient, FNetworkNotify* InNotify, const FU
 
 	Notify = InNotify;
 
+#if UE_WITH_IRIS
+	// At the time being using NetTokens and NetTokenStore is only available when compiling with iris.
+	// Init NetToken stores
+	if (!GetNetTokenStore())
+	{
+		using namespace UE::Net;
+
+		NetTokenStore = MakeUnique<FNetTokenStore>();
+
+		FNetTokenStore::FInitParams NetTokenStoreInitParams;
+		NetTokenStoreInitParams.Authority = !bInitAsClient ? FNetToken::ENetTokenAuthority::Authority : FNetToken::ENetTokenAuthority::None;
+		NetTokenStoreInitParams.MaxConnections = ConnectionIdHandler.GetMaxConnectionIdCount();
+		NetTokenStore->Init(NetTokenStoreInitParams);
+
+		// TODO: make this configurable from config so that users can provide their own custom NetTokenStores
+		// Order is important, changing order will require new netversion
+		// When we move it to the config TypeId will be made explicit
+		NetTokenStore->CreateAndRegisterDataStore<FStringTokenStore>();
+		NetTokenStore->CreateAndRegisterDataStore<FNameTokenStore>();
+		NetTokenStore->CreateAndRegisterDataStore<FGameplayTagTokenStore>();
+	}
+#endif
+	
 #if UE_WITH_IRIS
 	if (IsUsingIrisReplication() && !ReplicationSystem)
 	{
@@ -2833,6 +2860,9 @@ void UNetDriver::ProcessRemoteFunctionForChannelPrivate(
 
 	FNetBitWriter TempWriter( Bunch.PackageMap, 0 );
 
+	// Create NetToken export scope for TempBitWriter, appending exports to Bunch.NetTokensPendingExport
+	UE::Net::FNetTokenExportScope ExportScope(TempWriter, GetNetTokenStore(), Bunch.NetTokensPendingExport, "ProcessRemoteFunctionForChannelPrivate");
+	
 #if UE_NET_TRACE_ENABLED
 	// Create trace collector if tracing is enabled for the target bunch
 	SetTraceCollector(TempWriter, GetTraceCollector(Bunch) ? UE_NET_TRACE_CREATE_COLLECTOR(ENetTraceVerbosity::Trace) : nullptr);
@@ -7117,6 +7147,7 @@ void UNetDriver::CreateReplicationSystem(bool bInitAsClient)
 		Params.bAllowObjectReplication = !bInitAsClient;
 		Params.ForwardNetRPCCallDelegate.BindUObject(this, &UNetDriver::ForwardRemoteFunction);
 		UE::Net::Private::ApplyReplicationSystemConfig(RepSystemConfig, Params);
+		Params.NetTokenStore = NetTokenStore.Get();
 
 		SetReplicationSystem(UE::Net::FReplicationSystemFactory::CreateReplicationSystem(Params));
 	}
@@ -7668,7 +7699,7 @@ void UNetDriver::ProcessRemoteFunction(
 					{
 						// We don't want to call this unless necessary, and it will internally handle being called multiple times before a clear
 						// Builds any shared serialization state for this rpc
-						RepLayout->BuildSharedSerializationForRPC(Parameters);
+						RepLayout->BuildSharedSerializationForRPC(Parameters, GetNetTokenStore());
 
 						InternalProcessRemoteFunctionPrivate(Actor, SubObject, Connection, Function, Parameters, OutParms, Stack, bIsServer, RemoteFunctionFlags);
 					}
diff --git a/Engine/Source/Runtime/Engine/Private/PackageMapClient.cpp b/Engine/Source/Runtime/Engine/Private/PackageMapClient.cpp
--- a/Engine/Source/Runtime/Engine/Private/PackageMapClient.cpp
+++ b/Engine/Source/Runtime/Engine/Private/PackageMapClient.cpp
@@ -10,6 +10,7 @@
 #include "Engine/Engine.h"
 #include "UObject/UObjectIterator.h"
 #include "Engine/NetConnection.h"
+#include "Net/DataBunch.h"
 #include "Net/NetworkProfiler.h"
 #include "Engine/ActorChannel.h"
 #include "ProfilingDebugging/ScopedTimers.h"
@@ -21,6 +22,10 @@
 #include "Net/NetworkGranularMemoryLogging.h"
 #include "Misc/CommandLine.h"
 
+#if UE_WITH_IRIS
+#include "Iris/ReplicationSystem/NetTokenStore.h"
+#endif
+
 #include UE_INLINE_GENERATED_CPP_BY_NAME(PackageMapClient)
 
 #if WITH_EDITOR
@@ -280,6 +285,23 @@ UPackageMapClient::UPackageMapClient(const FObjectInitializer& ObjectInitializer
 {
 }
 
+void UPackageMapClient::Initialize(UNetConnection * InConnection, TSharedPtr<FNetGUIDCache> InNetGUIDCache)
+{
+	Connection = InConnection;
+	GuidCache = InNetGUIDCache;
+	ExportNetGUIDCount = 0;
+	OverrideAckState = &AckState;
+
+#if UE_WITH_IRIS
+	// Initialize NetTokenResolveContext
+	NetTokenResolveContext.NetTokenStore = InConnection->GetDriver()->GetNetTokenStore();
+	if (NetTokenResolveContext.NetTokenStore)
+	{
+		NetTokenResolveContext.RemoteNetTokenStoreState = NetTokenResolveContext.NetTokenStore->GetRemoteNetTokenStoreState(InConnection->GetConnectionId());
+	}
+#endif
+}
+
 /**
  *	This is the meat of the PackageMap class which serializes a reference to Object.
  */
@@ -1353,7 +1375,7 @@ bool UPackageMapClient::ExportNetGUID( FNetworkGUID NetGUID, UObject* Object, FS
 			CurrentExportBunch->DebugString = TEXT("NetGUIDs");
 #endif
 
-			UE_NET_TRACE_SCOPE(NetGUIDExportBunchHeader, *CurrentExportBunch, GetTraceCollector(*CurrentExportBunch), ENetTraceVerbosity::Verbose);
+			UE_NET_TRACE_SCOPE(NetGUIDExportBuncheader, *CurrentExportBunch, GetTraceCollector(*CurrentExportBunch), ENetTraceVerbosity::Verbose);
 
 			CurrentExportBunch->WriteBit( 0 );		// To signify this is NOT a rep layout export
 
@@ -1502,7 +1524,7 @@ void UPackageMapClient::ReceiveNetGUIDBunch( FInBunch &InBunch )
 
 	UE_LOG(LogNetPackageMap, Log, TEXT("UPackageMapClient::ReceiveNetGUIDBunch %d NetGUIDs. PacketId %d. ChSequence %d. ChIndex %d"), NumGUIDsInBunch, InBunch.PacketId, InBunch.ChSequence, InBunch.ChIndex );
 
-	UE_NET_TRACE(NetGUIDExportBunchHeader, Connection->GetInTraceCollector(), StartingBitPos, InBunch.GetPosBits(), ENetTraceVerbosity::Verbose);
+	UE_NET_TRACE(NetGUIDExportBuncheader, Connection->GetInTraceCollector(), StartingBitPos, InBunch.GetPosBits(), ENetTraceVerbosity::Verbose);
 
 	int32 NumGUIDsRead = 0;
 	while( NumGUIDsRead < NumGUIDsInBunch )
@@ -1735,7 +1757,8 @@ void UPackageMapClient::AppendNetFieldExportsInternal(FArchive& Archive, const T
 		check(NetFieldExportHandle == NetFieldExportGroup->NetFieldExports[NetFieldExportHandle].Handle);
 
 		// Export the path if we need to
-		const bool bForceExportDirty = EnumHasAnyFlags(Flags, EAppendNetExportFlags::ForceExportDirtyGroups) && NetFieldExportGroup->bDirtyForReplay;
+		// TODO: CVar this change so we can && this (old behavior) or || this (new behavior)
+		const bool bForceExportDirty = EnumHasAnyFlags(Flags, EAppendNetExportFlags::ForceExportDirtyGroups) || NetFieldExportGroup->bDirtyForReplay;
 
 		uint32 NeedsExport = ((bForceExportDirty || !OverrideAckState->NetFieldExportGroupPathAcked.Contains(PathNameIndex)) && !ExportedPathInThisBunchAlready.Contains(PathNameIndex)) ? 1 : 0;
 
@@ -1929,6 +1952,11 @@ void UPackageMapClient::ReceiveNetFieldExports(FArchive& Archive)
 
 				GuidCache->NetFieldExportGroupMap.Add(PathName, NewNetFieldExportGroup);
 			}
+			else if (NumExportsInGroup > (uint32)NetFieldExportGroup->NetFieldExports.Num())
+			{
+				// Allow our exports to dynamically grow as more are encountered
+				NetFieldExportGroup->NetFieldExports.SetNum(NumExportsInGroup);
+			}
 
 			GuidCache->NetFieldExportGroupPathToIndex.Add(PathName, PathNameIndex);
 			GuidCache->NetFieldExportGroupIndexToGroup.Add(PathNameIndex, NetFieldExportGroup);
@@ -2066,6 +2094,203 @@ void UPackageMapClient::AppendExportBunches(TArray<FOutBunch *>& OutgoingBunches
 	}
 }
 
+/**
+ * Takes an outgoing custom exports bunch and splits it into multiple bunches if it will be too large to fit in a single Packet (specified by MaxBunchSizeInBits).
+ * Notes:
+ *  - The passed-in Bunch argument can be invalidated (if a split occurs).
+ *  - We are not setting bPartialInitial or bPartialFinal on the first/last bunches because this code is intended to be used before such flags get
+ * set on the final bundles of outgoing bunches.
+ * - returns true of the Orignal bunch was split, false if not.
+ */
+bool SplitAndAppendCustomExportsBunch(TArray<FOutBunch*>& AdditionalRequiredBunches, FOutBunch*& OriginalBunch, int32 MaxBunchSizeInBits)
+{
+	// Always land on a byte boundary (except for the final bunch)
+	MaxBunchSizeInBits = (MaxBunchSizeInBits & ~0x07);
+
+	TArray<FOutBunch*, TInlineAllocator<8>> OutBunches;
+
+	//-----------------------------------------------------
+	// Possibly split large bunch into list of smaller partial bunches
+	//-----------------------------------------------------
+	if(OriginalBunch->GetNumBits() > MaxBunchSizeInBits)
+	{
+		uint8* Data = OriginalBunch->GetData();
+		int64 BitsLeft = OriginalBunch->GetNumBits();
+
+		while(BitsLeft > 0)
+		{
+			FOutBunch* PartialBunch = OriginalBunch->PackageMap ? new FOutBunch(OriginalBunch->PackageMap, MaxBunchSizeInBits) : new FOutBunch(OriginalBunch->Channel, false);
+			int64 BitsThisBunch = FMath::Min<int64>(BitsLeft, MaxBunchSizeInBits);
+			PartialBunch->SerializeBits(Data, BitsThisBunch);
+			PartialBunch->bPartial = true;
+
+#if UE_NET_TRACE_ENABLED
+			// Attach tracecollector of split bunch to first partial bunch
+			SetTraceCollector(*PartialBunch, GetTraceCollector(*OriginalBunch));
+			SetTraceCollector(*OriginalBunch, nullptr);
+#endif
+
+			OutBunches.Add(PartialBunch);
+
+#if !(UE_BUILD_SHIPPING || UE_BUILD_TEST)
+			PartialBunch->DebugString = FString::Printf(TEXT("Partial[%d]: %s"), OutBunches.Num(), *OriginalBunch->GetDebugString());
+#endif
+		
+			BitsLeft -= BitsThisBunch;
+			Data += (BitsThisBunch >> 3);
+
+			UE_LOG(LogNetPackageMap, Verbose, TEXT("	Creating partial bunch %s. bitsThisBunch: %d bitsLeft: %d"), *OriginalBunch->GetDebugString(), BitsThisBunch, BitsLeft);
+			
+			ensure(BitsLeft == 0 || BitsThisBunch % 8 == 0); // Byte aligned or it was the last bunch
+		}
+
+		// The last bunch has the packagemap extension data
+		FOutBunch* LastPartialBunch = OutBunches.Last();
+		if (ensure(LastPartialBunch))
+		{
+			LastPartialBunch->bPartialCustomExportsFinal = OriginalBunch->bPartialCustomExportsFinal;
+			LastPartialBunch->NetTokensPendingExport = OriginalBunch->NetTokensPendingExport;
+		}
+
+		// The incoming Bunch is no longer valid
+		delete OriginalBunch;
+		OriginalBunch = nullptr;
+	}
+	else
+	{
+		OutBunches.Add(OriginalBunch);
+	}
+
+	AdditionalRequiredBunches.Append(OutBunches);
+	return OutBunches.Num() > 1;
+}
+
+TArray<FOutBunch*> UPackageMapClient::GetAdditionalRequiredBunches(const FOutBunch& OutgoingBunch, EChannelGetAdditionalRequiredBunchesFlags Flags)
+{
+	TArray<FOutBunch*> AdditionalBunches;
+
+	// Fastpath does not want to append NetGUID exports.
+	if (!EnumHasAnyFlags(Flags, EChannelGetAdditionalRequiredBunchesFlags::SkipNetGUIDExports))
+	{
+		PRAGMA_DISABLE_DEPRECATION_WARNINGS
+		AppendExportBunches(AdditionalBunches);
+		PRAGMA_ENABLE_DEPRECATION_WARNINGS
+	}
+
+	if (FOutBunch* CustomExportsBunch = CreateCustomExportsBunch(OutgoingBunch))
+	{
+		ensureMsgf(CustomExportsBunch->bPartialCustomExportsFinal, TEXT("The Custom Exports Bunch must have bPartialCustomExportsFinal set."));
+
+		// This function will append the bunch to the AdditionalBunches, splitting into multiple bunches if necessary. 
+		// If the orignal bunch it split it will be destroyed
+		SplitAndAppendCustomExportsBunch(AdditionalBunches, CustomExportsBunch, Connection->GetMaxSingleBunchSizeBits());
+	}
+
+	return AdditionalBunches;
+}
+
+FOutBunch* UPackageMapClient::CreateCustomExportsBunch(const FOutBunch& OutgoingBunch)
+{
+	FOutBunch* ExportsBunch = nullptr;
+
+	// NetTokenExports is still an experimental feature that depends on Iris code.
+#if UE_WITH_IRIS
+	// NetTokens - begin
+	using namespace UE::Net;
+	TArray<UE::Net::FNetToken, TInlineAllocator<4>> NetTokensExportedInThisBunch;
+
+	if (OutgoingBunch.NetTokensPendingExport.Num() > 0)
+	{
+		// See if we have something to export
+		const bool bIsNetTokenAuthority = NetTokenResolveContext.NetTokenStore->IsAuthority();
+		for (const UE::Net::FNetToken& NetToken : OutgoingBunch.NetTokensPendingExport)
+		{
+			// We do not need to export tokens assigned by authority unless we are the authority or if they are not already exported in this bunch.
+			const bool bShouldExportToken = bIsNetTokenAuthority || (NetToken.IsAssignedByAuthority() == bIsNetTokenAuthority);
+			if (bShouldExportToken && !(AcknowledgedExportedNetTokens.Contains(NetToken) || NetTokensExportedInThisBunch.Contains(NetToken)))
+			{
+				NetTokensExportedInThisBunch.Add(NetToken);
+			}
+		}
+	}
+
+	if (NetTokensExportedInThisBunch.IsEmpty())
+	{
+		return ExportsBunch;
+	}
+
+	// Need to create an exports bunch that is resizeable (can grow indefinitely)
+	// At a higher level, this (potentially mega-)bunch will be split into partial bunches
+	bool bHasAnyExports = false;
+	ExportsBunch = new FOutBunch(this);
+	ExportsBunch->bPartialCustomExportsFinal = true;
+	ExportsBunch->SetAllowResize(true);
+	ExportsBunch->SetDebugString(FString::Printf(TEXT("%s ExportsBunch"), *OutgoingBunch.GetDebugString()));
+
+#if UE_NET_TRACE_ENABLED
+	// Only enable this if we are doing verbose tracing
+	// We leave it to the bunch to destroy the trace collector
+	SetTraceCollector(*ExportsBunch, UE_NET_TRACE_CREATE_COLLECTOR(ENetTraceVerbosity::Verbose));
+#endif
+	UE_NET_TRACE_SCOPE(CustomExportsBunch, *ExportsBunch, GetTraceCollector(*ExportsBunch), ENetTraceVerbosity::Verbose);
+
+	// NetToken exports
+	{
+		UE_NET_TRACE_SCOPE(NetTokenExports, *ExportsBunch, GetTraceCollector(*ExportsBunch), ENetTraceVerbosity::Verbose);
+		for (FNetToken NetToken : NetTokensExportedInThisBunch)
+		{
+			bool bWroteExport = true; 
+			ExportsBunch->SerializeBits(&bWroteExport, 1);
+
+			UE_NET_TRACE_DYNAMIC_NAME_SCOPE(*NetToken.ToString(), *ExportsBunch, GetTraceCollector(*ExportsBunch), ENetTraceVerbosity::VeryVerbose);
+
+			// Write token 
+			WriteNetToken(*ExportsBunch, NetToken);
+			NetTokenResolveContext.NetTokenStore->WriteTokenData(*ExportsBunch, NetToken);
+
+			// Track what exports we exported until we know what packet we actually did commit to.
+			ExportsBunch->NetTokensPendingExport.Add(NetToken);
+		}
+		bool bNetTokenExportStopBit = false; 
+		ExportsBunch->SerializeBits(&bNetTokenExportStopBit, 1);
+	}
+
+	check(!ExportsBunch->IsError());
+#endif // UE_WITH_IRIS
+
+	return ExportsBunch;
+}
+
+/** Receive Exports written by CreateCustomExportsBunch */
+void UPackageMapClient::ReceiveCustomExportsBunch(FInBunch &InBunch)
+{
+	check(InBunch.bPartialCustomExportsFinal);
+
+	// NetTokenExports is still an experimental feature that depends on Iris code.
+#if UE_WITH_IRIS
+	using namespace UE::Net;
+
+	// Read NetTokenExports
+	bool bHasExportsToRead = false; 
+	InBunch.SerializeBits(&bHasExportsToRead, 1);
+
+	UE::Net::FNetTokenStore* NetTokenStore = NetTokenResolveContext.NetTokenStore;
+	UE::Net::FNetTokenStoreState* RemoteNetTokenStoreState = NetTokenStore->GetRemoteNetTokenStoreState(Connection->GetConnectionId());
+
+	while (bHasExportsToRead && !InBunch.IsError())
+	{
+		FNetToken ImportedNetToken = ReadNetToken(InBunch);
+
+		NetTokenStore->ReadTokenData(InBunch, ImportedNetToken, *RemoteNetTokenStoreState);
+
+		UE_LOG(LogNetToken, VeryVerbose, TEXT("%hs: Imported token %s"), __func__, *ImportedNetToken.ToString());
+
+		// Continuation bit
+		InBunch.SerializeBits(&bHasExportsToRead, 1);
+	}
+#endif // UE_WITH_IRIS
+}
+
 int32 UPackageMapClient::GetNumExportBunches() const
 {
 	return ExportBunches.Num();
@@ -2107,41 +2332,51 @@ void UPackageMapClient::ResetAckState()
  *	Called when a bunch is committed to the connection's Out buffer.
  *	ExportNetGUIDs is the list of GUIDs stored on the bunch that we use to update the expected sequence for those exported GUIDs
  */
-void UPackageMapClient::NotifyBunchCommit( const int32 OutPacketId, const FOutBunch* OutBunch )
+void UPackageMapClient::NotifyBunchCommit(const int32 OutPacketId, const FOutBunch* OutBunch)
 {
 	// Mark all of the net field exports in this bunch as ack'd
 	// NOTE - This only currently works with reliable connections (i.e. InternalAck)
 	// For this to work with normal connections, we'll need to do real ack logic here
-	for ( int32 i = 0; i < OutBunch->NetFieldExports.Num(); i++ )
+	for (int32 i = 0; i < OutBunch->NetFieldExports.Num(); i++)
 	{
-		OverrideAckState->NetFieldExportGroupPathAcked.Add( OutBunch->NetFieldExports[i] >> 32 );
-		OverrideAckState->NetFieldExportAcked.Add( OutBunch->NetFieldExports[i] );
+		OverrideAckState->NetFieldExportGroupPathAcked.Add(OutBunch->NetFieldExports[i] >> 32);
+		OverrideAckState->NetFieldExportAcked.Add(OutBunch->NetFieldExports[i]);
+	}
+
+	// Register exported NetTokens in the PendingAckMap for this OutPackedId
+	if (OutBunch->bPartialCustomExportsFinal)
+	{
+		using namespace UE::Net;
+		for (const UE::Net::FNetToken& NetToken : OutBunch->NetTokensPendingExport)
+		{
+			NetTokenPendingAckMap.Add(OutPacketId, NetToken);
+		}
 	}
 
-	const TArray< FNetworkGUID >& ExportNetGUIDs = OutBunch->ExportNetGUIDs;
+	const TArray<FNetworkGUID>& ExportNetGUIDs = OutBunch->ExportNetGUIDs;
 
-	if ( ExportNetGUIDs.Num() == 0 )
+	if (ExportNetGUIDs.Num() == 0)
 	{
 		return;		// Nothing to do
 	}
 
-	check( OutPacketId > GUID_PACKET_ACKED );	// Assumptions break if this isn't true ( We assume ( OutPacketId > GUID_PACKET_ACKED ) == PENDING )
+	check(OutPacketId > GUID_PACKET_ACKED);	// Assumptions break if this isn't true ( We assume ( OutPacketId > GUID_PACKET_ACKED ) == PENDING )
 
-	for ( int32 i = 0; i < ExportNetGUIDs.Num(); i++ )
+	for (int32 i = 0; i < ExportNetGUIDs.Num(); i++)
 	{
-		if ( !OverrideAckState->NetGUIDAckStatus.Contains( ExportNetGUIDs[i] ) )
+		if (!OverrideAckState->NetGUIDAckStatus.Contains(ExportNetGUIDs[i]))
 		{
-			OverrideAckState->NetGUIDAckStatus.Add( ExportNetGUIDs[i], GUID_PACKET_NOT_ACKED );
+			OverrideAckState->NetGUIDAckStatus.Add(ExportNetGUIDs[i], GUID_PACKET_NOT_ACKED);
 		}
 
-		int32& ExpectedPacketIdRef = OverrideAckState->NetGUIDAckStatus.FindChecked( ExportNetGUIDs[i] );
+		int32& ExpectedPacketIdRef = OverrideAckState->NetGUIDAckStatus.FindChecked(ExportNetGUIDs[i]);
 
 		// Only update expected sequence if this guid was previously nak'd
 		// If we always update to the latest packet id, we risk prolonging the ack for no good reason
 		// (GUID information doesn't change, so updating to the latest expected sequence is unnecessary)
-		if ( ExpectedPacketIdRef == GUID_PACKET_NOT_ACKED )
+		if (ExpectedPacketIdRef == GUID_PACKET_NOT_ACKED)
 		{
-			if ( Connection->IsInternalAck() )
+			if (Connection->IsInternalAck())
 			{
 				// Auto ack now if the connection is 100% reliable
 				ExpectedPacketIdRef = GUID_PACKET_ACKED;
@@ -2149,8 +2384,8 @@ void UPackageMapClient::NotifyBunchCommit( const int32 OutPacketId, const FOutBu
 			}
 
 			ExpectedPacketIdRef = OutPacketId;
-			check( !PendingAckGUIDs.Contains( ExportNetGUIDs[i] ) );	// If we hit this assert, this means the lists are out of sync
-			PendingAckGUIDs.AddUnique( ExportNetGUIDs[i] );
+			check( !PendingAckGUIDs.Contains(ExportNetGUIDs[i]));	// If we hit this assert, this means the lists are out of sync
+			PendingAckGUIDs.AddUnique(ExportNetGUIDs[i]);
 		}
 	}
 }
@@ -2161,18 +2396,26 @@ void UPackageMapClient::NotifyBunchCommit( const int32 OutPacketId, const FOutBu
  */
 void UPackageMapClient::ReceivedAck( const int32 AckPacketId )
 {
-	for ( int32 i = PendingAckGUIDs.Num() - 1; i >= 0; i-- )
+	for ( int32 Index = PendingAckGUIDs.Num() - 1; Index >= 0; Index-- )
 	{
-		int32& ExpectedPacketIdRef = OverrideAckState->NetGUIDAckStatus.FindChecked( PendingAckGUIDs[i] );
+		int32& ExpectedPacketIdRef = OverrideAckState->NetGUIDAckStatus.FindChecked( PendingAckGUIDs[Index] );
 
 		check( ExpectedPacketIdRef > GUID_PACKET_ACKED );		// Make sure we really are pending, since we're on the list
 
 		if ( ExpectedPacketIdRef > GUID_PACKET_ACKED && ExpectedPacketIdRef <= AckPacketId )
 		{
 			ExpectedPacketIdRef = GUID_PACKET_ACKED;	// Fully acked
-			PendingAckGUIDs.RemoveAt( i );				// Remove from pending list, since we're now acked
+			PendingAckGUIDs.RemoveAt( Index );			// Remove from pending list, since we're now acked
 		}
 	}
+
+	// Acknowledge NetTokenExports
+	for (TMultiMap<int32, UE::Net::FNetToken>::TConstKeyIterator It = NetTokenPendingAckMap.CreateConstKeyIterator(AckPacketId); It; ++It)
+	{
+		UE_LOG(LogNetToken, VeryVerbose, TEXT("%hs: Acked token %s"), __func__, *(It.Value().ToString()));
+		AcknowledgedExportedNetTokens.Add(It.Value());
+	}
+	NetTokenPendingAckMap.Remove(AckPacketId);
 }
 
 /**
@@ -2195,6 +2438,9 @@ void UPackageMapClient::ReceivedNak( const int32 NakPacketId )
 			PendingAckGUIDs.RemoveAt( i );	
 		}
 	}
+
+	// Just remove our entry
+	NetTokenPendingAckMap.Remove(NakPacketId);
 }
 
 /**
diff --git a/Engine/Source/Runtime/Engine/Private/RepLayout.cpp b/Engine/Source/Runtime/Engine/Private/RepLayout.cpp
--- a/Engine/Source/Runtime/Engine/Private/RepLayout.cpp
+++ b/Engine/Source/Runtime/Engine/Private/RepLayout.cpp
@@ -8,6 +8,7 @@
 #include "Containers/StaticBitArray.h"
 #include "Misc/MemStack.h"
 #include "Net/Core/PropertyConditions/RepChangedPropertyTracker.h"
+#include "Net/Core/NetToken/NetTokenExportContext.h"
 #include "Net/Serialization/FastArraySerializer.h"
 #include "EngineStats.h"
 #include "Engine/PackageMapClient.h"
@@ -2117,7 +2118,8 @@ bool FRepLayout::ReplicateProperties(
 		// if no shared serialization info exists, build it
 		if (!RepChangelistState->SharedSerialization.IsValid())
 		{
-			BuildSharedSerialization(Data, Changed, true, RepChangelistState->SharedSerialization);
+			UE::Net::FNetTokenStore* NetTokenStore = OwningChannel->Connection->GetDriver()->GetNetTokenStore();
+			BuildSharedSerialization(Data, Changed, true, RepChangelistState->SharedSerialization, NetTokenStore);
 		}
 	}
 
@@ -2151,6 +2153,12 @@ bool FRepLayout::ReplicateProperties(
 	else if (Changed.Num() > 0)
 	{
 		SendProperties(RepState, ChangeTracker, Data, ObjectClass, Writer, Changed, RepChangelistState->SharedSerialization, RepFlags.bSerializePropertyNames ? ESerializePropertyType::Name : ESerializePropertyType::Handle);
+
+		if (UE::Net::FNetTokenExportContext* NetTokenExportContext = RepChangelistState->SharedSerialization.IsValid() ? Writer.NetTokenExportContext.Get() : nullptr)
+		{
+			// For now just append all potential exports, we can be smarter about it if necessary
+			NetTokenExportContext->AppendNetTokensPendingExport(RepChangelistState->SharedSerialization.NetTokensPendingExport);
+		}
 	}
 
 	// See if something actually sent (this may be false due to conditional checks inside the send properties function
@@ -6629,6 +6637,7 @@ void FRepLayout::SerializeProperties_r(
 #if !(UE_BUILD_SHIPPING || UE_BUILD_TEST)
 			if ((GNetVerifyShareSerializedData != 0) && Ar.IsSaving())
 			{
+				// TODO: Is this actually a NetBitWriter?
 				FBitWriter& Writer = static_cast<FBitWriter&>(Ar);
 
 				FBitWriterMark BitWriterMark(Writer);
@@ -6746,7 +6755,8 @@ void FRepLayout::BuildSharedSerialization(
 	const FConstRepObjectDataBuffer Data,
 	TArray<uint16>& Changed,
 	const bool bWriteHandle,
-	FRepSerializationSharedInfo& SharedInfo) const
+	FRepSerializationSharedInfo& SharedInfo,
+	UE::Net::FNetTokenStore* NetTokenStore) const
 {
 #ifdef ENABLE_PROPERTY_CHECKSUMS
 	const bool bDoChecksum = (GDoPropertyChecksum == 1);
@@ -6759,6 +6769,9 @@ void FRepLayout::BuildSharedSerialization(
 
 	SharedInfo.Init();
 
+	// Create scope to export NetTokens, storing pending exports in SharedInfoRPC.NetTokensPendingExport
+	UE::Net::FNetTokenExportScope NetTokenExportScope(*SharedInfo.SerializedProperties, NetTokenStore, SharedInfo.NetTokensPendingExport, "BuildSharedSerialization");
+
 	BuildSharedSerialization_r(HandleIterator, Data, bWriteHandle, bDoChecksum, 0, SharedInfo);
 
 	SharedInfo.SetValid();
@@ -6870,13 +6883,16 @@ void FRepLayout::BuildSharedSerializationForRPC_r(
 	}
 }
 
-void FRepLayout::BuildSharedSerializationForRPC(const FConstRepObjectDataBuffer Data)
+void FRepLayout::BuildSharedSerializationForRPC(const FConstRepObjectDataBuffer Data, UE::Net::FNetTokenStore* NetTokenStore)
 {
 	if ((GNetSharedSerializedData != 0) && !SharedInfoRPC.IsValid())
 	{
 		SharedInfoRPC.Init();
 		SharedInfoRPCParentsChanged.Init(false, Parents.Num());
 
+		// Create scope to export NetTokens, storing pending exports in SharedInfoRPC.NetTokensPendingExport
+		UE::Net::FNetTokenExportScope ExportScope(*SharedInfoRPC.SerializedProperties, NetTokenStore, SharedInfoRPC.NetTokensPendingExport, "BuildSharedSerializationForRPC");
+
 		for (int32 i = 0; i < Parents.Num(); i++)
 		{
 			if (Parents[i].Property->HasAnyPropertyFlags(CPF_OutParm))
@@ -6941,6 +6957,7 @@ void FRepLayout::SendPropertiesForRPC(
 		}
 		else
 		{
+			bool bAppendNetTokenExports = false;
 			for (int32 i = 0; i < Parents.Num(); i++)
 			{
 				bool Send = true;
@@ -6965,10 +6982,16 @@ void FRepLayout::SendPropertiesForRPC(
 
 				if (Send)
 				{
+					bAppendNetTokenExports = true;
 					bool bHasUnmapped = false;
 					SerializeProperties_r(Writer, Writer.PackageMap, Parents[i].CmdStart, Parents[i].CmdEnd, const_cast<uint8*>(Data.Data), bHasUnmapped, 0, 0, SharedInfoRPC, GetTraceCollector(Writer), nullptr);
 				}
 			}
+			// Append potential exports from shared serialization			
+			if (UE::Net::FNetTokenExportContext* NetTokenExportContext = bAppendNetTokenExports && SharedInfoRPC.IsValid() ? Writer.NetTokenExportContext.Get() : nullptr)
+			{
+				NetTokenExportContext->AppendNetTokensPendingExport(SharedInfoRPC.NetTokensPendingExport);
+			}
 		}	
 	}
 }
diff --git a/Engine/Source/Runtime/Engine/Public/Net/DataBunch.h b/Engine/Source/Runtime/Engine/Public/Net/DataBunch.h
--- a/Engine/Source/Runtime/Engine/Public/Net/DataBunch.h
+++ b/Engine/Source/Runtime/Engine/Public/Net/DataBunch.h
@@ -10,6 +10,7 @@
 #include "UObject/CoreNet.h"
 #include "EngineLogs.h"
 #include "Net/Core/Trace/NetTraceConfig.h"
+#include "Net/Core/NetToken/NetToken.h"
 
 class UChannel;
 class UNetConnection;
@@ -41,11 +42,13 @@ public:
 	uint8					bPartialFinal:1;			// The final bunch of a partial bunch
 	uint8					bHasPackageMapExports:1;	// This bunch has networkGUID name/id pairs
 	uint8					bHasMustBeMappedGUIDs:1;	// This bunch has guids that must be mapped before we can process this bunch
+	uint8					bPartialCustomExportsFinal:1;	// This bunch marks the end of the CustomExports data that needs to be processed immediately (not queued)
 
 	EChannelCloseReason		CloseReason;
 
 	TArray< FNetworkGUID >	ExportNetGUIDs;			// List of GUIDs that went out on this bunch
 	TArray< uint64 >		NetFieldExports;
+	TArray<UE::Net::FNetToken, TInlineAllocator<4>> NetTokensPendingExport; // List of NetTokens that will be exported if needed with this bunch.
 
 #if !(UE_BUILD_SHIPPING || UE_BUILD_TEST)
 	FString			DebugString;
@@ -53,7 +56,7 @@ public:
 	{
 			DebugString = DebugStr;
 	}
-	FString	GetDebugString()
+	FString	GetDebugString() const
 	{
 		return DebugString;
 	}
@@ -62,7 +65,7 @@ public:
 	{
 
 	}
-	FORCEINLINE FString	GetDebugString()
+	FORCEINLINE FString	GetDebugString() const
 	{
 		return FString();
 	}
@@ -103,7 +106,8 @@ public:
 		PRAGMA_ENABLE_DEPRECATION_WARNINGS
 		Str += FString::Printf(TEXT("bReliable: %d "), bReliable);
 		Str += FString::Printf(TEXT("bPartial: %d//%d//%d "), bPartial, bPartialInitial, bPartialFinal);
-		Str += FString::Printf( TEXT( "bHasPackageMapExports: %d " ), bHasPackageMapExports );
+		Str += FString::Printf(TEXT("bHasPackageMapExports: %d "), bHasPackageMapExports);
+		Str += FString::Printf(TEXT("NetTokensPendingExport: %d "), NetTokensPendingExport.Num());
 		Str += GetDebugString();
 #else
 		FString Str = FString::Printf(TEXT("Channel[%d]. Seq %d. PacketId: %d"), ChIndex, ChSequence, PacketId);
@@ -112,6 +116,7 @@ public:
 	}
 
 	ENGINE_API virtual void CountMemory(FArchive& Ar) const override;
+	ENGINE_API virtual void Reset() override;
 };
 
 //
@@ -137,6 +142,7 @@ public:
 	uint8				bPartialFinal:1;			// The final bunch of a partial bunch
 	uint8				bHasPackageMapExports:1;	// This bunch has networkGUID name/id pairs
 	uint8				bHasMustBeMappedGUIDs:1;	// This bunch has guids that must be mapped before we can process this bunch
+	uint8				bPartialCustomExportsFinal:1;	// This bunch marks the end of the extensions data that needs to be processed immediately (not queued)
 	uint8				bIgnoreRPCs:1;
 
 	EChannelCloseReason		CloseReason;
@@ -163,6 +169,7 @@ public:
 		Str += FString::Printf(TEXT("bPartial: %d//%d//%d "), bPartial, bPartialInitial, bPartialFinal);
 		Str += FString::Printf(TEXT("bHasPackageMapExports: %d "), bHasPackageMapExports );
 		Str += FString::Printf(TEXT("bHasMustBeMappedGUIDs: %d "), bHasMustBeMappedGUIDs );
+		Str += FString::Printf(TEXT("bPartialCustomExportsFinal: %d "), bPartialCustomExportsFinal);
 		Str += FString::Printf(TEXT("bIgnoreRPCs: %d "), bIgnoreRPCs );
 #else
 		FString Str = FString::Printf(TEXT("Channel[%d]. Seq %d. PacketId: %d"), ChIndex, ChSequence, PacketId);
@@ -208,16 +215,4 @@ struct FControlChannelOutBunch : public FOutBunch
 	}
 };
 
-/** Helper methods to allow us to instrument different type of BitStreams */
-inline uint32 GetBitStreamPositionForNetTrace(const FBitWriter& Stream) { return (uint32(Stream.IsError()) - 1U) & (uint32)Stream.GetNumBits(); }
-inline uint32 GetBitStreamPositionForNetTrace(const FBitReader& Stream) { return (uint32(Stream.IsError()) - 1U) & (uint32)Stream.GetPosBits(); }
-
-#if UE_NET_TRACE_ENABLED
-inline FNetTraceCollector* GetTraceCollector(FNetBitWriter& BitWriter) { return BitWriter.TraceCollector.Get(); }
-inline void SetTraceCollector(FNetBitWriter& BitWriter, FNetTraceCollector* Collector) { BitWriter.TraceCollector.Set(Collector); }
-#else
-inline FNetTraceCollector* GetTraceCollector(FNetBitWriter& BitWriter) { return nullptr; }
-inline void SetTraceCollector(FNetBitWriter& BitWriter, FNetTraceCollector* Collector) {}
-#endif
-
 
diff --git a/Engine/Source/Runtime/Engine/Public/Net/DataReplication.h b/Engine/Source/Runtime/Engine/Public/Net/DataReplication.h
--- a/Engine/Source/Runtime/Engine/Public/Net/DataReplication.h
+++ b/Engine/Source/Runtime/Engine/Public/Net/DataReplication.h
@@ -151,14 +151,6 @@ public:
 
 	ENGINE_API void CountBytes(FArchive& Ar) const;
 
-	/** Writes dirty properties to bunch */
-	UE_DEPRECATED(5.1, "Now takes an additional out param")
-	void ReplicateCustomDeltaProperties(FNetBitWriter& Bunch, FReplicationFlags RepFlags)
-	{
-		bool bSkippedPropertyCondition = false;
-		ReplicateCustomDeltaProperties(Bunch, RepFlags, bSkippedPropertyCondition);
-	}
-
 	ENGINE_API void ReplicateCustomDeltaProperties(FNetBitWriter& Bunch, FReplicationFlags RepFlags, bool& bSkippedPropertyCondition);
 	ENGINE_API bool ReplicateProperties(FOutBunch& Bunch, FReplicationFlags RepFlags);
 	ENGINE_API bool ReplicateProperties(FOutBunch& Bunch, FReplicationFlags RepFlags, FNetBitWriter& Writer);
diff --git a/Engine/Source/Runtime/Engine/Public/Net/NetConnectionIdHandler.h b/Engine/Source/Runtime/Engine/Public/Net/NetConnectionIdHandler.h
--- a/Engine/Source/Runtime/Engine/Public/Net/NetConnectionIdHandler.h
+++ b/Engine/Source/Runtime/Engine/Public/Net/NetConnectionIdHandler.h
@@ -16,6 +16,8 @@ public:
 	ENGINE_API uint32 Allocate();
 	ENGINE_API void Free(uint32 Id);
 
+	ENGINE_API uint32 GetMaxConnectionIdCount() const;
+
 private:
 	TBitArray<> UsedIds;
 	uint32 IdHint;
diff --git a/Engine/Source/Runtime/Engine/Public/Net/RepLayout.h b/Engine/Source/Runtime/Engine/Public/Net/RepLayout.h
--- a/Engine/Source/Runtime/Engine/Public/Net/RepLayout.h
+++ b/Engine/Source/Runtime/Engine/Public/Net/RepLayout.h
@@ -21,6 +21,7 @@
 #include "UObject/GCObject.h"
 #include "Containers/StaticBitArray.h"
 #include "Net/Core/Misc/GuidReferences.h"
+#include "Net/Core/NetToken/NetTokenExportContext.h"
 #include "Net/Core/PushModel/PushModel.h"
 #include "Net/Core/PropertyConditions/RepChangedPropertyTracker.h"
 #include "Templates/CopyQualifiersFromTo.h"
@@ -274,6 +275,7 @@ struct FRepSerializationSharedInfo
 		{
 			SharedPropertyInfo.Reset();
 			SerializedProperties->Reset();
+			NetTokensPendingExport.Reset();
 
 			bIsValid = false;
 		}
@@ -305,6 +307,9 @@ struct FRepSerializationSharedInfo
 	/** Binary blob of net serialized data to be shared */
 	TUniquePtr<FNetBitWriter> SerializedProperties;
 
+	/** NetTokenExports for this RepLayout, will need to index to avoid adding exports that we do not want to use. */
+	UE::Net::FNetTokenExportContext::FNetTokenExports NetTokensPendingExport;
+
 	void CountBytes(FArchive& Ar) const;
 
 private:
@@ -1411,7 +1416,7 @@ public:
 		TSet<FNetworkGUID>& UnmappedGuids) const;
 
 	/** Builds shared serialization state for a multicast rpc */
-	void ENGINE_API BuildSharedSerializationForRPC(const FConstRepObjectDataBuffer Data);
+	void ENGINE_API BuildSharedSerializationForRPC(const FConstRepObjectDataBuffer Data, UE::Net::FNetTokenStore* NetTokenStore = nullptr);
 
 	/** Clears shared serialization state for a multicast rpc */
 	void ENGINE_API ClearSharedSerializationForRPC();
@@ -1648,7 +1653,8 @@ private:
 		const FConstRepObjectDataBuffer Data,
 		TArray<uint16>& Changed,
 		const bool bWriteHandle,
-		FRepSerializationSharedInfo& SharedInfo) const;
+		FRepSerializationSharedInfo& SharedInfo, 
+		UE::Net::FNetTokenStore* NetTokenStore = nullptr) const;
 
 	void BuildSharedSerialization_r(
 		FRepHandleIterator& RepHandleIterator,
diff --git a/Engine/Source/Runtime/Experimental/Iris/Core/Private/Iris/Core/IrisDebugging.cpp b/Engine/Source/Runtime/Experimental/Iris/Core/Private/Iris/Core/IrisDebugging.cpp
--- a/Engine/Source/Runtime/Experimental/Iris/Core/Private/Iris/Core/IrisDebugging.cpp
+++ b/Engine/Source/Runtime/Experimental/Iris/Core/Private/Iris/Core/IrisDebugging.cpp
@@ -311,14 +311,13 @@ void NetObjectStateToString(FStringBuilderBase& StringBuilder, FNetRefHandle Ref
 	// If this is a client handle we assume that we only have a single connections and use the first valid connection to get the remote token store.
 	FReplicationConnections& Connections = ReplicationSystemInternal->GetConnections();
 	const uint32 FirstValidConnectionId = Connections.GetValidConnections().FindFirstOne();
-	FNetTokenStoreState* TokenStoreState = (bIsServer || (FirstValidConnectionId == FNetBitArray::InvalidIndex)) ? ReplicationSystemInternal->GetNetTokenStore().GetLocalNetTokenStoreState() : &Connections.GetRemoteNetTokenStoreState(FirstValidConnectionId);
 
 	// Setup Context
 	FInternalNetSerializationContext InternalContext;
 	FInternalNetSerializationContext::FInitParameters InternalContextInitParams;
 	InternalContextInitParams.ReplicationSystem = ReplicationSystem;
 	InternalContextInitParams.PackageMap = ReplicationSystemInternal->GetIrisObjectReferencePackageMap();
-	InternalContextInitParams.ObjectResolveContext.RemoteNetTokenStoreState = TokenStoreState;
+	InternalContextInitParams.ObjectResolveContext.RemoteNetTokenStoreState = FirstValidConnectionId == FNetBitArray::InvalidIndex ? nullptr :ReplicationSystem->GetNetTokenStore()->GetRemoteNetTokenStoreState(FirstValidConnectionId);
 	InternalContextInitParams.ObjectResolveContext.ConnectionId = (FirstValidConnectionId == FNetBitArray::InvalidIndex ? InvalidConnectionId : FirstValidConnectionId);
 	InternalContext.Init(InternalContextInitParams);
 
diff --git a/Engine/Source/Runtime/Experimental/Iris/Core/Private/Iris/ReplicationSystem/NameTokenStore.cpp b/Engine/Source/Runtime/Experimental/Iris/Core/Private/Iris/ReplicationSystem/NameTokenStore.cpp
--- a/Engine/Source/Runtime/Experimental/Iris/Core/Private/Iris/ReplicationSystem/NameTokenStore.cpp
+++ b/Engine/Source/Runtime/Experimental/Iris/Core/Private/Iris/ReplicationSystem/NameTokenStore.cpp
@@ -8,16 +8,17 @@
 #include "Iris/Serialization/NetSerializationContext.h"
 #include "Hash/CityHash.h"
 #include "Net/Core/Trace/NetTrace.h"
+#include "UObject/CoreNet.h"
 
 #define UE_NET_ENABLE_FNAME_TOKEN_LOG 1
 
 #if UE_NET_ENABLE_FNAME_TOKEN_LOG
-#	define UE_LOG_FNAMETOKEN(Format, ...)  UE_LOG(LogIris, Verbose, Format, ##__VA_ARGS__)
+#	define UE_LOG_FNAMETOKEN(Format, ...)  UE_LOG(LogNetToken, Verbose, Format, ##__VA_ARGS__)
 #else
 #	define UE_LOG_FNAMETOKEN(...)
 #endif
 
-#define UE_LOG_FNAMETOKEN_WARNING(Format, ...)  UE_LOG(LogIris, Warning, Format, ##__VA_ARGS__)
+#define UE_LOG_FNAMETOKEN_WARNING(Format, ...)  UE_LOG(LogNetToken, Warning, Format, ##__VA_ARGS__)
 
 namespace UE::Net
 {
@@ -75,9 +76,6 @@ FNameTokenStore::FNameTokenStore(FNetTokenStore& InTokenStore)
 	// Reserve 0
 	StoredFNames.Add(FName());
 	StoredTokens.Add(FNetToken());
-
-	// Register
-	TokenStore.RegisterDataStore(this, FNameTokenStore::GetTokenStoreName());
 }
 
 FName FNameTokenStore::ResolveToken(FNetToken Token, const FNetTokenStoreState* NetTokenStoreState) const
@@ -98,10 +96,19 @@ FName FNameTokenStore::ResolveToken(FNetToken Token, const FNetTokenStoreState*
 void FNameTokenStore::WriteTokenData(FNetSerializationContext& Context, FNetTokenStoreKey TokenStoreKey) const
 {
 	UE_NET_TRACE_DYNAMIC_NAME_SCOPE(StoredFNames[TokenStoreKey.GetKeyValue()], *Context.GetBitStreamWriter(), Context.GetTraceCollector(), ENetTraceVerbosity::VeryVerbose);
+	UE_LOG_FNAMETOKEN(TEXT("FNameTokenStore::WriteTokenData %s %s"), *(StoredTokens[TokenStoreKey.GetKeyValue()].ToString()), *(StoredFNames[TokenStoreKey.GetKeyValue()].ToString()));
 	// $TODO: $IRIS: We can be a bit smarter here and utilize the string-number split of FNames to export less data.. JIRA: UE-221753
 	WriteString(Context.GetBitStreamWriter(), StoredFNames[TokenStoreKey.GetKeyValue()].ToString());
 }
 
+void FNameTokenStore::WriteTokenData(FArchive& Ar, FNetTokenStoreKey TokenStoreKey) const
+{
+	UE_NET_TRACE_DYNAMIC_NAME_SCOPE(StoredFNames[TokenStoreKey.GetKeyValue()], static_cast<FNetBitWriter&>(Ar), GetTraceCollector(static_cast<FNetBitWriter&>(Ar)), ENetTraceVerbosity::VeryVerbose);
+	UE_LOG_FNAMETOKEN(TEXT("FNameTokenStore::WriteTokenData %s %s"), *(StoredTokens[TokenStoreKey.GetKeyValue()].ToString()), *(StoredFNames[TokenStoreKey.GetKeyValue()].ToString()));
+	FName Name = StoredFNames[TokenStoreKey.GetKeyValue()];
+	UPackageMap::StaticSerializeName(Ar, Name);
+}
+
 FNetTokenStoreKey FNameTokenStore::ReadTokenData(FNetSerializationContext& Context, const FNetToken& NetToken)
 {
 	UE_NET_TRACE_NAMED_DYNAMIC_NAME_SCOPE(TokenScope, FName(), *Context.GetBitStreamReader(), Context.GetTraceCollector(), ENetTraceVerbosity::VeryVerbose);
@@ -136,4 +143,30 @@ FNetTokenStoreKey FNameTokenStore::ReadTokenData(FNetSerializationContext& Conte
 	}
 }
 
+FNetTokenStoreKey FNameTokenStore::ReadTokenData(FArchive& Ar, const FNetToken& NetToken)
+{
+	FName Name;
+	UPackageMap::StaticSerializeName(Ar, Name);
+
+	if (!Ar.IsError())
+	{
+		FNetTokenStoreKey Key = GetOrCreateTokenStoreKey(Name);
+
+		// If this is an authtoken and we are not the authority, update the stored key
+		if (Key.IsValid() && NetToken.IsAssignedByAuthority() && !TokenStore.IsAuthority())
+		{
+			UE_LOG_FNAMETOKEN(TEXT("NameTokenStore::ReadTokenData - Replaced local key %s with %s for name %s"), *StoredTokens[Key.GetKeyValue()].ToString(), *NetToken.ToString(), *Name.ToString());	
+
+			// This way the next time we lookup this key during assignment use the imported authoritative key which we do not have to export
+			StoredTokens[Key.GetKeyValue()] = NetToken;
+		}
+
+		return Key;
+	}
+	else
+	{
+		return FNetTokenStoreKey();
+	}
+}
+
 }
diff --git a/Engine/Source/Runtime/Experimental/Iris/Core/Private/Iris/ReplicationSystem/NetTokenDataStream.cpp b/Engine/Source/Runtime/Experimental/Iris/Core/Private/Iris/ReplicationSystem/NetTokenDataStream.cpp
--- a/Engine/Source/Runtime/Experimental/Iris/Core/Private/Iris/ReplicationSystem/NetTokenDataStream.cpp
+++ b/Engine/Source/Runtime/Experimental/Iris/Core/Private/Iris/ReplicationSystem/NetTokenDataStream.cpp
@@ -60,11 +60,11 @@ void UNetTokenDataStream::Init(const FInitParameters& Params)
 
 	ReplicationSystemId = Params.ReplicationSystemId;
 	ConnectionId = Params.ConnectionId;
-	RemoteNetTokenStoreState = Params.RemoteTokenStoreState;
 	NetExports = Params.NetExports;
 
 	UReplicationSystem* ReplicationSystem = UE::Net::GetReplicationSystem(ReplicationSystemId);
-	NetTokenStore = &ReplicationSystem->GetReplicationSystemInternal()->GetNetTokenStore();
+	NetTokenStore = ReplicationSystem->GetNetTokenStore();
+	RemoteNetTokenStoreState = NetTokenStore->GetRemoteNetTokenStoreState(Params.ConnectionId);
 
 	// $IRIS $TODO: if we want to make this into a real feature we need to expose some sort of api to mark tokens for pre-export
 	if (Private::bIrisPreExportExistingNetTokensOnConnect)
diff --git a/Engine/Source/Runtime/Experimental/Iris/Core/Private/Iris/ReplicationSystem/NetTokenDataStream.h b/Engine/Source/Runtime/Experimental/Iris/Core/Private/Iris/ReplicationSystem/NetTokenDataStream.h
--- a/Engine/Source/Runtime/Experimental/Iris/Core/Private/Iris/ReplicationSystem/NetTokenDataStream.h
+++ b/Engine/Source/Runtime/Experimental/Iris/Core/Private/Iris/ReplicationSystem/NetTokenDataStream.h
@@ -3,7 +3,7 @@
 #pragma once
 
 #include "Iris/DataStream/DataStream.h"
-#include "Iris/ReplicationSystem/NetToken.h"
+#include "Net/Core/NetToken/NetToken.h"
 #include "Containers/RingBuffer.h"
 
 #include "NetTokenDataStream.generated.h"
@@ -32,7 +32,6 @@ public:
 	{
 		uint32 ReplicationSystemId;
 		uint32 ConnectionId;
-		UE::Net::FNetTokenStoreState* RemoteTokenStoreState;
 		UE::Net::Private::FNetExports* NetExports;
 	};
 
diff --git a/Engine/Source/Runtime/Experimental/Iris/Core/Private/Iris/ReplicationSystem/NetTokenStore.cpp b/Engine/Source/Runtime/Experimental/Iris/Core/Private/Iris/ReplicationSystem/NetTokenStore.cpp
--- a/Engine/Source/Runtime/Experimental/Iris/Core/Private/Iris/ReplicationSystem/NetTokenStore.cpp
+++ b/Engine/Source/Runtime/Experimental/Iris/Core/Private/Iris/ReplicationSystem/NetTokenStore.cpp
@@ -28,7 +28,8 @@ FNetTokenStoreKey FNetTokenDataStore::GetTokenKey(FNetToken Token, const FNetTok
 	if (Token.GetIndex() < (uint32)TokenStoreState.TokenInfos.Num())
 	{
 		const FNetTokenStoreKey& Key = TokenStoreState.TokenInfos[Token.GetIndex()];
-		if (ensureAlwaysMsgf(GetTypeId() == Key.GetTypeId(), TEXT("Cannot resolve NetToken %s with TypeId: %u in DataStore with TypeId: %u"), *Token.ToString(), Key.GetTypeId(), GetTypeId()))
+		//if (ensureAlwaysMsgf(Key.IsValid() && GetTypeId() == Key.GetTypeId(), TEXT("Cannot resolve NetToken %s with TypeId: %u in DataStore with TypeId: %u"), *Token.ToString(), Key.GetTypeId(), GetTypeId()))
+		if (ensureMsgf(Key.IsValid() && GetTypeId() == Key.GetTypeId(), TEXT("Cannot resolve NetToken %s with TypeId: %u in DataStore with TypeId: %u"), *Token.ToString(), Key.GetTypeId(), GetTypeId()))
 		{
 			return Key;
 		}
@@ -54,30 +55,70 @@ FNetToken FNetTokenDataStore::CreateToken(FNetToken::ENetTokenAuthority Authorit
 
 FNetToken ReadNetToken(UE::Net::FNetSerializationContext& Context)
 {
-	UE_NET_TRACE_NAMED_DYNAMIC_NAME_SCOPE(TokenScope, FName(), *Context.GetBitStreamReader(), Context.GetTraceCollector(), ENetTraceVerbosity::VeryVerbose);
 	FNetBitStreamReader* Reader = Context.GetBitStreamReader();
 
-	const uint32 Index = ReadPackedUint32(Reader);
-	const bool bIsAssignedByAuthority = Reader->ReadBool();
-
 	FNetToken ReadToken;
-	
-	if (!Reader->IsOverflown())
+
+	UE_NET_TRACE_NAMED_DYNAMIC_NAME_SCOPE(TokenScope, FName(), *Context.GetBitStreamReader(), Context.GetTraceCollector(), ENetTraceVerbosity::VeryVerbose);
+	const uint32 TokenIndex = ReadPackedUint32(Reader);
+	if (const bool bIsValid = (TokenIndex != FNetToken::InvalidTokenIndex))
 	{
-		ReadToken = FNetToken::MakeNetToken(Index, bIsAssignedByAuthority ? FNetToken::ENetTokenAuthority::Authority : FNetToken::ENetTokenAuthority::None);
-		UE_NET_TRACE_SET_SCOPE_NAME(TokenScope, *ReadToken.ToString());
+		const bool bIsAssignedByAuthority = Reader->ReadBool();
+		if (!Reader->IsOverflown())
+		{
+			ReadToken = FNetToken::MakeNetToken(TokenIndex, bIsAssignedByAuthority ? FNetToken::ENetTokenAuthority::Authority : FNetToken::ENetTokenAuthority::None);
+			UE_NET_TRACE_SET_SCOPE_NAME(TokenScope, *ReadToken.ToString());
+		}
 	}
 
 	return ReadToken;
 }
 
+// Note: Be cereful when modifying this methods to not affect replay compatibility.
 void WriteNetToken(UE::Net::FNetSerializationContext& Context, FNetToken Token)
 {
 	UE_NET_TRACE_DYNAMIC_NAME_SCOPE(*Token.ToString(), *Context.GetBitStreamWriter(), Context.GetTraceCollector(), ENetTraceVerbosity::VeryVerbose);
 	FNetBitStreamWriter* Writer = Context.GetBitStreamWriter();
 
-	WritePackedUint32(Writer, Token.GetIndex());
-	Writer->WriteBool(Token.IsAssignedByAuthority());
+	const uint32 TokenIndex = Token.GetIndex();
+	WritePackedUint32(Writer, TokenIndex);
+	if (TokenIndex != FNetToken::InvalidTokenIndex)
+	{
+		Writer->WriteBool(Token.IsAssignedByAuthority());
+	}
+}
+
+FNetToken ReadNetToken(FArchive& Ar)
+{
+	FNetToken ReadToken;
+
+	uint32 TokenIndex = 0;
+	Ar.SerializeIntPacked(TokenIndex);
+
+	if (const bool bIsValid = (TokenIndex != FNetToken::InvalidTokenIndex))
+	{
+		bool bIsAssignedByAuthority = false;
+		Ar.SerializeBits(&bIsAssignedByAuthority, 1);
+
+		if (!Ar.IsError())
+		{
+			ReadToken = FNetToken::MakeNetToken(TokenIndex, bIsAssignedByAuthority ? FNetToken::ENetTokenAuthority::Authority : FNetToken::ENetTokenAuthority::None);
+		}
+	}
+
+	return ReadToken;
+}
+
+// Note: Be cereful when modifying this methods to not affect replay compatibility.
+void WriteNetToken(FArchive& Ar, FNetToken Token)
+{
+	uint32 TokenIndex = Token.GetIndex();
+	Ar.SerializeIntPacked(TokenIndex);
+	if (TokenIndex != FNetToken::InvalidTokenIndex)
+	{
+		bool bIsAssignedByAuthority = Token.IsAssignedByAuthority();
+		Ar.SerializeBits(&bIsAssignedByAuthority, 1);
+	}
 }
 
 FNetTokenStore::FNetTokenStore()
@@ -87,27 +128,61 @@ FNetTokenStore::FNetTokenStore()
 
 FNetTokenStore::~FNetTokenStore()
 {
-	delete LocalNetTokenStoreState;
 }
 
 void FNetTokenStore::Init(FNetTokenStore::FInitParams& InParams)
 {
 	Params = InParams;
+	RemoteNetTokenStoreStates.SetNum(InParams.MaxConnections);
+}
+
+void FNetTokenStore::InitRemoteNetTokenStoreState(uint32 ConnectionId)
+{
+	if (ensureMsgf((ConnectionId != InvalidConnectionId) && (ConnectionId < (uint32)RemoteNetTokenStoreStates.Num()), TEXT("Trying to init RemoteNetTokenStoreState for invalid connection %u"), ConnectionId))
+	{
+		if (FNetTokenStoreState* ExistingState = RemoteNetTokenStoreStates[ConnectionId].Get())
+		{
+			ExistingState->ReserveTokenCount(0U);
+
+		}
+		else
+		{
+			RemoteNetTokenStoreStates[ConnectionId] = MakeUnique<FNetTokenStoreState>();
+		}
+	}
+}
+
+const FNetTokenStoreState* FNetTokenStore::GetRemoteNetTokenStoreState(uint32 ConnectionId) const
+{
+	if (!ensureMsgf((ConnectionId != InvalidConnectionId) && (ConnectionId < (uint32)RemoteNetTokenStoreStates.Num()), TEXT("Trying to access RemoteNetTokenStoreState for ConnectionID: %u"), ConnectionId))
+	{
+		return nullptr;
+	}
+	return RemoteNetTokenStoreStates[ConnectionId].Get();
+}
+
+FNetTokenStoreState* FNetTokenStore::GetRemoteNetTokenStoreState(uint32 ConnectionId)
+{
+	if (!ensureMsgf((ConnectionId != InvalidConnectionId) && (ConnectionId < (uint32)RemoteNetTokenStoreStates.Num()), TEXT("Trying to access non existing RemoteNetTokenStoreState for ConnectionId: %u"), ConnectionId))
+	{
+		return nullptr;
+	}
+	return RemoteNetTokenStoreStates[ConnectionId].Get();
 }
 
 const FNetTokenDataStore* FNetTokenStore::GetDataStore(FName Name) const
 {
-	const TTuple<FName, FNetTokenDataStore*>* Entry = TokenDataStores.FindByPredicate([Name](const TTuple<FName, FNetTokenDataStore*>& Entry){ return Name == Entry.Get<0>(); });
-	return Entry ? Entry->Get<1>() : nullptr;
+	const TTuple<FName, TUniquePtr<FNetTokenDataStore>>* Entry = TokenDataStores.FindByPredicate([Name](const TTuple<FName, TUniquePtr<FNetTokenDataStore>>& Entry){ return Name == Entry.Get<0>(); });
+	return Entry ? Entry->Get<1>().Get() : nullptr;
 }
 
  FNetTokenDataStore* FNetTokenStore::GetDataStore(FName Name)
 {
-	TTuple<FName, FNetTokenDataStore*>* Entry = TokenDataStores.FindByPredicate([Name](TTuple<FName, FNetTokenDataStore*>& Entry){ return Name == Entry.Get<0>(); });
-	return Entry ? Entry->Get<1>() : nullptr;
+	TTuple<FName, TUniquePtr<FNetTokenDataStore>>* Entry = TokenDataStores.FindByPredicate([Name](TTuple<FName, TUniquePtr<FNetTokenDataStore>>& Entry){ return Name == Entry.Get<0>(); });
+	return Entry ? Entry->Get<1>().Get() : nullptr;
 }
 
-bool FNetTokenStore::RegisterDataStore(FNetTokenDataStore* DataStore, FName TokenStoreName)
+bool FNetTokenStore::RegisterDataStore(TUniquePtr<FNetTokenDataStore> DataStore, FName TokenStoreName)
 {
 	if (TokenDataStores.Num() >= FNetToken::MaxTypeIdCount)
 	{
@@ -132,7 +207,7 @@ bool FNetTokenStore::RegisterDataStore(FNetTokenDataStore* DataStore, FName Toke
 	}
 
 	DataStore->TypeId = TokenDataStores.Num();
-	TokenDataStores.Emplace(TokenStoreName, DataStore);
+	TokenDataStores.Emplace(TokenStoreName, MoveTemp(DataStore));
 
 	return true;
 }
@@ -148,7 +223,22 @@ void FNetTokenStore::WriteTokenData(FNetSerializationContext& Context, const FNe
 		Writer->WriteBits(TokenKey.TypeId, FNetToken::TokenTypeIdBits);
 
 		// Write token data
-		TokenDataStores[TokenKey.TypeId].Get<1>()->WriteTokenData(Context, TokenKey);
+		TokenDataStores[TokenKey.TypeId].Get<1>().Get()->WriteTokenData(Context, TokenKey);
+	}
+}
+
+void FNetTokenStore::WriteTokenData(FArchive& Ar, const FNetToken NetToken) const
+{
+	if (NetToken.IsValid())
+	{
+		// Write type
+		const FNetTokenStoreKey TokenKey = LocalNetTokenStoreState->TokenInfos[NetToken.GetIndex()];
+		uint32 TypeId = TokenKey.TypeId;
+
+		Ar.SerializeBits(&TypeId, FNetToken::TokenTypeIdBits);
+
+		// Write token data
+		TokenDataStores[TokenKey.TypeId].Get<1>()->WriteTokenData(Ar, TokenKey);
 	}
 }
 
@@ -198,6 +288,52 @@ void FNetTokenStore::ReadTokenData(FNetSerializationContext& Context, const FNet
 	}
 }
 
+void FNetTokenStore::ReadTokenData(FArchive& Ar, const FNetToken NetToken, FNetTokenStoreState& RemoteNetTokenStoreState)
+{
+	if (NetToken.IsValid())
+	{
+		// TODO: Guard this better, a map might be better after-all!
+		// Also need to protect this if we process inbound data in parallel.
+		if (!RemoteNetTokenStoreState.ReserveTokenCount(NetToken.GetIndex() + 1))
+		{
+			Ar.SetError();
+			return;
+		}
+
+		// Read type
+		FNetToken::FTypeId TokenTypeId = 0U;
+		Ar.SerializeBits(&TokenTypeId, FNetToken::TokenTypeIdBits);
+
+		// Validate that we managed to read and verify the type
+		if (Ar.IsError() || TokenTypeId >= (uint32)TokenDataStores.Num() )
+		{
+			Ar.SetError();
+			return;
+		}
+
+		const FNetTokenStoreKey StoreKey = TokenDataStores[TokenTypeId].Get<1>()->ReadTokenData(Ar, NetToken);
+
+		// Validate that we managed to read the data for the key
+		if (Ar.IsError() || !StoreKey.IsValid())
+		{
+			Ar.SetError();
+			return;
+		}
+
+		// Since the same tokendata might be exported multiple times we better validate that it is the same data
+		const FNetTokenStoreKey& ExistingStoreKey = RemoteNetTokenStoreState.TokenInfos[NetToken.GetIndex()];
+		if (!ensureAlways(!ExistingStoreKey.IsValid() || StoreKey == ExistingStoreKey))
+		{
+			Ar.SetError();
+			return;
+		}
+	
+		// Store
+		RemoteNetTokenStoreState.TokenInfos[NetToken.GetIndex()] = StoreKey;
+	}
+}
+
+
 void FNetTokenStore::ConditionalWriteNetTokenData(FNetSerializationContext& Context, Private::FNetExportContext* ExportContext, const FNetToken NetToken) const
 {
 	FNetBitStreamWriter* Writer = Context.GetBitStreamWriter();
diff --git a/Engine/Source/Runtime/Experimental/Iris/Core/Private/Iris/ReplicationSystem/NetTokenStoreState.h b/Engine/Source/Runtime/Experimental/Iris/Core/Private/Iris/ReplicationSystem/NetTokenStoreState.h
--- a/Engine/Source/Runtime/Experimental/Iris/Core/Private/Iris/ReplicationSystem/NetTokenStoreState.h
+++ b/Engine/Source/Runtime/Experimental/Iris/Core/Private/Iris/ReplicationSystem/NetTokenStoreState.h
@@ -2,7 +2,7 @@
 
 #pragma once
 
-#include "Iris/ReplicationSystem/NetToken.h"
+#include "Net/Core/NetToken/NetToken.h"
 
 namespace UE::Net
 {
diff --git a/Engine/Source/Runtime/Experimental/Iris/Core/Private/Iris/ReplicationSystem/ObjectReferenceCache.cpp b/Engine/Source/Runtime/Experimental/Iris/Core/Private/Iris/ReplicationSystem/ObjectReferenceCache.cpp
--- a/Engine/Source/Runtime/Experimental/Iris/Core/Private/Iris/ReplicationSystem/ObjectReferenceCache.cpp
+++ b/Engine/Source/Runtime/Experimental/Iris/Core/Private/Iris/ReplicationSystem/ObjectReferenceCache.cpp
@@ -93,7 +93,7 @@ void FObjectReferenceCache::Init(UReplicationSystem* InReplicationSystem)
 {
 	ReplicationSystem = InReplicationSystem;
 	ReplicationBridge = InReplicationSystem->GetReplicationBridgeAs<UObjectReplicationBridge>();
-	NetTokenStore =  &InReplicationSystem->GetReplicationSystemInternal()->GetNetTokenStore();
+	NetTokenStore =  InReplicationSystem->GetNetTokenStore();
 	StringTokenStore = NetTokenStore->GetDataStore<FStringTokenStore>();
 	NetRefHandleManager = &InReplicationSystem->GetReplicationSystemInternal()->GetNetRefHandleManager();
 	bIsAuthority = InReplicationSystem->IsServer();
diff --git a/Engine/Source/Runtime/Experimental/Iris/Core/Private/Iris/ReplicationSystem/ObjectReplicationBridgeDebugging.cpp b/Engine/Source/Runtime/Experimental/Iris/Core/Private/Iris/ReplicationSystem/ObjectReplicationBridgeDebugging.cpp
--- a/Engine/Source/Runtime/Experimental/Iris/Core/Private/Iris/ReplicationSystem/ObjectReplicationBridgeDebugging.cpp
+++ b/Engine/Source/Runtime/Experimental/Iris/Core/Private/Iris/ReplicationSystem/ObjectReplicationBridgeDebugging.cpp
@@ -136,16 +136,12 @@ void PrintDefaultNetObjectState(UReplicationSystem* ReplicationSystem, uint32 Co
 {
 	FReplicationSystemInternal* ReplicationSystemInternal = ReplicationSystem->GetReplicationSystemInternal();
 
-	// In order to be able to output object references we need the TokenStoreState, for the server we just use the local one but if we are a client we must use the remote token store state
-	FReplicationConnections& Connections = ReplicationSystemInternal->GetConnections();
-	FNetTokenStoreState* TokenStoreState = ReplicationSystem->IsServer() ? ReplicationSystemInternal->GetNetTokenStore().GetLocalNetTokenStoreState() : &Connections.GetRemoteNetTokenStoreState(ConnectionId);
-
 	// Setup Context
 	FInternalNetSerializationContext InternalContext;
 	FInternalNetSerializationContext::FInitParameters InternalContextInitParams;
 	InternalContextInitParams.ReplicationSystem = ReplicationSystem;
 	InternalContextInitParams.PackageMap = ReplicationSystemInternal->GetIrisObjectReferencePackageMap();
-	InternalContextInitParams.ObjectResolveContext.RemoteNetTokenStoreState = TokenStoreState;
+	InternalContextInitParams.ObjectResolveContext.RemoteNetTokenStoreState = ReplicationSystem->GetNetTokenStore()->GetRemoteNetTokenStoreState(ConnectionId);
 	InternalContextInitParams.ObjectResolveContext.ConnectionId = ConnectionId;
 	InternalContext.Init(InternalContextInitParams);
 
diff --git a/Engine/Source/Runtime/Experimental/Iris/Core/Private/Iris/ReplicationSystem/ReplicationConnections.cpp b/Engine/Source/Runtime/Experimental/Iris/Core/Private/Iris/ReplicationSystem/ReplicationConnections.cpp
--- a/Engine/Source/Runtime/Experimental/Iris/Core/Private/Iris/ReplicationSystem/ReplicationConnections.cpp
+++ b/Engine/Source/Runtime/Experimental/Iris/Core/Private/Iris/ReplicationSystem/ReplicationConnections.cpp
@@ -30,7 +30,6 @@ void FReplicationConnections::InitDataStreams(uint32 ReplicationSystemId, uint32
 		UNetTokenDataStream::FInitParameters Params;
 		Params.ConnectionId = ConnectionId;
 		Params.ReplicationSystemId = ReplicationSystemId;
-		Params.RemoteTokenStoreState = &RemoteNetTokenStoreStates[ConnectionId];
 		Params.NetExports = &DataStreamManager->GetNetExports();
 
 		NetTokenDataStream->Init(Params);
@@ -38,7 +37,6 @@ void FReplicationConnections::InitDataStreams(uint32 ReplicationSystemId, uint32
 
 	// Init ReplicationDataStream
 	Connection->ReplicationWriter->SetNetExports(DataStreamManager->GetNetExports());
-	Connection->ReplicationReader->SetRemoteNetTokenStoreState(&RemoteNetTokenStoreStates[ConnectionId]);
 
 	UReplicationDataStream* ReplicationDataStream = StaticCast<UReplicationDataStream*>(DataStreamManager->GetStream(FName("Replication")));
 	if (ReplicationDataStream != nullptr) 
@@ -58,7 +56,6 @@ void FReplicationConnections::RemoveConnection(uint32 ConnectionId)
 {
 	check(ValidConnections.GetBit(ConnectionId));
 	SetReplicationView(ConnectionId, FReplicationView());
-	RemoteNetTokenStoreStates[ConnectionId] = FNetTokenStoreState();
 	DestroyReplicationReaderAndWriter(ConnectionId);
 
 	Connections[ConnectionId] = FReplicationConnection();
diff --git a/Engine/Source/Runtime/Experimental/Iris/Core/Private/Iris/ReplicationSystem/ReplicationConnections.h b/Engine/Source/Runtime/Experimental/Iris/Core/Private/Iris/ReplicationSystem/ReplicationConnections.h
--- a/Engine/Source/Runtime/Experimental/Iris/Core/Private/Iris/ReplicationSystem/ReplicationConnections.h
+++ b/Engine/Source/Runtime/Experimental/Iris/Core/Private/Iris/ReplicationSystem/ReplicationConnections.h
@@ -6,7 +6,6 @@
 #include "Containers/Array.h"
 #include "Net/Core/NetBitArray.h"
 #include "Iris/ReplicationSystem/ReplicationView.h"
-#include "Iris/ReplicationSystem/NetTokenStoreState.h"
 #include "UObject/ObjectPtr.h"
 #include "UObject/WeakObjectPtr.h"
 
@@ -37,7 +36,6 @@ public:
 	{
 		Connections.SetNumZeroed(MaxConnections);
 		ReplicationViews.SetNum(MaxConnections);
-		RemoteNetTokenStoreStates.SetNum(MaxConnections);
 	}
 
 	void Deinit();
@@ -92,9 +90,6 @@ public:
 	void SetReplicationView(uint32 ConnectionId, const FReplicationView& ViewInfo);
 	const FReplicationView& GetReplicationView(uint32 ConnectionId) const { return ReplicationViews[ConnectionId]; }
 
-	const FNetTokenStoreState& GetRemoteNetTokenStoreState(uint32 ConnectionId) const { return RemoteNetTokenStoreStates[ConnectionId]; }
-	FNetTokenStoreState& GetRemoteNetTokenStoreState(uint32 ConnectionId) { return RemoteNetTokenStoreStates[ConnectionId]; }
-
 	// Flag a connection as being in a graceful-close state meant to flush pending reliable data.
 	void SetConnectionIsClosing(uint32 ConnectionId)
 	{
@@ -108,7 +103,6 @@ private:
 private:
 	TArray<FReplicationConnection> Connections;
 	TArray<FReplicationView> ReplicationViews;
-	TArray<FNetTokenStoreState> RemoteNetTokenStoreStates;
 	FNetBitArray ValidConnections;
 };
 
diff --git a/Engine/Source/Runtime/Experimental/Iris/Core/Private/Iris/ReplicationSystem/ReplicationReader.cpp b/Engine/Source/Runtime/Experimental/Iris/Core/Private/Iris/ReplicationSystem/ReplicationReader.cpp
--- a/Engine/Source/Runtime/Experimental/Iris/Core/Private/Iris/ReplicationSystem/ReplicationReader.cpp
+++ b/Engine/Source/Runtime/Experimental/Iris/Core/Private/Iris/ReplicationSystem/ReplicationReader.cpp
@@ -279,8 +279,6 @@ void FReplicationReader::Init(const FReplicationParameters& InParameters)
 	// Store copy of parameters
 	Parameters = InParameters;
 
-	ResolveContext.ConnectionId = InParameters.ConnectionId;
-
 	// Cache internal systems
 	ReplicationSystemInternal = Parameters.ReplicationSystem->GetReplicationSystemInternal();
 	NetRefHandleManager = &ReplicationSystemInternal->GetNetRefHandleManager();
@@ -289,6 +287,10 @@ void FReplicationReader::Init(const FReplicationParameters& InParameters)
 	ObjectReferenceCache = &ReplicationSystemInternal->GetObjectReferenceCache();
 	ReplicationBridge = Parameters.ReplicationSystem->GetReplicationBridge();
 
+	// Init resolve context
+	ResolveContext.ConnectionId = InParameters.ConnectionId;
+	ResolveContext.RemoteNetTokenStoreState = Parameters.ReplicationSystem->GetNetTokenStore()->GetRemoteNetTokenStoreState(InParameters.ConnectionId);
+
 	// Find out if there's a PartialNetObjectAttachmentHandler so we can re-assemble split blobs
 	if (const UPartialNetObjectAttachmentHandler* Handler = ReplicationSystemInternal->GetNetBlobManager().GetPartialNetObjectAttachmentHandler())
 	{
@@ -2660,9 +2662,4 @@ void FReplicationReader::ResolveAndDispatchAttachments(FNetSerializationContext&
 	ReplicationInfo->bHasAttachments = bHasUnresolvedReferences;
 }
 
-void FReplicationReader::SetRemoteNetTokenStoreState(FNetTokenStoreState* InRemoteTokenStoreState)
-{
-	ResolveContext.RemoteNetTokenStoreState = InRemoteTokenStoreState;
-}
-
 } // end namespace UE::Net::Private
diff --git a/Engine/Source/Runtime/Experimental/Iris/Core/Private/Iris/ReplicationSystem/ReplicationReader.h b/Engine/Source/Runtime/Experimental/Iris/Core/Private/Iris/ReplicationSystem/ReplicationReader.h
--- a/Engine/Source/Runtime/Experimental/Iris/Core/Private/Iris/ReplicationSystem/ReplicationReader.h
+++ b/Engine/Source/Runtime/Experimental/Iris/Core/Private/Iris/ReplicationSystem/ReplicationReader.h
@@ -55,8 +55,6 @@ public:
 
 	// Read incoming replication data
 	void Read(FNetSerializationContext& Context);
-
-	void SetRemoteNetTokenStoreState(FNetTokenStoreState* RemoteTokenStoreState);
 	
 	// Mark objects pending destroy as unresolvable.
 	void UpdateUnresolvableReferenceTracking();
diff --git a/Engine/Source/Runtime/Experimental/Iris/Core/Private/Iris/ReplicationSystem/ReplicationSystem.cpp b/Engine/Source/Runtime/Experimental/Iris/Core/Private/Iris/ReplicationSystem/ReplicationSystem.cpp
--- a/Engine/Source/Runtime/Experimental/Iris/Core/Private/Iris/ReplicationSystem/ReplicationSystem.cpp
+++ b/Engine/Source/Runtime/Experimental/Iris/Core/Private/Iris/ReplicationSystem/ReplicationSystem.cpp
@@ -11,6 +11,7 @@
 #include "Net/Core/Trace/NetTrace.h"
 #include "Net/Core/Misc/NetConditionGroupManager.h"
 #include "Net/Core/Connection/NetEnums.h"
+#include "Net/Core/NetToken/NetToken.h"
 
 #include "Iris/ReplicationState/ReplicationStateUtil.h"
 #include "Iris/ReplicationSystem/ChangeMaskCache.h"
@@ -56,6 +57,7 @@ class FReplicationSystemImpl
 public:
 	TMap<FObjectKey, ENetObjectAttachmentSendPolicyFlags> AttachmentSendPolicyFlags;
 	UReplicationSystem* ReplicationSystem;
+	FNetTokenStore* NetTokenStore;
 	FReplicationSystemInternal ReplicationSystemInternal;
 	uint64 IrisDebugHelperDummy = 0U;
 	FNetObjectGroupHandle NotReplicatedNetObjectGroupHandle;
@@ -69,6 +71,7 @@ public:
 
 	explicit FReplicationSystemImpl(UReplicationSystem* InReplicationSystem, const UReplicationSystem::FReplicationSystemParams& Params)
 	: ReplicationSystem(InReplicationSystem)
+	, NetTokenStore(Params.NetTokenStore)
 	, ReplicationSystemInternal(
 		FReplicationSystemInternalInitParams(
 		{ 
@@ -107,11 +110,15 @@ public:
 #endif
 		const uint32 ReplicationSystemId = ReplicationSystem->GetId();
 
-		FNetTokenStore& NetTokenStore = ReplicationSystemInternal.GetNetTokenStore();
+		// Verify that we got a NetTokenStore and that it is configured as we expect.
+		bool bHasValidNetTokenStore = ensureAlwaysMsgf(Params.NetTokenStore, TEXT("ReplicationSystem cannot be initialized without a valid NetTokenStore"));
+		bHasValidNetTokenStore = bHasValidNetTokenStore && ensureAlwaysMsgf(Params.NetTokenStore->GetDataStore<FStringTokenStore>(), TEXT("ReplicationSystem cannot be initialized without a StringTokenStore"));
+		bHasValidNetTokenStore = bHasValidNetTokenStore && ensureAlwaysMsgf(Params.NetTokenStore->GetDataStore<FNameTokenStore>(), TEXT("ReplicationSystem cannot be initialized without a NameTokenStore"));
+
+		if (!bHasValidNetTokenStore)
 		{
-			FNetTokenStore::FInitParams NetTokenStoreInitParams;
-			NetTokenStoreInitParams.Authority = Params.bIsServer ? FNetToken::ENetTokenAuthority::Authority : FNetToken::ENetTokenAuthority::None;
-			NetTokenStore.Init(NetTokenStoreInitParams);
+			LowLevelFatalError(TEXT("Cannot initialize ReplicationSystem with invalid NetTokenStore"));
+			return;
 		}
 
 		FNetRefHandleManager& NetRefHandleManager = ReplicationSystemInternal.GetNetRefHandleManager();
@@ -1025,18 +1032,24 @@ UNetObjectPrioritizer* UReplicationSystem::GetPrioritizer(const FName Prioritize
 
 const UE::Net::FNetTokenStore* UReplicationSystem::GetNetTokenStore() const
 {
-	return &Impl->ReplicationSystemInternal.GetNetTokenStore();
+	return Impl->NetTokenStore;
 }
 
 UE::Net::FNetTokenStore* UReplicationSystem::GetNetTokenStore()
 {
-	return &Impl->ReplicationSystemInternal.GetNetTokenStore();
+	return Impl->NetTokenStore;
 }
 
-void UReplicationSystem::InitNetTokenExportContext(UE::Net::FNetTokenExportContext& Context, uint32 ConnectionId)
+UE::Net::FNetTokenResolveContext UReplicationSystem::GetNetTokenResolveContext(uint32 ConnectionId) const
 {
-	Context.TokenStore = &Impl->ReplicationSystemInternal.GetNetTokenStore();
-	Context.RemoteState = &Impl->ReplicationSystemInternal.GetConnections().GetRemoteNetTokenStoreState(ConnectionId);
+	using namespace UE::Net;
+
+	UE::Net::FNetTokenResolveContext NetTokenResolveContext;
+
+	NetTokenResolveContext.NetTokenStore = Impl->NetTokenStore;
+	NetTokenResolveContext.RemoteNetTokenStoreState = Impl->NetTokenStore->GetRemoteNetTokenStoreState(ConnectionId);
+
+	return NetTokenResolveContext;
 }
 
 bool UReplicationSystem::RegisterNetBlobHandler(UNetBlobHandler* Handler)
diff --git a/Engine/Source/Runtime/Experimental/Iris/Core/Private/Iris/ReplicationSystem/ReplicationSystemInternal.h b/Engine/Source/Runtime/Experimental/Iris/Core/Private/Iris/ReplicationSystem/ReplicationSystemInternal.h
--- a/Engine/Source/Runtime/Experimental/Iris/Core/Private/Iris/ReplicationSystem/ReplicationSystemInternal.h
+++ b/Engine/Source/Runtime/Experimental/Iris/Core/Private/Iris/ReplicationSystem/ReplicationSystemInternal.h
@@ -20,7 +20,7 @@
 #include "Iris/ReplicationSystem/Prioritization/ReplicationPrioritization.h"
 #include "Iris/ReplicationSystem/ReplicationProtocolManager.h"
 #include "Iris/ReplicationSystem/NetBlob/NetBlobManager.h"
-#include "Iris/ReplicationSystem/NetTokenStore.h"
+#include "Net/Core/NetToken/NetToken.h"
 #include "Iris/ReplicationSystem/StringTokenStore.h"
 #include "Iris/ReplicationSystem/NameTokenStore.h"
 #include "Iris/ReplicationSystem/WorldLocations.h"
@@ -49,8 +49,6 @@ public:
 	, DirtyNetObjectTracker()
 	, ReplicationBridge(nullptr)
 	, IrisObjectReferencePackageMap(nullptr)
-	, StringTokenStore(NetTokenStore)
-	, NameTokenStore(NetTokenStore)
 	, Id(Params.ReplicationSystemId)
 	{}
 
@@ -93,14 +91,6 @@ public:
 	FNetBlobHandlerManager& GetNetBlobHandlerManager() { return NetBlobManager.GetNetBlobHandlerManager(); }
 	const FNetBlobHandlerManager& GetNetBlobHandlerManager() const { return NetBlobManager.GetNetBlobHandlerManager(); }
 
-	const FStringTokenStore& GetStringTokenStore() const { return StringTokenStore; }
-	FStringTokenStore& GetStringTokenStore() { return StringTokenStore; }
-
-	const FNameTokenStore& GetNameTokenStore() const { return NameTokenStore; }
-	FNameTokenStore& GetNameTokenStore() { return NameTokenStore; }
-
-	FNetTokenStore& GetNetTokenStore() { return NetTokenStore; }
-
 	FNetCullDistanceOverrides& GetNetCullDistanceOverrides() { return NetCullDistanceOverrides; }
 
 	FWorldLocations& GetWorldLocations() { return WorldLocations; }
@@ -143,9 +133,6 @@ private:
 	FReplicationPrioritization Prioritization;
 	FObjectReferenceCache ObjectReferenceCache;
 	FNetBlobManager NetBlobManager;
-	FNetTokenStore NetTokenStore;
-	FStringTokenStore StringTokenStore;
-	FNameTokenStore NameTokenStore;
 	FNetCullDistanceOverrides NetCullDistanceOverrides;
 	FWorldLocations WorldLocations;
 	FDeltaCompressionBaselineManager DeltaCompressionBaselineManager;
diff --git a/Engine/Source/Runtime/Experimental/Iris/Core/Private/Iris/ReplicationSystem/StringTokenStore.cpp b/Engine/Source/Runtime/Experimental/Iris/Core/Private/Iris/ReplicationSystem/StringTokenStore.cpp
--- a/Engine/Source/Runtime/Experimental/Iris/Core/Private/Iris/ReplicationSystem/StringTokenStore.cpp
+++ b/Engine/Source/Runtime/Experimental/Iris/Core/Private/Iris/ReplicationSystem/StringTokenStore.cpp
@@ -15,12 +15,12 @@
 #endif 
 
 #if UE_NET_ENABLE_STRING_TOKEN_LOG
-#	define UE_LOG_STRINGTOKEN(Format, ...)  UE_LOG(LogIris, Verbose, Format, ##__VA_ARGS__)
+#	define UE_LOG_STRINGTOKEN(Format, ...)  UE_LOG(LogNetToken, Verbose, Format, ##__VA_ARGS__)
 #else
 #	define UE_LOG_STRINGTOKEN(...)
 #endif
 
-#define UE_LOG_STRINGTOKEN_WARNING(Format, ...)  UE_LOG(LogIris, Warning, Format, ##__VA_ARGS__)
+#define UE_LOG_STRINGTOKEN_WARNING(Format, ...)  UE_LOG(LogNetToken, Warning, Format, ##__VA_ARGS__)
 
 namespace UE::Net
 {
@@ -92,9 +92,6 @@ FStringTokenStore::FStringTokenStore(FNetTokenStore& InTokenStore)
 	// Reserve 0
 	StoredStrings.Add(nullptr);
 	StoredTokens.Add(FNetToken());
-
-	// Register
-	TokenStore.RegisterDataStore(this, GetTokenStoreName());
 }
 
 const TCHAR* FStringTokenStore::ResolveToken(FNetToken Token, const FNetTokenStoreState* NetTokenStoreState) const
@@ -118,6 +115,12 @@ void FStringTokenStore::WriteTokenData(FNetSerializationContext& Context, FNetTo
 	WriteString(Context.GetBitStreamWriter(), FStringView(StoredStrings[TokenStoreKey.GetKeyValue()]));
 }
 
+void FStringTokenStore::WriteTokenData(FArchive& Ar, FNetTokenStoreKey TokenStoreKey) const
+{
+	FString Temp(StoredStrings[TokenStoreKey.GetKeyValue()]);
+	Ar << Temp;
+}
+
 FNetTokenStoreKey FStringTokenStore::ReadTokenData(FNetSerializationContext& Context, const FNetToken& NetToken)
 {
 	FNetBitStreamReader* Reader = Context.GetBitStreamReader();
@@ -146,4 +149,31 @@ FNetTokenStoreKey FStringTokenStore::ReadTokenData(FNetSerializationContext& Con
 	}
 }
 
+FNetTokenStoreKey FStringTokenStore::ReadTokenData(FArchive& Ar, const FNetToken& NetToken)
+{
+	// Read the token data and add it to the string store without assigning LocalToken
+	FString Temp;
+	Ar << Temp;
+
+	if (!Ar.IsError())
+	{
+		FNetTokenStoreKey Key = GetOrCreatePersistentString(*Temp, Temp.Len());
+
+		// If this is an authtoken and we are not the authority, update the stored key
+		if (Key.IsValid() && NetToken.IsAssignedByAuthority() && !TokenStore.IsAuthority())
+		{
+			UE_LOG_STRINGTOKEN(TEXT("FStringTokenStore::ReadTokenData - Replaced local key %s with %s for string %s"), *StoredTokens[Key.GetKeyValue()].ToString(), *NetToken.ToString(), *Temp);	
+
+			// This way the next time we lookup this key during assignment use the imported authoritative key which we do not have to export
+			StoredTokens[Key.GetKeyValue()] = NetToken;
+		}
+		return Key;
+	}
+	else
+	{
+		return FNetTokenStoreKey();
+	}
+}
+
+
 }
diff --git a/Engine/Source/Runtime/Experimental/Iris/Core/Private/Iris/Serialization/IrisObjectReferencePackageMap.cpp b/Engine/Source/Runtime/Experimental/Iris/Core/Private/Iris/Serialization/IrisObjectReferencePackageMap.cpp
--- a/Engine/Source/Runtime/Experimental/Iris/Core/Private/Iris/Serialization/IrisObjectReferencePackageMap.cpp
+++ b/Engine/Source/Runtime/Experimental/Iris/Core/Private/Iris/Serialization/IrisObjectReferencePackageMap.cpp
@@ -114,59 +114,10 @@ bool UIrisObjectReferencePackageMap::SerializeName(FArchive& Ar, FName& InName)
 	return true;
 }
 
-bool UIrisObjectReferencePackageMap::SerializeNetToken(FArchive& Ar, UE::Net::FNetToken& InNetToken)
-{
-	if (!PackageMapExports)
-	{
-		return false;
-	}
-
-	UE::Net::FIrisPackageMapExports::FNetTokensArray* NetTokens = &PackageMapExports->NetTokens;
-	constexpr uint8 MaxNumNetTokens = 255U;
-
-	if (Ar.IsSaving())
-	{
-		int32 Index = MaxNumNetTokens;
-		if (!NetTokens->Find(InNetToken, Index) && NetTokens->Num() < MaxNumNetTokens)
-		{
-			Index = NetTokens->Add(InNetToken);
-		}
-
-		if (NetTokens->IsValidIndex(Index))
-		{
-			uint8 IndexByte = static_cast<uint8>(Index);
-			Ar << IndexByte;
-		}
-		else
-		{
-			ensureMsgf(false, TEXT("UIrisObjectReferencePackageMap::SerializeNetToken, failed to serialize NetToken with Index %u (%s). A Maximum of %u nettokens are currently supported by this PackageMap"),
-				Index, *InNetToken.ToString(), MaxNumNetTokens);
-			uint8 IndexByte = MaxNumNetTokens;
-			Ar << IndexByte;
-			return false;
-		}
-	}
-	else
-	{
-		uint8 IndexByte = MaxNumNetTokens;
-		Ar << IndexByte;
-		if (NetTokens->IsValidIndex(IndexByte) && IndexByte < MaxNumNetTokens)
-		{
-			InNetToken = (*NetTokens)[IndexByte];
-		}
-		else
-		{
-			ensureMsgf(false, TEXT("UIrisObjectReferencePackageMap::SerializeName, failed to read name index %u is out of bounds. Current Name num: %u"), IndexByte, NetTokens->Num());
-			return false;
-		}
-	}
-
-	return true;
-}
-
-void UIrisObjectReferencePackageMap::InitForRead(const UE::Net::FIrisPackageMapExports* InPackageMapExports)
+void UIrisObjectReferencePackageMap::InitForRead(const UE::Net::FIrisPackageMapExports* InPackageMapExports, const UE::Net::FNetTokenResolveContext& InNetTokenResolveContext)
 { 
 	PackageMapExports = const_cast<UE::Net::FIrisPackageMapExports*>(InPackageMapExports);
+	NetTokenResolveContext = InNetTokenResolveContext;
 }
 
 void UIrisObjectReferencePackageMap::InitForWrite(UE::Net::FIrisPackageMapExports* InPackageMapExports)
diff --git a/Engine/Source/Runtime/Experimental/Iris/Core/Private/Iris/Serialization/IrisPackageMapExportUtil.cpp b/Engine/Source/Runtime/Experimental/Iris/Core/Private/Iris/Serialization/IrisPackageMapExportUtil.cpp
--- a/Engine/Source/Runtime/Experimental/Iris/Core/Private/Iris/Serialization/IrisPackageMapExportUtil.cpp
+++ b/Engine/Source/Runtime/Experimental/Iris/Core/Private/Iris/Serialization/IrisPackageMapExportUtil.cpp
@@ -10,8 +10,9 @@
 #include "Iris/Serialization/NetReferenceCollector.h"
 #include "Iris/Serialization/NetSerializerArrayStorage.h"
 #include "Iris/Serialization/IrisObjectReferencePackageMap.h"
-#include "Iris/ReplicationSystem/NetTokenStore.h"
+#include "Net/Core/NetToken/NetToken.h"
 #include "Net/Core/Trace/NetTrace.h"
+#include "NetExportContext.h"
 
 namespace UE::Net
 {
@@ -62,27 +63,12 @@ void FIrisPackageMapExportsUtil::Serialize(FNetSerializationContext& Context, co
 		}
 	}
 
-	// For NetTokens
+	// We now serialize NetTokens directly in the data but we still need to append exports.
+	if (UE::Net::Private::FNetExportContext* ExportContext = Context.GetExportContext())
 	{
-		UE_NET_TRACE_SCOPE(NetTokens, *Writer, Context.GetTraceCollector(), ENetTraceVerbosity::VeryVerbose);
-		// $TODO: For asymmetrical data like nettokens we will need to do something special to support validating the default state hash as the local quantized state will differ
-		// For now we ignore this in default state hash
-		if (Context.IsInitializingDefaultState())
+		for (const FNetToken& NetToken : MakeArrayView(Value.NetTokenStorage.GetData(), Value.NetTokenStorage.Num()))
 		{
-			return;
-		}
-
-		const uint32 NumNetTokens = Value.NetTokenStorage.Num();
-		if (Writer->WriteBool(NumNetTokens != 0))
-		{
-			UE::Net::WritePackedUint32(Writer, NumNetTokens);
-			for (const FNetToken& NetToken : MakeArrayView(Value.NetTokenStorage.GetData(), NumNetTokens))
-			{
-				// Always write the token
-				WriteNetToken(Context, NetToken);
-				// Export or add to pending exports for later export
-				FNetTokenStore::AppendExportOrWriteInlinedExportData(Context, NetToken);
-			}
+			ExportContext->AddPendingExport(NetToken);
 		}
 	}
 }
@@ -157,47 +143,10 @@ void FIrisPackageMapExportsUtil::Deserialize(FNetSerializationContext& Context,
 		}
 	}
 
-	// Read any exported nettokens
-	{
-		UE_NET_TRACE_SCOPE(NetTokens, *Reader, Context.GetTraceCollector(), ENetTraceVerbosity::VeryVerbose);
-		const bool bHasNetTokens = Reader->ReadBool();
-		if (bHasNetTokens)
-		{
-			const uint32 NumNetTokens = UE::Net::ReadPackedUint32(Reader);
-
-			if (NumNetTokens > MaxExports)
-			{
-				UE_LOG(LogIris, Error, TEXT("FIrisPackageMapExportsUtil::Received too many NetToken exports %u > max:%u"), NumNetTokens, MaxExports);
-				Context.SetError(GNetError_ArraySizeTooLarge);
-				return;
-			}
-
-			Value.NetTokenStorage.AdjustSize(Context, NumNetTokens);
-
-			for (FNetToken& TargetToken : MakeArrayView(Value.NetTokenStorage.GetData(), Value.NetTokenStorage.Num()))
-			{
-				// Always Read the token
-				FNetToken NetToken = ReadNetToken(Context);
-
-				if (Reader->IsOverflown())
-				{
-					return;
-				}
-
-				// Read inlined exports if there are any
-				FNetTokenStore::ReadInlinedExportData(Context, NetToken);
-
-				TargetToken = NetToken;
-			}
-		}
-		else
-		{
-			Value.NetTokenStorage.Free(Context);
-		}
-	}
+	Value.NetTokenStorage.Free(Context);
 }
 
-void FIrisPackageMapExportsUtil::Quantize(FNetSerializationContext& Context, const UE::Net::FIrisPackageMapExports& PackageMapExports, QuantizedType& Value)
+void FIrisPackageMapExportsUtil::Quantize(FNetSerializationContext& Context, const UE::Net::FIrisPackageMapExports& PackageMapExports, TArrayView<const UE::Net::FNetToken> NetTokensPendingExport, QuantizedType& Value)
 {
 	// Quantize captured references
 	{
@@ -243,27 +192,20 @@ void FIrisPackageMapExportsUtil::Quantize(FNetSerializationContext& Context, con
 		}
 	}
 
-	// Just store captured nettokens
+	// Just store captured NetTokenExports they will be added as pending exports during serialization.
 	{
-		const FIrisPackageMapExports::FNetTokensArray& NetTokens = PackageMapExports.NetTokens;
-		const uint32 NumNetTokens = NetTokens.Num();
+		const uint32 NumNetTokens = NetTokensPendingExport.Num();
 		Value.NetTokenStorage.AdjustSize(Context, NumNetTokens);
-		if (NumNetTokens > 0)
+		FNetToken* TargetTokens = Value.NetTokenStorage.GetData();
+		for (uint32 Index = 0; Index < NumNetTokens; ++Index)
 		{
-			const FNetToken* SourceTokens = NetTokens.GetData();
-			FNetToken* TargetTokens = Value.NetTokenStorage.GetData();
-			for (uint32 ReferenceIndex = 0; ReferenceIndex < NumNetTokens; ++ReferenceIndex)
-			{
-				TargetTokens[ReferenceIndex] = SourceTokens[ReferenceIndex];
-			}
+			TargetTokens[Index] = NetTokensPendingExport[Index];
 		}
 	}
 }
 
 void FIrisPackageMapExportsUtil::Dequantize(FNetSerializationContext& Context, const QuantizedType& Source, UE::Net::FIrisPackageMapExports& PackageMapExports)
 {
-	Private::FInternalNetSerializationContext* InternalContext = Context.GetInternalContext();
-
 	// References
 	{
 		UE::Net::FIrisPackageMapExports::FObjectReferenceArray& ObjectReferences = PackageMapExports.References;
@@ -310,24 +252,6 @@ void FIrisPackageMapExportsUtil::Dequantize(FNetSerializationContext& Context, c
 			}
 		}
 	}
-
-	// NetTokens
-	{
-		UE::Net::FIrisPackageMapExports::FNetTokensArray& NetTokens = PackageMapExports.NetTokens;
-
-		const uint32 NumNetTokens = Source.NetTokenStorage.Num();
-		if (NumNetTokens > 0U)
-		{		
-			NetTokens.SetNumUninitialized(NumNetTokens);
-
-			const FNetToken* SourceTokens = Source.NetTokenStorage.GetData();
-			FNetToken* TargetTokens = NetTokens.GetData();
-			for (uint32 ReferenceIndex = 0; ReferenceIndex < NumNetTokens; ++ReferenceIndex)
-			{
-				TargetTokens[ReferenceIndex] = SourceTokens[ReferenceIndex];
-			}
-		}
-	}
 }
 
 bool FIrisPackageMapExportsUtil::IsEqual(FNetSerializationContext& Context, const QuantizedType& Value0, const QuantizedType& Value1)
@@ -347,13 +271,6 @@ bool FIrisPackageMapExportsUtil::IsEqual(FNetSerializationContext& Context, cons
 		return false;
 	}
 
-	// Note: this can be improved by implementing a resolved token compare. Currently this is pessimistic and does not allow compare between auth and non-auth tokens even if they would
-	// resolve to same value.
-	if (Value0.NetTokenStorage.Num() > 0 && FMemory::Memcmp(Value0.NetTokenStorage.GetData(), Value1.NetTokenStorage.GetData(), sizeof(FNetToken) * Value0.NetTokenStorage.Num()) != 0)
-	{
-		return false;
-	}
-
 	return true;
 }
 
@@ -383,7 +300,7 @@ void FIrisPackageMapExportsUtil::CollectNetReferences(FNetSerializationContext&
 
 bool FIrisPackageMapExportsUtil::Validate(FNetSerializationContext& Context, const QuantizedType& SourceValue)
 {
-	if ((SourceValue.ObjectReferenceStorage.Num() > MaxExports) || (SourceValue.NameStorage.Num() > MaxExports) || (SourceValue.NetTokenStorage.Num() > MaxExports))
+	if ((SourceValue.ObjectReferenceStorage.Num() > MaxExports) || (SourceValue.NameStorage.Num() > MaxExports))
 	{
 		return false;
 	}
diff --git a/Engine/Source/Runtime/Experimental/Iris/Core/Private/Iris/Serialization/LastResortPropertyNetSerializer.cpp b/Engine/Source/Runtime/Experimental/Iris/Core/Private/Iris/Serialization/LastResortPropertyNetSerializer.cpp
--- a/Engine/Source/Runtime/Experimental/Iris/Core/Private/Iris/Serialization/LastResortPropertyNetSerializer.cpp
+++ b/Engine/Source/Runtime/Experimental/Iris/Core/Private/Iris/Serialization/LastResortPropertyNetSerializer.cpp
@@ -12,6 +12,7 @@
 #include "Iris/Serialization/IrisObjectReferencePackageMap.h"
 #include "Iris/Serialization/IrisPackageMapExportUtil.h"
 #include "Iris/Serialization/NetExportContext.h"
+#include "Net/Core/NetToken/NetTokenExportContext.h"
 #include "Net/Core/Trace/NetTrace.h"
 
 namespace UE::Net
@@ -81,6 +82,12 @@ void FLastResortPropertyNetSerializer::Serialize(FNetSerializationContext& Conte
 	const QuantizedType& Value = *reinterpret_cast<const QuantizedType*>(Args.Source);
 	FNetBitStreamWriter* Writer = Context.GetBitStreamWriter();
 
+	// For now we ignore this in default state hash if it has exported NetTokens as they will differ
+	if (Context.IsInitializingDefaultState() && (Value.QuantizedExports.NetTokenStorage.Num() > 0U))
+	{
+		return;
+	}
+
 	UE_NET_TRACE_DYNAMIC_NAME_SCOPE(Config->Property.Get()->GetName(), *Writer, Context.GetTraceCollector(), ENetTraceVerbosity::VeryVerbose);
 
 	// If we have any captured exports, serialize them.
@@ -96,6 +103,12 @@ void FLastResortPropertyNetSerializer::Serialize(FNetSerializationContext& Conte
 
 void FLastResortPropertyNetSerializer::Deserialize(FNetSerializationContext& Context, const FNetDeserializeArgs& Args)
 {
+	// For consistency, we should never get here. For now we ignore this in default state hash due to complications with asymmetrically serialized state.
+	if (Context.IsInitializingDefaultState())
+	{
+		return;
+	}
+
 	const ConfigType* Config = static_cast<const ConfigType*>(Args.NetSerializerConfig);
 	QuantizedType& Value = *reinterpret_cast<QuantizedType*>(Args.Target);
 	const uint32 CurrentBitCount = Value.BitCount;
@@ -130,22 +143,18 @@ void FLastResortPropertyNetSerializer::Quantize(FNetSerializationContext& Contex
 	Private::FInternalNetSerializationContext* InternalContext = Context.GetInternalContext();
 	UIrisObjectReferencePackageMap* PackageMap = InternalContext ? InternalContext->PackageMap : nullptr;
 
+	// Since this struct uses custom serialization path we need to explicitly capture exports in order to forward them to iris
 	UE::Net::FIrisPackageMapExports PackageMapExports;
+	UE::Net::FNetTokenExportContext::FNetTokenExports NetTokensPendingExport;
 
 	if (PackageMap)
 	{
 		PackageMap->InitForWrite(&PackageMapExports);
-
-		// $IRIS: $TODO: TBD: how to propagate context
-		//{
-		//	UE::Net::FNetTokenExportContext* NetTokenExportContext = PackageMap->GetNetTokenExportContext();
-		//	NetTokenExportContext->TokenStore = Context.GetNetTokenStore();
-		//	NetTokenExportContext->RemoteState = Context.GetInternalContext()->ResolveContext.RemoteNetTokenStoreState;
-		//}
 	}
 
 	// Use the Property serialization and store as binary blob.
-	FNetBitWriter Archive(PackageMap, 8192);
+	FNetBitWriter Archive(PackageMap, 8192);	
+	FNetTokenExportScope NetTokenExportScope(Archive, Context.GetNetTokenStore(), NetTokensPendingExport);
 	Property->NetSerializeItem(Archive, PackageMap, reinterpret_cast<void*>(Args.Source));
 
 	const uint64 BitCount = Archive.GetNumBits();
@@ -156,7 +165,7 @@ void FLastResortPropertyNetSerializer::Quantize(FNetSerializationContext& Contex
 	}
 
 	// Quantize captured exports
-	FIrisPackageMapExportsUtil::Quantize(Context, PackageMapExports, Value.QuantizedExports);
+	FIrisPackageMapExportsUtil::Quantize(Context, PackageMapExports, NetTokensPendingExport, Value.QuantizedExports);
 
 	// Deal with serialized data
 	AdjustStorageSize(Context, Value, static_cast<uint16>(BitCount));
@@ -176,19 +185,16 @@ void FLastResortPropertyNetSerializer::Dequantize(FNetSerializationContext& Cont
 	// Dequantize and inject exports
 	Private::FInternalNetSerializationContext* InternalContext = Context.GetInternalContext();
 	UIrisObjectReferencePackageMap* PackageMap = InternalContext ? InternalContext->PackageMap : nullptr;
-
 	UE::Net::FIrisPackageMapExports PackageMapExports;
 
 	FIrisPackageMapExportsUtil::Dequantize(Context, Source.QuantizedExports, PackageMapExports);	
 
-	PackageMap->InitForRead(&PackageMapExports);
+	// Setup resolve context for call into NetSerialize
+	UE::Net::FNetTokenResolveContext ResolveContext;
+	ResolveContext.RemoteNetTokenStoreState = Context.GetRemoteNetTokenStoreState();
+	ResolveContext.NetTokenStore = Context.GetNetTokenStore();
 
-	// $IRIS: $TODO: TBD: how to propagate context
-	//{
-	//	UE::Net::FNetTokenExportContext* NetTokenExportContext = PackageMap->GetNetTokenExportContext();
-	//	NetTokenExportContext->TokenStore = Context.GetNetTokenStore();
-	//	NetTokenExportContext->RemoteState = Context.GetInternalContext()->ResolveContext.RemoteNetTokenStoreState;
-	//}
+	PackageMap->InitForRead(&PackageMapExports, ResolveContext);
 
 	// Read data
 	if (Source.BitCount)
diff --git a/Engine/Source/Runtime/Experimental/Iris/Core/Private/Iris/Serialization/NetExportContext.h b/Engine/Source/Runtime/Experimental/Iris/Core/Private/Iris/Serialization/NetExportContext.h
--- a/Engine/Source/Runtime/Experimental/Iris/Core/Private/Iris/Serialization/NetExportContext.h
+++ b/Engine/Source/Runtime/Experimental/Iris/Core/Private/Iris/Serialization/NetExportContext.h
@@ -6,7 +6,7 @@
 #include "Containers/Array.h"
 #include "Containers/Set.h"
 #include "Iris/ReplicationSystem/NetRefHandle.h"
-#include "Iris/ReplicationSystem/NetToken.h"
+#include "Net/Core/NetToken/NetToken.h"
 #include "Iris/Serialization/NetSerializationContext.h"
 #include "Iris/Core/NetObjectReference.h"
 
diff --git a/Engine/Source/Runtime/Experimental/Iris/Core/Private/Iris/Serialization/NetSerializationContext.cpp b/Engine/Source/Runtime/Experimental/Iris/Core/Private/Iris/Serialization/NetSerializationContext.cpp
--- a/Engine/Source/Runtime/Experimental/Iris/Core/Private/Iris/Serialization/NetSerializationContext.cpp
+++ b/Engine/Source/Runtime/Experimental/Iris/Core/Private/Iris/Serialization/NetSerializationContext.cpp
@@ -104,5 +104,10 @@ FNetTokenStore* FNetSerializationContext::GetNetTokenStore()
 	return ReplicationSystem ? ReplicationSystem->GetNetTokenStore() : nullptr;
 }
 
+const UE::Net::FNetTokenStoreState* FNetSerializationContext::GetRemoteNetTokenStoreState() const
+{
+	return InternalContext ? InternalContext->ResolveContext.RemoteNetTokenStoreState : nullptr;
+}
+
 
 }
diff --git a/Engine/Source/Runtime/Experimental/Iris/Core/Public/Iris/Core/NetObjectReference.h b/Engine/Source/Runtime/Experimental/Iris/Core/Public/Iris/Core/NetObjectReference.h
--- a/Engine/Source/Runtime/Experimental/Iris/Core/Public/Iris/Core/NetObjectReference.h
+++ b/Engine/Source/Runtime/Experimental/Iris/Core/Public/Iris/Core/NetObjectReference.h
@@ -4,7 +4,7 @@
 
 #include "CoreTypes.h"
 #include "Iris/ReplicationSystem/NetRefHandle.h"
-#include "Iris/ReplicationSystem/NetToken.h"
+#include "Net/Core/NetToken/NetToken.h"
 #include "Misc/EnumClassFlags.h"
 #include "Templates/IsPODType.h"
 
diff --git a/Engine/Source/Runtime/Experimental/Iris/Core/Public/Iris/ReplicationSystem/NameTokenStore.h b/Engine/Source/Runtime/Experimental/Iris/Core/Public/Iris/ReplicationSystem/NameTokenStore.h
--- a/Engine/Source/Runtime/Experimental/Iris/Core/Public/Iris/ReplicationSystem/NameTokenStore.h
+++ b/Engine/Source/Runtime/Experimental/Iris/Core/Public/Iris/ReplicationSystem/NameTokenStore.h
@@ -5,18 +5,15 @@
 #include "CoreTypes.h"
 #include "Containers/Map.h"
 #include "Iris/ReplicationSystem/NetTokenStore.h"
-#include "Misc/MemStack.h"
 
 namespace UE::Net
 {
 
-// Simple token store used to store string tokens
-// When the PackageMapRefactor is complete we will most likely rely on NetTagManager for persistent storage
 class FNameTokenStore : public FNetTokenDataStore
 {
 	UE_NONCOPYABLE(FNameTokenStore);
 public:
-	explicit FNameTokenStore(FNetTokenStore& TokenStore);
+	IRISCORE_API explicit FNameTokenStore(FNetTokenStore& TokenStore);
 
 	// Create a NetToken for the provided name
 	IRISCORE_API FNetToken GetOrCreateToken(FName Name);
@@ -28,12 +25,14 @@ public:
 
 protected:
 	// Serialize data for a token, note there is not validation in this function
-	virtual void WriteTokenData(FNetSerializationContext& Context, FNetTokenStoreKey TokenStoreKey) const override;
+	IRISCORE_API virtual void WriteTokenData(FNetSerializationContext& Context, FNetTokenStoreKey TokenStoreKey) const override;
+	IRISCORE_API virtual void WriteTokenData(FArchive& Archive, FNetTokenStoreKey TokenStoreKey) const override;
 
 	// Read data for a token, returns a valid StoreKey if successful read
-	virtual FNetTokenStoreKey ReadTokenData(FNetSerializationContext& Context, const FNetToken& NetToken) override;
+	IRISCORE_API virtual FNetTokenStoreKey ReadTokenData(FNetSerializationContext& Context, const FNetToken& NetToken) override;
+	IRISCORE_API virtual FNetTokenStoreKey ReadTokenData(FArchive& Archive, const FNetToken& NetToken) override;
 
-	// Create a persistent string
+	// Create and store data for Token
 	FNetTokenStoreKey GetOrCreateTokenStoreKey(FName Name);
 
 private:
diff --git a/Engine/Source/Runtime/Experimental/Iris/Core/Public/Iris/ReplicationSystem/NetBlob/ShrinkWrapNetBlob.h b/Engine/Source/Runtime/Experimental/Iris/Core/Public/Iris/ReplicationSystem/NetBlob/ShrinkWrapNetBlob.h
--- a/Engine/Source/Runtime/Experimental/Iris/Core/Public/Iris/ReplicationSystem/NetBlob/ShrinkWrapNetBlob.h
+++ b/Engine/Source/Runtime/Experimental/Iris/Core/Public/Iris/ReplicationSystem/NetBlob/ShrinkWrapNetBlob.h
@@ -3,7 +3,7 @@
 #pragma once
 #include "CoreTypes.h"
 #include "Iris/ReplicationSystem/NetBlob/NetBlob.h"
-#include "Iris/ReplicationSystem/NetToken.h"
+#include "Net/Core/NetToken/NetToken.h"
 
 namespace UE::Net
 {
diff --git a/Engine/Source/Runtime/Experimental/Iris/Core/Public/Iris/ReplicationSystem/NetToken.h b/Engine/Source/Runtime/Experimental/Iris/Core/Public/Iris/ReplicationSystem/NetToken.h
deleted file mode 100644
diff --git a/Engine/Source/Runtime/Experimental/Iris/Core/Public/Iris/ReplicationSystem/NetTokenStore.h b/Engine/Source/Runtime/Experimental/Iris/Core/Public/Iris/ReplicationSystem/NetTokenStore.h
--- a/Engine/Source/Runtime/Experimental/Iris/Core/Public/Iris/ReplicationSystem/NetTokenStore.h
+++ b/Engine/Source/Runtime/Experimental/Iris/Core/Public/Iris/ReplicationSystem/NetTokenStore.h
@@ -3,11 +3,33 @@
 #pragma once
 
 #include "CoreTypes.h"
-#include "NetToken.h"
+#include "Net/Core/NetToken/NetToken.h"
 #include "UObject/NameTypes.h"
 #include "Containers/Array.h"
 #include "Templates/Tuple.h"
 
+/** 
+* The idea with NetTokens is to allow export of "stable"-pieces of data such as string and names by replacing them with a NetToken during quantization/serialization.
+* The data associated with a NetToken is then exported separately from the data. As soon as the exported data has been acknowledged the data will only be serialized using the NetToken.
+* 
+* As both Client and Server can communicate NetTokens, each side can end up assigning different tokens that will differe from each other.
+* Here is a high-level overview of the algorithm which works slightly differently based on if using Iris replication or the old replication system.
+*
+* The sending side looks-up or creates a NetToken for the Data being serialized. Servers will mark assigned NetTokens as authoritive, while clients generate a temporary NetToken.
+* When a network bunch/batch that contains NetToken is being sent, there is a per-connection look-up to see if we need to append and serialize exports or not, if the data associated with the token
+* has been acknowledged the token will not be exported again.
+*
+* On the receiving side, imported exports are always guaranteed to have been processed before we attempt to read received data containing NetTokens which allows
+* the receiving side to resolve the NetToken to get the actual data.
+*
+* The implmentaion details differs a bit depending on if we are using iris replication or old style replication.
+*
+* A current example that us used by both systems is GameplayTags, For Iris: See GameplayTagNetSerializer.cpp, for the old replication system: See: GameplayTagContainer.cpp.
+*
+* Note: this is still an experimental feature and is currently only available if compiling UE_WITH_IRIS. 
+* Implementation datails is subject to change.
+*/
+
 class UNetTokenDataStream;
 namespace UE::Net
 {
@@ -29,22 +51,25 @@ namespace UE::Net
 class FNetTokenDataStore
 {
 public:
-	virtual ~FNetTokenDataStore();
+	IRISCORE_API virtual ~FNetTokenDataStore();
 
 protected:
-	FNetTokenDataStore();
+	IRISCORE_API FNetTokenDataStore();
 
 	virtual void WriteTokenData(FNetSerializationContext& Context, FNetTokenStoreKey Key) const = 0;
 	virtual FNetTokenStoreKey ReadTokenData(FNetSerializationContext& Context, const FNetToken& NetToken) = 0;
 
+	virtual void WriteTokenData(FArchive& Ar, FNetTokenStoreKey Key) const = 0;
+	virtual FNetTokenStoreKey ReadTokenData(FArchive& Ar, const FNetToken& NetToken) = 0;
+
 	FNetTokenStoreKey GetTokenKey(FNetToken Token, const FNetTokenStoreState& TokenStoreState) const;
 	inline FNetToken::FTypeId GetTypeId() const { return TypeId; }
 
 	// Create new NetToken
-	FNetToken CreateToken(FNetToken::ENetTokenAuthority Authority, FNetTokenStoreKey Key, FNetTokenStoreState& TokenStoreState);
+	IRISCORE_API FNetToken CreateToken(FNetToken::ENetTokenAuthority Authority, FNetTokenStoreKey Key, FNetTokenStoreState& TokenStoreState);
 
 	// Make NetTokenStoreKey
-	static FNetTokenStoreKey MakeNetTokenStoreKey(FNetToken::FTypeId TokenTypeId, uint32 TokenKeyValue);
+	IRISCORE_API static FNetTokenStoreKey MakeNetTokenStoreKey(FNetToken::FTypeId TokenTypeId, uint32 TokenKeyValue);
 
 private:
 	friend FNetTokenStore;
@@ -56,27 +81,37 @@ class FNetTokenStore
 {
 	UE_NONCOPYABLE(FNetTokenStore);
 public:
-	FNetTokenStore();
-	~FNetTokenStore();
+	IRISCORE_API FNetTokenStore();
+	IRISCORE_API ~FNetTokenStore();
 
 	/** External configuration variables used to initialize the NetTokenStore */
 	struct FInitParams
 	{
 		FNetToken::ENetTokenAuthority Authority;
+		uint32 MaxConnections = 256;
 	};
-	void Init(FInitParams& InitParams);
+	IRISCORE_API void Init(FInitParams& InitParams);
 
 	bool IsAuthority() const { return Params.Authority == FNetToken::ENetTokenAuthority::Authority; }
 
 	// A token is local if the authority of the NetTokenStore and the token matches, Invalid tokens are always local.
 	bool IsLocalToken(const FNetToken NetToken) const { return !NetToken.IsValid() || IsAuthority() == NetToken.IsAssignedByAuthority(); }
 
-	// Register DataStore and return true if it was registered.
-	bool RegisterDataStore(FNetTokenDataStore* DataStore, FName TokenStoreName);
+	// Register DataStore and return true if it was registered, 
+	IRISCORE_API bool RegisterDataStore(TUniquePtr<FNetTokenDataStore> DataStore, FName TokenStoreName);
 
 	IRISCORE_API const FNetTokenDataStore* GetDataStore(FName Name) const;
 	IRISCORE_API FNetTokenDataStore* GetDataStore(FName Name);
 
+	/** Create and return data store of specified type, it will be owned by NetTokenStore */
+	template<typename T>
+	bool CreateAndRegisterDataStore()
+	{
+		TUniquePtr<FNetTokenDataStore> NetTokenDataStore(MakeUnique<T>(*this));
+		
+		return RegisterDataStore(MakeUnique<T>(*this), T::GetTokenStoreName());
+	}
+
 	/** Return data store of specified type. */
 	template<typename T>
 	T* GetDataStore() { return static_cast<T*>(GetDataStore(T::GetTokenStoreName())); }
@@ -85,20 +120,35 @@ public:
 	template<typename T>
 	const T* GetDataStore() const { return static_cast<const T*>(GetDataStore(T::GetTokenStoreName())); }
 
-	const FNetTokenStoreState* GetLocalNetTokenStoreState() const { return LocalNetTokenStoreState; }
-	FNetTokenStoreState* GetLocalNetTokenStoreState() { return LocalNetTokenStoreState; }
+	const FNetTokenStoreState* GetLocalNetTokenStoreState() const { return LocalNetTokenStoreState.Get(); }
+	FNetTokenStoreState* GetLocalNetTokenStoreState() { return LocalNetTokenStoreState.Get();; }
+
+	/** Init RemoteNetTokenStoreState for given ConnectionId, if it already exists it will be reset. */
+	IRISCORE_API void InitRemoteNetTokenStoreState(uint32 ConnectionId);
+
+	/** Get RemoteNetTokenStoreState for given ConnectionId. */
+	IRISCORE_API const FNetTokenStoreState* GetRemoteNetTokenStoreState(uint32 ConnectionId) const;
+
+	/** Get RemoteNetTokenStoreState for given ConnectionId. */
+	IRISCORE_API FNetTokenStoreState* GetRemoteNetTokenStoreState(uint32 ConnectionId);
 
 	// Write data associated with the NetToken
-	void WriteTokenData(FNetSerializationContext& Context, const FNetToken NetToken) const;
+	IRISCORE_API void WriteTokenData(FNetSerializationContext& Context, const FNetToken NetToken) const;
 
 	// Read data associated with the NetToken
-	void ReadTokenData(FNetSerializationContext& Context, const FNetToken NetToken, FNetTokenStoreState& RemoteNetTokenStoreState);
+	IRISCORE_API void ReadTokenData(FNetSerializationContext& Context, const FNetToken NetToken, FNetTokenStoreState& RemoteNetTokenStoreState);
+
+	// Write data associated with the NetToken
+	IRISCORE_API void WriteTokenData(FArchive& Ar, const FNetToken NetToken) const;
+
+	// Read data associated with the NetToken
+	IRISCORE_API void ReadTokenData(FArchive& Ar, const FNetToken NetToken, FNetTokenStoreState& RemoteNetTokenStoreState);
 
 	// Conditionally write NetTokenData unless already exported
-	void ConditionalWriteNetTokenData(FNetSerializationContext& Context, Private::FNetExportContext* ExportContext, const FNetToken NetToken) const;
+	IRISCORE_API void ConditionalWriteNetTokenData(FNetSerializationContext& Context, Private::FNetExportContext* ExportContext, const FNetToken NetToken) const;
 
 	// Conditionally read NetTokenData if exported
-	void ConditionalReadNetTokenData(FNetSerializationContext& Context, const FNetToken NetToken);
+	IRISCORE_API void ConditionalReadNetTokenData(FNetSerializationContext& Context, const FNetToken NetToken);
 
 	// Utility methods, consolidate with other changes to NetTokenStore as next step.
 	IRISCORE_API static void AppendExportOrWriteInlinedExportData(FNetSerializationContext&, FNetToken NetToken);
@@ -108,8 +158,9 @@ private:
 	friend UNetTokenDataStream;
 	friend FNetTokenDataStore;
 
-	FNetTokenStoreState* LocalNetTokenStoreState;
-	TArray<TTuple<FName, FNetTokenDataStore*> > TokenDataStores;
+	TUniquePtr<FNetTokenStoreState> LocalNetTokenStoreState;
+	TArray<TUniquePtr<FNetTokenStoreState>> RemoteNetTokenStoreStates;
+	TArray<TTuple<FName, TUniquePtr<FNetTokenDataStore>>> TokenDataStores;
 	FInitParams Params;
 };
 
@@ -125,5 +176,11 @@ inline FNetTokenStoreKey FNetTokenDataStore::MakeNetTokenStoreKey(FNetToken::FTy
 	return TokenKey;
 }
 
+// Read and write tokens
+IRISCORE_API FNetToken ReadNetToken(UE::Net::FNetSerializationContext& Context);
+IRISCORE_API void WriteNetToken(UE::Net::FNetSerializationContext& Context, FNetToken Token);
+
+IRISCORE_API FNetToken ReadNetToken(FArchive& Ar);
+IRISCORE_API void WriteNetToken(FArchive& Ar, FNetToken Token);
 
 }
diff --git a/Engine/Source/Runtime/Experimental/Iris/Core/Public/Iris/ReplicationSystem/ReplicationSystem.h b/Engine/Source/Runtime/Experimental/Iris/Core/Public/Iris/ReplicationSystem/ReplicationSystem.h
--- a/Engine/Source/Runtime/Experimental/Iris/Core/Public/Iris/ReplicationSystem/ReplicationSystem.h
+++ b/Engine/Source/Runtime/Experimental/Iris/Core/Public/Iris/ReplicationSystem/ReplicationSystem.h
@@ -36,6 +36,7 @@ namespace UE::Net
 	class FStringTokenStore;
 	class FNameTokenStore;
 	class FWorldLocations;
+	class FNetTokenResolveContext;
 
 	struct FReplicationSystemUtil;
 	struct FReplicationView;
@@ -48,8 +49,6 @@ namespace UE::Net
 	enum class ENetRefHandleError : uint32;
 	enum class EReplicationCondition : uint32;
 
-	struct FNetTokenExportContext;
-
 	namespace Private
 	{
 		class FReplicationSystemImpl;
@@ -127,6 +126,9 @@ public:
 
 		/** Delegate that receives every RPC executed locally. */
 		UE::Net::FForwardNetRPCCallDelegate ForwardNetRPCCallDelegate;
+
+		/** NetTokenStore */
+		UE::Net::FNetTokenStore* NetTokenStore = nullptr;
 	};
 
 	/** @return The unique ID of the ReplicationSystem. */
@@ -345,10 +347,10 @@ public:
 	IRISCORE_API const UE::Net::FNetTokenStore* GetNetTokenStore() const;
 
 	/**
-	 * Init NetTokenExportContext required to quantize and resolve NetTokens
+	 * Get NetTokenResolveContext resolve NetTokens
 	 * @see UE::Net::FNetTokenStore
 	 */
-	IRISCORE_API void InitNetTokenExportContext(UE::Net::FNetTokenExportContext& NetTokenExportContext, uint32 ConnectionId);
+	IRISCORE_API UE::Net::FNetTokenResolveContext GetNetTokenResolveContext(uint32 ConnectionId) const;
 
 	/**
 	 * Check whether a FNetRefHandle is still associated with a replicated object.
diff --git a/Engine/Source/Runtime/Experimental/Iris/Core/Public/Iris/ReplicationSystem/StringTokenStore.h b/Engine/Source/Runtime/Experimental/Iris/Core/Public/Iris/ReplicationSystem/StringTokenStore.h
--- a/Engine/Source/Runtime/Experimental/Iris/Core/Public/Iris/ReplicationSystem/StringTokenStore.h
+++ b/Engine/Source/Runtime/Experimental/Iris/Core/Public/Iris/ReplicationSystem/StringTokenStore.h
@@ -16,7 +16,7 @@ class FStringTokenStore : public FNetTokenDataStore
 {
 	UE_NONCOPYABLE(FStringTokenStore);
 public:
-	explicit FStringTokenStore(FNetTokenStore& TokenStore);
+	IRISCORE_API explicit FStringTokenStore(FNetTokenStore& TokenStore);
 
 	// Create a string token for the provided string
 	IRISCORE_API FNetToken GetOrCreateToken(const FString& String);
@@ -33,9 +33,11 @@ public:
 protected:
 	// Serialize data for a token, note there is not validation in this function
 	virtual void WriteTokenData(FNetSerializationContext& Context, FNetTokenStoreKey TokenStoreKey) const override;
+	virtual void WriteTokenData(FArchive& Archive, FNetTokenStoreKey TokenStoreKey) const override;
 
 	// Read data for a token, returns a valid StoreKey if successful read
 	virtual FNetTokenStoreKey ReadTokenData(FNetSerializationContext& Context, const FNetToken& NetToken) override;
+	virtual FNetTokenStoreKey ReadTokenData(FArchive& Archive, const FNetToken& NetToken) override;
 
 	// Create a persistent string
 	FNetTokenStoreKey GetOrCreatePersistentString(const TCHAR* Name, uint32 Length);
diff --git a/Engine/Source/Runtime/Experimental/Iris/Core/Public/Iris/Serialization/IrisObjectReferencePackageMap.h b/Engine/Source/Runtime/Experimental/Iris/Core/Public/Iris/Serialization/IrisObjectReferencePackageMap.h
--- a/Engine/Source/Runtime/Experimental/Iris/Core/Public/Iris/Serialization/IrisObjectReferencePackageMap.h
+++ b/Engine/Source/Runtime/Experimental/Iris/Core/Public/Iris/Serialization/IrisObjectReferencePackageMap.h
@@ -3,7 +3,7 @@
 #pragma once
 
 #include "UObject/CoreNet.h"
-#include "Iris/ReplicationSystem/NetToken.h"
+#include "Net/Core/NetToken/NetToken.h"
 #include "IrisObjectReferencePackageMap.generated.h"
 
 // Forward declarations
@@ -11,40 +11,28 @@ class FNetworkGUID;
 
 namespace UE::Net
 {
-	class FNetTokenStoreState;
-	class FNetTokenStore;
-
-	struct FNetTokenExportContext
-	{
-		FNetTokenStore* TokenStore = nullptr;
-		FNetTokenStoreState* RemoteState = nullptr;		
-	};
-
 	// In order to properly capture exported data when calling in to old style NetSerialize methods
 	// we need to capture and inject certain types.
 	struct FIrisPackageMapExports
 	{
 		typedef TArray<TObjectPtr<UObject>, TInlineAllocator<4>> FObjectReferenceArray;
 		typedef TArray<FName, TInlineAllocator<4>> FNameArray;
-		typedef TArray<UE::Net::FNetToken, TInlineAllocator<4>> FNetTokensArray;
 
 		void Reset()
 		{
 			References.Reset();
 			Names.Reset();
-			NetTokens.Reset();
 		}
 
 		FObjectReferenceArray References;
 		FNameArray Names;
-		FNetTokensArray NetTokens;
 	};
 }
 
 /**
- * Custom packagemap implementation used to be able to capture UObject* references from external serialization.
- * Any object references written when using this packagemap will be added to the References array and serialized as an index.
- * When reading using this packagemap references will be read as an index and resolved by picking the corresponding entry from the provided array containing the references.
+ * Custom packagemap implementation used to be able to capture exports such as UObject* references, names and NetTokens from external serialization.
+ * Exports written when using this packagemap will be captured in an array and serialized as an index.
+ * When reading using this packagemap exports will be read as an index and resolved by picking the corresponding entry from the provided array containing the data associated with the export.
  */
 UCLASS(transient, MinimalAPI)
 class UIrisObjectReferencePackageMap : public UPackageMap
@@ -52,32 +40,21 @@ class UIrisObjectReferencePackageMap : public UPackageMap
 public:
 	GENERATED_BODY()
 
-	typedef TArray<TObjectPtr<UObject>, TInlineAllocator<4>> FObjectReferenceArray;
-	typedef TArray<FName, TInlineAllocator<4>> FNameArray;
-	typedef TArray<UE::Net::FNetToken, TInlineAllocator<4>> FNetTokensArray;
-
 	// We override SerializeObject in order to be able to capture object references
 	virtual bool SerializeObject(FArchive& Ar, UClass* InClass, UObject*& Obj, FNetworkGUID* OutNetGUID) override;
 
 	// Override SerializeName in order to be able to capture name and serialize them with iris instead.
 	virtual bool SerializeName(FArchive& Ar, FName& InName);
 
-	// Allow serialization of NetTokens to be able to support export of custom token types from existing NetSerialize implementations.
-	IRISCORE_API bool SerializeNetToken(FArchive& Ar, UE::Net::FNetToken& InNetToken);
-
-	// Init for read, we need a reference array to be able to resolve references
-	IRISCORE_API void InitForRead(const UE::Net::FIrisPackageMapExports* PackageMapExports);
+	// Init for read, we need to set the exports from which we are going to read our data.
+	IRISCORE_API void InitForRead(const UE::Net::FIrisPackageMapExports* PackageMapExports, const UE::Net::FNetTokenResolveContext& InNetTokenResolveContext);
 
 	// Init for write, all captured exports will be serialized as in index and added to the PackageMapExports for later export using iris.
 	IRISCORE_API void InitForWrite(UE::Net::FIrisPackageMapExports* PackageMapExports);
 
-	// Gets NetTokenExportContext to allow access to NetTokenStore
-	UE::Net::FNetTokenExportContext* GetNetTokenExportContext() { return &NetTokenExportContext; }
-
-	// Gets NetTokenExportContext to allow access to NetTokenStore
-	const UE::Net::FNetTokenExportContext* GetNetTokenExportContext() const { return &NetTokenExportContext; }
+	virtual const UE::Net::FNetTokenResolveContext* GetNetTokenResolveContext() const  override { return &NetTokenResolveContext; }
 
 protected:
 	UE::Net::FIrisPackageMapExports* PackageMapExports = nullptr;
-	UE::Net::FNetTokenExportContext NetTokenExportContext;
+	UE::Net::FNetTokenResolveContext NetTokenResolveContext;
 };
diff --git a/Engine/Source/Runtime/Experimental/Iris/Core/Public/Iris/Serialization/IrisPackageMapExportUtil.h b/Engine/Source/Runtime/Experimental/Iris/Core/Public/Iris/Serialization/IrisPackageMapExportUtil.h
--- a/Engine/Source/Runtime/Experimental/Iris/Core/Public/Iris/Serialization/IrisPackageMapExportUtil.h
+++ b/Engine/Source/Runtime/Experimental/Iris/Core/Public/Iris/Serialization/IrisPackageMapExportUtil.h
@@ -6,7 +6,7 @@
 #include "NetSerializer.h"
 #include "UObject/ObjectMacros.h"
 #include "Iris/Core/NetObjectReference.h"
-#include "Iris/ReplicationSystem/NetToken.h"
+#include "Net/Core/NetToken/NetToken.h"
 #include "Iris/Serialization/NetSerializerArrayStorage.h"
 #include "Iris/Serialization/StringNetSerializers.h"
 
@@ -54,7 +54,7 @@ struct FIrisPackageMapExportsUtil
 	// Matches NetSerializer functions
 	IRISCORE_API static void Serialize(FNetSerializationContext& Context, const QuantizedType& Value);
 	IRISCORE_API static void Deserialize(FNetSerializationContext& Context, QuantizedType& Value);
-	IRISCORE_API static void Quantize(FNetSerializationContext& Context, const UE::Net::FIrisPackageMapExports& PackageMapExport, QuantizedType& Target);
+	IRISCORE_API static void Quantize(FNetSerializationContext& Context, const UE::Net::FIrisPackageMapExports& PackageMapExport, TArrayView<const UE::Net::FNetToken> NetTokensPendingExport, QuantizedType& Target);
 	IRISCORE_API static void Dequantize(FNetSerializationContext& Context, const QuantizedType& Source, UE::Net::FIrisPackageMapExports& PackageMapExports);
 	IRISCORE_API static bool IsEqual(FNetSerializationContext& Context, const QuantizedType& Value0, const QuantizedType& Value1);
 	IRISCORE_API static void CloneDynamicState(FNetSerializationContext& Context, QuantizedType& Target, const QuantizedType& Source);
diff --git a/Engine/Source/Runtime/Experimental/Iris/Core/Public/Iris/Serialization/NetSerializationContext.h b/Engine/Source/Runtime/Experimental/Iris/Core/Public/Iris/Serialization/NetSerializationContext.h
--- a/Engine/Source/Runtime/Experimental/Iris/Core/Public/Iris/Serialization/NetSerializationContext.h
+++ b/Engine/Source/Runtime/Experimental/Iris/Core/Public/Iris/Serialization/NetSerializationContext.h
@@ -14,6 +14,7 @@ namespace UE::Net
 	class FNetBitStreamReader;
 	class FNetBitStreamWriter;
 	class FNetTokenStore;
+	class FNetTokenStoreState;
 
 	namespace Private
 	{
@@ -91,6 +92,7 @@ public:
 
 	IRISCORE_API const UE::Net::FNetTokenStore* GetNetTokenStore() const;
 	IRISCORE_API UE::Net::FNetTokenStore* GetNetTokenStore();
+	IRISCORE_API const UE::Net::FNetTokenStoreState* GetRemoteNetTokenStoreState() const;
 
 	void SetInternalContext(Private::FInternalNetSerializationContext* InInternalContext) { InternalContext = InInternalContext; }
 	Private::FInternalNetSerializationContext* GetInternalContext() { return InternalContext; }
diff --git a/Engine/Source/Runtime/GameplayTags/Classes/GameplayTagContainer.h b/Engine/Source/Runtime/GameplayTags/Classes/GameplayTagContainer.h
--- a/Engine/Source/Runtime/GameplayTags/Classes/GameplayTagContainer.h
+++ b/Engine/Source/Runtime/GameplayTags/Classes/GameplayTagContainer.h
@@ -219,6 +219,8 @@ struct FGameplayTag
 
 protected:
 
+	bool NetSerialize_ForReplayUsingFastReplication(FArchive& Ar, class UPackageMapClient& PackageMapClient);
+
 	/** Intentionally private so only the tag manager can use */
 	GAMEPLAYTAGS_API explicit FGameplayTag(const FName& InTagName);
 
diff --git a/Engine/Source/Runtime/GameplayTags/Classes/GameplayTagsManager.h b/Engine/Source/Runtime/GameplayTags/Classes/GameplayTagsManager.h
--- a/Engine/Source/Runtime/GameplayTags/Classes/GameplayTagsManager.h
+++ b/Engine/Source/Runtime/GameplayTags/Classes/GameplayTagsManager.h
@@ -559,6 +559,12 @@ public:
 		return bUseFastReplication;
 	}
 
+	/** Should use dynamic replication (Gameplay Tags need not match between client/server) */
+	bool ShouldUseDynamicReplication() const
+	{
+		return !bUseFastReplication && bUseDynamicReplication;
+	}
+
 	/** If we are allowed to unload tags */
 	GAMEPLAYTAGS_API bool ShouldUnloadTags() const;
 
@@ -905,6 +911,9 @@ private:
 	/** Cached runtime value for whether we are using fast replication or not. Initialized from config setting. */
 	bool bUseFastReplication;
 
+	/** Cached runtime value for whether we are using dynamic replication or not. Initialized from the config setting. */
+	bool bUseDynamicReplication;
+
 	/** Cached runtime value for whether we should warn when loading invalid tags */
 	bool bShouldWarnOnInvalidTags;
 
diff --git a/Engine/Source/Runtime/GameplayTags/Classes/GameplayTagsSettings.h b/Engine/Source/Runtime/GameplayTags/Classes/GameplayTagsSettings.h
--- a/Engine/Source/Runtime/GameplayTags/Classes/GameplayTagsSettings.h
+++ b/Engine/Source/Runtime/GameplayTags/Classes/GameplayTagsSettings.h
@@ -124,6 +124,10 @@ class UGameplayTagsSettings : public UGameplayTagsList
 	UPROPERTY(config, EditAnywhere, Category = "Advanced Replication")
 	bool FastReplication;
 
+	/** If true, will replicate gameplay tags dynamically by index per connection. Slightly higher cost than FastReplication, but tags can differ between client and server */
+	UPROPERTY(config, EditAnywhere, Category = "Advanced Replication", meta=(EditCondition="!FastReplication"))
+	bool bDynamicReplication;
+
 	/** These characters cannot be used in gameplay tags, in addition to special ones like newline*/
 	UPROPERTY(config, EditAnywhere, Category = GameplayTags)
 	FString InvalidTagCharacters;
diff --git a/Engine/Source/Runtime/GameplayTags/GameplayTags.Build.cs b/Engine/Source/Runtime/GameplayTags/GameplayTags.Build.cs
--- a/Engine/Source/Runtime/GameplayTags/GameplayTags.Build.cs
+++ b/Engine/Source/Runtime/GameplayTags/GameplayTags.Build.cs
@@ -20,6 +20,7 @@ namespace UnrealBuildTool.Rules
 				new string[]
 				{
 					"Projects",
+					"NetCore",
 					"Json",
 					"JsonUtilities"
 				}
diff --git a/Engine/Source/Runtime/GameplayTags/Private/GameplayTagContainer.cpp b/Engine/Source/Runtime/GameplayTags/Private/GameplayTagContainer.cpp
--- a/Engine/Source/Runtime/GameplayTags/Private/GameplayTagContainer.cpp
+++ b/Engine/Source/Runtime/GameplayTags/Private/GameplayTagContainer.cpp
@@ -12,6 +12,13 @@
 #include "Misc/OutputDeviceNull.h"
 #include "JsonObjectConverter.h"
 
+// For dynamic serialization support
+#if UE_WITH_IRIS
+#include "Net/Core/NetToken/NetTokenExportContext.h"
+#include "GameplayTagTokenStore.h"
+#endif
+#include <Net/Core/Trace/NetTrace.h>
+
 #include UE_INLINE_GENERATED_CPP_BY_NAME(GameplayTagContainer)
 
 const FGameplayTag FGameplayTag::EmptyTag;
@@ -37,6 +44,9 @@ DECLARE_CYCLE_STAT(TEXT("FGameplayTag::NetSerialize"), STAT_FGameplayTag_NetSeri
 static bool GEnableGameplayTagDetailedStats = false;
 static FAutoConsoleVariableRef CVarGameplayTagDetailedStats(TEXT("GameplayTags.EnableDetailedStats"), GEnableGameplayTagDetailedStats, TEXT("Runtime toggle for verbose CPU profiling stats"), ECVF_Default);
 
+static bool GOldReplaysUseGameplayTagFastReplication = true;
+static FAutoConsoleVariableRef CVarOldReplaysUseGameplayTagFastReplication(TEXT("GameplayTags.OldReplaysUseFastReplication"), GOldReplaysUseGameplayTagFastReplication, TEXT("When loading an outdated replay (before dynamic replication), do we assume it used fast replication?"), ECVF_Default);
+
 /**
  *	Replicates a tag in a packed format:
  *	-A segment of NetIndexFirstBitSegment bits are always replicated.
@@ -1213,88 +1223,292 @@ static TSharedPtr<FNetFieldExportGroup> CreateNetfieldExportGroupForNetworkGamep
 	return NetFieldExportGroup;
 }
 
-bool FGameplayTag::NetSerialize_Packed(FArchive& Ar, class UPackageMap* Map, bool& bOutSuccess)
+bool FGameplayTag::NetSerialize_ForReplayUsingFastReplication(FArchive& Ar, UPackageMapClient& PackageMapClient)
 {
-	CONDITIONAL_SCOPE_CYCLE_COUNTER(STAT_FGameplayTag_NetSerialize, GEnableGameplayTagDetailedStats);
-
 	UGameplayTagsManager& TagManager = UGameplayTagsManager::Get();
+	FGameplayTagNetIndex NetIndex = INVALID_TAGNETINDEX;
+
+	// For replays, use a net field export group to guarantee we can send the name reliably (without having to rely on the client having a deterministic NetworkGameplayTagNodeIndex array)
+	const TCHAR* NetFieldExportGroupName = TEXT("NetworkGameplayTagNodeIndex");
+
+	// Find this net field export group
+	TSharedPtr<FNetFieldExportGroup> NetFieldExportGroup = PackageMapClient.GetNetFieldExportGroup(NetFieldExportGroupName);
 
-	if (TagManager.ShouldUseFastReplication())
+	if (Ar.IsSaving())
 	{
-		FGameplayTagNetIndex NetIndex = INVALID_TAGNETINDEX;
+		// If we didn't find it, we need to create it (only when saving though, it should be here on load since it was exported at save time)
+		if (!NetFieldExportGroup.IsValid())
+		{
+			NetFieldExportGroup = CreateNetfieldExportGroupForNetworkGameplayTags(TagManager, NetFieldExportGroupName);
+			PackageMapClient.AddNetFieldExportGroup(NetFieldExportGroupName, NetFieldExportGroup);
+		}
 
-		UPackageMapClient* PackageMapClient = Cast<UPackageMapClient>(Map);
-		const bool bIsReplay = PackageMapClient && PackageMapClient->GetConnection() && PackageMapClient->GetConnection()->IsInternalAck();
+		NetIndex = TagManager.GetNetIndexFromTag(*this);
+
+		if (NetIndex != TagManager.GetInvalidTagNetIndex() && NetIndex != INVALID_TAGNETINDEX)
+		{
+			PackageMapClient.TrackNetFieldExport(NetFieldExportGroup.Get(), NetIndex);
+		}
+		else
+		{
+			NetIndex = INVALID_TAGNETINDEX;		// We can't save InvalidTagNetIndex, since the remote side could have a different value for this
+		}
+	}
 
-		TSharedPtr<FNetFieldExportGroup> NetFieldExportGroup;
+	uint32 NetIndex32 = NetIndex;
+	Ar.SerializeIntPacked(NetIndex32);
+	NetIndex = IntCastChecked<uint16, uint32>(NetIndex32);
 
-		if (bIsReplay)
+	if (Ar.IsLoading())
+	{
+		// Get the tag name from the net field export group entry
+		if (NetIndex != INVALID_TAGNETINDEX && ensure(NetFieldExportGroup.IsValid()) && ensure(NetIndex < NetFieldExportGroup->NetFieldExports.Num()))
+		{
+			TagName = NetFieldExportGroup->NetFieldExports[NetIndex].ExportName;
+
+			// Validate the tag name
+			const FGameplayTag Tag = TagManager.RequestGameplayTag(TagName, false);
+
+			// Warn (once) if the tag isn't found
+			if (!Tag.IsValid() && !NetFieldExportGroup->NetFieldExports[NetIndex].bIncompatible)
+			{ 
+				UE_LOG(LogGameplayTags, Warning, TEXT( "Gameplay tag not found (marking incompatible): %s"), *TagName.ToString());
+				NetFieldExportGroup->NetFieldExports[NetIndex].bIncompatible = true;
+			}
+
+			TagName = Tag.TagName;
+		}
+		else
+		{
+			TagName = NAME_None;
+		}
+	}
+
+	return true;
+}
+
+// DynamicSerialization currently relies on experimental code only available when compiling with Iris.
+namespace UE::GameplayTags::GameplayTagDynamicSerialization
+{
+#if UE_WITH_IRIS
+	// Can we make this generic and handle arbitrary export payloads? Probably something that should be handled in PackageMapClient
+	static bool NetSerialize_ForReplay(FGameplayTag& GameplayTag, FArchive& Ar, UPackageMapClient& PackageMapClient)
+	{
+		using namespace UE::Net;
+
+		const FNetTokenResolveContext* NetTokenResolveContext = PackageMapClient.GetNetTokenResolveContext();
+		FGameplayTagTokenStore* TagTokenDataStore = NetTokenResolveContext ? NetTokenResolveContext->NetTokenStore->GetDataStore<UE::Net::FGameplayTagTokenStore>() : nullptr;
+		if (!TagTokenDataStore)
 		{
-			// For replays, use a net field export group to guarantee we can send the name reliably (without having to rely on the client having a deterministic NetworkGameplayTagNodeIndex array)
-			const TCHAR* NetFieldExportGroupName = TEXT("NetworkGameplayTagNodeIndex");
+			return false;
+		}
+
+		const TCHAR* NetFieldExportGroupName = TEXT("NetworkGameplayTagDynamicIndex");
+		TSharedPtr<FNetFieldExportGroup> NetFieldExportGroup = PackageMapClient.GetNetFieldExportGroup(NetFieldExportGroupName);
+		FNetToken TagToken;
+
+		if (Ar.IsSaving())
+		{
+			TagToken = TagTokenDataStore->GetOrCreateToken(GameplayTag);
 
-			// Find this net field export group
-			NetFieldExportGroup = PackageMapClient->GetNetFieldExportGroup(NetFieldExportGroupName);
+			// Write token
+			WriteNetToken(Ar, TagToken);
 
-			if (Ar.IsSaving())
+			// Register replay export if needed
+			if (TagToken.IsValid())
 			{
 				// If we didn't find it, we need to create it (only when saving though, it should be here on load since it was exported at save time)
 				if (!NetFieldExportGroup.IsValid())
 				{
-					NetFieldExportGroup = CreateNetfieldExportGroupForNetworkGameplayTags(TagManager, NetFieldExportGroupName);
-
-					PackageMapClient->AddNetFieldExportGroup(NetFieldExportGroupName, NetFieldExportGroup);
+					NetFieldExportGroup = TSharedPtr<FNetFieldExportGroup>(new FNetFieldExportGroup());
+					NetFieldExportGroup->PathName = NetFieldExportGroupName;
+					PackageMapClient.AddNetFieldExportGroup(NetFieldExportGroupName, NetFieldExportGroup);
 				}
 
-				NetIndex = TagManager.GetNetIndexFromTag(*this);
-
-				if (NetIndex != TagManager.GetInvalidTagNetIndex() && NetIndex != INVALID_TAGNETINDEX)
+				// Make sure we have enough room in the NetFieldExports to hold this entry...
+				const uint32 TagTokenIndex = TagToken.GetIndex();
+				if (!NetFieldExportGroup->NetFieldExports.IsValidIndex(TagTokenIndex))
 				{
-					PackageMapClient->TrackNetFieldExport(NetFieldExportGroup.Get(), NetIndex);
+					NetFieldExportGroup->NetFieldExports.SetNum(TagTokenIndex + 1, EAllowShrinking::No);
+					NetFieldExportGroup->bDirtyForReplay = true;
 				}
-				else
+				ensure(NetFieldExportGroup->NetFieldExports.IsValidIndex(TagTokenIndex));
+				FNetFieldExport& NetFieldExport = NetFieldExportGroup->NetFieldExports[TagTokenIndex];
+
+				// If it's not yet exported, export it now
+				if (!NetFieldExport.bExported)
 				{
-					NetIndex = INVALID_TAGNETINDEX;		// We can't save InvalidTagNetIndex, since the remote side could have a different value for this
+					NetFieldExport = FNetFieldExport(TagTokenIndex, 0, GameplayTag.GetTagName());
+					UE_LOG(LogGameplayTags, Log, TEXT("Replay> Exported Tag %s as NetFieldIndex %u"), *NetFieldExport.ExportName.ToString(), TagTokenIndex);
 				}
-			}
 
-			uint32 NetIndex32 = NetIndex;
-			Ar.SerializeIntPacked(NetIndex32);
-			NetIndex = IntCastChecked<uint16, uint32>(NetIndex32);
+				// Track the export so that it gets added to the replay index
+				PackageMapClient.TrackNetFieldExport(NetFieldExportGroup.Get(), TagTokenIndex);
+			}
+			return true;
+		}
+		else if (Ar.IsLoading())
+		{
+			// Read TagNetToken
+			TagToken = ReadNetToken(Ar);
+			if (Ar.IsError())
+			{
+				return false;
+			}
 
-			if (Ar.IsLoading())
+			if (TagToken.IsValid())
 			{
-				// Get the tag name from the net field export group entry
-				if (NetIndex != INVALID_TAGNETINDEX && ensure(NetFieldExportGroup.IsValid()) && ensure(NetIndex < NetFieldExportGroup->NetFieldExports.Num()))
+				const uint32 TagTokenIndex = TagToken.GetIndex();
+				if (ensure(NetFieldExportGroup.IsValid()) && ensure(TagTokenIndex < (uint32)NetFieldExportGroup->NetFieldExports.Num()))
 				{
-					TagName = NetFieldExportGroup->NetFieldExports[NetIndex].ExportName;
+					FName TagName = NetFieldExportGroup->NetFieldExports[TagTokenIndex].ExportName;
 
 					// Validate the tag name
-					const FGameplayTag Tag = TagManager.RequestGameplayTag(TagName, false);
+					// TODO: Should we be able to add tags through this?
+					UGameplayTagsManager& TagManager = UGameplayTagsManager::Get();
+					GameplayTag = TagManager.RequestGameplayTag(TagName, false);
 
 					// Warn (once) if the tag isn't found
-					if (!Tag.IsValid() && !NetFieldExportGroup->NetFieldExports[NetIndex].bIncompatible)
+					if (!GameplayTag.IsValid() && !NetFieldExportGroup->NetFieldExports[TagTokenIndex].bIncompatible)
 					{ 
 						UE_LOG(LogGameplayTags, Warning, TEXT( "Gameplay tag not found (marking incompatible): %s"), *TagName.ToString());
-						NetFieldExportGroup->NetFieldExports[NetIndex].bIncompatible = true;
+						NetFieldExportGroup->NetFieldExports[TagTokenIndex].bIncompatible = true;
+					}
+					return true;
+				}
+				else
+				{
+					GameplayTag = FGameplayTag();
+					return false;
+				}
+			}
+			GameplayTag = FGameplayTag();
+			return true;
+		}
+		return false;
+	}
+
+	bool NetSerialize(FGameplayTag& GameplayTag, FArchive& Ar, UPackageMap* Map)
+	{
+		using namespace UE::Net;
+
+		// For now special case replays
+		UPackageMapClient* PackageMapClient = Cast<UPackageMapClient>(Map);
+		if (const bool bIsReplay = PackageMapClient && PackageMapClient->GetConnection() && PackageMapClient->GetConnection()->IsInternalAck())
+		{
+			return NetSerialize_ForReplay(GameplayTag, Ar, *PackageMapClient);
+		}
+	
+		if (Ar.IsSaving())
+		{
+			UE::Net::FNetTokenExportContext* ExportContext = FNetTokenExportContext::GetNetTokenExportContext(Ar);
+			UE::Net::FNetTokenStore* NetTokenStore = ExportContext ? ExportContext->GetNetTokenStore() : nullptr;
+			if (ensure(NetTokenStore))
+			{
+				FNetToken TagToken;
+
+				if (GameplayTag.IsValid())
+				{
+					if (FGameplayTagTokenStore* TagTokenStore = NetTokenStore->GetDataStore<UE::Net::FGameplayTagTokenStore>())
+					{
+						TagToken = TagTokenStore->GetOrCreateToken(GameplayTag);
 					}
+					else
+					{
+						UE_LOG(LogGameplayTags, Error, TEXT("FGameplayTag::NetSerialize::Could not find required FGameplayTagTokenStore"));
+						ensure(false);
+					}
+				}
+
+				UE_NET_TRACE_DYNAMIC_NAME_SCOPE(*TagToken.ToString(), static_cast<FNetBitWriter&>(Ar), GetTraceCollector(static_cast<FNetBitWriter&>(Ar)), ENetTraceVerbosity::VeryVerbose);
+
+				// Write token
+				WriteNetToken(Ar, TagToken);
 
-					TagName = Tag.TagName;
+				// Should we append export from WriteNetToken method?
+				if (TagToken.IsValid())
+				{
+					ExportContext->AddNetTokenPendingExport(TagToken);
+				}
+				return true;
+			}
+		}
+		else if (Ar.IsLoading())
+		{
+			// When reading data we always have a PackageMap so we can get the necessary resolve context from here.
+			const FNetTokenResolveContext* NetTokenResolveContext = Map ? Map->GetNetTokenResolveContext() : nullptr;
+			if (ensure(NetTokenResolveContext && NetTokenResolveContext->NetTokenStore))
+			{
+				FNetToken TagToken = ReadNetToken(Ar);
+				if (Ar.IsError())
+				{
+					return false;
+				}
+				if (FGameplayTagTokenStore* TagTokenStore = NetTokenResolveContext->NetTokenStore->GetDataStore<UE::Net::FGameplayTagTokenStore>())
+				{
+					GameplayTag = TagTokenStore->ResolveToken(TagToken, NetTokenResolveContext->RemoteNetTokenStoreState);
 				}
 				else
 				{
-					TagName = NAME_None;
+					UE_LOG(LogGameplayTags, Error, TEXT("FGameplayTag::NetSerialize::Could not find required FGameplayTagTokenStore"));	
+					ensure(false);
+					Ar.SetError();
+					return false;
 				}
+				return true;
 			}
+		}
+		return false;
+	}
+#else
+	bool NetSerialize(FGameplayTag& GameplayTag, FArchive& Ar, UPackageMap* Map)
+	{
+		LowLevelFatalError(TEXT("Cannot use dynamic serialization without compiling with Iris."));
+		return false;
+	}
+#endif
+}
 
-			bOutSuccess = true;
-			return true;
+
+bool FGameplayTag::NetSerialize_Packed(FArchive& Ar, class UPackageMap* Map, bool& bOutSuccess)
+{
+	CONDITIONAL_SCOPE_CYCLE_COUNTER(STAT_FGameplayTag_NetSerialize, GEnableGameplayTagDetailedStats);
+	UGameplayTagsManager& TagManager = UGameplayTagsManager::Get();
+
+	// We now save off which method we're using (below)
+	bool bUseFastReplication = TagManager.ShouldUseFastReplication();
+	bool bUseDynamicReplication = TagManager.ShouldUseDynamicReplication();
+
+	Ar.UsingCustomVersion(FEngineNetworkCustomVersion::Guid);
+	const bool bSerializeReplicationMethod = (Ar.EngineNetVer() >= FEngineNetworkCustomVersion::CustomExports);
+	if (bSerializeReplicationMethod)
+	{
+		Ar.SerializeBits(&bUseFastReplication, 1);
+		if (!bUseFastReplication)
+		{
+			Ar.SerializeBits(&bUseDynamicReplication, 1);
+		}
+	}
+	else
+	{
+		bUseFastReplication = GOldReplaysUseGameplayTagFastReplication;
+		bUseDynamicReplication = false; // this didn't exist in prior versions
+	}
+
+	if (bUseFastReplication)
+	{
+		UPackageMapClient* PackageMapClient = Cast<UPackageMapClient>(Map);
+
+		const bool bIsReplay = PackageMapClient && PackageMapClient->GetConnection() && PackageMapClient->GetConnection()->IsInternalAck();
+		if (bIsReplay)
+		{
+			return NetSerialize_ForReplayUsingFastReplication(Ar, *PackageMapClient);
 		}
 
+		FGameplayTagNetIndex NetIndex = INVALID_TAGNETINDEX;
 		if (Ar.IsSaving())
 		{
 			NetIndex = TagManager.GetNetIndexFromTag(*this);
-			
 			SerializeTagNetIndexPacked(Ar, NetIndex, TagManager.GetNetIndexFirstBitSegment(), TagManager.GetNetIndexTrueBitNum());
 		}
 		else
@@ -1303,6 +1517,10 @@ bool FGameplayTag::NetSerialize_Packed(FArchive& Ar, class UPackageMap* Map, boo
 			TagName = TagManager.GetTagNameFromNetIndex(NetIndex);
 		}
 	}
+	else if (bUseDynamicReplication)
+	{
+		UE::GameplayTags::GameplayTagDynamicSerialization::NetSerialize(*this, Ar, Map);
+	}
 	else
 	{
 		Ar << TagName;
diff --git a/Engine/Source/Runtime/GameplayTags/Private/GameplayTagContainerNetSerializer.cpp b/Engine/Source/Runtime/GameplayTags/Private/GameplayTagContainerNetSerializer.cpp
--- a/Engine/Source/Runtime/GameplayTags/Private/GameplayTagContainerNetSerializer.cpp
+++ b/Engine/Source/Runtime/GameplayTags/Private/GameplayTagContainerNetSerializer.cpp
@@ -118,19 +118,31 @@ void FGameplayTagContainerNetSerializer::DeserializeDelta(FNetSerializationConte
 
 void FGameplayTagContainerNetSerializer::Quantize(FNetSerializationContext& Context, const FNetQuantizeArgs& Args)
 {
-	FGameplayTagContainerNetSerializerSerializationHelper IntermediateValue;
+	const UGameplayTagsManager& TagManager = UGameplayTagsManager::Get();
+	if (const bool bUseFastReplication = TagManager.ShouldUseFastReplication())
+	{
+		FGameplayTagContainerNetSerializerSerializationHelper IntermediateValue;
+
+		// It's preferred but not vital that the default state represents the true default state. Because of the brittle tag loading/adding we choose to not store tags in the default state.
+		if (!Context.IsInitializingDefaultState())
+		{
+			const SourceType& SourceValue = *reinterpret_cast<const SourceType*>(Args.Source);
+			GameplayTagContainerToSerializationHelper(SourceValue, IntermediateValue);
+		}
 
-	// It's preferred but not vital that the default state represents the true default state. Because of the brittle tag loading/adding we choose to not store tags in the default state.
-	if (!Context.IsInitializingDefaultState())
+		FNetQuantizeArgs InternalArgs = Args;
+		InternalArgs.Source = NetSerializerValuePointer(&IntermediateValue);
+		InternalArgs.NetSerializerConfig = &StructNetSerializerConfig;
+		return StructNetSerializer->Quantize(Context, InternalArgs);
+	}
+	else
 	{
-		const SourceType& SourceValue = *reinterpret_cast<const SourceType*>(Args.Source);
-		GameplayTagContainerToSerializationHelper(SourceValue, IntermediateValue);
+		// If we do not use FastReplication we just quantize the tags without sorting them.
+		FNetQuantizeArgs InternalArgs = Args;
+		InternalArgs.Source = Args.Source;
+		InternalArgs.NetSerializerConfig = &StructNetSerializerConfig;
+		return StructNetSerializer->Quantize(Context, InternalArgs);
 	}
-
-	FNetQuantizeArgs InternalArgs = Args;
-	InternalArgs.Source = NetSerializerValuePointer(&IntermediateValue);
-	InternalArgs.NetSerializerConfig = &StructNetSerializerConfig;
-	return StructNetSerializer->Quantize(Context, InternalArgs);
 }
 
 void FGameplayTagContainerNetSerializer::Dequantize(FNetSerializationContext& Context, const FNetDequantizeArgs& Args)
@@ -151,9 +163,49 @@ bool FGameplayTagContainerNetSerializer::IsEqual(FNetSerializationContext& Conte
 {
 	if (Args.bStateIsQuantized)
 	{
-		FNetIsEqualArgs InternalArgs = Args;
-		InternalArgs.NetSerializerConfig = &StructNetSerializerConfig;
-		return StructNetSerializer->IsEqual(Context, InternalArgs);
+		const UGameplayTagsManager& TagManager = UGameplayTagsManager::Get();
+		if (const bool bUseFastReplication = TagManager.ShouldUseFastReplication())
+		{
+			FNetIsEqualArgs InternalArgs = Args;
+			InternalArgs.NetSerializerConfig = &StructNetSerializerConfig;
+			return StructNetSerializer->IsEqual(Context, InternalArgs);
+		}
+		else
+		{
+			// For now as we are rarely using this path and with dynamic tag indices we cannot stable sort them as we do for the fast path
+			// so for now we dequantize and compare using the native compare
+			SourceType SourceValue0;
+			SourceType SourceValue1;
+
+			FNetDequantizeArgs DequantizeArgs = {};
+			DequantizeArgs.NetSerializerConfig = Args.NetSerializerConfig;
+
+			DequantizeArgs.Source = Args.Source0;
+			DequantizeArgs.Target = NetSerializerValuePointer(&SourceValue0);
+			Dequantize(Context, DequantizeArgs);
+
+			DequantizeArgs.Source = Args.Source1;
+			DequantizeArgs.Target = NetSerializerValuePointer(&SourceValue1);
+			Dequantize(Context, DequantizeArgs);
+
+			if (SourceValue0.GetGameplayTags().Num() != SourceValue1.GetGameplayTags().Num())
+			{
+				return false;
+			}
+
+			if (!SourceValue0.HasAllExact(SourceValue1))
+			{
+				return false;
+			}
+
+			// This will detect duplicate tags.
+			if (!SourceValue1.HasAllExact(SourceValue0))
+			{
+				return false;
+			}
+
+			return true;
+		}
 	}
 	else
 	{
diff --git a/Engine/Source/Runtime/GameplayTags/Private/GameplayTagNetSerializer.cpp b/Engine/Source/Runtime/GameplayTags/Private/GameplayTagNetSerializer.cpp
--- a/Engine/Source/Runtime/GameplayTags/Private/GameplayTagNetSerializer.cpp
+++ b/Engine/Source/Runtime/GameplayTags/Private/GameplayTagNetSerializer.cpp
@@ -8,9 +8,11 @@
 
 #include "GameplayTagContainer.h"
 #include "GameplayTagsManager.h"
+#include "GameplayTagTokenStore.h"
 #include "Iris/ReplicationState/PropertyNetSerializerInfoRegistry.h"
 #include "Iris/Serialization/NetBitStreamUtil.h"
 #include "Iris/Serialization/NetSerializerDelegates.h"
+#include "Iris/ReplicationSystem/ReplicationSystem.h"
 
 static_assert(sizeof(FGameplayTagNetIndex) == 2, "Unexpected GameplayTagNetIndex size. Expected 2.");
 
@@ -22,19 +24,31 @@ struct FGameplayTagAccessorForNetSerializer : public FGameplayTag
 	void SetTagName(const FName InTagName) { TagName = InTagName; }
 };
 
-struct FGameplayTagNetSerializer
+// Types
+struct FFGameplayTagNetSerializerQuantizedType
 {
-	// Version
-	static const uint32 Version = 0;
-
-	// Types
-	struct FQuantizedType
+	union
 	{
+		FNetToken TagNetToken;		
 		FGameplayTagNetIndex TagIndex;
 	};
+	bool bUseFastReplication;
+};
+
+}
+
+template <> struct TIsPODType<UE::Net::FFGameplayTagNetSerializerQuantizedType> { enum { Value = true }; };
+
+namespace UE::Net
+{
 
-	typedef FGameplayTagAccessorForNetSerializer SourceType;
-	typedef FQuantizedType QuantizedType;
+struct FGameplayTagNetSerializer
+{
+	// Version
+	static const uint32 Version = 0;
+
+	typedef FGameplayTag SourceType;
+	typedef FFGameplayTagNetSerializerQuantizedType QuantizedType;
 	typedef struct FGameplayTagNetSerializerConfig ConfigType;
 
 	static const ConfigType DefaultConfig;
@@ -72,7 +86,25 @@ void FGameplayTagNetSerializer::Serialize(FNetSerializationContext& Context, con
 	const QuantizedType& Value = *reinterpret_cast<QuantizedType*>(Args.Source);
 	
 	FNetBitStreamWriter* Writer = Context.GetBitStreamWriter();
-	WritePackedUint16(Writer, Value.TagIndex);
+	
+	if (Writer->WriteBool(Value.bUseFastReplication))
+	{
+		WritePackedUint16(Writer, Value.TagIndex);
+	}
+	else
+	{
+		// Tokens will differ, so we cannot store them in the default statehash.
+		if (Context.IsInitializingDefaultState())
+		{
+			return;
+		}
+
+		// Always write the token
+		WriteNetToken(Context, Value.TagNetToken);
+
+		// Export or add to pending exports for later export
+		FNetTokenStore::AppendExportOrWriteInlinedExportData(Context, Value.TagNetToken);
+	}
 }
 
 void FGameplayTagNetSerializer::Deserialize(FNetSerializationContext& Context, const FNetDeserializeArgs& Args)
@@ -80,7 +112,33 @@ void FGameplayTagNetSerializer::Deserialize(FNetSerializationContext& Context, c
 	QuantizedType& TargetValue = *reinterpret_cast<QuantizedType*>(Args.Target);
 	
 	FNetBitStreamReader* Reader = Context.GetBitStreamReader();
-	TargetValue.TagIndex = ReadPackedUint16(Reader);
+
+	TargetValue = {};
+
+	if (const bool bUseFastReplication = Reader->ReadBool())
+	{
+		TargetValue.TagIndex = ReadPackedUint16(Reader);
+		TargetValue.bUseFastReplication = true;
+	}
+	else
+	{
+		FNetToken NetToken = ReadNetToken(Context);
+		if (Reader->IsOverflown())
+		{
+			return;
+		}
+
+		// Read inlined exports if we should
+		FNetTokenStore::ReadInlinedExportData(Context, NetToken);
+
+		if (Reader->IsOverflown())
+		{
+			return;
+		}
+
+		// Store 
+		TargetValue.TagNetToken = NetToken;
+	}
 }
 
 void FGameplayTagNetSerializer::Quantize(FNetSerializationContext& Context, const FNetQuantizeArgs& Args)
@@ -88,15 +146,38 @@ void FGameplayTagNetSerializer::Quantize(FNetSerializationContext& Context, cons
 	const SourceType& SourceValue = *reinterpret_cast<const SourceType*>(Args.Source);
 	QuantizedType& TargetValue = *reinterpret_cast<QuantizedType*>(Args.Target);
 
+	TargetValue = {};
+
 	const UGameplayTagsManager& TagManager = UGameplayTagsManager::Get();
-	FGameplayTagNetIndex TagIndex = TagManager.GetNetIndexFromTag(SourceValue);
-	if (TagIndex == TagManager.GetInvalidTagNetIndex())
+
+	const bool bUseFastReplication = TagManager.ShouldUseFastReplication();
+	
+	if (bUseFastReplication)
 	{
-		TargetValue.TagIndex = InvalidTagIndex;
+		TargetValue.bUseFastReplication = true;
+		
+		// We use a stable value for invalid TagIndex as the value from the TagManager is dynamic.
+		FGameplayTagNetIndex TagIndex = TagManager.GetNetIndexFromTag(SourceValue);
+		if (TagIndex == TagManager.GetInvalidTagNetIndex())
+		{
+			TargetValue.TagIndex = InvalidTagIndex;
+		}
+		else
+		{
+			TargetValue.TagIndex = TagIndex;
+		}
 	}
 	else
 	{
-		TargetValue.TagIndex = TagIndex;
+		if (FGameplayTagTokenStore* TagTokenStore = Context.GetNetTokenStore()->GetDataStore<FGameplayTagTokenStore>())
+		{
+			TargetValue.TagNetToken = TagTokenStore->GetOrCreateToken(SourceValue);
+		}
+		else
+		{
+			UE_LOG(LogGameplayTags, Error, TEXT("FGameplayTagNetSerializer::Quantize Could not find required FGameplayTagTokenStore"));
+			ensure(false);
+		}
 	}
 }
 
@@ -105,14 +186,36 @@ void FGameplayTagNetSerializer::Dequantize(FNetSerializationContext& Context, co
 	const QuantizedType& Source = *reinterpret_cast<const QuantizedType*>(Args.Source);
 	SourceType& Target = *reinterpret_cast<SourceType*>(Args.Target);
 
-	if (Source.TagIndex != InvalidTagIndex)
+	if (Source.bUseFastReplication)
 	{
-		const UGameplayTagsManager& TagManager = UGameplayTagsManager::Get();
-		Target.SetTagName(TagManager.GetTagNameFromNetIndex(Source.TagIndex));
+		if (Source.TagIndex != InvalidTagIndex)
+		{
+			FGameplayTagAccessorForNetSerializer& TargetAccessor = *reinterpret_cast<FGameplayTagAccessorForNetSerializer*>(Args.Target);
+			const UGameplayTagsManager& TagManager = UGameplayTagsManager::Get();
+			TargetAccessor.SetTagName(TagManager.GetTagNameFromNetIndex(Source.TagIndex));
+		}
+		else
+		{
+			// Invalid Tag
+			Target = FGameplayTag();
+		}
 	}
 	else
 	{
-		Target.SetTagName(FName());
+		if (FGameplayTagTokenStore* TagTokenStore = Context.GetNetTokenStore()->GetDataStore<FGameplayTagTokenStore>())
+		{
+			Target = TagTokenStore->ResolveToken(Source.TagNetToken, Context.GetRemoteNetTokenStoreState());
+		}
+		else
+		{
+			UE_LOG(LogGameplayTags, Error, TEXT("FGameplayTagNetSerializer::Dequantize Could not find required FGameplayTagTokenStore"));
+			ensure(false);
+
+			// Invalid Tag
+			Target = FGameplayTag();
+
+			return;
+		}
 	}
 }
 
@@ -122,12 +225,44 @@ bool FGameplayTagNetSerializer::IsEqual(FNetSerializationContext& Context, const
 	{
 		const QuantizedType& Value0 = *reinterpret_cast<const QuantizedType*>(Args.Source0);
 		const QuantizedType& Value1 = *reinterpret_cast<const QuantizedType*>(Args.Source1);
-		return Value0.TagIndex == Value1.TagIndex;
+
+		if (Value0.bUseFastReplication != Value1.bUseFastReplication)
+		{
+			return false;
+		}
+
+		if (Value0.bUseFastReplication)
+		{
+			return Value0.TagIndex == Value1.TagIndex;
+		}
+		else
+		{
+			// Need to compare actual Tags to properly compare non-auth and auth token
+			if (Value0.TagNetToken.IsAssignedByAuthority() != Value1.TagNetToken.IsAssignedByAuthority())
+			{
+				FGameplayTagTokenStore* TagTokenStore = Context.GetNetTokenStore()->GetDataStore<FGameplayTagTokenStore>();
+				const UE::Net::FNetTokenStoreState* RemoteNetTokenStoreState = Context.GetRemoteNetTokenStoreState();
+	
+				const FGameplayTag Tag0 = TagTokenStore->ResolveToken(Value0.TagNetToken, RemoteNetTokenStoreState);
+				const FGameplayTag Tag1 = TagTokenStore->ResolveToken(Value1.TagNetToken, RemoteNetTokenStoreState);
+				
+				if (Tag0 != Tag1)
+				{
+					return false;
+				}
+			}
+			else if (Value0.TagNetToken != Value1.TagNetToken)
+			{
+				return false;
+			}
+			return true;
+		}
 	}
 	else
 	{
 		const SourceType& Value0 = *reinterpret_cast<SourceType*>(Args.Source0);
 		const SourceType& Value1 = *reinterpret_cast<SourceType*>(Args.Source1);
+
 		return Value0.GetTagName() == Value1.GetTagName();
 	}
 }
diff --git a/Engine/Source/Runtime/GameplayTags/Private/GameplayTagTokenStore.cpp b/Engine/Source/Runtime/GameplayTags/Private/GameplayTagTokenStore.cpp
new file mode 100644
--- /dev/null
+++ b/Engine/Source/Runtime/GameplayTags/Private/GameplayTagTokenStore.cpp
@@ -0,0 +1,30 @@
+// Copyright Epic Games, Inc. All Rights Reserved.
+
+// NetTokenExports is still an experimental feature that depends on Iris code.
+#if UE_WITH_IRIS
+
+#include "GameplayTagTokenStore.h"
+#include "GameplayTagContainer.h"
+#include "GameplayTagsManager.h"
+
+namespace UE::Net
+{
+
+ FGameplayTagTokenStore::FGameplayTagTokenStore(FNetTokenStore& TokenStore) : FNameTokenStore(TokenStore)
+ {
+ }
+
+FNetToken FGameplayTagTokenStore::GetOrCreateToken(FGameplayTag Tag)
+{ 
+	return FNameTokenStore::GetOrCreateToken(Tag.GetTagName());
+}
+
+FGameplayTag FGameplayTagTokenStore::ResolveToken(FNetToken Token, const FNetTokenStoreState* RemoteTokenStoreState) const
+{
+	const FName TagName = FNameTokenStore::ResolveToken(Token, RemoteTokenStoreState);
+	return TagName.IsNone() ? FGameplayTag() : UGameplayTagsManager::Get().RequestGameplayTag(TagName);
+}
+
+}
+
+#endif // UE_WITH_IRIS
\ No newline at end of file
diff --git a/Engine/Source/Runtime/GameplayTags/Private/GameplayTagsManager.cpp b/Engine/Source/Runtime/GameplayTags/Private/GameplayTagsManager.cpp
--- a/Engine/Source/Runtime/GameplayTags/Private/GameplayTagsManager.cpp
+++ b/Engine/Source/Runtime/GameplayTags/Private/GameplayTagsManager.cpp
@@ -547,6 +547,7 @@ void UGameplayTagsManager::ConstructGameplayTagTree()
 			}
 
 			bUseFastReplication = MutableDefault->FastReplication;
+			bUseDynamicReplication = MutableDefault->bDynamicReplication;
 			bShouldWarnOnInvalidTags = MutableDefault->WarnOnInvalidTags;
 			NumBitsForContainerSize = MutableDefault->NumBitsForContainerSize;
 			NetIndexFirstBitSegment = MutableDefault->NetIndexFirstBitSegment;
diff --git a/Engine/Source/Runtime/GameplayTags/Private/GameplayTagsSettings.cpp b/Engine/Source/Runtime/GameplayTags/Private/GameplayTagsSettings.cpp
--- a/Engine/Source/Runtime/GameplayTags/Private/GameplayTagsSettings.cpp
+++ b/Engine/Source/Runtime/GameplayTags/Private/GameplayTagsSettings.cpp
@@ -64,6 +64,7 @@ UGameplayTagsSettings::UGameplayTagsSettings(const FObjectInitializer& ObjectIni
 	ImportTagsFromConfig = true;
 	WarnOnInvalidTags = true;
 	FastReplication = false;
+	bDynamicReplication = true;
 	AllowEditorTagUnloading = true;
 	AllowGameTagUnloading = false;
 	InvalidTagCharacters = ("\"',");
diff --git a/Engine/Source/Runtime/GameplayTags/Public/GameplayTagTokenStore.h b/Engine/Source/Runtime/GameplayTags/Public/GameplayTagTokenStore.h
new file mode 100644
--- /dev/null
+++ b/Engine/Source/Runtime/GameplayTags/Public/GameplayTagTokenStore.h
@@ -0,0 +1,34 @@
+// Copyright Epic Games, Inc. All Rights Reserved.
+
+#pragma once
+#include "CoreMinimal.h"
+
+// NetTokenExports is still an experimental feature that depends on Iris code.
+#if UE_WITH_IRIS
+
+#include "Iris/ReplicationSystem/NameTokenStore.h"
+
+struct FGameplayTag;
+namespace UE::Net
+{
+
+// For now, this is just a specialization of NameTokenStore
+class FGameplayTagTokenStore : public FNameTokenStore
+{
+	UE_NONCOPYABLE(FGameplayTagTokenStore);
+public:
+	GAMEPLAYTAGS_API explicit FGameplayTagTokenStore(FNetTokenStore& TokenStore);
+
+	// Create a NetToken for the provided name
+	GAMEPLAYTAGS_API FNetToken GetOrCreateToken(FGameplayTag Tag);
+
+	// Resolve NetToken, to resolve remote tokens RemoteTokenStoreState must be valid
+	GAMEPLAYTAGS_API FGameplayTag ResolveToken(FNetToken Token, const FNetTokenStoreState* RemoteTokenStoreState = nullptr) const;
+
+	static FName GetTokenStoreName() { return NameTokenStoreName; }
+private:
+	inline static FName NameTokenStoreName = TEXT("GameplayTagTokenStore");
+};
+
+}
+#endif // UE_WITH_IRIS
\ No newline at end of file
diff --git a/Engine/Source/Runtime/Net/Core/Private/Net/Core/NetToken/NetTokenExportContext.cpp b/Engine/Source/Runtime/Net/Core/Private/Net/Core/NetToken/NetTokenExportContext.cpp
new file mode 100644
--- /dev/null
+++ b/Engine/Source/Runtime/Net/Core/Private/Net/Core/NetToken/NetTokenExportContext.cpp
@@ -0,0 +1,77 @@
+// Copyright Epic Games, Inc. All Rights Reserved.
+
+#include "Net/Core/NetToken/NetTokenExportContext.h"
+#include "UObject/CoreNet.h"
+
+DEFINE_LOG_CATEGORY(LogNetToken);
+
+namespace UE::Net
+{
+
+// Simple scope to make sure we set the correct ExportContext and restore the old one when we exit the scope
+FNetTokenExportScope::FNetTokenExportScope(FNetBitWriter& InNetBitWriter, UE::Net::FNetTokenStore* InNetTokenStore, FNetTokenExportContext::FNetTokenExports& TargetExports, const char* InDebugName)
+: ExportContext(InNetTokenStore, &TargetExports)
+, NetBitWriter(InNetBitWriter)
+, OldExportContext(InNetBitWriter.NetTokenExportContext.Get())
+, DebugName(InDebugName)
+{
+	NetBitWriter.NetTokenExportContext.Set(&ExportContext);
+}
+
+FNetTokenExportScope::~FNetTokenExportScope()
+{
+	// Restore old export context
+	NetBitWriter.NetTokenExportContext.Set(OldExportContext);
+
+	if (!ExportContext.NetTokensPendingExport->IsEmpty() && (UE_GET_LOG_VERBOSITY(LogNetToken) >= ELogVerbosity::Verbose))
+	{
+		UE_LOG(LogNetToken, Verbose, TEXT("FNetTokenExportScope %hs added %d pending exports"), DebugName, ExportContext.NetTokensPendingExport->Num());
+		for (UE::Net::FNetToken NetToken : *ExportContext.NetTokensPendingExport)
+		{
+			UE_LOG(LogNetToken, Verbose, TEXT("Pending export %s"), *NetToken.ToString());
+		}
+	}
+}
+
+FNetTokenExportContext::FNetTokenExportContext(FNetTokenStore* InNetTokenStore, FNetTokenExports* InNetTokensPendingExport)
+: NetTokenStore(InNetTokenStore)
+, NetTokensPendingExport(InNetTokensPendingExport)
+{
+}
+
+FNetTokenExportContext* FNetTokenExportContext::GetNetTokenExportContext(FArchive& Ar)
+{
+	if (Ar.IsSaving() && ensureMsgf(Ar.IsNetArchive(), TEXT("Trying to export net tokens for archive that is not a NetBitWriter")))
+	{
+		// Assume that we are a NetBitWriter so that we have a way to pass NetTokenExportContext to NetSerialize() methods
+		FNetBitWriter& NetBitWriter = static_cast<FNetBitWriter&>(Ar);
+		return NetBitWriter.NetTokenExportContext.Get();
+	}
+
+	return nullptr;
+}
+
+void FNetTokenExportContext::AppendNetTokensPendingExport(TArrayView<const UE::Net::FNetToken> NetTokens)
+{ 
+	if (ensure(NetTokensPendingExport))
+	{
+		NetTokensPendingExport->Append(NetTokens);
+	}
+}
+
+void FNetTokenExportContext::AddNetTokenPendingExport(UE::Net::FNetToken NetToken)
+{
+	if (ensure(NetTokensPendingExport))
+	{
+		NetTokensPendingExport->Add(NetToken);
+		UE_LOG(LogNetToken, Verbose, TEXT("FNetTokenExportContext::AddNetTokenPendingExport %s"), *NetToken.ToString());
+	}	
+}
+
+TArrayView<const UE::Net::FNetToken> FNetTokenExportContext::GetNetTokensPendingExport() const
+{ 
+	return NetTokensPendingExport ? MakeArrayView(*NetTokensPendingExport) : MakeArrayView<const UE::Net::FNetToken>(nullptr, 0U);
+}
+
+
+}
diff --git a/Engine/Source/Runtime/Net/Core/Public/Net/Core/Connection/NetCloseResult.h b/Engine/Source/Runtime/Net/Core/Public/Net/Core/Connection/NetCloseResult.h
--- a/Engine/Source/Runtime/Net/Core/Public/Net/Core/Connection/NetCloseResult.h
+++ b/Engine/Source/Runtime/Net/Core/Public/Net/Core/Connection/NetCloseResult.h
@@ -195,6 +195,9 @@ enum class ENetCloseResult : uint32
 	/** ReceiveNetGUIDBunch serialization failed */
 	ReceivedNetGUIDBunchFail,
 
+	/** ReceiveNetExtBunch serialization failed */
+	ReceivedNetExtBunchFail,
+
 	/** Too many reliable bunches queued */
 	MaxReliableExceeded,
 
diff --git a/Engine/Source/Runtime/Net/Core/Public/Net/Core/NetToken/NetToken.h b/Engine/Source/Runtime/Net/Core/Public/Net/Core/NetToken/NetToken.h
new file mode 100644
--- /dev/null
+++ b/Engine/Source/Runtime/Net/Core/Public/Net/Core/NetToken/NetToken.h
@@ -0,0 +1,124 @@
+// Copyright Epic Games, Inc. All Rights Reserved.
+
+#pragma once
+
+#include "CoreTypes.h"
+#include "Containers/UnrealString.h"
+#include "Templates/TypeHash.h"
+#include "Logging/LogMacros.h"
+
+NETCORE_API DECLARE_LOG_CATEGORY_EXTERN(LogNetToken, Log, All);
+
+namespace UE::Net
+{
+	class FNetTokenStore;
+	class FNetTokenStoreState;
+}
+
+namespace UE::Net
+{
+
+class FNetToken
+{
+public:
+	typedef uint32 FTypeId;
+
+	enum : uint32 { Invalid = 0U };
+
+	enum : uint32 { InvalidTokenTypeId = ~FTypeId(0) };
+	enum : uint32 { InvalidTokenIndex = 0U };
+
+	/** How many bits we use to express the TypeId for NetTokens, Increasing this value will break network compatibility and might need versioning for replays. */
+	enum : uint32 { TokenTypeIdBits = 3U };
+
+	/** How many bits we use to express the index part of NetTokens */
+	enum : uint32 { TokenBits = 20U };
+
+	enum : uint32 { MaxTypeIdCount = 1U << TokenTypeIdBits };
+	enum : uint32 { MaxNetTokenCount = 1U << TokenBits };
+
+	enum class ENetTokenAuthority
+	{
+		None,
+		Authority,
+	};
+
+public:	
+	FNetToken() : Value(Invalid) {}
+	inline bool IsValid() const { return Value != Invalid; }
+	bool IsAssignedByAuthority() const { return bIsAssignedByAuthority != 0U; }
+	uint32 GetIndex() const { return Index; }
+	bool operator==(const FNetToken& Other) const { return IsAssignedByAuthority() == Other.IsAssignedByAuthority() && Index == Other.Index; }
+	FString ToString() const;
+
+	friend uint32 GetTypeHash(const FNetToken& Token)
+	{
+		return ::GetTypeHash(Token.Value);
+	}
+	
+	static FNetToken MakeNetToken(uint32 Index, ENetTokenAuthority Authority) { check(Index < MaxNetTokenCount); return FNetToken(Index, Authority); }
+	
+private:
+	explicit FNetToken(uint32 InIndex, ENetTokenAuthority Authority) { Padding = 0U, Index = InIndex, bIsAssignedByAuthority = Authority == ENetTokenAuthority::Authority ? 1U : 0U; }
+
+private:
+
+	union 
+	{
+		struct
+		{
+			uint32 Index : TokenBits;
+			uint32 bIsAssignedByAuthority : 1U;
+			uint32 Padding : 32 - TokenBits - 1U;
+		};
+		uint32 Value;
+	};
+};
+
+class FNetTokenStoreKey
+{
+public:
+	enum { Invalid = 0U };
+
+	FNetTokenStoreKey() : Value(Invalid) {}
+	bool IsValid() const { return Value != Invalid; }
+	uint32 GetKeyValue() const { return Key; }
+	FNetToken::FTypeId GetTypeId() const { return TypeId; }
+	bool operator==(const FNetTokenStoreKey& Other) const { return Value == Other.Value; }
+
+private:
+	friend class FNetTokenStore;
+	friend class FNetTokenDataStore;
+
+private:
+	union
+	{
+		struct
+		{
+			uint32 TypeId : FNetToken::TokenTypeIdBits;
+			uint32 Key : FNetToken::TokenBits;
+		};
+		uint32 Value;
+	};
+};
+
+static_assert(sizeof(FNetTokenStoreKey) == sizeof(uint32), "FNetTokenKey should fit in a uint32");
+
+// Contains necessary context to resolve NetTokens
+class FNetTokenResolveContext
+{
+public:
+	FNetTokenStore* NetTokenStore = nullptr;
+	const FNetTokenStoreState* RemoteNetTokenStoreState = nullptr;
+};
+
+inline FString FNetToken::ToString() const
+{
+	FString Result;
+	Result = FString::Printf(TEXT("NetToken (Auth:%u Index=%u)"), IsAssignedByAuthority(), Index);
+	return Result;
+}
+
+}
+
+template <> struct TIsPODType<UE::Net::FNetToken> { enum { Value = true }; };
diff --git a/Engine/Source/Runtime/Net/Core/Public/Net/Core/NetToken/NetTokenExportContext.h b/Engine/Source/Runtime/Net/Core/Public/Net/Core/NetToken/NetTokenExportContext.h
new file mode 100644
--- /dev/null
+++ b/Engine/Source/Runtime/Net/Core/Public/Net/Core/NetToken/NetTokenExportContext.h
@@ -0,0 +1,62 @@
+// Copyright Epic Games, Inc. All Rights Reserved.
+
+#pragma once
+
+#include "Containers/Array.h"
+#include "NetToken.h"
+
+class FNetBitWriter;
+
+namespace UE::Net
+{
+	class FNetTokenStore;
+	class FNetTokenStoreState;
+}
+
+namespace UE::Net
+{
+
+// Contains necessary context to export NetTokens
+class FNetTokenExportContext
+{
+public:
+	typedef TArray<UE::Net::FNetToken, TInlineAllocator<4>> FNetTokenExports;
+
+	inline FNetTokenStore* GetNetTokenStore();		
+	NETCORE_API void AddNetTokenPendingExport(UE::Net::FNetToken NetToken);
+	NETCORE_API void AppendNetTokensPendingExport(TArrayView<const UE::Net::FNetToken> NetTokens);
+	NETCORE_API TArrayView<const UE::Net::FNetToken> GetNetTokensPendingExport() const;
+	NETCORE_API static FNetTokenExportContext* GetNetTokenExportContext(FArchive& Ar);
+
+private:
+	friend class FNetTokenExportScope;
+
+	FNetTokenExportContext(FNetTokenStore* InNetTokenStore, FNetTokenExports* InNetTokensPendingExport);
+
+	// We can only export if we have a NetTokenStore
+	FNetTokenStore* NetTokenStore = nullptr;
+
+	// We also need a target where to store exports.
+	FNetTokenExports* NetTokensPendingExport = nullptr;
+};
+
+// Simple scope to make sure we set the correct ExportContext and restore the old one when we exit the scope
+class FNetTokenExportScope
+{
+public:
+	NETCORE_API FNetTokenExportScope(FNetBitWriter& InNetBitWriter, UE::Net::FNetTokenStore* InNetTokenStore, FNetTokenExportContext::FNetTokenExports& TargetExports, const char* InDebugName = "None");
+	NETCORE_API ~FNetTokenExportScope();
+
+	FNetTokenExportContext ExportContext;
+	FNetBitWriter& NetBitWriter;
+	FNetTokenExportContext* OldExportContext;
+	const char* DebugName;
+};
+
+FNetTokenStore* FNetTokenExportContext::GetNetTokenStore()
+{ 
+	return NetTokenStore; 
+}
+
+}
+
diff --git a/Engine/Source/Runtime/Net/Core/Public/Net/Core/Trace/Private/NetTraceInternal.h b/Engine/Source/Runtime/Net/Core/Public/Net/Core/Trace/Private/NetTraceInternal.h
--- a/Engine/Source/Runtime/Net/Core/Public/Net/Core/Trace/Private/NetTraceInternal.h
+++ b/Engine/Source/Runtime/Net/Core/Public/Net/Core/Trace/Private/NetTraceInternal.h
@@ -73,6 +73,7 @@ union FNetTraceBunchInfo
 		uint64 ChannelCloseReason : 4;
 		uint64 bPartial : 1;
 		uint64 bPartialInitial : 1;
+		uint64 bPartialCustomExportsFinal : 1;
 		uint64 bPartialFinal : 1;
 		UE_DEPRECATED(5.3, "Replication pausing is deprecated")
 		uint64 bIsReplicationPaused : 1;
@@ -81,7 +82,7 @@ union FNetTraceBunchInfo
 		uint64 bReliable : 1;
 		uint64 bHasPackageMapExports : 1;
 		uint64 bHasMustBeMappedGUIDs : 1;
-		uint64 Padding : 19;
+		uint64 Padding : 18;
 	};
 	uint64 Value;
 };
@@ -299,6 +300,7 @@ FNetTraceBunchInfo MakeBunchInfo(const T& Bunch)
 	BunchInfo.ChannelCloseReason = uint64(Bunch.CloseReason);
 	BunchInfo.bPartial = Bunch.bPartial;
 	BunchInfo.bPartialInitial = Bunch.bPartialInitial;
+	BunchInfo.bPartialCustomExportsFinal = Bunch.bPartialCustomExportsFinal;
 	BunchInfo.bPartialFinal = Bunch.bPartialFinal;
 	PRAGMA_DISABLE_DEPRECATION_WARNINGS
 	BunchInfo.bIsReplicationPaused = Bunch.bIsReplicationPaused;
