Restore the input name of the primary input of TransformPosition material expression as this name is used by plugins.

#rnx

diff --git a/Engine/Source/Runtime/Engine/Private/Materials/MaterialExpressions.cpp b/Engine/Source/Runtime/Engine/Private/Materials/MaterialExpressions.cpp
--- a/Engine/Source/Runtime/Engine/Private/Materials/MaterialExpressions.cpp
+++ b/Engine/Source/Runtime/Engine/Private/Materials/MaterialExpressions.cpp
@@ -13653,15 +13653,17 @@ FExpressionInput* UMaterialExpressionTransformPosition::GetInput(int32 InputInde
 }
 #undef IF_INPUT_RETURN
 
-#define IF_INPUT_RETURN(Name) if(!InputIndex) return Name; --InputIndex
 FName UMaterialExpressionTransformPosition::GetInputName(int32 InputIndex) const
 {
-	IF_INPUT_RETURN(TEXT(""));
-	IF_INPUT_RETURN(TEXT("Periodic World Tile Size"));
+	const FExpressionInput* FoundInput = static_cast<const UMaterialExpression*>(this)->GetInput(InputIndex);
 
-	return TEXT("");
+	if (FoundInput == &PeriodicWorldTileSize)
+	{
+		return TEXT("Periodic World Tile Size");
+	}
+
+	return Super::GetInputName(InputIndex);
 }
-#undef IF_INPUT_RETURN
 
 void UMaterialExpressionTransformPosition::PostEditChangeProperty(struct FPropertyChangedEvent& PropertyChangedEvent)
 {
