Add Substrate's auto-conversion versioning.

diff --git a/Engine/Source/Editor/MaterialEditor/Private/MaterialEditor.cpp b/Engine/Source/Editor/MaterialEditor/Private/MaterialEditor.cpp
--- a/Engine/Source/Editor/MaterialEditor/Private/MaterialEditor.cpp
+++ b/Engine/Source/Editor/MaterialEditor/Private/MaterialEditor.cpp
@@ -201,6 +201,18 @@ TAutoConsoleVariable<bool> CVarMaterialEdAllowIgnoringCompilationErrors(
 	true,
 	TEXT("Allow ignoring compilation errors of platform shaders and derived materials."));
 
+static void InvalidateSubstrateConversionVersion(UMaterial* Material)
+{
+	// The material has been updated, invalid the AutoConversion version
+	if (Substrate::IsSubstrateEnabled())
+	{
+		if (UMaterialEditorOnlyData* WritableEditorOnly = Material->GetEditorOnlyData())
+		{
+			WritableEditorOnly->ResetSubstrateConversionVersion();
+		}
+	}
+}
+
 ///////////////////////////
 // FMatExpressionPreview //
 ///////////////////////////
@@ -6778,6 +6790,7 @@ void FMaterialEditor::CreateDerivedMaterialInstancesPreviews()
 void FMaterialEditor::UpdateMaterialAfterGraphChange()
 {
 	FlushRenderingCommands();
+	InvalidateSubstrateConversionVersion(Material);
 	Material->MaterialGraph->LinkMaterialExpressionsFromGraph();
 
 	// Update the current preview material.
@@ -7374,6 +7387,7 @@ void FMaterialEditor::NotifyPostChange( const FPropertyChangedEvent& PropertyCha
 		&& PropertyName != GET_MEMBER_NAME_CHECKED(UMaterialExpressionTextureSampleParameter, ChannelNames))
 		{
 			// Update the current preview material.
+			InvalidateSubstrateConversionVersion(Material);
 			UpdatePreviewMaterial();
 			RefreshExpressionPreviews();
 			RegenerateCodeView();
diff --git a/Engine/Source/Runtime/Engine/Private/Materials/Material.cpp b/Engine/Source/Runtime/Engine/Private/Materials/Material.cpp
--- a/Engine/Source/Runtime/Engine/Private/Materials/Material.cpp
+++ b/Engine/Source/Runtime/Engine/Private/Materials/Material.cpp
@@ -3289,6 +3289,11 @@ EBlendMode ConvertLegacyBlendMode(EBlendMode InBlendMode, FMaterialShadingModelF
 
 #if WITH_EDITOR
 
+static int32 GetSubstrateConversionVersion()
+{ 
+	return 0;
+}
+
 bool UMaterial::ConvertMaterialToSubstrateMaterial()
 {
 	/*
@@ -3312,6 +3317,16 @@ bool UMaterial::ConvertMaterialToSubstrateMaterial()
 		return false;
 	}
 
+	UMaterialEditorOnlyData* EditorOnly = GetEditorOnlyData();
+
+	// * If the current material has already been converted, skip the conversion
+	// * Store the current version of Subtrate's auto-conversion
+	//   This allows to version the conversion, and safely update auto-converted materials if they have been saved.
+	if (EditorOnly->SubstrateConversionVersion == GetSubstrateConversionVersion())
+	{
+		return false;
+	}
+
 	// Store current node post from the root node.
 	int32 CurrentNodePosX = EditorX;
 	const int32 TranslationOffsetX = 350.0f;
@@ -3405,8 +3420,6 @@ bool UMaterial::ConvertMaterialToSubstrateMaterial()
 		}
 	};
 
-	UMaterialEditorOnlyData* EditorOnly = GetEditorOnlyData();
-
 	bool bCustomNodesGathered = false;
 	UMaterialExpressionThinTranslucentMaterialOutput* ThinTranslucentOutput = nullptr;
 	UMaterialExpressionSingleLayerWaterMaterialOutput* SingleLayerWaterOutput = nullptr;
@@ -3459,6 +3472,7 @@ bool UMaterial::ConvertMaterialToSubstrateMaterial()
 	const bool bRequireNoSubsurfaceProfile	= !bHasShadingModelMixture && (ShadingModel == MSM_Subsurface || ShadingModel == MSM_PreintegratedSkin); // Insure there is no profile, as this would take priority otherwise
 
 	bool bInvalidateShader = false;
+	bool bEmptyShader = false;
 	bool bRelinkCustomOutputNodes = false;
 	UMaterialExpressionSubstrateShadingModels* ConvertNode = nullptr;
 	// Connect all the legacy pin into the conversion node
@@ -3646,6 +3660,7 @@ bool UMaterial::ConvertMaterialToSubstrateMaterial()
 			}
 			bRelinkCustomOutputNodes = false;
 			bInvalidateShader = true;
+			bEmptyShader = true;
 		}
 		else if (!bUseMaterialAttributes && !EditorOnly->FrontMaterial.IsConnected() && !bHasAnySubstrateNodes)
 		{
@@ -3954,6 +3969,12 @@ bool UMaterial::ConvertMaterialToSubstrateMaterial()
 		UpdateCachedExpressionData();
 	}
 
+	// Store the conversion version, only for converted shaders
+	if (bInvalidateShader && !bEmptyShader)
+	{
+		EditorOnly->SubstrateConversionVersion = GetSubstrateConversionVersion();
+	}
+
 	return bInvalidateShader;
 }
 #endif // WITH_EDITOR
@@ -4895,13 +4916,23 @@ void UMaterial::PostEditChangePropertyInternal(FPropertyChangedEvent& PropertyCh
 
 	const UMaterialEditorOnlyData* EditorOnly = GetEditorOnlyData();
 
-	// Apply Substrate material conversion if needed.
-	// For imported materials, the conversion won't happen during the PostLoad() phase. 
-	// Usually material importers create material by code, and then calls PostEditChange(), which calls PostEditChangeProperty()
-	// This late conversion ensures that imported materials are converted properly with Substrate.
-	if (Substrate::IsSubstrateEnabled() && !IsMaterialAlreadyConvertedToSubstrate(this, EditorOnly))
+	if (Substrate::IsSubstrateEnabled())
 	{
-		ConvertMaterialToSubstrateMaterial();
+		// If a modification is made to the material, revert the auto-conversion version, to mark the material as 'edited'
+		const bool bResetConversionVersion = PropertyThatChanged != nullptr || PropertyChangedEvent.ChangeType != EPropertyChangeType::Unspecified;
+		if (bResetConversionVersion && EditorOnly)
+		{
+			GetEditorOnlyData()->ResetSubstrateConversionVersion();
+		}
+
+		// Apply Substrate material conversion if needed.
+		// For imported materials, the conversion won't happen during the PostLoad() phase. 
+		// Usually material importers create material by code, and then calls PostEditChange(), which calls PostEditChangeProperty()
+		// This late conversion ensures that imported materials are converted properly with Substrate.
+		if (!IsMaterialAlreadyConvertedToSubstrate(this, EditorOnly))
+		{
+			ConvertMaterialToSubstrateMaterial();
+		}
 	}
 
 	// Check for distortion in material 
@@ -7646,6 +7677,11 @@ UMaterialEditorOnlyData::UMaterialEditorOnlyData()
 	AmbientOcclusion.Constant = FMaterialAttributeDefinitionMap::GetDefaultValue(MP_AmbientOcclusion).X;
 	Refraction.Constant = FMaterialAttributeDefinitionMap::GetDefaultValue(MP_Refraction).X;
 	SurfaceThickness.Constant = FMaterialAttributeDefinitionMap::GetDefaultValue(MP_SurfaceThickness).X;
+	ResetSubstrateConversionVersion();
 }
 
+void UMaterialEditorOnlyData::ResetSubstrateConversionVersion()
+{
+	SubstrateConversionVersion = -1; // Invalid version
+}
 #undef LOCTEXT_NAMESPACE
diff --git a/Engine/Source/Runtime/Engine/Public/Materials/Material.h b/Engine/Source/Runtime/Engine/Public/Materials/Material.h
--- a/Engine/Source/Runtime/Engine/Public/Materials/Material.h
+++ b/Engine/Source/Runtime/Engine/Public/Materials/Material.h
@@ -400,6 +400,13 @@ public:
 	/** Controls where this parameter group is displayed in a material instance parameter list.  The lower the number the higher up in the parameter list. */
 	UPROPERTY(EditAnywhere, EditFixedSize, Category = "Group Sorting")
 	TArray<FParameterGroupData> ParameterGroupData;
+
+	/** Store the version of the Substrate's auto-conversion applied to this material if any (default, no conversion =-1). */
+	UPROPERTY()
+	int32 SubstrateConversionVersion;
+
+	/** Reset the version of the Substrate's auto-conversion applied to this material */
+	ENGINE_API void ResetSubstrateConversionVersion();
 };
 
 /**
