[SkinnedAsset]
 - Fix race condition when building multiple groom bindings having dependencies on the same skeletal mesh leading to SerializeLODInfoForDDC calls from multiple threads

#rnx

diff --git a/Engine/Source/Runtime/Engine/Classes/Engine/SkinnedAssetCommon.h b/Engine/Source/Runtime/Engine/Classes/Engine/SkinnedAssetCommon.h
--- a/Engine/Source/Runtime/Engine/Classes/Engine/SkinnedAssetCommon.h
+++ b/Engine/Source/Runtime/Engine/Classes/Engine/SkinnedAssetCommon.h
@@ -245,10 +245,95 @@ struct FSkeletalMeshLODInfo
 	UPROPERTY()
 	uint8 bImportWithBaseMesh:1;
 
+	// Protects access to the build guid.
+	struct FThreadSafeBuildGUID
+	{
+	private:
+		FGuid   BuildGuid;
+		mutable UE::FMutex Mutex;
+	public:
+		FThreadSafeBuildGUID() = default;
+		FThreadSafeBuildGUID(const FGuid& InGuid)
+			: BuildGuid(InGuid)
+		{
+		}
+
+		FThreadSafeBuildGUID(const FThreadSafeBuildGUID& Other)
+		{
+			UE::TUniqueLock OtherLock(Other.Mutex);
+			UE::TUniqueLock Lock(Mutex);
+			BuildGuid = Other.BuildGuid;
+		}
+
+		FThreadSafeBuildGUID& operator=(const FThreadSafeBuildGUID& Other)
+		{
+			if (&Other != this)
+			{
+				UE::TUniqueLock OtherLock(Other.Mutex);
+				UE::TUniqueLock Lock(Mutex);
+				BuildGuid = Other.BuildGuid;
+			}
+			return *this;
+		}
+
+		bool operator==(const FGuid& Other) const
+		{
+			UE::TUniqueLock Lock(Mutex);
+			return BuildGuid == Other;
+		}
+
+		bool operator==(const FThreadSafeBuildGUID& Other) const
+		{
+			if (&Other == this)
+			{
+				return true;
+			}
+
+			UE::TUniqueLock OtherLock(Other.Mutex);
+			UE::TUniqueLock Lock(Mutex);
+			return BuildGuid == Other.BuildGuid;
+		}
+
+		void operator=(const FGuid& Other)
+		{
+			UE::TUniqueLock Lock(Mutex);
+			BuildGuid = Other;
+		}
+
+		void Invalidate()
+		{
+			UE::TUniqueLock Lock(Mutex);
+			BuildGuid.Invalidate();
+		}
+
+		operator FGuid() const
+		{
+			UE::TUniqueLock Lock(Mutex);
+			return BuildGuid;
+		}
+
+		bool IsValid() const
+		{
+			UE::TUniqueLock Lock(Mutex);
+			return BuildGuid.IsValid();
+		}
+
+		FString ToString(EGuidFormats Format = EGuidFormats::Digits) const
+		{
+			UE::TUniqueLock Lock(Mutex);
+			FString Out;
+			BuildGuid.AppendString(Out, Format);
+			return Out;
+		}
+
+		FThreadSafeBuildGUID(FThreadSafeBuildGUID&&) = delete;
+		FThreadSafeBuildGUID& operator=(FThreadSafeBuildGUID&&) = delete;
+	};
+
 	//Temporary build GUID data
 	//We use this GUID to store the LOD Key so we can know if the LOD needs to be rebuilt
 	//This GUID is set when we Cache the render data (build function)
-	FGuid BuildGUID;
+	FThreadSafeBuildGUID BuildGUID;
 
 	ENGINE_API FGuid ComputeDeriveDataCacheKey(const FSkeletalMeshLODGroupSettings* SkeletalMeshLODGroupSettings);
 #endif
