- Added incremental pre-registration of actor/components in level streaming (optional).
- When enabled, allows asynchronous physics state creation when making a streaming level visible (depends on p.Chaos.EnableAsyncInitBody being set to true).
- Can be enabled with s.AllowIncrementalPreRegisterComponents=true.

diff --git a/Engine/Source/Runtime/Engine/Classes/Components/ActorComponent.h b/Engine/Source/Runtime/Engine/Classes/Components/ActorComponent.h
--- a/Engine/Source/Runtime/Engine/Classes/Components/ActorComponent.h
+++ b/Engine/Source/Runtime/Engine/Classes/Components/ActorComponent.h
@@ -12,6 +12,7 @@
 #include "UObject/ScriptMacros.h"
 #include "Interfaces/Interface_AssetUserData.h"
 #include "UObject/StructOnScope.h"
+#include "PhysicsInterfaceDeclaresCore.h"
 #include "PropertyPairsMap.h"
 #include "ComponentInstanceDataCache.h"
 #include "Experimental/ConcurrentLinearAllocator.h"
@@ -37,12 +38,33 @@ class FRegisterComponentContext
 public:
 	FRegisterComponentContext(UWorld* InWorld)
 		: World(InWorld)
+		, IncrementalUpdateTimeLimitMs(0)
+		, IncrementalUpdateAddPrimitiveBatchSize(0)
+		, IncrementalUpdateNumComponentsToUpdate(0)
 	{}
 
+	~FRegisterComponentContext()
+	{
+		Process();
+	}
+
+	FRegisterComponentContext& SetIncrementalUpdateTimeLimitMs(double Value) { IncrementalUpdateTimeLimitMs = Value; return *this; }
+	FRegisterComponentContext& SetIncrementalUpdateAddPrimitiveBatchSize(int32 Value) { IncrementalUpdateAddPrimitiveBatchSize = Value; return *this; }
+	FRegisterComponentContext& SetIncrementalUpdateNumComponentsToUpdate(int32 Value) { IncrementalUpdateNumComponentsToUpdate = Value; return *this; }
+
+	double GetIncrementalUpdateTimeLimitMs() const { return IncrementalUpdateTimeLimitMs; }
+	int32 GetIncrementalUpdateAddPrimitiveBatchSize() const { return IncrementalUpdateAddPrimitiveBatchSize; }
+	int32 GetIncrementalUpdateNumComponentsToUpdate() const { return IncrementalUpdateNumComponentsToUpdate; }
+
 	void AddPrimitive(UPrimitiveComponent* PrimitiveComponent)
 	{
 		checkSlow(!AddPrimitiveBatches.Contains(PrimitiveComponent));
 		AddPrimitiveBatches.Add(PrimitiveComponent);
+		
+		if ((IncrementalUpdateAddPrimitiveBatchSize > 0) && (AddPrimitiveBatches.Num() > IncrementalUpdateAddPrimitiveBatchSize))
+		{
+			Process();
+		}
 	}
 
 	void AddSendRenderDynamicData(UPrimitiveComponent* PrimitiveComponent)
@@ -53,11 +75,16 @@ public:
 
 	ENGINE_API static void SendRenderDynamicData(FRegisterComponentContext* Context, UPrimitiveComponent* PrimitiveComponent);
 
+	UE_DEPRECATED(5.5, "This function is not used anymore")
 	int32 Count() const { return AddPrimitiveBatches.Num(); }
+
 	void Process();
 
 private:
 	UWorld* World;
+	double	IncrementalUpdateTimeLimitMs;
+	int32	IncrementalUpdateAddPrimitiveBatchSize;
+	int32	IncrementalUpdateNumComponentsToUpdate;
 	TArray<UPrimitiveComponent*, FConcurrentLinearArrayAllocator> AddPrimitiveBatches;
 	TArray<UPrimitiveComponent*, FConcurrentLinearArrayAllocator> SendRenderDynamicDataPrimitives;
 };
@@ -169,6 +196,20 @@ private:
 	int32 UCSSerializationIndex;
 
 protected:
+
+	enum class EComponentRegistrationState
+	{
+		/** Component is not in any of the listed states, code should rely on bRegistered to determine if component is registered or not. */
+		None,
+		/** Component is pre-registering (not registered) */
+		PreRegistering,
+		/** Component is pre-registered (not registered) */
+		PreRegistered
+	};
+
+	/** Indicates if this ActorComponent is currently pre-registering or pre-registered with a scene. */
+	EComponentRegistrationState RegistrationState;
+
 	/** 
 	 *  Indicates if this ActorComponent is currently registered with a scene. 
 	 */
@@ -188,6 +229,9 @@ protected:
 	/** If the physics state is currently created for this component */
 	uint8 bPhysicsStateCreated:1;
 
+	/** If the physics state was created while component was pre-registering, notification to GameThread is deferred next time CreatePhysicsState is called. */
+	uint8 bCallAsyncPhysicsStateCreatedDelegate:1;
+
 	/** Is this component safe to ID over the network by name?  */
 	UPROPERTY()
 	uint8 bNetAddressable:1;
@@ -331,6 +375,9 @@ private:
 	/** Whether to use use the async physics tick with this component. */
 	uint8 bAsyncPhysicsTickEnabled : 1;
 
+	/** If non-null, physics state creation is done asynchronously by FPhysScene_Chaos::ProcessAsyncPhysicsStateJobs. */
+	FPhysScene* AsyncCreatePhysicsStateScene;
+
 public:
 
 	/** Describes how a component instance will be created */
@@ -357,6 +404,9 @@ private:
 
 public:
 
+	/** Whether asynchronous physics state creation is running or not. */
+	ENGINE_API bool IsAsyncCreatePhysicsStateRunning() const { return AsyncCreatePhysicsStateScene != nullptr; }
+
 	/** Tracks whether the component has been added to one of the world's end of frame update lists */
 	uint32 GetMarkedForEndOfFrameUpdateState() const { return MarkedForEndOfFrameUpdateState; }
 
@@ -665,11 +715,6 @@ private:
 	/** Private version without inlining that does *not* check Dedicated server build flags (which should already have been done). */
 	ENGINE_API ENetMode InternalGetNetMode() const;
 
-protected:
-	/** Return true if this component is in a state where it can be activated normally. */
-	ENGINE_API virtual bool ShouldActivate() const;
-
-private:
 	/** Calls OnUnregister, DestroyRenderState_Concurrent and OnDestroyPhysicsState. */
 	ENGINE_API void ExecuteUnregisterEvents();
 
@@ -678,7 +723,19 @@ private:
 
 	/** Utility function for each of the PostEditChange variations to call for the same behavior */
 	ENGINE_API void ConsolidatedPostEditChange(const FPropertyChangedEvent& PropertyChangedEvent);
+
 protected:
+	/** Return true if this component is in a state where it can be activated normally. */
+	ENGINE_API virtual bool ShouldActivate() const;
+
+	/** Returns whether this component should create its physics state asynchronously. */
+	ENGINE_API virtual bool ShouldAsyncCreatePhysicsState(UWorld* WorldContext) const;
+
+	/** Called by PreRegisterComponentWithWorld when the component is being pre-registered. */
+	ENGINE_API virtual void OnPreRegister();
+
+	/** Called once the component is pre-registered. */
+	ENGINE_API virtual void OnPreRegistered();
 
 	/**
 	 * Called when a component is registered, after Scene is set, but before CreateRenderState_Concurrent or OnCreatePhysicsState are called.
@@ -743,6 +800,19 @@ protected:
 	ENGINE_API virtual void RegisterComponentTickFunctions(bool bRegister);
 
 public:
+
+	/** Returns whether PreRegisterComponentWithWorld should be called for this component when its Level is being added to the world. */
+	ENGINE_API virtual bool ShouldIncrementalPreRegister(UWorld* WorldContext) const;
+
+	/** Returns whether this component allows having its physics state being created and destroyed asynchronously (outside of the GameThread). */
+	ENGINE_API virtual bool AllowsAsyncPhysicsStateCreationAndDestruction() const { return false; }
+
+	/** Used to create any physics engine information for this component outside of the GameThread. */
+	ENGINE_API virtual void OnAsyncCreatePhysicsState();
+
+	/** Called on the GameThread once the component's physic engine information is created. */
+	ENGINE_API virtual void OnAsyncCreatePhysicsStateDone_GameThread();
+
 	/**
 	 * Initializes the component.  Occurs at level startup or actor spawn. This is before BeginPlay (Actor or Component).  
 	 * All Components in the level will be Initialized on load before any Actor/Component gets BeginPlay
@@ -863,6 +933,13 @@ public:
 	UFUNCTION(BlueprintCallable, Category="Components|Tick")
 	ENGINE_API float GetComponentTickInterval() const;
 
+	/**
+	 * Pre-registers a component with a specific world. Only called on components that return ShouldIncrementalPreRegister().
+	 * Will initialize the necessary before calling OnPreRegister().
+	 * @param InWorld - The world to pre-register the component with.
+	 */
+	ENGINE_API void PreRegisterComponentWithWorld(UWorld* InWorld);
+
 	/**
 	 * Registers a component with a specific world, which creates any visual/physical state
 	 * @param InWorld - The world to register the component with.
@@ -1092,6 +1169,12 @@ public:
 	/** See if the owning Actor is currently running the UCS */
 	ENGINE_API bool IsOwnerRunningUserConstructionScript() const;
 
+	/** See if this component is currently pre-registering */
+	FORCEINLINE bool IsPreRegistering() const { return RegistrationState == EComponentRegistrationState::PreRegistering; }
+
+	/** See if this component is currently pre-registered */
+	FORCEINLINE bool IsPreRegistered() const { return RegistrationState == EComponentRegistrationState::PreRegistered; }
+
 	/** See if this component is currently registered */
 	FORCEINLINE bool IsRegistered() const { return bRegistered; }
 
diff --git a/Engine/Source/Runtime/Engine/Classes/Components/PrimitiveComponent.h b/Engine/Source/Runtime/Engine/Classes/Components/PrimitiveComponent.h
--- a/Engine/Source/Runtime/Engine/Classes/Components/PrimitiveComponent.h
+++ b/Engine/Source/Runtime/Engine/Classes/Components/PrimitiveComponent.h
@@ -2424,6 +2424,7 @@ public:
 	ENGINE_API virtual void OnUnregister()  override;
 	ENGINE_API virtual void DestroyRenderState_Concurrent() override;
 	ENGINE_API virtual void OnCreatePhysicsState() override;
+	ENGINE_API virtual void OnAsyncCreatePhysicsStateDone_GameThread() override;
 	ENGINE_API virtual void OnDestroyPhysicsState() override;
 	ENGINE_API virtual void OnActorEnableCollisionChanged() override;
 	ENGINE_API virtual void InvalidateLightingCacheDetailed(bool bInvalidateBuildEnqueuedLighting, bool bTranslationOnly) override;
diff --git a/Engine/Source/Runtime/Engine/Classes/Components/SceneComponent.h b/Engine/Source/Runtime/Engine/Classes/Components/SceneComponent.h
--- a/Engine/Source/Runtime/Engine/Classes/Components/SceneComponent.h
+++ b/Engine/Source/Runtime/Engine/Classes/Components/SceneComponent.h
@@ -314,6 +314,9 @@ private:
 	ENGINE_API void BeginScopedMovementUpdate(FScopedMovementUpdate& ScopedUpdate);
 	ENGINE_API void EndScopedMovementUpdate(FScopedMovementUpdate& ScopedUpdate);
 
+	/** Internally called by OnRegister() or OnPreRegister() to handle parent attachement. */
+	void AttachToParent();
+
 	/** Cache that avoids Quat<->Rotator conversions if possible. Only to be used with GetComponentTransform().GetRotation(). */
 	FRotationConversionCache WorldRotationCache;
 
@@ -886,6 +889,7 @@ public:
 
 	//~ Begin ActorComponent Interface
 	ENGINE_API virtual void GetLifetimeReplicatedProps(TArray<FLifetimeProperty>& OutLifetimeProps) const override;
+	ENGINE_API virtual void OnPreRegister() override;
 	ENGINE_API virtual void OnRegister() override;
 	ENGINE_API virtual void OnUnregister() override;
 	ENGINE_API virtual void EndPlay(EEndPlayReason::Type Reason) override;
diff --git a/Engine/Source/Runtime/Engine/Classes/Components/StaticMeshComponent.h b/Engine/Source/Runtime/Engine/Classes/Components/StaticMeshComponent.h
--- a/Engine/Source/Runtime/Engine/Classes/Components/StaticMeshComponent.h
+++ b/Engine/Source/Runtime/Engine/Classes/Components/StaticMeshComponent.h
@@ -540,6 +540,7 @@ public:
 
 	//~ Begin UActorComponent Interface.
 protected: 
+	ENGINE_API virtual void OnPreRegister() override;
 	ENGINE_API virtual void OnRegister() override;
 	ENGINE_API virtual void OnUnregister() override;
 	ENGINE_API virtual void BeginPlay() override;
@@ -549,6 +550,8 @@ protected:
 	ENGINE_API virtual void OnDestroyPhysicsState() override;
 	ENGINE_API virtual bool ShouldCreatePhysicsState() const override;
 	ENGINE_API virtual bool ShouldCreateRenderState() const override;
+	ENGINE_API virtual bool AllowsAsyncPhysicsStateCreationAndDestruction() const override;
+
 public:
 	ENGINE_API virtual void InvalidateLightingCacheDetailed(bool bInvalidateBuildEnqueuedLighting, bool bTranslationOnly) override;
 	ENGINE_API virtual UObject const* AdditionalStatObject() const override;
diff --git a/Engine/Source/Runtime/Engine/Classes/Engine/Level.h b/Engine/Source/Runtime/Engine/Classes/Engine/Level.h
--- a/Engine/Source/Runtime/Engine/Classes/Engine/Level.h
+++ b/Engine/Source/Runtime/Engine/Classes/Engine/Level.h
@@ -16,6 +16,7 @@
 #include "Engine/MaterialMerging.h"
 #include "Engine/TextureStreamingTypes.h"
 #include "Misc/EditorPathObjectInterface.h"
+#include "Components/ActorComponent.h"
 #include <atomic>
 
 #include "Level.generated.h"
@@ -644,6 +645,7 @@ public:
 	enum class EIncrementalComponentState : uint8
 	{
 		Init,
+		PreRegisterInitialComponents,
 		RegisterInitialComponents,
 #if WITH_EDITOR
 		RunConstructionScripts,
@@ -651,7 +653,7 @@ public:
 		Finalize
 	};
 
-	/** Whether the actor referenced by CurrentActorIndexForUpdateComponents has called PreRegisterAllComponents */
+	UE_DEPRECATED(5.5, "This property shouldn't be used anymore.")
 	uint8										bHasCurrentActorCalledPreRegister:1;
 	/** The current stage for incrementally updating actor components in the level*/
 	EIncrementalComponentState					IncrementalComponentState;
@@ -775,6 +777,13 @@ private:
 	/** Cached level collection that this level is contained in, for faster access than looping through the collections in the world. */
 	FLevelCollection* CachedLevelCollection;
 
+	ENGINE_API void OnPreRegisterComponent(UActorComponent* Component) { check(IsValid(Component)); PreRegisteringActorComponents.Add(Component); }
+	ENGINE_API void OnRemovedPreRegisteringComponent(UActorComponent* Component) { check(IsValid(Component)); verify(PreRegisteringActorComponents.Remove(Component)); }
+	ENGINE_API void OnPreRegisteredComponent(UActorComponent* Component) { OnRemovedPreRegisteringComponent(Component); }
+	ENGINE_API bool HasPreRegisteringComponents() const { return PreRegisteringActorComponents.Num() > 0; }
+	TSet<UActorComponent*> PreRegisteringActorComponents;
+	friend struct FLevelRegistrationAccessor;
+
 protected:
 
 	/** Array of user data stored with the asset */
@@ -1394,7 +1403,8 @@ public:
 #endif
 
 private:
-	bool IncrementalRegisterComponents(bool bPreRegisterComponents, int32 NumComponentsToUpdate, FRegisterComponentContext* Context);
+	bool IncrementalPreRegisterComponents(FRegisterComponentContext& Context);
+	bool IncrementalRegisterComponents(FRegisterComponentContext& Context);
 #if WITH_EDITOR
 	bool IncrementalRunConstructionScripts(bool bProcessAllActors);
 	TOptional<bool> bCachedHasStaticMeshCompilationPending;
@@ -1455,6 +1465,18 @@ private:
 #endif // WITH_EDITOR
 };
 
+/** Accessor used by World/ActorComponent for component registration. */
+struct FLevelRegistrationAccessor
+{
+private:
+	ENGINE_API static void OnPreRegisterComponent(ULevel* InLevel, UActorComponent* InComponent) { InLevel->OnPreRegisterComponent(InComponent); }
+	ENGINE_API static void OnPreRegisteredComponent(ULevel* InLevel, UActorComponent* InComponent) { InLevel->OnPreRegisteredComponent(InComponent); }
+	ENGINE_API static void OnRemovedPreRegisteringComponent(ULevel* InLevel, UActorComponent* InComponent) { InLevel->OnRemovedPreRegisteringComponent(InComponent); }
+	ENGINE_API static bool HasPreRegisteringComponents(const ULevel* InLevel) { return InLevel->HasPreRegisteringComponents(); }
+	friend class UActorComponent;
+	friend class UWorld;
+};
+
 #if WITH_EDITOR
 struct ENGINE_API FScopedOverrideSpawningLevelMountPointObject
 {
diff --git a/Engine/Source/Runtime/Engine/Classes/GameFramework/Actor.h b/Engine/Source/Runtime/Engine/Classes/GameFramework/Actor.h
--- a/Engine/Source/Runtime/Engine/Classes/GameFramework/Actor.h
+++ b/Engine/Source/Runtime/Engine/Classes/GameFramework/Actor.h
@@ -576,6 +576,9 @@ private:
 	/** True if this actor is currently running user construction script (used to defer component registration) */
 	uint8 bRunningUserConstructionScript:1;
 
+	/** True if PreRegisterAllComponents was called on this actor. */
+	uint8 bHasPreRegisteredAllComponents : 1;
+
 	/** Set true just before PostRegisterAllComponents() is called and false just before PostUnregisterAllComponents() is called */
 	uint8 bHasRegisteredAllComponents:1;
 
@@ -1472,6 +1475,12 @@ private:
 	ENGINE_API TArray<const UDataLayerInstance*> GetDataLayerInstancesInternal(bool bUseLevelContext, bool bIncludeParentDataLayers = true) const;
 	ENGINE_API const IWorldPartitionCell* GetWorldPartitionRuntimeCell() const;
 
+	/** Incrementally pre-registers components associated with this actor, used during level streaming */
+	ENGINE_API void IncrementalPreRegisterComponents();
+
+	/** Returns bHasPreRegisteredAllComponents which indicates whether this actor has pre-registered all its components. */
+	ENGINE_API bool HasPreRegisteredAllComponents() const { return bHasPreRegisteredAllComponents; }
+
 	//~=============================================================================
 	// General functions.
 public:
@@ -4299,6 +4308,7 @@ private:
 	friend struct FMarkActorIsBeingDestroyed;
 	friend struct FActorParentComponentSetter;
 	friend struct FSetActorWantsDestroyDuringBeginPlay;
+	friend struct FActorPreRegistrationAccessor;
 #if WITH_EDITOR
 	ENGINE_API bool IsActorFolderValid() const;
 	ENGINE_API void SetFolderPathInternal(const FName& InNewFolderPath, bool bInBroadcastChange = true);
@@ -4570,6 +4580,15 @@ inline FActorLastRenderTime* FActorLastRenderTime::GetPtr(AActor* InActor)
 	return (InActor ? &InActor->LastRenderTime : nullptr);
 }
 
+/** Accessor use by Level for Actor pre-pregistration. */
+struct FActorPreRegistrationAccessor
+{
+private:
+	ENGINE_API static void IncrementalPreRegisterComponents(AActor* InActor) { InActor->IncrementalPreRegisterComponents(); }
+	ENGINE_API static bool HasPreRegisteredAllComponents(const AActor* InActor) { return InActor->HasPreRegisteredAllComponents(); }
+	friend class ULevel;
+};
+
 #if WITH_EDITOR
 struct FSetActorHiddenInSceneOutliner
 {
diff --git a/Engine/Source/Runtime/Engine/Private/Actor.cpp b/Engine/Source/Runtime/Engine/Private/Actor.cpp
--- a/Engine/Source/Runtime/Engine/Private/Actor.cpp
+++ b/Engine/Source/Runtime/Engine/Private/Actor.cpp
@@ -237,6 +237,7 @@ void AActor::InitializeDefaults()
 	DefaultUpdateOverlapsMethodDuringLevelStreaming = EActorUpdateOverlapsMethod::OnlyUpdateMovable;
 	
 	bHasDeferredComponentRegistration = false;
+	bHasPreRegisteredAllComponents = false;
 	bHasRegisteredAllComponents = false;
 
 #if WITH_EDITORONLY_DATA
@@ -3678,6 +3679,8 @@ void AActor::DisableComponentsSimulatePhysics()
 
 void AActor::PreRegisterAllComponents()
 {
+	ensureMsgf(!bHasPreRegisteredAllComponents, TEXT("Actor %s PreRegisterAllComponents has already been called (or bHasPreRegisteredAllComponents was not properly reset)."), *this->GetFullName());
+	bHasPreRegisteredAllComponents = true;
 }
 
 void AActor::PostRegisterAllComponents() 
@@ -5479,6 +5482,8 @@ void AActor::UnregisterAllComponents(const bool bForReregister)
 
 		PostUnregisterAllComponents();
 	}
+
+	bHasPreRegisteredAllComponents = false;
 }
 
 void AActor::PostUnregisterAllComponents()
@@ -5493,12 +5498,17 @@ void AActor::RegisterAllComponents()
 	QUICK_SCOPE_CYCLE_COUNTER(STAT_AActor_RegisterAllComponents);
 	
 	PreRegisterAllComponents();
+	ensureMsgf(HasPreRegisteredAllComponents(), TEXT("Actor %s PreRegisterAllComponents implementation seems be to missing a call to AActor::PreRegisterAllComponents."), *this->GetFullName());
 
 	// 0 - means register all components
 	bool bAllRegistered = IncrementalRegisterComponents(0);
 	check(bAllRegistered);
 }
 
+namespace UE
+{
+namespace Private
+{
 // Walks through components hierarchy and returns closest to root parent component that is unregistered
 // Only for components that belong to the same owner
 static USceneComponent* GetUnregisteredParent(UActorComponent* Component)
@@ -5523,6 +5533,96 @@ static USceneComponent* GetUnregisteredParent(UActorComponent* Component)
 	return ParentComponent;
 }
 
+static USceneComponent* GetNonPreRegisteredParent(UActorComponent* Component)
+{
+	USceneComponent* ParentComponent = nullptr;
+	USceneComponent* SceneComponent = Cast<USceneComponent>(Component);
+
+	while (SceneComponent && 
+		   SceneComponent->GetAttachParent() && 
+		   SceneComponent->GetAttachParent()->GetOwner() == Component->GetOwner() && 
+		   !SceneComponent->GetAttachParent()->IsPreRegistering() &&
+		   !SceneComponent->GetAttachParent()->IsPreRegistered())
+	{
+		SceneComponent = SceneComponent->GetAttachParent();
+		if (SceneComponent->bAutoRegister && IsValidChecked(SceneComponent))
+		{
+			// We found non-preregistered parent that should be preregistered
+			// Keep looking up the tree
+			ParentComponent = SceneComponent;
+		}
+	}
+
+	return ParentComponent;
+}
+} // namespace Private
+} // namespace UE
+
+void AActor::IncrementalPreRegisterComponents()
+{
+	TRACE_CPUPROFILER_EVENT_SCOPE(AActor::IncrementalPreRegisterComponents);
+
+	// First call PreRegisterAllComponents (IncrementalPreRegisterComponents includes PreRegisterAllComponents)
+#if WITH_EDITOR
+	// In PIE, WorldSettings is spawned when world partition cell level is created, AActor::PostSpawnInitialize will 
+	// have already called RegisterAllComponents (which calls PreRegisterAllComponents)
+	if (!HasPreRegisteredAllComponents())
+#endif
+	{
+		PreRegisterAllComponents();
+	}
+	ensureMsgf(HasPreRegisteredAllComponents(), TEXT("Actor %s PreRegisterAllComponents implementation seems be to missing a call to AActor::PreRegisterAllComponents."), *this->GetFullName());
+
+	UWorld* const World = GetWorld();
+	TInlineComponentArray<UActorComponent*> Components;
+	ForEachComponent<UActorComponent>(false, [World, &Components](UActorComponent* Component)
+	{
+		if (Component->ShouldIncrementalPreRegister(World))
+		{
+			Components.Add(Component);
+		}
+	});
+
+	// No component found to pre-register
+	if (Components.IsEmpty())
+	{
+		return;
+	}
+
+	// Pre-register RootComponent first so all other children components can reliably use it (i.e., call GetLocation) when they pre-register
+	if (RootComponent && RootComponent->bAutoRegister && !RootComponent->IsPreRegistered())
+	{
+		check(!RootComponent->IsPreRegistering());
+		RootComponent->PreRegisterComponentWithWorld(World);
+	}
+
+	TSet<UActorComponent*> PreRegisteredParents;
+	for (int32 CompIdx = 0; CompIdx < Components.Num(); CompIdx++)
+	{
+		UActorComponent* Component = Components[CompIdx];
+		if (!Component->IsPreRegistered() && !Component->IsPreRegistering() && Component->bAutoRegister && IsValidChecked(Component))
+		{
+			// Ensure that all parents are pre-registered first
+			USceneComponent* NonPreRegisteredParentComponent = UE::Private::GetNonPreRegisteredParent(Component);
+			if (NonPreRegisteredParentComponent)
+			{
+				bool bParentAlreadyHandled = false;
+				PreRegisteredParents.Add(NonPreRegisteredParentComponent, &bParentAlreadyHandled);
+				if (bParentAlreadyHandled)
+				{
+					UE_LOG(LogActor, Error, TEXT("AActor::IncrementalPreRegisterComponents parent component '%s' cannot be pre-registered in actor '%s'"), *GetPathNameSafe(NonPreRegisteredParentComponent), *GetPathName());
+					break;
+				}
+
+				// Register parent first, then return to this component on a next iteration
+				Component = NonPreRegisteredParentComponent;
+				CompIdx--;
+			}
+			Component->PreRegisterComponentWithWorld(World);
+		}
+	}
+}
+
 bool AActor::IncrementalRegisterComponents(int32 NumComponentsToRegister, FRegisterComponentContext* Context)
 {
 	if (NumComponentsToRegister == 0)
@@ -5569,7 +5669,7 @@ bool AActor::IncrementalRegisterComponents(int32 NumComponentsToRegister, FRegis
 		if (!Component->IsRegistered() && Component->bAutoRegister && IsValidChecked(Component))
 		{
 			// Ensure that all parent are registered first
-			USceneComponent* UnregisteredParentComponent = GetUnregisteredParent(Component);
+			USceneComponent* UnregisteredParentComponent = UE::Private::GetUnregisteredParent(Component);
 			if (UnregisteredParentComponent)
 			{
 				bool bParentAlreadyHandled = false;
diff --git a/Engine/Source/Runtime/Engine/Private/Components/ActorComponent.cpp b/Engine/Source/Runtime/Engine/Private/Components/ActorComponent.cpp
--- a/Engine/Source/Runtime/Engine/Private/Components/ActorComponent.cpp
+++ b/Engine/Source/Runtime/Engine/Private/Components/ActorComponent.cpp
@@ -122,44 +122,47 @@ void FRegisterComponentContext::Process()
 {
 	TRACE_CPUPROFILER_EVENT_SCOPE(FRegisterComponentContext::Process)
 
-	bool bSingleThreaded = !FApp::ShouldUseThreadingForPerformance();
+	if (AddPrimitiveBatches.Num())
+	{
+		bool bSingleThreaded = !FApp::ShouldUseThreadingForPerformance();
 #if WITH_EDITOR
-	// This is required for async static mesh compilation in case a scene proxy is not async aware.
-	// A stall until the compilation is finished might occur, and this is only supported on the game thread for now.
-	bSingleThreaded |= FStaticMeshCompilingManager::Get().IsAsyncStaticMeshCompilationEnabled();
+		// This is required for async static mesh compilation in case a scene proxy is not async aware.
+		// A stall until the compilation is finished might occur, and this is only supported on the game thread for now.
+		bSingleThreaded |= FStaticMeshCompilingManager::Get().IsAsyncStaticMeshCompilationEnabled();
 #endif
 
-	FSceneInterface* Scene = World->Scene;
-	const bool bAppCanEverRender = FApp::CanEverRender();
-
-	ParallelFor(AddPrimitiveBatches.Num(),
-		[&](int32 Index)
-		{
-			FOptionalTaskTagScope Scope(ETaskTag::EParallelGameThread);
-			UPrimitiveComponent* Component = AddPrimitiveBatches[Index];
+		FSceneInterface* Scene = World->Scene;
+		const bool bAppCanEverRender = FApp::CanEverRender();
 
-			// AActor::PostRegisterAllComponents (called by AActor::IncrementalRegisterComponents) can trigger code 
-			// that either unregisters or re-registers components. If unregistered, skip this component.
-			// If re-registered, FRegisterComponentContext is not passed, so SceneProxy can be created.
-			if (IsValid(Component) && Component->IsRegistered())
+		ParallelFor(AddPrimitiveBatches.Num(),
+			[&](int32 Index)
 			{
-				if (Component->IsRenderStateCreated() || !bAppCanEverRender)
+				FOptionalTaskTagScope Scope(ETaskTag::EParallelGameThread);
+				UPrimitiveComponent* Component = AddPrimitiveBatches[Index];
+
+				// AActor::PostRegisterAllComponents (called by AActor::IncrementalRegisterComponents) can trigger code 
+				// that either unregisters or re-registers components. If unregistered, skip this component.
+				// If re-registered, FRegisterComponentContext is not passed, so SceneProxy can be created.
+				if (IsValid(Component) && Component->IsRegistered())
 				{
-					// Skip if SceneProxy is already created
-					if (Component->SceneProxy == nullptr)
+					if (Component->IsRenderStateCreated() || !bAppCanEverRender)
 					{
-						Scene->AddPrimitive(Component);
+						// Skip if SceneProxy is already created
+						if (Component->SceneProxy == nullptr)
+						{
+							Scene->AddPrimitive(Component);
+						}
+					}
+					else // Fallback for some edge case where the component renderstate are missing
+					{
+						Component->CreateRenderState_Concurrent(nullptr);
 					}
 				}
-				else // Fallback for some edge case where the component renderstate are missing
-				{
-					Component->CreateRenderState_Concurrent(nullptr);
-				}
-			}
-		},
-		bSingleThreaded
-	);
-	AddPrimitiveBatches.Empty();
+			},
+			bSingleThreaded
+		);
+		AddPrimitiveBatches.Empty();
+	}
 
 	for (UPrimitiveComponent* Primitive : SendRenderDynamicDataPrimitives)
 	{
@@ -434,6 +437,8 @@ FTransactionallySafeRWLock UActorComponent::AllUCSModifiedPropertiesLock;
 
 UActorComponent::UActorComponent(const FObjectInitializer& ObjectInitializer /*= FObjectInitializer::Get()*/)
 	: Super(ObjectInitializer)
+	, RegistrationState(EComponentRegistrationState::None)
+	, AsyncCreatePhysicsStateScene(nullptr)
 {
 	OwnerPrivate = GetTypedOuter<AActor>();
 
@@ -461,6 +466,7 @@ UActorComponent::UActorComponent(const FObjectInitializer& ObjectInitializer /*=
 
 	bMarkedForPreEndOfFrameSync = false;
 	bAsyncPhysicsTickEnabled = false;
+	bCallAsyncPhysicsStateCreatedDelegate = false;
 
 	bReplicateUsingRegisteredSubObjectList = GDefaultUseSubObjectReplicationList;
 }
@@ -1115,8 +1121,67 @@ void UActorComponent::PostEditChangeChainProperty(FPropertyChangedChainEvent& Pr
 
 #endif // WITH_EDITOR
 
+bool UActorComponent::ShouldAsyncCreatePhysicsState(UWorld* WorldContext) const
+{
+	check(WorldContext);
+	check(FPhysScene::SupportsAsyncPhysicsStateCreation());
+	TGuardValue<EComponentRegistrationState> Guard(const_cast<UActorComponent*>(this)->RegistrationState, EComponentRegistrationState::PreRegistering);
+	return (!bPhysicsStateCreated && WorldContext->GetPhysicsScene() && AllowsAsyncPhysicsStateCreationAndDestruction() && ShouldCreatePhysicsState());
+}
+
+bool UActorComponent::ShouldIncrementalPreRegister(UWorld* WorldContext) const
+{
+	// For now the only criteria is if the component should create physics state asynchronously
+	return ShouldAsyncCreatePhysicsState(WorldContext);
+}
+
+void UActorComponent::OnPreRegister()
+{
+	TRACE_CPUPROFILER_EVENT_SCOPE(UActorComponent::OnPreRegister)
+	checkf(WorldPrivate, TEXT("OnPreRegister: %s to %s"), *GetDetailedInfo(), *GetFullNameSafe(GetOwner()));
+	checkf(RegistrationState == EComponentRegistrationState::None, TEXT("OnPreRegister: %s to %s"), *GetDetailedInfo(), *GetFullNameSafe(GetOwner()));
+	checkf(!bRegistered, TEXT("OnPreRegister: %s to %s"), *GetDetailedInfo(), *GetFullNameSafe(GetOwner()));
+
+	RegistrationState = EComponentRegistrationState::PreRegistering;
+
+	AActor* MyOwner = GetOwner();
+	if (ULevel* Level = MyOwner ? MyOwner->GetLevel() : nullptr)
+	{
+		FLevelRegistrationAccessor::OnPreRegisterComponent(Level, this);
+	}
+
+	UpdateComponentToWorld();
+
+	if (ShouldAsyncCreatePhysicsState(WorldPrivate))
+	{
+		FPhysScene_Chaos* PhysScene = WorldPrivate->GetPhysicsScene();
+		check(PhysScene);
+		check(!IsAsyncCreatePhysicsStateRunning());
+		verify(PhysScene->PushAsyncCreatePhysicsStateJob(this));
+		AsyncCreatePhysicsStateScene = PhysScene;
+		check(IsAsyncCreatePhysicsStateRunning());
+	}
+	else
+	{
+		OnPreRegistered();
+	}
+}
+
+void UActorComponent::OnPreRegistered()
+{
+	check(RegistrationState == EComponentRegistrationState::PreRegistering);
+	RegistrationState = EComponentRegistrationState::PreRegistered;
+
+	AActor* MyOwner = GetOwner();
+	if (ULevel* Level = MyOwner ? MyOwner->GetLevel() : nullptr)
+	{
+		FLevelRegistrationAccessor::OnPreRegisteredComponent(Level, this);
+	}
+}
+
 void UActorComponent::OnRegister()
 {
+	TRACE_CPUPROFILER_EVENT_SCOPE(UActorComponent::OnRegister)
 #if !UE_BUILD_SHIPPING
 	// These are removed in shipping because they are still likely to fail in Test and Development builds, and checks in shipping makes this rather expensive.
 	checkf(!IsUnreachable(), TEXT("%s"), *GetDetailedInfo());
@@ -1127,9 +1192,17 @@ void UActorComponent::OnRegister()
 	checkf(WorldPrivate, TEXT("OnRegister: %s to %s"), *GetDetailedInfo(), GetOwner() ? *GetOwner()->GetFullName() : TEXT("*** No Owner ***") );
 	checkf(!bRegistered, TEXT("OnRegister: %s to %s"), *GetDetailedInfo(), GetOwner() ? *GetOwner()->GetFullName() : TEXT("*** No Owner ***") );
 
+	const bool bHasBeenPreRegistered = IsPreRegistered();
 	bRegistered = true;
+	
+	// Reset pre-register state as the component is now marked registered
+	RegistrationState = EComponentRegistrationState::None;
 
-	UpdateComponentToWorld();
+	// No need to call UpdateComponentToWorld if pre-registered as it has already been called in UActorComponent::OnPreRegister
+	if (!bHasBeenPreRegistered)
+	{
+		UpdateComponentToWorld();
+	}
 
 	if (bAutoActivate)
 	{
@@ -1146,6 +1219,9 @@ void UActorComponent::OnUnregister()
 	check(bRegistered);
 	bRegistered = false;
 
+	check(!IsPreRegistering());
+	RegistrationState = EComponentRegistrationState::None;
+
 	ClearNeedEndOfFrameUpdate();
 }
 
@@ -1410,6 +1486,66 @@ void UActorComponent::TickComponent(float DeltaTime, enum ELevelTick TickType, F
 	}
 }
 
+void UActorComponent::PreRegisterComponentWithWorld(UWorld* InWorld)
+{
+	checkf(!IsUnreachable(), TEXT("%s"), *GetFullName());
+
+	if (!IsValid(this))
+	{
+		UE_LOG(LogActorComponent, Log, TEXT("Trying to pre-register component with IsValid() == false. Aborting."));
+		return;
+	}
+
+	if (IsPreRegistering())
+	{
+		UE_LOG(LogActorComponent, Log, TEXT("Component already pre-registering. Aborting."));
+		return;
+	}
+
+	// If the component was already pre-registered, do nothing
+	if (IsPreRegistered())
+	{
+		UE_LOG(LogActorComponent, Log, TEXT("Component already pre-registered. Aborting."));
+		return;
+	}
+
+	if (InWorld == nullptr)
+	{
+		return;
+	}
+
+	// If not pre-registered, should not have a world
+	checkf(WorldPrivate == nullptr, TEXT("%s"), *GetFullName());
+
+	AActor* MyOwner = GetOwner();
+	checkSlow(MyOwner == nullptr || MyOwner->OwnsComponent(this));
+
+	if (MyOwner && MyOwner->GetClass()->HasAnyClassFlags(CLASS_NewerVersionExists))
+	{
+		UE_LOG(LogActorComponent, Log, TEXT("Component owner belongs to a DEADCLASS"));
+		return;
+	}
+
+#if !(UE_BUILD_SHIPPING || UE_BUILD_TEST)
+	// Can only pre-register with an Actor if we are created within one
+	if (MyOwner)
+	{
+		checkf(!MyOwner->IsUnreachable(), TEXT("%s"), *GetFullName());
+		// The only time you should specify a scene that is not Owner->GetWorld() is when you don't have an Actor
+		UE_CLOG(InWorld != MyOwner->GetWorld(), LogActorComponent, Log, TEXT("Specifying a world, but an Owner Actor found, and InWorld is not GetOwner()->GetWorld()"));
+	}
+#endif // !(UE_BUILD_SHIPPING || UE_BUILD_TEST)
+
+	if (!bHasBeenCreated)
+	{
+		OnComponentCreated();
+	}
+
+	WorldPrivate = InWorld;
+
+	OnPreRegister();
+}
+
 void UActorComponent::RegisterComponentWithWorld(UWorld* InWorld, FRegisterComponentContext* Context)
 {
 	SCOPE_CYCLE_COUNTER(STAT_RegisterComponent);
@@ -1436,8 +1572,8 @@ void UActorComponent::RegisterComponentWithWorld(UWorld* InWorld, FRegisterCompo
 		return;
 	}
 
-	// If not registered, should not have a scene
-	checkf(WorldPrivate == nullptr, TEXT("%s"), *GetFullName());
+	// If not registered/pre-registered, should not have a world
+	checkf(WorldPrivate == nullptr || (IsPreRegistered()), TEXT("%s"), *GetFullName());
 
 	AActor* MyOwner = GetOwner();
 	checkSlow(MyOwner == nullptr || MyOwner->OwnsComponent(this));
@@ -1453,9 +1589,6 @@ void UActorComponent::RegisterComponentWithWorld(UWorld* InWorld, FRegisterCompo
 	if(MyOwner)
 	{
 		checkf(!MyOwner->IsUnreachable(), TEXT("%s"), *GetFullName());
-		// can happen with undo because the owner will be restored "next"
-		//checkf(IsValid(MyOwner), TEXT("%s"), *GetFullName());
-
 		if(InWorld != MyOwner->GetWorld())
 		{
 			// The only time you should specify a scene that is not Owner->GetWorld() is when you don't have an Actor
@@ -1706,13 +1839,36 @@ void UActorComponent::DestroyRenderState_Concurrent()
 
 void UActorComponent::OnCreatePhysicsState()
 {
-	check(IsRegistered());
+	TRACE_CPUPROFILER_EVENT_SCOPE(UActorComponent::OnCreatePhysicsState)
+	check(IsRegistered() || IsPreRegistering());
 	check(ShouldCreatePhysicsState());
 	check(WorldPrivate->GetPhysicsScene());
 	check(!bPhysicsStateCreated);
 	bPhysicsStateCreated = true;
 }
 
+void UActorComponent::OnAsyncCreatePhysicsState()
+{
+	TRACE_CPUPROFILER_EVENT_SCOPE(UActorComponent::OnAsyncCreatePhysicsState)
+	check(AllowsAsyncPhysicsStateCreationAndDestruction());
+	OnCreatePhysicsState();
+}
+
+void UActorComponent::OnAsyncCreatePhysicsStateDone_GameThread()
+{
+	check(IsPhysicsStateCreated());
+	check(IsPreRegistering())
+	check(IsAsyncCreatePhysicsStateRunning());
+	check(!IsA<UPrimitiveComponent>() || !Cast<UPrimitiveComponent>(this)->DeferredCreatePhysicsStateScene);
+
+	// Reset members used by asynchronous physics state creation
+	AsyncCreatePhysicsStateScene = nullptr;
+	bCallAsyncPhysicsStateCreatedDelegate = true;
+
+	OnPreRegistered();
+	check(IsPreRegistered());
+}
+
 void UActorComponent::OnDestroyPhysicsState()
 {
 	ensure(bPhysicsStateCreated);
@@ -1727,6 +1883,7 @@ void UActorComponent::CreatePhysicsState(bool bAllowDeferral)
 	UE_TRACE_METADATA_SCOPE_ASSET_FNAME(NAME_None, NAME_None, GetPackage()->GetFName());
 
 	SCOPE_CYCLE_COUNTER(STAT_ComponentCreatePhysicsState);
+	TRACE_CPUPROFILER_EVENT_SCOPE(UActorComponent::CreatePhysicsState)
 
 	if (!bPhysicsStateCreated && WorldPrivate->GetPhysicsScene() && ShouldCreatePhysicsState())
 	{
@@ -1766,13 +1923,39 @@ void UActorComponent::CreatePhysicsState(bool bAllowDeferral)
 			GlobalCreatePhysicsDelegate.Broadcast(this);
 		}
 	}
-
+	// Asynchronous physics state was created, but GlobalCreatePhysicsDelegate needs to be broadcasted
+	else if (bCallAsyncPhysicsStateCreatedDelegate)
+	{
+		check(bPhysicsStateCreated);
+		bCallAsyncPhysicsStateCreatedDelegate = false;
+		if (ensure(WorldPrivate->GetPhysicsScene()))
+		{
+			GlobalCreatePhysicsDelegate.Broadcast(this);
+		}
+	}
 }
 
 void UActorComponent::DestroyPhysicsState()
 {
+	TRACE_CPUPROFILER_EVENT_SCOPE(UActorComponent::DestroyPhysicsState)
 	SCOPE_CYCLE_COUNTER(STAT_ComponentDestroyPhysicsState);
 
+	if (IsAsyncCreatePhysicsStateRunning())
+	{
+		check(!IsA<UPrimitiveComponent>() || !Cast<UPrimitiveComponent>(this)->DeferredCreatePhysicsStateScene);
+		AsyncCreatePhysicsStateScene->RemoveAsyncPhysicsStateJobs(this);
+		
+		// Reset members used by asynchronous physics state creation
+		AsyncCreatePhysicsStateScene = nullptr;
+		bCallAsyncPhysicsStateCreatedDelegate = false;
+		
+		AActor* MyOwner = GetOwner();
+		if (ULevel* Level = MyOwner ? MyOwner->GetLevel() : nullptr)
+		{
+			FLevelRegistrationAccessor::OnRemovedPreRegisteringComponent(Level, this);
+		}
+	}
+
 	if (bPhysicsStateCreated)
 	{
 		// Broadcast delegate
@@ -1822,7 +2005,6 @@ void UActorComponent::ExecuteRegisterEvents(FRegisterComponentContext* Context)
 
 void UActorComponent::ExecuteUnregisterEvents()
 {
-	// Delay the destroying the physics, as well as the rendering state until after we commit
 	DestroyPhysicsState();
 
 	if (bRenderStateCreated)
@@ -1883,6 +2065,9 @@ void UActorComponent::RecreateRenderState_Concurrent()
 
 void UActorComponent::RecreatePhysicsState()
 {
+	TRACE_CPUPROFILER_EVENT_SCOPE(UActorComponent::RecreatePhysicsState);
+	check(!IsAsyncCreatePhysicsStateRunning());
+	
 	DestroyPhysicsState();
 
 	if (IsRegistered())
diff --git a/Engine/Source/Runtime/Engine/Private/Components/PrimitiveComponent.cpp b/Engine/Source/Runtime/Engine/Private/Components/PrimitiveComponent.cpp
--- a/Engine/Source/Runtime/Engine/Private/Components/PrimitiveComponent.cpp
+++ b/Engine/Source/Runtime/Engine/Private/Components/PrimitiveComponent.cpp
@@ -891,8 +891,20 @@ void UPrimitiveComponent::OnCreatePhysicsState()
 		}
 	}
 
+	if (!IsAsyncCreatePhysicsStateRunning())
+	{
+		OnComponentPhysicsStateChanged.Broadcast(this, EComponentPhysicsStateChange::Created);
+	}
+}
+
+void UPrimitiveComponent::OnAsyncCreatePhysicsStateDone_GameThread()
+{
+	check(IsPhysicsStateCreated());
+
 	OnComponentPhysicsStateChanged.Broadcast(this, EComponentPhysicsStateChange::Created);
-}	
+
+	Super::OnAsyncCreatePhysicsStateDone_GameThread();
+}
 
 void UPrimitiveComponent::EnsurePhysicsStateCreated()
 {
@@ -1689,7 +1701,7 @@ bool UPrimitiveComponent::ShouldCreatePhysicsState() const
 		return false;
 	}
 
-	bool bShouldCreatePhysicsState = IsRegistered() && (bAlwaysCreatePhysicsState || BodyInstance.GetCollisionEnabled() != ECollisionEnabled::NoCollision);
+	bool bShouldCreatePhysicsState = (IsRegistered() || (IsPreRegistering() && AllowsAsyncPhysicsStateCreationAndDestruction())) && (bAlwaysCreatePhysicsState || BodyInstance.GetCollisionEnabled() != ECollisionEnabled::NoCollision);
 
 #if WITH_EDITOR
 	if (BodyInstance.bSimulatePhysics)
diff --git a/Engine/Source/Runtime/Engine/Private/Components/SceneComponent.cpp b/Engine/Source/Runtime/Engine/Private/Components/SceneComponent.cpp
--- a/Engine/Source/Runtime/Engine/Private/Components/SceneComponent.cpp
+++ b/Engine/Source/Runtime/Engine/Private/Components/SceneComponent.cpp
@@ -763,7 +763,7 @@ void USceneComponent::UpdateComponentToWorldWithParent(USceneComponent* Parent,F
 	}
 }
 
-void USceneComponent::OnRegister()
+void USceneComponent::AttachToParent()
 {
 	// If we need to perform a call to AttachTo, do that now
 	// At this point scene component still has no any state (rendering, physics),
@@ -781,7 +781,19 @@ void USceneComponent::OnRegister()
 			SetShouldSnapScaleWhenAttached(false);
 		}
 	}
-	
+}
+
+void USceneComponent::OnPreRegister()
+{
+	TRACE_CPUPROFILER_EVENT_SCOPE(USceneComponent::OnPreRegister)
+	AttachToParent();
+	Super::OnPreRegister();
+}
+
+void USceneComponent::OnRegister()
+{
+	TRACE_CPUPROFILER_EVENT_SCOPE(USceneComponent::OnRegister)
+	AttachToParent();
 	// Cache the level collection that contains the level in which this component is registered for fast access in IsVisible().
 	const UWorld* const World = GetWorld();
 	if (World)
@@ -975,6 +987,7 @@ bool USceneComponent::IsDeferringMovementUpdates(const FScopedMovementUpdate& Sc
 
 bool USceneComponent::UpdateOverlaps(const TOverlapArrayView* PendingOverlaps /* = nullptr */, bool bDoNotifies /* = true */, const TOverlapArrayView* OverlapsAtEndLocation /* = nullptr */)
 {
+	TRACE_CPUPROFILER_EVENT_SCOPE(USceneComponent::UpdateOverlaps)
 	if (IsDeferringMovementUpdates())
 	{
 		GetCurrentScopedMovement()->ForceOverlapUpdate();
diff --git a/Engine/Source/Runtime/Engine/Private/Components/StaticMeshComponent.cpp b/Engine/Source/Runtime/Engine/Private/Components/StaticMeshComponent.cpp
--- a/Engine/Source/Runtime/Engine/Private/Components/StaticMeshComponent.cpp
+++ b/Engine/Source/Runtime/Engine/Private/Components/StaticMeshComponent.cpp
@@ -774,15 +774,28 @@ void UStaticMeshComponent::PostEditImport()
 
 #endif // #if WITH_EDITOR
 
+void UStaticMeshComponent::OnPreRegister()
+{
+	TRACE_CPUPROFILER_EVENT_SCOPE(UStaticMeshComponent::OnPreRegister)
+	UpdateCollisionFromStaticMesh();
+	Super::OnPreRegister();
+}
+
 void UStaticMeshComponent::OnRegister()
 {
+	TRACE_CPUPROFILER_EVENT_SCOPE(UStaticMeshComponent::OnRegister)
+
 #if WITH_EDITORONLY_DATA
 	FGuardValue_Bitfield(bRegistering, true);
 #endif
 
 	NotifyIfStaticMeshChanged();
 
-	UpdateCollisionFromStaticMesh();
+	// No need to call UpdateCollisionFromStaticMesh if pre-registered as it has already been called in UStaticMeshComponent::OnPreRegister
+	if (!IsPreRegistered())
+	{
+		UpdateCollisionFromStaticMesh();
+	}
 
 #if WITH_EDITORONLY_DATA
 	//Remap the override materials if the import version is different
@@ -856,7 +869,15 @@ void UStaticMeshComponent::OnCreatePhysicsState()
 	Super::OnCreatePhysicsState();
 
 	bNavigationRelevant = IsNavigationRelevant();
-	FNavigationSystem::UpdateComponentData(*this);
+
+	// Async physics state creation is currently only done through component pre-registration.
+	// No need to call FNavigationSystem::UpdateComponentData for a pre-registering component.
+	// Once async physics state creation completes, component registration will call FNavigationSystem::OnComponentRegistered.
+	if (!IsAsyncCreatePhysicsStateRunning())
+	{
+		check(!IsPreRegistering());
+		FNavigationSystem::UpdateComponentData(*this);
+	}
 }
 
 void UStaticMeshComponent::OnDestroyPhysicsState()
@@ -2278,6 +2299,12 @@ bool UStaticMeshComponent::ShouldCreatePhysicsState() const
 	return Super::ShouldCreatePhysicsState() && GetStaticMesh() && !GetStaticMesh()->IsCompiling();
 }
 
+bool UStaticMeshComponent::AllowsAsyncPhysicsStateCreationAndDestruction() const
+{
+	// Async physics state creation is supported for static collision or primitive component not using welding
+	return (FInitBodySpawnParams(this).bStaticPhysics || !BodyInstance.bAutoWeld || !GetAttachParent());
+}
+
 void UStaticMeshComponent::SetStaticMeshInternal(UStaticMesh* NewMesh)
 {
 	StaticMesh = NewMesh;
diff --git a/Engine/Source/Runtime/Engine/Private/GameInstance.cpp b/Engine/Source/Runtime/Engine/Private/GameInstance.cpp
--- a/Engine/Source/Runtime/Engine/Private/GameInstance.cpp
+++ b/Engine/Source/Runtime/Engine/Private/GameInstance.cpp
@@ -526,7 +526,6 @@ FGameInstancePIEResult UGameInstance::StartPlayInEditorGameInstance(ULocalPlayer
 		{
 			FRegisterComponentContext Context(PlayWorld);
 			PlayWorld->InitializeActorsForPlay(URL, true, &Context);
-			Context.Process();
 		}
 		// calling it after InitializeActorsForPlay has been called to have all potential bounding boxed initialized
 		FNavigationSystem::AddNavigationSystemToWorld(*PlayWorld, FNavigationSystemRunMode::PIEMode);
diff --git a/Engine/Source/Runtime/Engine/Private/InstancedStaticMesh.cpp b/Engine/Source/Runtime/Engine/Private/InstancedStaticMesh.cpp
--- a/Engine/Source/Runtime/Engine/Private/InstancedStaticMesh.cpp
+++ b/Engine/Source/Runtime/Engine/Private/InstancedStaticMesh.cpp
@@ -2613,7 +2613,7 @@ void UInstancedStaticMeshComponent::CreateAllInstanceBodies()
 
 		if (InstanceBodiesSanitized.Num() > 0 && Mobility != EComponentMobility::Movable)
 		{
-			FBodyInstance::InitStaticBodies(InstanceBodiesSanitized, Transforms, BodySetup, this, GetWorld()->GetPhysicsScene());
+			FBodyInstance::InitStaticBodies(MoveTemp(InstanceBodiesSanitized), MoveTemp(Transforms), BodySetup, this, GetWorld()->GetPhysicsScene());
 		}
 	}
 	else
@@ -2705,11 +2705,13 @@ bool UInstancedStaticMeshComponent::CanEditSimulatePhysics()
 
 FBoxSphereBounds UInstancedStaticMeshComponent::CalcBounds(const FTransform& BoundTransform) const
 {
+	TRACE_CPUPROFILER_EVENT_SCOPE(UInstancedStaticMeshComponent::CalcBounds);
 	return CalcBoundsImpl(BoundTransform, /*bForNavigation*/false);
 }
 
 void UInstancedStaticMeshComponent::CalcAndCacheNavigationBounds()
 {
+	TRACE_CPUPROFILER_EVENT_SCOPE(UInstancedStaticMeshComponent::CalcAndCacheNavigationBounds);
 	NavigationBounds = CalcBoundsImpl(GetComponentTransform(), /*bForNavigation*/true).GetBox();
 }
 
@@ -2749,7 +2751,7 @@ void UInstancedStaticMeshComponent::UpdateBounds()
 		Super::UpdateBounds();
 	}
 
-	if (bCanEverAffectNavigation && IsRegistered())
+	if (bCanEverAffectNavigation && (IsRegistered() || IsPreRegistering()))
 	{
 		CalcAndCacheNavigationBounds();
 	}
@@ -4499,7 +4501,7 @@ TArray<int32> UInstancedStaticMeshComponent::GetInstancesOverlappingBox(const FB
 
 bool UInstancedStaticMeshComponent::ShouldCreatePhysicsState() const
 {
-	return !bDisableCollision && IsRegistered() && !IsBeingDestroyed() && GetStaticMesh() && !GetStaticMesh()->IsCompiling() && (bAlwaysCreatePhysicsState || IsCollisionEnabled());
+	return !bDisableCollision && (IsRegistered() || (IsPreRegistering() && AllowsAsyncPhysicsStateCreationAndDestruction())) && !IsBeingDestroyed() && GetStaticMesh() && !GetStaticMesh()->IsCompiling() && (bAlwaysCreatePhysicsState || IsCollisionEnabled());
 }
 
 float UInstancedStaticMeshComponent::GetTextureStreamingTransformScale() const
@@ -4759,6 +4761,7 @@ void UInstancedStaticMeshComponent::InitPerInstanceRenderData(bool InitializeFro
 
 void UInstancedStaticMeshComponent::OnRegister()
 {
+	TRACE_CPUPROFILER_EVENT_SCOPE(UInstancedStaticMeshComponent::OnRegister())
 	Super::OnRegister();
 	
 	PreviousComponentTransform = GetComponentTransform();
@@ -5184,6 +5187,7 @@ void UInstancedStaticMeshComponent::PartialNavigateUpdateForCurrentInstances()
 
 void UInstancedStaticMeshComponent::PartialNavigationUpdate(int32 InstanceIdx)
 {
+	TRACE_CPUPROFILER_EVENT_SCOPE(UInstancedStaticMeshComponent::PartialNavigationUpdate);
 	if (!IsNavigationRelevant())
 	{
 		return;
diff --git a/Engine/Source/Runtime/Engine/Private/Level.cpp b/Engine/Source/Runtime/Engine/Private/Level.cpp
--- a/Engine/Source/Runtime/Engine/Private/Level.cpp
+++ b/Engine/Source/Runtime/Engine/Private/Level.cpp
@@ -125,6 +125,12 @@ static FAutoConsoleVariableRef CVarRouteActorInitializationWorkUnitWeighting(
 	ECVF_Default
 );
 
+bool GAllowIncrementalPreRegisterComponents = false;
+static FAutoConsoleVariableRef CVarAllowIncrementalPreregisterComponents(
+	TEXT("s.AllowIncrementalPreRegisterComponents"),
+	GAllowIncrementalPreRegisterComponents,
+	TEXT("Whether to allow level streaming incremental pre-registration of actor components."));
+
 #if WITH_EDITOR
 
 TArray<ULevel::FLevelExternalActorsPathsProviderDelegate> ULevel::LevelExternalActorsPathsProviders;
@@ -1404,6 +1410,7 @@ bool ULevel::CanBeClusterRoot() const
 
 void ULevel::CreateCluster()
 {
+	TRACE_CPUPROFILER_EVENT_SCOPE(ULevel::CreateCluster);
 	// ULevels are not cluster roots themselves, instead they create a special actor container
 	// that holds a reference to all actors that are to be clustered. This is because only
 	// specific actor types can be clustered so the remaining actors that are not clustered
@@ -1575,6 +1582,7 @@ void ULevel::UpdateLevelComponents(bool bRerunConstructionScripts, FRegisterComp
  */
 static void SortActorsHierarchy(TArray<TObjectPtr<AActor>>& Actors, ULevel* Level)
 {
+	TRACE_CPUPROFILER_EVENT_SCOPE(SortActorsHierarchy);
 	const double StartTime = FPlatformTime::Seconds();
 
 	TMap<AActor*, int32> DepthMap;
@@ -1697,7 +1705,7 @@ static void SortActorsHierarchy(TArray<TObjectPtr<AActor>>& Actors, ULevel* Leve
 }
 
 DECLARE_CYCLE_STAT(TEXT("Deferred Init Bodies"), STAT_DeferredUpdateBodies, STATGROUP_Physics);
-void ULevel::IncrementalUpdateComponents(int32 NumComponentsToUpdate, bool bRerunConstructionScripts, FRegisterComponentContext* Context)
+void ULevel::IncrementalUpdateComponents(int32 NumComponentsToUpdate, bool bRerunConstructionScripts, FRegisterComponentContext* InContext)
 {
 #if !WITH_EDITOR
 	ensure(!bRerunConstructionScripts);
@@ -1717,20 +1725,56 @@ void ULevel::IncrementalUpdateComponents(int32 NumComponentsToUpdate, bool bReru
 	// The editor is never allowed to incrementally update components.  Make sure to pass in a value of zero for NumActorsToUpdate.
 	check(bFullyUpdateComponents || OwningWorld->IsGameWorld());
 
+	auto IsServer = [this]()
+	{
+		const ENetMode NetMode = OwningWorld->GetNetMode();
+		return NetMode == NM_DedicatedServer || NetMode == NM_ListenServer;
+	};
+
+	FRegisterComponentContext LocalContext(OwningWorld);
+	FRegisterComponentContext& Context = InContext ? *InContext : LocalContext;
+	Context.SetIncrementalUpdateTimeLimitMs(InContext ? InContext->GetIncrementalUpdateTimeLimitMs() : 0);
+	Context.SetIncrementalUpdateNumComponentsToUpdate(NumComponentsToUpdate);
+	Context.SetIncrementalUpdateAddPrimitiveBatchSize(InContext ? InContext->GetIncrementalUpdateAddPrimitiveBatchSize() : 0);
+
+	bool bAllowLoop;
 	do
 	{
+		bAllowLoop = false;
 		switch (IncrementalComponentState)
 		{
 		case EIncrementalComponentState::Init:
-			// Do BSP on the first pass.
-			UpdateModelComponents();
-			// Sort actors to ensure that parent actors will be registered before child actors
-			SortActorsHierarchy(Actors, this);
-			IncrementalComponentState = EIncrementalComponentState::RegisterInitialComponents;
-
-			//pass through
+			{
+				// Do BSP on the first pass.
+				UpdateModelComponents();
+				// Sort actors to ensure that parent actors will be registered before child actors
+				SortActorsHierarchy(Actors, this);
+				// Skip PreRegisterInitialComponents step if any of these conditions are true
+				// - running on server 
+				// - not running in incremental update mode 
+				// - asynchronous physics state creation is not supported
+				// - world is inside a blocking level streaming
+				const bool bIsPreRegisterInitialComponentsAllowed = !bFullyUpdateComponents &&
+					GAllowIncrementalPreRegisterComponents &&
+					!IsServer() &&
+					FPhysScene::SupportsAsyncPhysicsStateCreation() &&
+					!OwningWorld->bRequestedBlockOnAsyncLoading && 
+					!OwningWorld->GetIsInBlockTillLevelStreamingCompleted();
+				// Because pre-register is only done during incremental update, it is implicitly only used in a game world
+				check(!bIsPreRegisterInitialComponentsAllowed || OwningWorld->IsGameWorld());
+				IncrementalComponentState = bIsPreRegisterInitialComponentsAllowed ? EIncrementalComponentState::PreRegisterInitialComponents : EIncrementalComponentState::RegisterInitialComponents;
+				bAllowLoop = true;
+			}
+			break;
+		case EIncrementalComponentState::PreRegisterInitialComponents:
+			if (IncrementalPreRegisterComponents(Context))
+			{
+				IncrementalComponentState = EIncrementalComponentState::RegisterInitialComponents;
+				bAllowLoop = true;
+			}
+			break;
 		case EIncrementalComponentState::RegisterInitialComponents:
-			if (IncrementalRegisterComponents(true, NumComponentsToUpdate, Context))
+			if (IncrementalRegisterComponents(Context))
 			{
 #if WITH_EDITOR
 				const bool bShouldRunConstructionScripts = !bHasRerunConstructionScripts && bRerunConstructionScripts && !IsTemplate();
@@ -1753,14 +1797,12 @@ void ULevel::IncrementalUpdateComponents(int32 NumComponentsToUpdate, bool bReru
 		case EIncrementalComponentState::Finalize:
 			IncrementalComponentState = EIncrementalComponentState::Init;
 			CurrentActorIndexForIncrementalUpdate = 0;
-			bHasCurrentActorCalledPreRegister = false;
 			bAreComponentsCurrentlyRegistered = true;
 			CreateCluster();
 			break;
 		}
-	}while(bFullyUpdateComponents && !bAreComponentsCurrentlyRegistered);
+	} while (bAllowLoop || (bFullyUpdateComponents && !bAreComponentsCurrentlyRegistered));
 	
-
 	{
 		SCOPE_CYCLE_COUNTER(STAT_DeferredUpdateBodies);
 
@@ -1772,15 +1814,75 @@ void ULevel::IncrementalUpdateComponents(int32 NumComponentsToUpdate, bool bReru
 	}
 }
 
-bool ULevel::IncrementalRegisterComponents(bool bPreRegisterComponents, int32 NumComponentsToUpdate, FRegisterComponentContext* Context)
+bool ULevel::IncrementalPreRegisterComponents(FRegisterComponentContext& Context)
 {
-	// Find next valid actor to process components registration
+	auto ProcessPreRegisteringComponents = [this]()
+	{
+		if (HasPreRegisteringComponents())
+		{
+			if (FPhysScene* PhysScene = OwningWorld->GetPhysicsScene())
+			{
+				PhysScene->ProcessAsyncPhysicsStateJobs();
+			}
+		}
+	};
+
+	const double StartTime = FPlatformTime::Seconds();
+	const double TimeLimitMs = Context.GetIncrementalUpdateTimeLimitMs();
+	const bool bConsiderTimeLimit = TimeLimitMs >= 0;
+	
+	ON_SCOPE_EXIT
+	{
+		if (bConsiderTimeLimit)
+		{
+			const double ElapsedTimeMs = (FPlatformTime::Seconds() - StartTime) * 1000;
+			const double IncrementalUpdateTimeLimitMs = bConsiderTimeLimit ? FMath::Max<double>(0.0, TimeLimitMs - ElapsedTimeMs) : -1;
+			Context.SetIncrementalUpdateTimeLimitMs(IncrementalUpdateTimeLimitMs);
+		}
+	};
+
+	check(OwningWorld);
+	
+	while (CurrentActorIndexForIncrementalUpdate < Actors.Num())
+	{
+		AActor* Actor = Actors[CurrentActorIndexForIncrementalUpdate];
+		CurrentActorIndexForIncrementalUpdate++;
+		if (IsValid(Actor))
+		{
+			FActorPreRegistrationAccessor::IncrementalPreRegisterComponents(Actor);
+		}
+
+		if (bConsiderTimeLimit)
+		{
+			if (const double ElapseTimeMs = (FPlatformTime::Seconds() - StartTime) * 1000; ElapseTimeMs > TimeLimitMs)
+			{
+				break;
+			}
+		}
+	}
+
+	// Process pending preregistering components
+	ProcessPreRegisteringComponents();
+
+	if (!HasPreRegisteringComponents() && (CurrentActorIndexForIncrementalUpdate >= Actors.Num()))
+	{
+		// Reset CurrentActorIndexForIncrementalUpdate for IncrementalRegisterComponents
+		CurrentActorIndexForIncrementalUpdate = 0;
+		return true;
+	}
+	return false;
+}
+
+bool ULevel::IncrementalRegisterComponents(FRegisterComponentContext& Context)
+{
+	const int32 NumComponentsToUpdate = Context.GetIncrementalUpdateNumComponentsToUpdate();
 
 	if (OwningWorld)
 	{
 		OwningWorld->SetAllowDeferredPhysicsStateCreation(true);
 	}
 
+	// Find next valid actor to process components registration
 	while (CurrentActorIndexForIncrementalUpdate < Actors.Num())
 	{
 		AActor* Actor = Actors[CurrentActorIndexForIncrementalUpdate];
@@ -1792,12 +1894,11 @@ bool ULevel::IncrementalRegisterComponents(bool bPreRegisterComponents, int32 Nu
 #if PERF_TRACK_DETAILED_ASYNC_STATS
 				FScopeCycleCounterUObject ContextScope(Actor);
 #endif
-				if (bPreRegisterComponents && !bHasCurrentActorCalledPreRegister)
+				if (!FActorPreRegistrationAccessor::HasPreRegisteredAllComponents(Actor))
 				{
 					Actor->PreRegisterAllComponents();
-					bHasCurrentActorCalledPreRegister = true;
 				}
-				bAllComponentsRegistered = Actor->IncrementalRegisterComponents(NumComponentsToUpdate, Context);
+				bAllComponentsRegistered = Actor->IncrementalRegisterComponents(NumComponentsToUpdate, &Context);
 			}
 #if !UE_BUILD_SHIPPING
 			else if (GOptimizeActorRegistration == 2)
@@ -1815,7 +1916,6 @@ bool ULevel::IncrementalRegisterComponents(bool bPreRegisterComponents, int32 Nu
 		{
 			// All components have been registered for this actor, move to a next one
 			CurrentActorIndexForIncrementalUpdate++;
-			bHasCurrentActorCalledPreRegister = false;
 		}
 
 		// If we do an incremental registration return to outer loop after each processed actor 
@@ -1830,10 +1930,7 @@ bool ULevel::IncrementalRegisterComponents(bool bPreRegisterComponents, int32 Nu
 	{
 		// We need to process pending adds prior to rerunning the construction scripts, which may internally
 		// perform removals / adds themselves.
-		if (Context)
-		{
-			Context->Process();
-		}
+		Context.Process();
 		CurrentActorIndexForIncrementalUpdate = 0;
 		return true;
 	}
diff --git a/Engine/Source/Runtime/Engine/Private/PhysicsEngine/Experimental/PhysScene_AsyncPhysicsStateJobQueue.cpp b/Engine/Source/Runtime/Engine/Private/PhysicsEngine/Experimental/PhysScene_AsyncPhysicsStateJobQueue.cpp
new file mode 100644
--- /dev/null
+++ b/Engine/Source/Runtime/Engine/Private/PhysicsEngine/Experimental/PhysScene_AsyncPhysicsStateJobQueue.cpp
@@ -0,0 +1,169 @@
+// Copyright Epic Games, Inc. All Rights Reserved.
+
+#include "Physics/Experimental/PhysScene_AsyncPhysicsStateJobQueue.h"
+#include "Physics/Experimental/PhysScene_Chaos.h"
+#include "Components/ActorComponent.h"
+#include "Components/PrimitiveComponent.h"
+#include "PhysicsEngine/BodySetup.h"
+#include "Engine/World.h"
+
+void FPhysScene_AsyncPhysicsStateJobQueue::AddJob(const FJob& Job)
+{
+	TRACE_CPUPROFILER_EVENT_SCOPE(FPhysScene_AsyncPhysicsStateJobQueue::AddJob);
+	check(IsInGameThread());
+
+	if (!PendingJobs.IsValid())
+	{
+		PendingJobs = MakeUnique<FJobSet>();
+	}
+	PendingJobs->Add(Job);
+}
+
+void FPhysScene_AsyncPhysicsStateJobQueue::RemoveJob(const FJob& Job)
+{
+	TRACE_CPUPROFILER_EVENT_SCOPE(FPhysScene_AsyncPhysicsStateJobQueue::RemoveJob);
+	check(IsInGameThread());
+	check(Job.IsValid());
+
+	if (PendingJobs.IsValid() && PendingJobs->Remove(Job))
+	{
+		return;
+	}
+
+	WaitForCompletion();
+	
+	for (FJobs& Jobs : PipedJobs)
+	{
+		if (Jobs->Remove(Job))
+		{
+			break;
+		}
+	}
+}
+
+void FPhysScene_AsyncPhysicsStateJobQueue::Launch()
+{
+	TRACE_CPUPROFILER_EVENT_SCOPE(FPhysScene_AsyncPhysicsStateJobQueue::Launch);
+	check(IsInGameThread());
+
+	if (!PendingJobs.IsValid() || PendingJobs->IsEmpty())
+	{
+		return;
+	}
+
+	PendingJobs->OnPreExecute_GameThread();
+
+	// Move Jobs to the async queue
+	FJobs& NewJobs = PipedJobs.Emplace_GetRef();
+	NewJobs = MoveTemp(PendingJobs);
+	// Launch a new task in the pipe
+	PipeTasks.Emplace(Pipe.Launch(UE_SOURCE_LOCATION, [this, NewJobs = NewJobs.Get()]() { NewJobs->Execute(); }, UE::Tasks::ETaskPriority::BackgroundHigh));
+}
+
+void FPhysScene_AsyncPhysicsStateJobQueue::WaitForCompletion()
+{
+	TRACE_CPUPROFILER_EVENT_SCOPE(FPhysScene_AsyncTaskQueue::WaitForCompletion);
+	check(IsInGameThread());
+
+	UE::Tasks::Wait(PipeTasks);
+}
+
+void FPhysScene_AsyncPhysicsStateJobQueue::Tick()
+{
+	TRACE_CPUPROFILER_EVENT_SCOPE(FPhysScene_AsyncPhysicsStateJobQueue::Tick);
+	check(IsInGameThread());
+
+	// Flush all queued jobs
+	Launch();
+
+	// Wait for tasks if the world is inside a blocking load
+	if (PhysScene->GetOwningWorld()->GetIsInBlockTillLevelStreamingCompleted())
+	{
+		WaitForCompletion();
+	}
+
+	if (IsCompleted())
+	{
+		TRACE_CPUPROFILER_EVENT_SCOPE(FPhysScene_AsyncPhysicsStateJobQueue::OnCompleted);
+		for (const FJobs& Jobs : PipedJobs)
+		{
+			Jobs->OnPostExecute_GameThread();
+			
+		}
+		PipedJobs.Reset();
+		PipeTasks.Reset();
+		check(!PendingJobs.IsValid());
+	}
+}
+
+void FPhysScene_AsyncPhysicsStateJobQueue::FJobSet::Execute() const
+{
+	TRACE_CPUPROFILER_EVENT_SCOPE(FPhysScene_AsyncPhysicsStateJobQueue::ExecuteJobs);
+	for (const FJob& Job : Jobs)
+	{
+		Job.Execute();
+	}
+}
+
+void FPhysScene_AsyncPhysicsStateJobQueue::FJobSet::OnPreExecute_GameThread() const
+{
+	TSet<UBodySetup*> UniqueBodySetups;
+
+	for (const FJob& Job : Jobs)
+	{
+		if (Job.Type == EJobType::CreatePhysicsState)
+		{
+			UPrimitiveComponent* PrimitiveComponent = Cast<UPrimitiveComponent>(Job.Component);
+			UBodySetup* BodySetup = PrimitiveComponent ? PrimitiveComponent->GetBodySetup() : nullptr;
+			if (BodySetup && !BodySetup->bCreatedPhysicsMeshes)
+			{
+				UniqueBodySetups.Add(BodySetup);
+			}
+		}
+	}
+
+	if (UniqueBodySetups.Num())
+	{
+		TArray<UBodySetup*> BodySetups = UniqueBodySetups.Array();
+		ParallelFor(TEXT("CreatePhysicsMeshes.PF"), BodySetups.Num(), 1, [this, &BodySetups](int32 Index) { BodySetups[Index]->CreatePhysicsMeshes(); });
+	}
+}
+
+void FPhysScene_AsyncPhysicsStateJobQueue::FJobSet::OnPostExecute_GameThread() const
+{
+	for (const FJob& Job : Jobs)
+	{
+		Job.OnPostExecute_GameThread();
+	}
+}
+
+void FPhysScene_AsyncPhysicsStateJobQueue::FJob::Execute() const
+{
+	if (const bool bPendingKill = Component->GetOwner() && !::IsValid(Component->GetOwner()); !bPendingKill)
+	{
+		switch (Type)
+		{
+			case EJobType::CreatePhysicsState:
+			{
+				TRACE_CPUPROFILER_EVENT_SCOPE(FPhysScene_AsyncPhysicsStateJobQueue::CreatePhysicsState);
+				Component->OnAsyncCreatePhysicsState();
+			}
+			break;
+		}
+	}
+}
+
+void FPhysScene_AsyncPhysicsStateJobQueue::FJob::OnPostExecute_GameThread() const
+{
+	if (const bool bPendingKill = Component->GetOwner() && !::IsValid(Component->GetOwner()); !bPendingKill)
+	{
+		switch (Type)
+		{
+			case EJobType::CreatePhysicsState:
+			{
+				Component->OnAsyncCreatePhysicsStateDone_GameThread();
+			}
+			break;
+		}
+	}
+}
\ No newline at end of file
diff --git a/Engine/Source/Runtime/Engine/Private/PhysicsEngine/Experimental/PhysScene_Chaos.cpp b/Engine/Source/Runtime/Engine/Private/PhysicsEngine/Experimental/PhysScene_Chaos.cpp
--- a/Engine/Source/Runtime/Engine/Private/PhysicsEngine/Experimental/PhysScene_Chaos.cpp
+++ b/Engine/Source/Runtime/Engine/Private/PhysicsEngine/Experimental/PhysScene_Chaos.cpp
@@ -631,11 +631,13 @@ AActor* FPhysScene_Chaos::GetSolverActor() const
 
 void FPhysScene_Chaos::RegisterForCollisionEvents(UPrimitiveComponent* Component)
 {
+	UE_CHAOS_ASYNC_INITBODY_WRITESCOPELOCK(CollisionEventRegistrationsLock);
 	CollisionEventRegistrations.Add(Component);
 }
 
 void FPhysScene_Chaos::UnRegisterForCollisionEvents(UPrimitiveComponent* Component)
 {
+	UE_CHAOS_ASYNC_INITBODY_WRITESCOPELOCK(CollisionEventRegistrationsLock);
 	CollisionEventRegistrations.Remove(Component);
 }
 
@@ -1031,6 +1033,7 @@ void FPhysScene_Chaos::HandleCollisionEvents(const Chaos::FCollisionEventData& E
 	{
 		LastEventDispatchTime = Event.CollisionData.TimeCreated;
 
+		UE_CHAOS_ASYNC_INITBODY_READSCOPELOCK(CollisionEventRegistrationsLock);
 		// Iterate through the smallest between events registration and physics proxies
 		if (PhysicsProxyToCollisionIndicesMap.Num() <= CollisionEventRegistrations.Num())
 		{
@@ -1748,6 +1751,58 @@ void FPhysScene_Chaos::ProcessDeferredCreatePhysicsState()
 	DeferredCreatePhysicsStateComponents.Reset();
 }
 
+bool FPhysScene_Chaos::SupportsAsyncPhysicsStateCreation()
+{
+	return Chaos::CVars::bEnableAsyncInitBody;
+}
+
+bool FPhysScene_Chaos::PushAsyncCreatePhysicsStateJob(UActorComponent* Component)
+{
+	if (SupportsAsyncPhysicsStateCreation())
+	{
+		check(IsInGameThread());
+		check(!Component->IsPhysicsStateCreated());
+		check(Component->AllowsAsyncPhysicsStateCreationAndDestruction());
+
+		if (!AsyncPhysicsStateJobQueue.IsValid())
+		{
+			AsyncPhysicsStateJobQueue = MakeUnique<FPhysScene_AsyncPhysicsStateJobQueue>(this);
+		}
+		AsyncPhysicsStateJobQueue->AddJob({ Component, FPhysScene_AsyncPhysicsStateJobQueue::EJobType::CreatePhysicsState });
+		return true;
+	}
+
+	return false;
+}
+
+void FPhysScene_Chaos::RemoveAsyncPhysicsStateJobs(UActorComponent* Component)
+{
+	if (SupportsAsyncPhysicsStateCreation())
+	{
+		check(IsInGameThread());
+		if (AsyncPhysicsStateJobQueue.IsValid())
+		{
+			AsyncPhysicsStateJobQueue->RemoveJob({ Component, FPhysScene_AsyncPhysicsStateJobQueue::EJobType::CreatePhysicsState });
+		}
+	}
+}
+
+void FPhysScene_Chaos::ProcessAsyncPhysicsStateJobs()
+{
+	if (SupportsAsyncPhysicsStateCreation())
+	{
+		check(IsInGameThread());
+		if (AsyncPhysicsStateJobQueue.IsValid())
+		{
+			AsyncPhysicsStateJobQueue->Tick();
+			if (AsyncPhysicsStateJobQueue->IsCompleted())
+			{
+				AsyncPhysicsStateJobQueue.Reset();
+			}
+		}
+	}
+}
+
 // Collect the actors and transforms of all the bodies we have to move, and process them in bulk
 // to avoid locks in the Spatial Acceleration and the Solver's Dirty Proxy systems.
 void FPhysScene_Chaos::UpdateKinematicsOnDeferredSkelMeshes()
diff --git a/Engine/Source/Runtime/Engine/Private/UnrealEngine.cpp b/Engine/Source/Runtime/Engine/Private/UnrealEngine.cpp
--- a/Engine/Source/Runtime/Engine/Private/UnrealEngine.cpp
+++ b/Engine/Source/Runtime/Engine/Private/UnrealEngine.cpp
@@ -15901,7 +15901,6 @@ bool UEngine::LoadMap( FWorldContext& WorldContext, FURL URL, class UPendingNetG
 	{
 		FRegisterComponentContext Context(WorldContext.World());
 		WorldContext.World()->InitializeActorsForPlay(URL, true, &Context);
-		Context.Process();
 	}
 
 	// calling it after InitializeActorsForPlay has been called to have all potential bounding boxed initialized
diff --git a/Engine/Source/Runtime/Engine/Private/World.cpp b/Engine/Source/Runtime/Engine/Private/World.cpp
--- a/Engine/Source/Runtime/Engine/Private/World.cpp
+++ b/Engine/Source/Runtime/Engine/Private/World.cpp
@@ -3216,6 +3216,7 @@ void UWorld::AddToWorld( ULevel* Level, const FTransform& LevelTransform, bool b
 	// Don't make this level visible if it's currently being made invisible
 	if( bExecuteNextStep && CurrentLevelPendingVisibility == NULL && CurrentLevelPendingInvisibility != Level )
 	{
+		TRACE_CPUPROFILER_EVENT_SCOPE(AddToWorld_CurrentLevelPendingInvisibility);
 		TRACE_BEGIN_REGION(*WriteToString<256>(TEXT("AddToWorld: "), Level->GetOutermost()->GetName()));
 
 		Level->OwningWorld = this;
@@ -3251,6 +3252,7 @@ void UWorld::AddToWorld( ULevel* Level, const FTransform& LevelTransform, bool b
 
 	if( bExecuteNextStep && !Level->bAlreadyMovedActors )
 	{
+		TRACE_CPUPROFILER_EVENT_SCOPE(AddToWorld_MoveActors);
 		CSV_SCOPED_TIMING_STAT(LevelStreamingDetail, AddToWorld_MoveActors);
 		QUICK_SCOPE_CYCLE_COUNTER(STAT_AddToWorldTime_MoveActors);
 		SCOPE_TIME_TO_VAR(&MoveActorTime);
@@ -3265,6 +3267,7 @@ void UWorld::AddToWorld( ULevel* Level, const FTransform& LevelTransform, bool b
 
 	if( bExecuteNextStep && !Level->bAlreadyShiftedActors )
 	{
+		TRACE_CPUPROFILER_EVENT_SCOPE(AddToWorld_ShiftActors);
 		CSV_SCOPED_TIMING_STAT(LevelStreamingDetail, AddToWorld_ShiftActors);
 		QUICK_SCOPE_CYCLE_COUNTER(STAT_AddToWorldTime_ShiftActors);
 		SCOPE_TIME_TO_VAR(&ShiftActorsTime);
@@ -3323,6 +3326,7 @@ void UWorld::AddToWorld( ULevel* Level, const FTransform& LevelTransform, bool b
 	// Updates the level components (Actor components and UModelComponents).
 	if( bExecuteNextStep && !Level->bAlreadyUpdatedComponents )
 	{
+		TRACE_CPUPROFILER_EVENT_SCOPE(AddToWorld_UpdateComponents);
 		CSV_SCOPED_TIMING_STAT(LevelStreamingDetail, AddToWorld_UpdateComponents);
 		QUICK_SCOPE_CYCLE_COUNTER(STAT_AddToWorldTime_UpdatingComponents);
 		SCOPE_TIME_TO_VAR(&UpdateComponentsTime);
@@ -3339,25 +3343,26 @@ void UWorld::AddToWorld( ULevel* Level, const FTransform& LevelTransform, bool b
 		// from disk rather than duplicated
 		const bool bRerunConstructionScript = !FPlatformProperties::RequiresCookedData() && (!bIsGameWorld || !Level->bHasRerunConstructionScripts);
 
-		// Prepare context used to store batch/parallelize AddPrimitive calls
-		FRegisterComponentContext Context(this);
-		FRegisterComponentContext* ContextPtr = GLevelStreamingAddPrimitiveGranularity == 0 ? nullptr : &Context;
-
-		// Incrementally update components.
-		int32 NumComponentsToUpdate = (!bConsiderTimeLimit || !bIsGameWorld || IsRunningCommandlet() ? 0 : GLevelStreamingComponentsRegistrationGranularity);
-		do
+		// Incremental update components
 		{
-			Level->IncrementalUpdateComponents( NumComponentsToUpdate, bRerunConstructionScript, ContextPtr);
-			// Process AddPrimitives if threshold is reached
-			if (Context.Count() > GLevelStreamingAddPrimitiveGranularity)
+			// Prepare context used to store batch/parallelize AddPrimitive calls
+			const double ElapsedTimeMs = (FPlatformTime::Seconds() - StartTime) * 1000;
+			const double IncrementalUpdateTimeLimitMs = bConsiderTimeLimit ? FMath::Max<double>(0.0, TimeLimit - ElapsedTimeMs) : -1;
+			FRegisterComponentContext Context = FRegisterComponentContext(this)
+				.SetIncrementalUpdateTimeLimitMs(IncrementalUpdateTimeLimitMs)
+				.SetIncrementalUpdateAddPrimitiveBatchSize(FMath::Max<int32>(0, GLevelStreamingAddPrimitiveGranularity));
+
+			// Incrementally update components.
+			const int32 NumComponentsToUpdate = (!bConsiderTimeLimit || !bIsGameWorld || IsRunningCommandlet() ? 0 : GLevelStreamingComponentsRegistrationGranularity);
+			do
 			{
-				Context.Process();
-			}
+				Level->IncrementalUpdateComponents(NumComponentsToUpdate, bRerunConstructionScript, &Context);
+				if (Level->bAreComponentsCurrentlyRegistered || FLevelRegistrationAccessor::HasPreRegisteringComponents(Level))
+				{
+					break;
+				}
+			} while (!IsTimeLimitExceeded(TEXT("updating components"), StartTime, Level, TimeLimit));
 		}
-		while (!Level->bAreComponentsCurrentlyRegistered && !IsTimeLimitExceeded(TEXT("updating components"), StartTime, Level, TimeLimit));
-		
-		// Process remaining AddPrimitives
-		Context.Process();
 
 		// We are done once all components are attached.
 		Level->bAlreadyUpdatedComponents	= Level->bAreComponentsCurrentlyRegistered;
diff --git a/Engine/Source/Runtime/Engine/Public/Physics/Experimental/PhysScene_AsyncPhysicsStateJobQueue.h b/Engine/Source/Runtime/Engine/Public/Physics/Experimental/PhysScene_AsyncPhysicsStateJobQueue.h
new file mode 100644
--- /dev/null
+++ b/Engine/Source/Runtime/Engine/Public/Physics/Experimental/PhysScene_AsyncPhysicsStateJobQueue.h
@@ -0,0 +1,82 @@
+// Copyright Epic Games, Inc. All Rights Reserved.
+
+#pragma once
+
+#include "Tasks/Pipe.h"
+#include "Tasks/Task.h"
+#include "Async/Fundamental/Task.h"
+#include "Components/ActorComponent.h"
+#include "PhysicsInterfaceDeclaresCore.h"
+
+class UActorComponent;
+
+struct FPhysScene_AsyncPhysicsStateJobQueue
+{
+	FPhysScene_AsyncPhysicsStateJobQueue(FPhysScene* InPhysicScene)
+		: PhysScene(InPhysicScene)
+	{}
+
+	enum class EJobType
+	{
+		CreatePhysicsState,
+		DestroyPhysicsState
+	};
+
+	struct FJob
+	{
+		FJob(UActorComponent* InComponent, EJobType InType)
+			: Component(InComponent)
+			, Type(InType)
+			, bIsCancelled(false)
+		{}
+
+		bool IsValid() const { return ::IsValid(Component); }
+		void Execute() const;
+		void OnPostExecute_GameThread() const;
+
+		inline bool operator == (const FJob& Other) const
+		{
+			return (Component == Other.Component) && (Type == Other.Type);
+		}
+
+		friend inline uint32 GetTypeHash(const FJob& InJob)
+		{
+			return GetTypeHash(TPair<UActorComponent*, EJobType>{InJob.Component, InJob.Type});
+		}
+
+		UActorComponent* Component;
+		EJobType Type;
+		bool bIsCancelled;
+	};
+
+	struct FJobSet
+	{
+		void Add(const FJob& Job) { Jobs.Add(Job); }
+		int32 Remove(const FJob& Job) { return Jobs.Remove(Job); }
+		bool IsEmpty() const { return Jobs.IsEmpty(); }
+
+		void OnPreExecute_GameThread() const;
+		void Execute() const;
+		void OnPostExecute_GameThread() const;
+
+	private:
+		TSet<FJob> Jobs;
+	};
+
+	void Tick();
+	void AddJob(const FJob& Job);
+	void RemoveJob(const FJob& Job);
+	bool IsCompleted() const { return !Pipe.HasWork(); }
+
+private:
+
+	void Launch();
+	void WaitForCompletion();
+
+	FPhysScene* PhysScene;
+	using FJobs = TUniquePtr<FJobSet>;
+	FJobs PendingJobs;
+	TArray<FJobs> PipedJobs;
+	TArray<UE::Tasks::FTask> PipeTasks;
+	UE::Tasks::FPipe Pipe{ TEXT("FPhysScene::AsyncCreatePhysicStateQueue") };
+};
\ No newline at end of file
diff --git a/Engine/Source/Runtime/Engine/Public/Physics/Experimental/PhysScene_Chaos.h b/Engine/Source/Runtime/Engine/Public/Physics/Experimental/PhysScene_Chaos.h
--- a/Engine/Source/Runtime/Engine/Public/Physics/Experimental/PhysScene_Chaos.h
+++ b/Engine/Source/Runtime/Engine/Public/Physics/Experimental/PhysScene_Chaos.h
@@ -7,6 +7,7 @@
 #include "EventsData.h"
 #include "Physics/PhysScene.h"
 #include "Physics/Experimental/ChaosEventType.h"
+#include "Physics/Experimental/PhysScene_AsyncPhysicsStateJobQueue.h"
 #include "GameFramework/Actor.h"
 #include "PhysicsPublic.h"
 #include "PhysInterface_Chaos.h"
@@ -261,6 +262,11 @@ public:
 	ENGINE_API void RemoveDeferredPhysicsStateCreation(UPrimitiveComponent* Component);
 	ENGINE_API void ProcessDeferredCreatePhysicsState();
 
+	ENGINE_API static bool SupportsAsyncPhysicsStateCreation();
+	ENGINE_API bool PushAsyncCreatePhysicsStateJob(UActorComponent* Component);
+	ENGINE_API void RemoveAsyncPhysicsStateJobs(UActorComponent* Component);
+	ENGINE_API void ProcessAsyncPhysicsStateJobs();
+
 	UChaosEventRelay* GetChaosEventRelay() const { return ChaosEventRelay.Get(); }
 
 	/** Get cached state for replication, if no state is cached RegisterForReplicationCache() is called */
@@ -448,9 +454,15 @@ private:
 
 	TSet<UPrimitiveComponent*> DeferredCreatePhysicsStateComponents;
 
-	// RWLock for for thread safe access to PhysicsProxyToComponentMap and ComponentToPhysicsProxyMap.
+	// Queue used to create/destroy physics state asynchronously.
+	TUniquePtr<FPhysScene_AsyncPhysicsStateJobQueue> AsyncPhysicsStateJobQueue;
+
+	// RWLock for thread safe access to PhysicsProxyToComponentMap and ComponentToPhysicsProxyMap.
 	mutable FRWLock PhysicsProxyComponentMapsLock;
 
+	// RWLock for thread safe access to CollisionEventRegistrations
+	mutable FRWLock CollisionEventRegistrationsLock;
+
 	//Body Instances
 	TUniquePtr<Chaos::TArrayCollectionArray<FBodyInstance*>> BodyInstances;
 	// Temp Interface
