[World Partition] Properly empty AlwaysLoadedActorsForPIE for editor streaming generations.
This was causing accumulation of actors when using wp.Editor.DumpStreamingGenerationLog for example.

diff --git a/Engine/Source/Runtime/Engine/Private/WorldPartition/WorldPartitionRuntimeHash.cpp b/Engine/Source/Runtime/Engine/Private/WorldPartition/WorldPartitionRuntimeHash.cpp
--- a/Engine/Source/Runtime/Engine/Private/WorldPartition/WorldPartitionRuntimeHash.cpp
+++ b/Engine/Source/Runtime/Engine/Private/WorldPartition/WorldPartitionRuntimeHash.cpp
@@ -231,19 +231,13 @@ void UWorldPartitionRuntimeHash::PrepareEditorGameWorld()
 {
 	// Mark always loaded actors so that the Level will force reference to these actors for PIE.
 	// These actor will then be duplicated for PIE during the PIE world duplication process
-	ForceExternalActorLevelReference(/*bForceExternalActorLevelReferenceForPIE*/true);
+	ForceExternalActorLevelReference(/*bForceExternalActorLevelReference*/true);
 }
 
 void UWorldPartitionRuntimeHash::ShutdownEditorGameWorld()
 {
 	// Unmark always loaded actors
-	ForceExternalActorLevelReference(/*bForceExternalActorLevelReferenceForPIE*/false);
-
-	// Release references (will unload actors that were not already loaded in the Editor)
-	{
-		FWorldPartitionLoadingContext::FDeferred LoadingContext;
-		AlwaysLoadedActorsForPIE.Empty();
-	}
+	ForceExternalActorLevelReference(/*bForceExternalActorLevelReference*/false);
 }
 
 bool UWorldPartitionRuntimeHash::GenerateStreaming(class UWorldPartitionStreamingPolicy* StreamingPolicy, const IStreamingGenerationContext* StreamingGenerationContext, TArray<FString>* OutPackagesToGenerate)
@@ -254,6 +248,12 @@ bool UWorldPartitionRuntimeHash::GenerateStreaming(class UWorldPartitionStreamin
 void UWorldPartitionRuntimeHash::FlushStreamingContent()
 {
 	PackagesToGenerateForCook.Empty();
+
+	// Release references (will unload actors that were not already loaded in the editor)
+	{
+		FWorldPartitionLoadingContext::FDeferred LoadingContext;
+		EditorAlwaysLoadedActor.Empty();
+	}
 }
 
 bool UWorldPartitionRuntimeHash::PopulateCellActorInstances(const TArray<const IStreamingGenerationContext::FActorSetInstance*>& ActorSetInstances, bool bIsMainWorldPartition, bool bIsCellAlwaysLoaded, TArray<IStreamingGenerationContext::FActorInstance>& OutCellActorInstances)
@@ -307,7 +307,7 @@ bool UWorldPartitionRuntimeHash::PopulateCellActorInstances(const TArray<const I
 	{
 		if (AActor* AlwaysLoadedActor = FindObject<AActor>(nullptr, *ActorDescView->GetActorSoftPath().ToString()))
 		{
-			AlwaysLoadedActorsForPIE.Emplace(Reference, AlwaysLoadedActor);
+			EditorAlwaysLoadedActor.Emplace(Reference, AlwaysLoadedActor);
 		}
 	}
 
@@ -425,7 +425,7 @@ void UWorldPartitionRuntimeHash::DumpStateLog(FHierarchicalLogArchive& Ar) const
 
 		if (!IsRunningCookCommandlet())
 		{
-			for (const FAlwaysLoadedActorForPIE& AlwaysLoadedActor : AlwaysLoadedActorsForPIE)
+			for (const FEditorAlwaysLoadedActor& AlwaysLoadedActor : EditorAlwaysLoadedActor)
 			{
 				if (AlwaysLoadedActor.Actor.IsValid())
 				{
@@ -466,16 +466,16 @@ void UWorldPartitionRuntimeHash::DumpStateLog(FHierarchicalLogArchive& Ar) const
 	Ar.Printf(TEXT(""));
 }
 
-void UWorldPartitionRuntimeHash::ForceExternalActorLevelReference(bool bForceExternalActorLevelReferenceForPIE)
+void UWorldPartitionRuntimeHash::ForceExternalActorLevelReference(bool bForceExternalActorLevelReference)
 {
 	// Do this only on non game worlds prior to PIE so that always loaded actors get duplicated with the world
 	if (!GetWorld()->IsGameWorld())
 	{
-		for (const FAlwaysLoadedActorForPIE& AlwaysLoadedActor : AlwaysLoadedActorsForPIE)
+		for (const FEditorAlwaysLoadedActor& AlwaysLoadedActor : EditorAlwaysLoadedActor)
 		{
 			if (AActor* Actor = AlwaysLoadedActor.Actor.Get())
 			{
-				Actor->SetForceExternalActorLevelReferenceForPIE(bForceExternalActorLevelReferenceForPIE);
+				Actor->SetForceExternalActorLevelReferenceForPIE(bForceExternalActorLevelReference);
 			}
 		}
 	}
diff --git a/Engine/Source/Runtime/Engine/Public/WorldPartition/WorldPartitionRuntimeHash.h b/Engine/Source/Runtime/Engine/Public/WorldPartition/WorldPartitionRuntimeHash.h
--- a/Engine/Source/Runtime/Engine/Public/WorldPartition/WorldPartitionRuntimeHash.h
+++ b/Engine/Source/Runtime/Engine/Public/WorldPartition/WorldPartitionRuntimeHash.h
@@ -277,24 +277,25 @@ protected:
 	virtual EWorldPartitionStreamingPerformance GetStreamingPerformanceForCell(const UWorldPartitionRuntimeCell* Cell) const;
 
 #if WITH_EDITORONLY_DATA
-	struct FAlwaysLoadedActorForPIE
+	struct FEditorAlwaysLoadedActor
 	{
-		FAlwaysLoadedActorForPIE(const FWorldPartitionReference& InReference, AActor* InActor)
-			: Reference(InReference), Actor(InActor)
+		FEditorAlwaysLoadedActor(const FWorldPartitionReference& InReference, AActor* InActor)
+			: Reference(InReference)
+			, Actor(InActor)
 		{}
 
 		FWorldPartitionReference Reference;
 		TWeakObjectPtr<AActor> Actor;
 	};
 
-	TArray<FAlwaysLoadedActorForPIE> AlwaysLoadedActorsForPIE;
+	TArray<FEditorAlwaysLoadedActor> EditorAlwaysLoadedActor;
 
 	TMap<FString, TObjectPtr<UWorldPartitionRuntimeCell>> PackagesToGenerateForCook;
 #endif
 
 protected:
 #if WITH_EDITOR
-	ENGINE_API void ForceExternalActorLevelReference(bool bForceExternalActorLevelReferenceForPIE);
+	ENGINE_API void ForceExternalActorLevelReference(bool bForceExternalActorLevelReference);
 	ENGINE_API bool ResolveBlockOnSlowStreamingForCell(bool bInOwnerBlockOnSlowStreaming, bool bInIsHLODCell, const TArray<const UDataLayerInstance*>& InCellDataLayerInstances) const;
 #endif
 
