Update docs for UObject::CallRemoteFunction for the missing param and missing return description.

diff --git a/Engine/Source/Runtime/CoreUObject/Public/UObject/Object.h b/Engine/Source/Runtime/CoreUObject/Public/UObject/Object.h
--- a/Engine/Source/Runtime/CoreUObject/Public/UObject/Object.h
+++ b/Engine/Source/Runtime/CoreUObject/Public/UObject/Object.h
@@ -1479,8 +1479,10 @@ public:
 	 * Call the actor's function remotely
 	 *
 	 * @param Function function to call
-	 * @param Parameters arguments to the function call
+	 * @param Parms arguments to the function call
+	 * @param OutParms out parameter information (irrelevant for RPC's)
 	 * @param Stack stack frame for the function call
+	 * @return returns true if the function was processed by a NetDriver
 	 */
 	virtual bool CallRemoteFunction( UFunction* Function, void* Parms, struct FOutParmRec* OutParms, FFrame* Stack )
 	{
diff --git a/Engine/Source/Runtime/Engine/Classes/Engine/NetDriver.h b/Engine/Source/Runtime/Engine/Classes/Engine/NetDriver.h
--- a/Engine/Source/Runtime/Engine/Classes/Engine/NetDriver.h
+++ b/Engine/Source/Runtime/Engine/Classes/Engine/NetDriver.h
@@ -1536,11 +1536,12 @@ public:
 	 *
 	 * @param Actor actor making the function call
 	 * @param Function function definition called
-	 * @param Params parameters in a UObject memory layout
+	 * @param Parms parameters in a UObject memory layout
+	 * @param OutParms out parameter information (irrelevant for RPC's)
 	 * @param Stack stack frame the UFunction is called in
 	 * @param SubObject optional: sub object to actually call function on
 	 */
-	ENGINE_API virtual void ProcessRemoteFunction(class AActor* Actor, class UFunction* Function, void* Parameters, struct FOutParmRec* OutParms, struct FFrame* Stack, class UObject* SubObject = nullptr );
+	ENGINE_API virtual void ProcessRemoteFunction(class AActor* Actor, class UFunction* Function, void* Parms, struct FOutParmRec* OutParms, struct FFrame* Stack, class UObject* SubObject = nullptr );
 
 	/** Return a reference to the database that holds metrics calcluated by the networking system. */
 	ENGINE_API TObjectPtr<UNetworkMetricsDatabase> GetMetrics() { return NetworkMetricsDatabase; };
