Engine - setting UAnimSequenceBase::bLoop as BlueprintReadOnly

diff --git a/Engine/Source/Runtime/Engine/Classes/Animation/AnimSequenceBase.h b/Engine/Source/Runtime/Engine/Classes/Animation/AnimSequenceBase.h
--- a/Engine/Source/Runtime/Engine/Classes/Animation/AnimSequenceBase.h
+++ b/Engine/Source/Runtime/Engine/Classes/Animation/AnimSequenceBase.h
@@ -64,7 +64,7 @@ public:
 	 * The default looping behavior of this animation.
 	 * Asset players can override this
 	 */
-	UPROPERTY(EditAnywhere, Category=Animation)
+	UPROPERTY(EditAnywhere, Category=Animation, BlueprintReadOnly)
 	bool bLoop;
 #if WITH_EDITORONLY_DATA
 	// if you change Notifies array, this will need to be rebuilt
