Makes sure that mesh paint textures allocate their VT even after a texture compile resource update race condition fails the first AquireAllocatedVT()
This fixes intermittent bug where texture is opened in texture editor before entering mesh paint mode.

diff --git a/Engine/Source/Runtime/Engine/Private/VT/MeshPaintVirtualTexture.cpp b/Engine/Source/Runtime/Engine/Private/VT/MeshPaintVirtualTexture.cpp
--- a/Engine/Source/Runtime/Engine/Private/VT/MeshPaintVirtualTexture.cpp
+++ b/Engine/Source/Runtime/Engine/Private/VT/MeshPaintVirtualTexture.cpp
@@ -195,6 +195,24 @@ namespace MeshPaintVirtualTexture
 	 * Call this on texture resource creation.
 	 * This will acquire the virtual texture and store to our global set.
 	 */
+	static IAllocatedVirtualTexture* AcquireAllocatedVT_RenderThread(FVirtualTexture2DResource* VTResource)
+	{
+		IAllocatedVirtualTexture* AllocatedVT = VTResource->GetAllocatedVT();
+		if (AllocatedVT == nullptr)
+		{
+			AllocatedVT = VTResource->AcquireAllocatedVT();
+			if (ensure(AllocatedVT))
+			{
+				// Add virtual texture to the global set.
+				AddAllocatedVT(AllocatedVT);
+				// Queue on-destruction callback for removal from the global set.
+				GetRendererModule().AddVirtualTextureProducerDestroyedCallback(AllocatedVT->GetProducerHandle(0), &RemoveAllocatedVT, AllocatedVT);
+			}
+		}
+
+		return AllocatedVT;
+	}
+
 	static void AcquireAllocatedVT(FTextureResource* Resource)
 	{
 		FVirtualTexture2DResource* VTResource = Resource != nullptr ? Resource->GetVirtualTexture2DResource() : nullptr;
@@ -202,15 +220,9 @@ namespace MeshPaintVirtualTexture
 		{
 			ENQUEUE_RENDER_COMMAND(AcquireVT)([VTResource](FRHICommandListImmediate& RHICmdList)
 			{
-				if (IAllocatedVirtualTexture* AllocatedVT = VTResource->AcquireAllocatedVT())
-				{
-					// Add virtual texture to the global set.
-					AddAllocatedVT(AllocatedVT);
-					// Queue on-destruction callback for removal from the global set.
-					GetRendererModule().AddVirtualTextureProducerDestroyedCallback(AllocatedVT->GetProducerHandle(0), &RemoveAllocatedVT, AllocatedVT);
-				}
+				AcquireAllocatedVT_RenderThread(VTResource);
 			});
-		}
+		}	
 	}
 
 	FUniformParams GetUniformParams()
@@ -243,7 +255,7 @@ namespace MeshPaintVirtualTexture
 			return Descriptor;
 		}
 
-		IAllocatedVirtualTexture* AllocatedVT = VTTextureResource->GetAllocatedVT();
+		IAllocatedVirtualTexture* AllocatedVT = AcquireAllocatedVT_RenderThread(VTTextureResource);
 		if (AllocatedVT == nullptr)
 		{
 			return Descriptor;
