Fix copy&paste

diff --git a/Engine/Source/Runtime/Engine/Private/UnrealExporter.cpp b/Engine/Source/Runtime/Engine/Private/UnrealExporter.cpp
--- a/Engine/Source/Runtime/Engine/Private/UnrealExporter.cpp
+++ b/Engine/Source/Runtime/Engine/Private/UnrealExporter.cpp
@@ -726,13 +726,13 @@ void ExportProperties
 		FArrayProperty* ArrayProperty = CastField<FArrayProperty>(Property);
 		FMapProperty* MapProperty = CastField<FMapProperty>(Property);
 
-		FObjectPropertyBase* ExportObjectProp = (Property->PropertyFlags & CPF_ExportObject) != 0 ? CastField<FObjectPropertyBase>(Property) : NULL;
+		FObjectPropertyBase *const ExportObjectProp = (Property->PropertyFlags & CPF_ExportObject) != 0 ? CastField<FObjectPropertyBase>(Property) : NULL;
 		const uint32 ExportFlags = PortFlags | PPF_Delimited;
-		auto HandleExportObject = [Context, &Out, Indent, Parent, PortFlags, ExportObjectProp](const FProperty* Property, uint8* Data)
+		auto HandleExportObject = [Context, &Out, Indent, Parent, PortFlags](const FProperty* Property, uint8* Data, FObjectPropertyBase*const InExportObjectProp)
 		{
-			if (ExportObjectProp)
+			if (InExportObjectProp)
 			{
-				UObject* Obj = ExportObjectProp->GetObjectPropertyValue(Data);
+				UObject* Obj = InExportObjectProp->GetObjectPropertyValue(Data);
 				check(!Obj || Obj->IsValidLowLevel());
 				if (Obj && !Obj->HasAnyMarks(OBJECTMARK_TagImp))
 				{
@@ -772,7 +772,7 @@ void ExportProperties
 			{
 				// Export dynamic array.
 				FProperty* InnerProp = ArrayProperty->Inner;
-				ExportObjectProp = (Property->PropertyFlags & CPF_ExportObject) != 0 ? CastField<FObjectPropertyBase>(InnerProp) : NULL;
+				FObjectPropertyBase*const ArrayExportObjectProp = (Property->PropertyFlags & CPF_ExportObject) != 0 ? CastField<FObjectPropertyBase>(InnerProp) : NULL;
 				// This is used as the default value in the case of an array property that has
 				// fewer elements than the exported object.
 				uint8* StructDefaults = nullptr;
@@ -842,7 +842,7 @@ void ExportProperties
 						if (bExportItem)
 						{
 							InnerProp->ExportTextItem_Direct(Value, SourceData, DiffData, Parent, ExportFlags, ExportRootScope);
-							HandleExportObject(InnerProp, ArrayHelper.GetRawPtr(DynamicArrayIndex));
+							HandleExportObject(InnerProp, ArrayHelper.GetRawPtr(DynamicArrayIndex), ArrayExportObjectProp);
 
 							Out.Logf(TEXT("%s%s(%i)=%s\r\n"), FCString::Spc(Indent), *SanitizedPropertyName, DynamicArrayIndex, *Value);
 						}
@@ -930,7 +930,7 @@ void ExportProperties
 
 							FString Value;
 							InnerProp->ExportTextItem_Direct(Value, ArrayHelper.GetRawPtr(i), nullptr, Parent, ExportFlags, ExportRootScope);
-							HandleExportObject(InnerProp, ArrayHelper.GetRawPtr(i));
+							HandleExportObject(InnerProp, ArrayHelper.GetRawPtr(i), ArrayExportObjectProp);
 
 							Out.Logf(TEXT("%s%s<%s>=%s\r\n"), FCString::Spc(Indent), *SanitizedPropertyName, *GetOverriddenOperationString(EOverriddenPropertyOperation::Add), *Value);
 						}
@@ -1062,7 +1062,7 @@ void ExportProperties
 			uint8* DiffData = (DiffClass && Property->IsInContainer(DiffClass->GetPropertiesSize())) ? Diff : NULL;
 			if (Property->ExportText_InContainer(PropertyArrayIndex, Value, Object, DiffData, Parent, ExportFlags, ExportRootScope))
 			{
-				HandleExportObject(Property, Property->ContainerPtrToValuePtr<uint8>(Object, PropertyArrayIndex));
+				HandleExportObject(Property, Property->ContainerPtrToValuePtr<uint8>(Object, PropertyArrayIndex), ExportObjectProp);
 
 				FString OverridableOperation;
 				if (OverriddenProperties)
