[UBA] Minor fix of missing newline character to better separate additional information in UbaVisualizer breadcrumbs.

#rnx

diff --git a/Engine/Source/Runtime/Engine/Private/ShaderCompiler/ShaderCompilerDistributed.cpp b/Engine/Source/Runtime/Engine/Private/ShaderCompiler/ShaderCompilerDistributed.cpp
--- a/Engine/Source/Runtime/Engine/Private/ShaderCompiler/ShaderCompilerDistributed.cpp
+++ b/Engine/Source/Runtime/Engine/Private/ShaderCompiler/ShaderCompilerDistributed.cpp
@@ -125,7 +125,7 @@ static FString BuildCompactTaskDescription(const TArray<FShaderCommonCompileJobP
 		}
 		if (JobsToSerialize.Num() > NumJobsInDescription)
 		{
-			Description += FString::Printf(TEXT("%d more shaders ...\n"), JobsToSerialize.Num() - NumJobsInDescription);
+			Description += FString::Printf(TEXT("\n%d more shaders ...\n"), JobsToSerialize.Num() - NumJobsInDescription);
 		}
 	}
 
