Removed ensure in AActor::PreRegisterAllComponents() because there are cases where PreRegisterAllComponents() can be called multiple times (SpawnActor is one example).

diff --git a/Engine/Source/Runtime/Engine/Private/Actor.cpp b/Engine/Source/Runtime/Engine/Private/Actor.cpp
--- a/Engine/Source/Runtime/Engine/Private/Actor.cpp
+++ b/Engine/Source/Runtime/Engine/Private/Actor.cpp
@@ -3679,7 +3679,7 @@ void AActor::DisableComponentsSimulatePhysics()
 
 void AActor::PreRegisterAllComponents()
 {
-	ensureMsgf(!bHasPreRegisteredAllComponents, TEXT("Actor %s PreRegisterAllComponents has already been called (or bHasPreRegisteredAllComponents was not properly reset)."), *this->GetFullName());
+	// Here we don't ensure that bHasPreRegisteredAllComponents == false because there are cases where PreRegisterAllComponents() can be called multiple times (SpawnActor is one example)
 	bHasPreRegisteredAllComponents = true;
 }
 
