Fix crashes related to user code passing null USubsystem class.

diff --git a/Engine/Source/Runtime/Engine/Private/Subsystems/SubsystemCollection.cpp b/Engine/Source/Runtime/Engine/Private/Subsystems/SubsystemCollection.cpp
--- a/Engine/Source/Runtime/Engine/Private/Subsystems/SubsystemCollection.cpp
+++ b/Engine/Source/Runtime/Engine/Private/Subsystems/SubsystemCollection.cpp
@@ -49,6 +49,12 @@ FSubsystemCollectionBase::FSubsystemCollectionBase(UClass* InBaseType)
 
 USubsystem* FSubsystemCollectionBase::GetSubsystemInternal(UClass* SubsystemClass) const
 {
+	// It does not make sense to get a subsystem by null class.
+	if (!ensure(SubsystemClass))
+	{
+		return nullptr;
+	}
+	
 #if WITH_EDITOR && UE_BUILD_SHIPPING
 	TStringBuilder<200> DebugSubsystemClassName;
 	if (SubsystemClass && IsEngineExitRequested())
