PR #10243: Expose function to compute scalability from existing perf indices

diff --git a/Engine/Source/Runtime/Engine/Private/Scalability.cpp b/Engine/Source/Runtime/Engine/Private/Scalability.cpp
--- a/Engine/Source/Runtime/Engine/Private/Scalability.cpp
+++ b/Engine/Source/Runtime/Engine/Private/Scalability.cpp
@@ -674,19 +674,40 @@ ENGINE_API float GetResolutionQualityFromGPUPerfIndex(float GPUPerfIndex)
 	 return GetRenderScaleLevelFromQualityLevel(ComputeOptionFromPerfIndex(TEXT("ResolutionQuality"), 0, GPUPerfIndex));
 }
 
+ENGINE_API FQualityLevels ComputeQualityLevelsFromPerfIndex(float CPUPerfIndex, float GPUPerfIndex)
+{
+	FQualityLevels Results;
+	Results.ResolutionQuality = GetResolutionQualityFromGPUPerfIndex(GPUPerfIndex);
+	Results.ViewDistanceQuality = ComputeOptionFromPerfIndex(TEXT("ViewDistanceQuality"), CPUPerfIndex, GPUPerfIndex);
+	Results.AntiAliasingQuality = ComputeOptionFromPerfIndex(TEXT("AntiAliasingQuality"), CPUPerfIndex, GPUPerfIndex);
+	Results.ShadowQuality = ComputeOptionFromPerfIndex(TEXT("ShadowQuality"), CPUPerfIndex, GPUPerfIndex);
+	Results.GlobalIlluminationQuality = ComputeOptionFromPerfIndex(TEXT("GlobalIlluminationQuality"), CPUPerfIndex, GPUPerfIndex);
+	Results.ReflectionQuality = ComputeOptionFromPerfIndex(TEXT("ReflectionQuality"), CPUPerfIndex, GPUPerfIndex);
+	Results.PostProcessQuality = ComputeOptionFromPerfIndex(TEXT("PostProcessQuality"), CPUPerfIndex, GPUPerfIndex);
+	Results.TextureQuality = ComputeOptionFromPerfIndex(TEXT("TextureQuality"), CPUPerfIndex, GPUPerfIndex);
+	Results.EffectsQuality = ComputeOptionFromPerfIndex(TEXT("EffectsQuality"), CPUPerfIndex, GPUPerfIndex);
+	Results.FoliageQuality = ComputeOptionFromPerfIndex(TEXT("FoliageQuality"), CPUPerfIndex, GPUPerfIndex);
+	Results.ShadingQuality = ComputeOptionFromPerfIndex(TEXT("ShadingQuality"), CPUPerfIndex, GPUPerfIndex);
+	Results.LandscapeQuality = ComputeOptionFromPerfIndex( TEXT("LandscapeQuality"), CPUPerfIndex, GPUPerfIndex);
+	Results.CPUBenchmarkResults = CPUPerfIndex;
+	Results.GPUBenchmarkResults = GPUPerfIndex;
+	return Results;
+}
+
 FQualityLevels BenchmarkQualityLevels(uint32 WorkScale, float CPUMultiplier, float GPUMultiplier)
 {
 	ensure((CPUMultiplier > 0.0f) && (GPUMultiplier > 0.0f));
 
 	// benchmark the system
 
-	FQualityLevels Results;
+	TArray<float> CPUBenchmarkSteps;
+	TArray<float> GPUBenchmarkSteps;
 
 	FSynthBenchmarkResults SynthBenchmark;
 	ISynthBenchmark::Get().Run(SynthBenchmark, true, WorkScale);
 
-	float CPUPerfIndex = SynthBenchmark.ComputeCPUPerfIndex(/*out*/ &Results.CPUBenchmarkSteps) * CPUMultiplier;
-	float GPUPerfIndex = SynthBenchmark.ComputeGPUPerfIndex(/*out*/ &Results.GPUBenchmarkSteps) * GPUMultiplier;
+	float CPUPerfIndex = SynthBenchmark.ComputeCPUPerfIndex(/*out*/ &CPUBenchmarkSteps) * CPUMultiplier;
+	float GPUPerfIndex = SynthBenchmark.ComputeGPUPerfIndex(/*out*/ &GPUBenchmarkSteps) * GPUMultiplier;
 
 #if !UE_BUILD_SHIPPING
 	if (CVarTestCPUPerfIndexOverride.GetValueOnAnyThread() > 0.0f)
@@ -700,20 +721,9 @@ FQualityLevels BenchmarkQualityLevels(uint32 WorkScale, float CPUMultiplier, flo
 #endif
 
 	// decide on the actual quality needed
-	Results.ResolutionQuality = GetResolutionQualityFromGPUPerfIndex(GPUPerfIndex);
-	Results.ViewDistanceQuality = ComputeOptionFromPerfIndex(TEXT("ViewDistanceQuality"), CPUPerfIndex, GPUPerfIndex);
-	Results.AntiAliasingQuality = ComputeOptionFromPerfIndex(TEXT("AntiAliasingQuality"), CPUPerfIndex, GPUPerfIndex);
-	Results.ShadowQuality = ComputeOptionFromPerfIndex(TEXT("ShadowQuality"), CPUPerfIndex, GPUPerfIndex);
-	Results.GlobalIlluminationQuality = ComputeOptionFromPerfIndex(TEXT("GlobalIlluminationQuality"), CPUPerfIndex, GPUPerfIndex);
-	Results.ReflectionQuality = ComputeOptionFromPerfIndex(TEXT("ReflectionQuality"), CPUPerfIndex, GPUPerfIndex);
-	Results.PostProcessQuality = ComputeOptionFromPerfIndex(TEXT("PostProcessQuality"), CPUPerfIndex, GPUPerfIndex);
-	Results.TextureQuality = ComputeOptionFromPerfIndex(TEXT("TextureQuality"), CPUPerfIndex, GPUPerfIndex);
-	Results.EffectsQuality = ComputeOptionFromPerfIndex(TEXT("EffectsQuality"), CPUPerfIndex, GPUPerfIndex);
-	Results.FoliageQuality = ComputeOptionFromPerfIndex(TEXT("FoliageQuality"), CPUPerfIndex, GPUPerfIndex);
-	Results.ShadingQuality = ComputeOptionFromPerfIndex(TEXT("ShadingQuality"), CPUPerfIndex, GPUPerfIndex);
-	Results.LandscapeQuality = ComputeOptionFromPerfIndex( TEXT("LandscapeQuality"), CPUPerfIndex, GPUPerfIndex);
-	Results.CPUBenchmarkResults = CPUPerfIndex;
-	Results.GPUBenchmarkResults = GPUPerfIndex;
+	FQualityLevels Results = ComputeQualityLevelsFromPerfIndex(CPUPerfIndex, GPUPerfIndex);
+	Results.CPUBenchmarkSteps = MoveTemp(CPUBenchmarkSteps);
+	Results.GPUBenchmarkSteps = MoveTemp(GPUBenchmarkSteps);
 
 	return Results;
 }
diff --git a/Engine/Source/Runtime/Engine/Public/Scalability.h b/Engine/Source/Runtime/Engine/Public/Scalability.h
--- a/Engine/Source/Runtime/Engine/Public/Scalability.h
+++ b/Engine/Source/Runtime/Engine/Public/Scalability.h
@@ -210,6 +210,9 @@ namespace Scalability
 	 */
 	ENGINE_API void RecordQualityLevelsAnalytics(bool bAutoApplied);
 
+	/** Compute scalability quality levels from existing CPU and GPU perf indices **/
+	ENGINE_API FQualityLevels ComputeQualityLevelsFromPerfIndex(float CPUPerfIndex, float GPUPerfIndex);
+	
 	/** Run synthbenchmark and configure scalability based on results **/
 	ENGINE_API FQualityLevels BenchmarkQualityLevels(uint32 WorkScale=10, float CPUMultiplier = 1.0f, float GPUMultiplier = 1.0f);
 
