Fix for issue where looping level sequences could get stuck on the final frame

diff --git a/Engine/Source/Runtime/Engine/Private/Animation/AnimMontage.cpp b/Engine/Source/Runtime/Engine/Private/Animation/AnimMontage.cpp
--- a/Engine/Source/Runtime/Engine/Private/Animation/AnimMontage.cpp
+++ b/Engine/Source/Runtime/Engine/Private/Animation/AnimMontage.cpp
@@ -2088,7 +2088,9 @@ EMontageSubStepResult FMontageSubStepper::Advance(float& InOut_P_Original, const
 		const float NewPosition = MontageInstance->ForcedNextToPosition.GetValue();
 		if (MontageInstance->ForcedNextFromPosition.IsSet())
 		{
+			// We are modifying the current position so we also need to update the section and pos in section
 			InOut_P_Original = MontageInstance->ForcedNextFromPosition.GetValue();
+			CurrentSectionIndex = Montage->GetAnimCompositeSectionIndexFromPos(InOut_P_Original, PositionInSection);
 		}
 		DeltaMove = NewPosition - InOut_P_Original;
 		PlayRate = DeltaMove / TimeRemaining;
