Make SimulateAdvance symmetrical with Advance when dealing with ends of sections

diff --git a/Engine/Source/Runtime/Engine/Private/Animation/AnimMontage.cpp b/Engine/Source/Runtime/Engine/Private/Animation/AnimMontage.cpp
--- a/Engine/Source/Runtime/Engine/Private/Animation/AnimMontage.cpp
+++ b/Engine/Source/Runtime/Engine/Private/Animation/AnimMontage.cpp
@@ -2030,6 +2030,12 @@ bool FAnimMontageInstance::SimulateAdvance(float DeltaTime, float& InOutPosition
 			}
 			else
 			{
+				// Clamp position to prevent playing past the end of the current section
+				float CurrentSectionStartTime, CurrentSectionEndTime;
+				Montage->GetSectionStartAndEndTime(CurrentSectionIndex, CurrentSectionStartTime, CurrentSectionEndTime);
+
+				InOutPosition = bPlayingForward ? (CurrentSectionEndTime - UE_KINDA_SMALL_NUMBER) : CurrentSectionStartTime;
+
 				// Reached end of last section. Exit.
 				break;
 			}
