Fixed crash in FMeshBatchAndRelevance ctor.

diff --git a/Engine/Source/Runtime/Engine/Private/SceneManagement.cpp b/Engine/Source/Runtime/Engine/Private/SceneManagement.cpp
--- a/Engine/Source/Runtime/Engine/Private/SceneManagement.cpp
+++ b/Engine/Source/Runtime/Engine/Private/SceneManagement.cpp
@@ -335,10 +335,19 @@ FMeshBatchAndRelevance::FMeshBatchAndRelevance(const FMeshBatch& InMesh, const F
 	Mesh(&InMesh),
 	PrimitiveSceneProxy(InPrimitiveSceneProxy)
 {
-	const FMaterial& Material = InMesh.MaterialRenderProxy->GetIncompleteMaterialWithFallback(FeatureLevel);
-	bHasOpaqueMaterial = IsOpaqueBlendMode(Material);
-	bHasMaskedMaterial = IsMaskedBlendMode(Material);
-	bRenderInMainPass = PrimitiveSceneProxy->ShouldRenderInMainPass();
+	if (InMesh.MaterialRenderProxy)
+	{
+		const FMaterial& Material = InMesh.MaterialRenderProxy->GetIncompleteMaterialWithFallback(FeatureLevel);
+		bHasOpaqueMaterial = IsOpaqueBlendMode(Material);
+		bHasMaskedMaterial = IsMaskedBlendMode(Material);
+	}
+	else
+	{
+		bHasOpaqueMaterial = false;
+		bHasMaskedMaterial = false;
+	}
+
+	bRenderInMainPass = PrimitiveSceneProxy ? PrimitiveSceneProxy->ShouldRenderInMainPass() : false;
 }
 
 #if RHI_RAYTRACING
