PR #12280: Exclude VT samples that don't generate feedback when counting requests

diff --git a/Engine/Shaders/Private/VirtualTextureCommon.ush b/Engine/Shaders/Private/VirtualTextureCommon.ush
--- a/Engine/Shaders/Private/VirtualTextureCommon.ush
+++ b/Engine/Shaders/Private/VirtualTextureCommon.ush
@@ -13,6 +13,10 @@
 #error NUM_VIRTUALTEXTURE_SAMPLES should be set to the number of VT samples that will be done before including this file (exluding any lightmap samples)
 #endif
 
+#ifndef NUM_VIRTUALTEXTURE_FEEDBACK_REQUESTS
+#define NUM_VIRTUALTEXTURE_FEEDBACK_REQUESTS NUM_VIRTUALTEXTURE_SAMPLES
+#endif
+
 #ifndef LIGHTMAP_VT_ENABLED
 #define LIGHTMAP_VT_ENABLED 0
 #endif
@@ -21,7 +25,7 @@
 struct FVirtualTextureFeedbackParams
 {
 	uint Request;	// The tile to request for this pixel
-#if (NUM_VIRTUALTEXTURE_SAMPLES + LIGHTMAP_VT_ENABLED) > 1	
+#if (NUM_VIRTUALTEXTURE_FEEDBACK_REQUESTS + LIGHTMAP_VT_ENABLED) > 1	
 	uint RequestId;	// The id of the VT texture sample this pixel will generate feedback for
 #endif	
 };
@@ -29,8 +33,8 @@ struct FVirtualTextureFeedbackParams
 /** Initializes the FVirtualTextureFeedbackParams for the pixel shader. */
 void InitializeVirtualTextureFeedback(in out FVirtualTextureFeedbackParams Params, uint2 SvPosition, uint FrameNumber)
 {
-#if (NUM_VIRTUALTEXTURE_SAMPLES + LIGHTMAP_VT_ENABLED) > 1
-	const uint NumVTSamplesInShader = NUM_VIRTUALTEXTURE_SAMPLES + LIGHTMAP_VT_ENABLED;
+#if (NUM_VIRTUALTEXTURE_FEEDBACK_REQUESTS + LIGHTMAP_VT_ENABLED) > 1
+	const uint NumVTSamplesInShader = NUM_VIRTUALTEXTURE_FEEDBACK_REQUESTS + LIGHTMAP_VT_ENABLED;
 	const uint2 PixelPos = SvPosition >> View.VirtualTextureFeedbackShift;
 	const uint FeedbackPos = PixelPos.y * View.VirtualTextureFeedbackStride + PixelPos.x;
 	Params.RequestId = (View.VirtualTextureFeedbackSampleOffset + FeedbackPos) % NumVTSamplesInShader;
@@ -43,7 +47,7 @@ void InitializeVirtualTextureFeedback(in out FVirtualTextureFeedbackParams Param
 /** Store feedback info for a VT request in the passed in FVirtualTextureFeedbackParams. */
 void StoreVirtualTextureFeedback(in out FVirtualTextureFeedbackParams Params, uint RequestId, uint Request)
 {
-#if (NUM_VIRTUALTEXTURE_SAMPLES + LIGHTMAP_VT_ENABLED) > 1
+#if (NUM_VIRTUALTEXTURE_FEEDBACK_REQUESTS + LIGHTMAP_VT_ENABLED) > 1
 	Params.Request = (RequestId == Params.RequestId) ? Request : Params.Request;
 #else
 	Params.Request = Request;	
diff --git a/Engine/Source/Runtime/Core/Private/UObject/FortniteMainBranchObjectVersion.cpp b/Engine/Source/Runtime/Core/Private/UObject/FortniteMainBranchObjectVersion.cpp
--- a/Engine/Source/Runtime/Core/Private/UObject/FortniteMainBranchObjectVersion.cpp
+++ b/Engine/Source/Runtime/Core/Private/UObject/FortniteMainBranchObjectVersion.cpp
@@ -14,6 +14,6 @@ TMap<FGuid, FGuid> FFortniteMainBranchObjectVersion::GetSystemGuids()
 	SystemGuids.Add(DevGuids.Niagara_LatestScriptCompileVersion, FGuid("F8D995224324470E97ED718D33939A0E"));
 	SystemGuids.Add(DevGuids.SkeletalMeshDerivedDataVersion, FGuid("B9AEAA3EE7AC4FCFB3860E8F3E06DD75"));
 	SystemGuids.Add(DevGuids.STATICMESH_DERIVEDDATA_VER, FGuid("3231A08A87A248A09D959335B0F4060A"));
-	SystemGuids.Add(DevGuids.MaterialTranslationDDCVersion, FGuid("0F9C4453CDDC4B8BA1527D4942381ECA"));
+	SystemGuids.Add(DevGuids.MaterialTranslationDDCVersion, FGuid("C260AF26228241889456439BE5F5AF51"));
 	return SystemGuids;
 }
diff --git a/Engine/Source/Runtime/Engine/Private/Materials/HLSLMaterialTranslator.cpp b/Engine/Source/Runtime/Engine/Private/Materials/HLSLMaterialTranslator.cpp
--- a/Engine/Source/Runtime/Engine/Private/Materials/HLSLMaterialTranslator.cpp
+++ b/Engine/Source/Runtime/Engine/Private/Materials/HLSLMaterialTranslator.cpp
@@ -292,6 +292,7 @@ struct FHLSLMaterialTranslator::FEnvironmentDefines
 	uint8 MaterialPathTracingBufferRead;
 	bool MaterialNeuralPostProcess;
 	uint32 NumVirtualTextureSamples;
+	uint32 NumVirtualTextureFeedbackRequests;
 	bool bMaterialVirtualTextureFeedback;
 	TArray<uint8> VirtualPageTypes;
 	bool bSingleLayerWaterShadingQuality;
@@ -387,6 +388,7 @@ struct FHLSLMaterialTranslator::FEnvironmentDefines
 		Ar << MaterialPathTracingBufferRead;
 		Ar << MaterialNeuralPostProcess;
 		Ar << NumVirtualTextureSamples;
+		Ar << NumVirtualTextureFeedbackRequests;
 		Ar << bMaterialVirtualTextureFeedback;
 		Ar << VirtualPageTypes;
 		Ar << bSingleLayerWaterShadingQuality;
@@ -2683,6 +2685,7 @@ void FHLSLMaterialTranslator::GetMaterialEnvironment(EShaderPlatform InPlatform,
 	OutEnvironment.SetDefine(TEXT("MATERIAL_NEURAL_POST_PROCESS"), EnvironmentDefines->MaterialNeuralPostProcess);
 
 	OutEnvironment.SetDefine(TEXT("NUM_VIRTUALTEXTURE_SAMPLES"), EnvironmentDefines->NumVirtualTextureSamples);
+	OutEnvironment.SetDefine(TEXT("NUM_VIRTUALTEXTURE_FEEDBACK_REQUESTS"), EnvironmentDefines->NumVirtualTextureFeedbackRequests);
 	OutEnvironment.SetDefine(TEXT("MATERIAL_VIRTUALTEXTURE_FEEDBACK"), EnvironmentDefines->bMaterialVirtualTextureFeedback);
 
 	for (int i = 0; i < EnvironmentDefines->VirtualPageTypes.Num(); ++i)
@@ -6870,9 +6873,23 @@ uint32 FHLSLMaterialTranslator::AcquireVTStackIndex(
 	// Select LoadVirtualPageTable function name for this context
 	FString BaseFunctionName = bAdaptive ? TEXT("TextureLoadVirtualPageTableAdaptive") : TEXT("TextureLoadVirtualPageTable");
 
+	// Count number of existing samples that generate feedback
+	int32 FeedbackRequestIndex = 0;
+	if (bGenerateFeedback)
+	{
+		const int32 NumStacks = VTStacks.Num() - 1; // Exclude current stack
+		for (int32 Index = 0; Index < NumStacks; ++Index)
+		{
+			if (VTStacks[Index].bGenerateFeedback)
+			{
+				++FeedbackRequestIndex;
+			}
+		}
+	}
+
 	// Optionally sample without virtual texture feedback but only for miplevel mode
 	check(bGenerateFeedback || MipValueMode == TMVM_MipLevel)
-	FString FeedbackParameter = bGenerateFeedback ? FString::Printf(TEXT(", %dU + LIGHTMAP_VT_ENABLED, Parameters.VirtualTextureFeedback"), StackIndex) : TEXT("");
+	FString FeedbackParameter = bGenerateFeedback ? FString::Printf(TEXT(", %dU + LIGHTMAP_VT_ENABLED, Parameters.VirtualTextureFeedback"), FeedbackRequestIndex) : TEXT("");
 
 	EDerivativeStatus UVDerivativeStatus = GetDerivativeStatus(CoordinateIndex);
 	const bool bHasValidDerivative = IsAnalyticDerivEnabled() && IsDerivativeValid(UVDerivativeStatus);
@@ -15686,12 +15703,16 @@ void FHLSLMaterialTranslator::PrepareEnvironmentDefines()
 	EnvironmentDefines->NumVirtualTextureSamples = VTStacks.Num();
 
 	// Check if any feedback slots are in use. We can simplify shader and remove EARLYZ optimizations if none are.
-	bool bGenerateFeedback = false;
-	for (int i = 0; i < VTStacks.Num() && !bGenerateFeedback; ++i)
+	uint32 NumVirtualTextureFeedbackRequests = 0;
+	for (int i = 0; i < VTStacks.Num(); ++i)
 	{
-		bGenerateFeedback |= VTStacks[i].bGenerateFeedback;
+		if (VTStacks[i].bGenerateFeedback)
+		{
+			++NumVirtualTextureFeedbackRequests;
+		}
 	}
-	EnvironmentDefines->bMaterialVirtualTextureFeedback = bGenerateFeedback;
+	EnvironmentDefines->NumVirtualTextureFeedbackRequests = NumVirtualTextureFeedbackRequests;
+	EnvironmentDefines->bMaterialVirtualTextureFeedback = NumVirtualTextureFeedbackRequests > 0;
 
 	// Setup defines to map each VT stack to either 1 or 2 page table textures, depending on how many layers it uses
 	EnvironmentDefines->VirtualPageTypes.SetNumUninitialized(VTStacks.Num());
