Expose GetEditorMaterial to editor only blueprints to allow referencing of non-Nanite materials (i.e. ignore Nanite overrides and return base material interfaces)

diff --git a/Engine/Source/Runtime/Engine/Classes/Components/PrimitiveComponent.h b/Engine/Source/Runtime/Engine/Classes/Components/PrimitiveComponent.h
--- a/Engine/Source/Runtime/Engine/Classes/Components/PrimitiveComponent.h
+++ b/Engine/Source/Runtime/Engine/Classes/Components/PrimitiveComponent.h
@@ -1423,11 +1423,14 @@ public:
 	UFUNCTION(BlueprintPure, Category="Rendering|Material")
 	ENGINE_API virtual class UMaterialInterface* GetMaterial(int32 ElementIndex) const;
 
-	/** Returns the material to show in the editor details panel as being used. */
-	virtual class UMaterialInterface* GetEditorMaterial(int32 ElementIndex) const 
-	{ 
+#if WITH_EDITOR
+	/** Returns the material to show in the editor details panel as being used. Skips Nanite Override materials. */
+	UFUNCTION(BlueprintPure, Category = "Rendering|Material")
+	ENGINE_API virtual class UMaterialInterface* GetEditorMaterial(int32 ElementIndex) const
+	{
 		return GetMaterial(ElementIndex);
 	}
+#endif
 
 	UFUNCTION(BlueprintCallable, Category = "Rendering|Material")
 	ENGINE_API virtual int32 GetMaterialIndex(FName MaterialSlotName) const;
diff --git a/Engine/Source/Runtime/Engine/Classes/Components/StaticMeshComponent.h b/Engine/Source/Runtime/Engine/Classes/Components/StaticMeshComponent.h
--- a/Engine/Source/Runtime/Engine/Classes/Components/StaticMeshComponent.h
+++ b/Engine/Source/Runtime/Engine/Classes/Components/StaticMeshComponent.h
@@ -605,7 +605,9 @@ public:
 	ENGINE_API virtual void GetLightAndShadowMapMemoryUsage( int32& LightMapMemoryUsage, int32& ShadowMapMemoryUsage ) const override;
 	ENGINE_API virtual void GetUsedMaterials(TArray<UMaterialInterface*>& OutMaterials, bool bGetDebugMaterials = false) const final;
 	ENGINE_API virtual UMaterialInterface* GetMaterial(int32 MaterialIndex) const final;
+#if WITH_EDITOR
 	ENGINE_API virtual UMaterialInterface* GetEditorMaterial(int32 MaterialIndex) const override;
+#endif
 	ENGINE_API virtual int32 GetMaterialIndex(FName MaterialSlotName) const override;
 	ENGINE_API virtual UMaterialInterface* GetMaterialFromCollisionFaceIndex(int32 FaceIndex, int32& SectionIndex) const override;
 	ENGINE_API virtual TArray<FName> GetMaterialSlotNames() const override;
diff --git a/Engine/Source/Runtime/Engine/Private/Components/StaticMeshComponent.cpp b/Engine/Source/Runtime/Engine/Private/Components/StaticMeshComponent.cpp
--- a/Engine/Source/Runtime/Engine/Private/Components/StaticMeshComponent.cpp
+++ b/Engine/Source/Runtime/Engine/Private/Components/StaticMeshComponent.cpp
@@ -2911,6 +2911,7 @@ UMaterialInterface* UStaticMeshComponent::GetNaniteAuditMaterial(int32 MaterialI
 	return GetMaterial(MaterialIndex, true);
 }
 
+#if WITH_EDITOR
 UMaterialInterface* UStaticMeshComponent::GetEditorMaterial(int32 MaterialIndex) const
 {
 	// Same logic as GetMaterial() but without the nanite override.
@@ -2925,6 +2926,7 @@ UMaterialInterface* UStaticMeshComponent::GetEditorMaterial(int32 MaterialIndex)
 	}
 	return nullptr;
 }
+#endif
 
 void UStaticMeshComponent::GetUsedMaterials(TArray<UMaterialInterface*>& OutMaterials, bool bGetDebugMaterials) const
 {
