Allowed post process Animation Blueprints to be applied based on a LOD threshold per-component

In addition to per-skeletal mesh, as was previously the case

PR #11061: PostProcessAnimBPLODThreshold implementation for Skeletal Mesh Component
[RN] Allowed post process Animation Blueprints to be applied based on a LOD threshold per-component

diff --git a/Engine/Source/Runtime/Engine/Classes/Components/SkeletalMeshComponent.h b/Engine/Source/Runtime/Engine/Classes/Components/SkeletalMeshComponent.h
--- a/Engine/Source/Runtime/Engine/Classes/Components/SkeletalMeshComponent.h
+++ b/Engine/Source/Runtime/Engine/Classes/Components/SkeletalMeshComponent.h
@@ -897,6 +897,16 @@ public:
 	UPROPERTY(EditAnywhere, BlueprintReadWrite, Interp, Category = Clothing, meta = (UIMin = 0.0, UIMax = 10.0, ClampMin = 0.0, ClampMax = 10000.0))
 	float ClothGeometryScale = 1.f;
 
+private:
+	/*
+	 * Max LOD level that post-process AnimBPs are evaluated. Overrides the setting of the same name on the skeletal mesh.
+	 * For example if you have the threshold set to 2, it will evaluate until including LOD 2 (based on 0 index). In case the LOD level gets set to 3, it will stop evaluating the post-process AnimBP.
+	 * Setting it to -1 will always evaluate it and disable LODing overrides for this component.
+	 */
+	UPROPERTY(EditAnywhere, AdvancedDisplay, BlueprintReadWrite, Category=Animation, meta= (DisplayName = "Post-Process AnimBP LOD Threshold", AllowPrivateAccess = "true"))
+	int32 PostProcessAnimBPLODThreshold = INDEX_NONE;
+
+public:
 	/** Notification when constraint is broken. */
 	UPROPERTY(BlueprintAssignable)
 	FConstraintBrokenSignature OnConstraintBroken;
@@ -1985,6 +1995,7 @@ public:
 	//~ End USkinnedMeshComponent Interface
 
 	// Conditions used to gate when post process events happen
+	ENGINE_API bool ShouldEvaluatePostProcessAnimBP() const;
 	ENGINE_API bool ShouldUpdatePostProcessInstance() const;
 	ENGINE_API bool ShouldPostUpdatePostProcessInstance() const;
 	ENGINE_API bool ShouldEvaluatePostProcessInstance() const;
diff --git a/Engine/Source/Runtime/Engine/Private/Components/SkeletalMeshComponent.cpp b/Engine/Source/Runtime/Engine/Private/Components/SkeletalMeshComponent.cpp
--- a/Engine/Source/Runtime/Engine/Private/Components/SkeletalMeshComponent.cpp
+++ b/Engine/Source/Runtime/Engine/Private/Components/SkeletalMeshComponent.cpp
@@ -4463,27 +4463,40 @@ void USkeletalMeshComponent::FinalizeBoneTransform()
 	TRACE_SKELETAL_MESH_COMPONENT(this);
 }
 
-bool USkeletalMeshComponent::ShouldUpdatePostProcessInstance() const
+bool USkeletalMeshComponent::ShouldEvaluatePostProcessAnimBP() const
 {
+	const int32 LODLevel = GetPredictedLODLevel();
+	if (!(PostProcessAnimBPLODThreshold == INDEX_NONE || LODLevel <= PostProcessAnimBPLODThreshold))
+	{
+		return false;
+	}
+
 	if (USkeletalMesh* SkelMesh = GetSkeletalMeshAsset())
 	{
-		if (!SkelMesh->ShouldEvaluatePostProcessAnimBP(GetPredictedLODLevel()))
+		if (!SkelMesh->ShouldEvaluatePostProcessAnimBP(LODLevel))
 		{
 			return false;
 		}
 	}
 
+	return true;
+}
+
+bool USkeletalMeshComponent::ShouldUpdatePostProcessInstance() const
+{
+	if (!ShouldEvaluatePostProcessAnimBP())
+	{
+		return false;
+	}
+
 	return PostProcessAnimInstance && !bDisablePostProcessBlueprint;
 }
 
 bool USkeletalMeshComponent::ShouldPostUpdatePostProcessInstance() const
 {
-	if (USkeletalMesh* SkelMesh = GetSkeletalMeshAsset())
+	if (!ShouldEvaluatePostProcessAnimBP())
 	{
-		if (!SkelMesh->ShouldEvaluatePostProcessAnimBP(GetPredictedLODLevel()))
-		{
-			return false;
-		}
+		return false;
 	}
 
 	return PostProcessAnimInstance && PostProcessAnimInstance->NeedsUpdate() && !bDisablePostProcessBlueprint;
@@ -4518,12 +4531,9 @@ bool USkeletalMeshComponent::ShouldEvaluatePostProcessInstance() const
 {
 	if (!OverridePostProcessAnimBP)
 	{
-		if (USkeletalMesh* SkelMesh = GetSkeletalMeshAsset())
+		if (!ShouldEvaluatePostProcessAnimBP())
 		{
-			if (!SkelMesh->ShouldEvaluatePostProcessAnimBP(GetPredictedLODLevel()))
-			{
-				return false;
-			}
+			return false;
 		}
 	}
 
