Fixed static analysis error.

diff --git a/Engine/Source/Runtime/Engine/Private/SkeletalRenderStatic.cpp b/Engine/Source/Runtime/Engine/Private/SkeletalRenderStatic.cpp
--- a/Engine/Source/Runtime/Engine/Private/SkeletalRenderStatic.cpp
+++ b/Engine/Source/Runtime/Engine/Private/SkeletalRenderStatic.cpp
@@ -48,7 +48,7 @@ void FSkeletalMeshObjectStatic::InitResources(const FSkinnedMeshSceneProxyDesc&
 		// Skip LODs that have their render data stripped
 		if (LODData.GetNumVertices() > 0)
 		{
-			const FSkelMeshComponentLODInfo* CompLODInfo = InMeshDesc.LODInfo.IsValidIndex(LODIndex) ? CompLODInfo = &InMeshDesc.LODInfo[LODIndex] : nullptr;
+			const FSkelMeshComponentLODInfo* CompLODInfo = InMeshDesc.LODInfo.IsValidIndex(LODIndex) ? &InMeshDesc.LODInfo[LODIndex] : nullptr;
 			SkelLOD.InitResources(CompLODInfo);
 
 #if RHI_RAYTRACING
