Updated FColorMaterialInput to use FLinearColor instead of FColor.
This fixes the following issues:
- In-line editing of the material for emissive on a material node clamps the value at 1
- Inconsistent results when using inline color picker in root material node

diff --git a/Engine/Source/Editor/MaterialEditor/Private/MaterialEditor.cpp b/Engine/Source/Editor/MaterialEditor/Private/MaterialEditor.cpp
--- a/Engine/Source/Editor/MaterialEditor/Private/MaterialEditor.cpp
+++ b/Engine/Source/Editor/MaterialEditor/Private/MaterialEditor.cpp
@@ -5235,8 +5235,8 @@ void FMaterialEditor::OnResetToDefault(const FToolMenuContext& InMenuContext) co
 		UMaterialEditorOnlyData* EditorOnlyData = Material->GetEditorOnlyData();
 		switch (PropertyId)
 		{
-			case MP_BaseColor:
-				EditorOnlyData->BaseColor.Constant = FColor(128, 128, 128);
+		case MP_BaseColor:
+				EditorOnlyData->BaseColor.Constant = FLinearColor::Gray;
 				TargetPin->GetSchema()->TrySetDefaultValue(*TargetPin, EditorOnlyData->BaseColor.GetDefaultValue());
 				break;
 
@@ -5306,12 +5306,12 @@ void FMaterialEditor::OnResetToDefault(const FToolMenuContext& InMenuContext) co
 				break;
 			
 			case MP_EmissiveColor:
-				EditorOnlyData->EmissiveColor.Constant = FLinearColor(FMaterialAttributeDefinitionMap::GetDefaultValue(MP_EmissiveColor)).ToFColorSRGB();
+				EditorOnlyData->EmissiveColor.Constant = FLinearColor(FMaterialAttributeDefinitionMap::GetDefaultValue(MP_EmissiveColor));
 				TargetPin->GetSchema()->TrySetDefaultValue(*TargetPin, EditorOnlyData->EmissiveColor.GetDefaultValue());
 				break;
 			
 			case MP_SubsurfaceColor:
-				EditorOnlyData->SubsurfaceColor.Constant = FLinearColor(FMaterialAttributeDefinitionMap::GetDefaultValue(MP_SubsurfaceColor)).ToFColorSRGB(); 
+				EditorOnlyData->SubsurfaceColor.Constant = FLinearColor(FMaterialAttributeDefinitionMap::GetDefaultValue(MP_SubsurfaceColor)); 
 				TargetPin->GetSchema()->TrySetDefaultValue(*TargetPin, EditorOnlyData->SubsurfaceColor.GetDefaultValue());
 				break;
 			
diff --git a/Engine/Source/Runtime/Core/Private/UObject/FortniteMainBranchObjectVersion.cpp b/Engine/Source/Runtime/Core/Private/UObject/FortniteMainBranchObjectVersion.cpp
--- a/Engine/Source/Runtime/Core/Private/UObject/FortniteMainBranchObjectVersion.cpp
+++ b/Engine/Source/Runtime/Core/Private/UObject/FortniteMainBranchObjectVersion.cpp
@@ -14,6 +14,6 @@ TMap<FGuid, FGuid> FFortniteMainBranchObjectVersion::GetSystemGuids()
 	SystemGuids.Add(DevGuids.Niagara_LatestScriptCompileVersion, FGuid("F8D995224324470E97ED718D33939A0E"));
 	SystemGuids.Add(DevGuids.SkeletalMeshDerivedDataVersion, FGuid("B9AEAA3EE7AC4FCFB3860E8F3E06DD75"));
 	SystemGuids.Add(DevGuids.STATICMESH_DERIVEDDATA_VER, FGuid("3231A08A87A248A09D959335B0F4060A"));
-	SystemGuids.Add(DevGuids.MaterialTranslationDDCVersion, FGuid("C260AF26228241889456439BE5F5AF51"));
+	SystemGuids.Add(DevGuids.MaterialTranslationDDCVersion, FGuid("ED57D4E5371B47AB9D15BCB57144F9D2"));
 	return SystemGuids;
 }
diff --git a/Engine/Source/Runtime/Core/Public/UObject/FortniteMainBranchObjectVersion.h b/Engine/Source/Runtime/Core/Public/UObject/FortniteMainBranchObjectVersion.h
--- a/Engine/Source/Runtime/Core/Public/UObject/FortniteMainBranchObjectVersion.h
+++ b/Engine/Source/Runtime/Core/Public/UObject/FortniteMainBranchObjectVersion.h
@@ -527,6 +527,9 @@ struct FFortniteMainBranchObjectVersion
 		// Added support for to record the new solver iteration settings in CVD
 		SolverIterationsDataSupportInChaosVisualDebugger,
 
+		// Updated FColorMaterialInput to use FLinearColor instead of FColor
+		MaterialInputUsesLinearColor,
+		
 		// -----<new versions can be added above this line>-------------------------------------------------
 		VersionPlusOne,
 		LatestVersion = VersionPlusOne - 1
diff --git a/Engine/Source/Runtime/Engine/Private/Materials/Material.cpp b/Engine/Source/Runtime/Engine/Private/Materials/Material.cpp
--- a/Engine/Source/Runtime/Engine/Private/Materials/Material.cpp
+++ b/Engine/Source/Runtime/Engine/Private/Materials/Material.cpp
@@ -3381,7 +3381,7 @@ bool UMaterial::ConvertMaterialToSubstrateMaterial()
 	{
 		FVector4f DefaultPinValue = FMaterialAttributeDefinitionMap::GetDefaultValue(Property);
 		FLinearColor DefaultPinLinearColor = FLinearColor(DefaultPinValue.X, DefaultPinValue.Y, DefaultPinValue.Z);
-		FLinearColor CurrentPinLinearColor = OldNodeInput.Constant.ReinterpretAsLinear();
+		FLinearColor CurrentPinLinearColor = OldNodeInput.Constant;
 
 		if (OldNodeInput.IsConnected())
 		{
@@ -7659,19 +7659,19 @@ bool UMaterial::IsRefractionPinPluggedIn(const UMaterialEditorOnlyData* EditorOn
 
 UMaterialEditorOnlyData::UMaterialEditorOnlyData()
 {
-	BaseColor.Constant = FColor(128, 128, 128);
+	BaseColor.Constant = FLinearColor::Gray;
 	Metallic.Constant = FMaterialAttributeDefinitionMap::GetDefaultValue(MP_Metallic).X;
 	Specular.Constant = FMaterialAttributeDefinitionMap::GetDefaultValue(MP_Specular).X;
 	Roughness.Constant = FMaterialAttributeDefinitionMap::GetDefaultValue(MP_Roughness).X;
 	Anisotropy.Constant = FMaterialAttributeDefinitionMap::GetDefaultValue(MP_Anisotropy).X;
 	Normal.Constant = FMaterialAttributeDefinitionMap::GetDefaultValue(MP_Normal);
 	Tangent.Constant = FMaterialAttributeDefinitionMap::GetDefaultValue(MP_Tangent);
-	EmissiveColor.Constant = FLinearColor(FMaterialAttributeDefinitionMap::GetDefaultValue(MP_EmissiveColor)).ToFColorSRGB();
+	EmissiveColor.Constant = FLinearColor(FMaterialAttributeDefinitionMap::GetDefaultValue(MP_EmissiveColor));
 	Opacity.Constant = FMaterialAttributeDefinitionMap::GetDefaultValue(MP_Opacity).X;
 	OpacityMask.Constant = FMaterialAttributeDefinitionMap::GetDefaultValue(MP_OpacityMask).X;
 	WorldPositionOffset.Constant = FMaterialAttributeDefinitionMap::GetDefaultValue(MP_WorldPositionOffset);
 	Displacement.Constant = FMaterialAttributeDefinitionMap::GetDefaultValue(MP_Displacement).X;
-	SubsurfaceColor.Constant = FLinearColor(FMaterialAttributeDefinitionMap::GetDefaultValue(MP_SubsurfaceColor)).ToFColorSRGB();
+	SubsurfaceColor.Constant = FLinearColor(FMaterialAttributeDefinitionMap::GetDefaultValue(MP_SubsurfaceColor));
 	ClearCoat.Constant = FMaterialAttributeDefinitionMap::GetDefaultValue(MP_CustomData0).X;
 	ClearCoatRoughness.Constant = FMaterialAttributeDefinitionMap::GetDefaultValue(MP_CustomData1).X;
 	AmbientOcclusion.Constant = FMaterialAttributeDefinitionMap::GetDefaultValue(MP_AmbientOcclusion).X;
diff --git a/Engine/Source/Runtime/Engine/Private/Materials/MaterialShared.cpp b/Engine/Source/Runtime/Engine/Private/Materials/MaterialShared.cpp
--- a/Engine/Source/Runtime/Engine/Private/Materials/MaterialShared.cpp
+++ b/Engine/Source/Runtime/Engine/Private/Materials/MaterialShared.cpp
@@ -598,7 +598,32 @@ bool FExpressionInput::Serialize(FArchive& Ar)
 
 bool FColorMaterialInput::Serialize(FArchive& Ar)
 {
-	return SerializeMaterialInput<FColor>(Ar, *this);
+	if (Ar.IsLoading() && Ar.CustomVer(FFortniteMainBranchObjectVersion::GUID) < FFortniteMainBranchObjectVersion::MaterialInputUsesLinearColor)
+	{
+		FMaterialInput<FColor> OldValue;
+		if (SerializeMaterialInput<FColor>(Ar, OldValue))
+		{
+			this->Expression = OldValue.Expression;
+			this->OutputIndex = OldValue.OutputIndex;
+			this->InputName = OldValue.InputName;
+			this->Mask = OldValue.Mask;
+			this->MaskR = OldValue.MaskR;
+			this->MaskG = OldValue.MaskG;
+			this->MaskB = OldValue.MaskB;
+			this->MaskA = OldValue.MaskA;
+			this->UseConstant = OldValue.UseConstant;
+			this->Constant = OldValue.Constant.ReinterpretAsLinear();
+			return true;
+		}
+		else
+		{
+			return false;
+		}
+	}
+	else
+	{
+		return SerializeMaterialInput<FLinearColor>(Ar, *this);	
+	}
 }
 
 bool FScalarMaterialInput::Serialize(FArchive& Ar)
@@ -635,9 +660,7 @@ bool FMaterialAttributesInput::Serialize(FArchive& Ar)
 void FColorMaterialInput::DefaultValueChanged(const FString& DefaultValue)
 {
 #if WITH_EDITOR
-	FLinearColor Value;
-	Value.InitFromString(DefaultValue);
-	Constant = Value.QuantizeRound();
+	Constant.InitFromString(DefaultValue);
 	UseConstant = true;
 #endif
 }
@@ -646,8 +669,7 @@ FString FColorMaterialInput::GetDefaultValue() const
 {
 	FString DefaultValue;
 #if WITH_EDITOR
-	FLinearColor Color = Constant.ReinterpretAsLinear();
-	DefaultValue = Color.ToString();
+	DefaultValue = Constant.ToString();
 #endif
 	return DefaultValue;
 }
@@ -716,8 +738,7 @@ int32 FColorMaterialInput::CompileWithDefault(class FMaterialCompiler* Compiler,
 {
 	if (UseConstant)
 	{
-		FLinearColor LinearColor(Constant);
-		return Compiler->Constant3(LinearColor.R, LinearColor.G, LinearColor.B);
+		return Compiler->Constant3(Constant.R, Constant.G, Constant.B);
 	}
 	else if (Expression)
 	{
diff --git a/Engine/Source/Runtime/Engine/Public/MaterialExpressionIO.h b/Engine/Source/Runtime/Engine/Public/MaterialExpressionIO.h
--- a/Engine/Source/Runtime/Engine/Public/MaterialExpressionIO.h
+++ b/Engine/Source/Runtime/Engine/Public/MaterialExpressionIO.h
@@ -178,7 +178,7 @@ template<class InputType> struct FMaterialInput : FExpressionInput
 	FMaterialInput()
 	{
 		UseConstant = 0;
-		Constant = InputType(0);
+		Constant = {};
 	}
 
 #if WITH_EDITOR
@@ -189,7 +189,7 @@ template<class InputType> struct FMaterialInput : FExpressionInput
 	InputType	Constant;
 };
 
-struct FColorMaterialInput : FMaterialInput<FColor>
+struct FColorMaterialInput : FMaterialInput<FLinearColor>
 {
 #if WITH_EDITOR
 	ENGINE_API int32 CompileWithDefault(class FMaterialCompiler* Compiler, EMaterialProperty Property);
diff --git a/Engine/Source/Runtime/Engine/Public/Materials/Material.h b/Engine/Source/Runtime/Engine/Public/Materials/Material.h
--- a/Engine/Source/Runtime/Engine/Public/Materials/Material.h
+++ b/Engine/Source/Runtime/Engine/Public/Materials/Material.h
@@ -194,7 +194,7 @@ struct FColorMaterialInput : public FMaterialInput
 	uint32 UseConstant : 1;
 
 	UPROPERTY()
-	FColor Constant;
+	FLinearColor Constant;
 
 	FColorMaterialInput() :
 		UseConstant(0) {}
