Fixed unity build break.

diff --git a/Engine/Source/Runtime/Engine/Private/Components/SkinnedMeshComponent.cpp b/Engine/Source/Runtime/Engine/Private/Components/SkinnedMeshComponent.cpp
--- a/Engine/Source/Runtime/Engine/Private/Components/SkinnedMeshComponent.cpp
+++ b/Engine/Source/Runtime/Engine/Private/Components/SkinnedMeshComponent.cpp
@@ -2159,8 +2159,7 @@ bool USkinnedMeshComponent::GetTwistAndSwingAngleOfDeltaRotationFromRefPose(FNam
 
 bool USkinnedMeshComponent::IsSkinCacheAllowed(int32 LodIdx) const
 {
-	static const IConsoleVariable* CVarDefaultGPUSkinCacheBehavior = IConsoleManager::Get().FindConsoleVariable(TEXT("r.SkinCache.DefaultBehavior"));
-	const bool bGlobalDefault = CVarDefaultGPUSkinCacheBehavior && ESkinCacheDefaultBehavior(CVarDefaultGPUSkinCacheBehavior->GetInt()) == ESkinCacheDefaultBehavior::Inclusive;
+	const bool bGlobalDefault = GetSkinCacheDefaultBehavior() == ESkinCacheDefaultBehavior::Inclusive;
 
 	if (GetMeshDeformerInstanceForLOD(LodIdx) != nullptr)
 	{
diff --git a/Engine/Source/Runtime/Engine/Private/GPUSkinCache.cpp b/Engine/Source/Runtime/Engine/Private/GPUSkinCache.cpp
--- a/Engine/Source/Runtime/Engine/Private/GPUSkinCache.cpp
+++ b/Engine/Source/Runtime/Engine/Private/GPUSkinCache.cpp
@@ -247,6 +247,13 @@ bool ShouldWeCompileGPUSkinVFShaders(EShaderPlatform Platform, ERHIFeatureLevel:
 	return (PerPlatformCVar.Get(Platform) == false);
 }
 
+ESkinCacheDefaultBehavior GetSkinCacheDefaultBehavior()
+{
+	return ESkinCacheDefaultBehavior(CVarDefaultGPUSkinCacheBehavior->GetInt()) == ESkinCacheDefaultBehavior::Inclusive
+		? ESkinCacheDefaultBehavior::Inclusive
+		: ESkinCacheDefaultBehavior::Exclusive;
+}
+
 ENGINE_API bool GPUSkinCacheNeedsDuplicatedVertices()
 {
 #if WITH_EDITOR // Duplicated vertices are used in the editor when merging meshes
diff --git a/Engine/Source/Runtime/Engine/Private/SkinnedMeshSceneProxyDesc.cpp b/Engine/Source/Runtime/Engine/Private/SkinnedMeshSceneProxyDesc.cpp
--- a/Engine/Source/Runtime/Engine/Private/SkinnedMeshSceneProxyDesc.cpp
+++ b/Engine/Source/Runtime/Engine/Private/SkinnedMeshSceneProxyDesc.cpp
@@ -10,6 +10,7 @@
 #include "Engine/SkinnedAsset.h"
 #include "Components/SkinnedMeshComponent.h"
 #include "SceneInterface.h"
+#include "GPUSkinCache.h"
 
 extern ENGINE_API int32 GEnableGPUSkinCache;
 
@@ -228,8 +229,7 @@ FSkeletalMeshObject* FSkinnedMeshSceneProxyDesc::CreateMeshObject(const FSkinned
 
 bool FSkinnedMeshSceneProxyDesc::IsSkinCacheAllowed(int32 LodIdx) const
 {
-	static const IConsoleVariable* CVarDefaultGPUSkinCacheBehavior = IConsoleManager::Get().FindConsoleVariable(TEXT("r.SkinCache.DefaultBehavior"));
-	const bool bGlobalDefault = CVarDefaultGPUSkinCacheBehavior && ESkinCacheDefaultBehavior(CVarDefaultGPUSkinCacheBehavior->GetInt()) == ESkinCacheDefaultBehavior::Inclusive;
+	const bool bGlobalDefault = GetSkinCacheDefaultBehavior() == ESkinCacheDefaultBehavior::Inclusive;
 
 	if (GetMeshDeformerInstance() != nullptr)
 	{
@@ -272,8 +272,7 @@ UMeshDeformerInstance* FSkinnedMeshSceneProxyDesc::GetMeshDeformerInstance() con
 
 bool FSkinnedMeshSceneProxyDynamicData::IsSkinCacheAllowed(int32 LodIdx, const USkinnedAsset* InSkinnedAsset) const
 {
-	static const IConsoleVariable* CVarDefaultGPUSkinCacheBehavior = IConsoleManager::Get().FindConsoleVariable(TEXT("r.SkinCache.DefaultBehavior"));
-	const bool bGlobalDefault = CVarDefaultGPUSkinCacheBehavior && ESkinCacheDefaultBehavior(CVarDefaultGPUSkinCacheBehavior->GetInt()) == ESkinCacheDefaultBehavior::Inclusive;
+	const bool bGlobalDefault = GetSkinCacheDefaultBehavior() == ESkinCacheDefaultBehavior::Inclusive;
 	
 	if (bHasMeshDeformerInstance)
 	{
diff --git a/Engine/Source/Runtime/Engine/Public/GPUSkinCache.h b/Engine/Source/Runtime/Engine/Public/GPUSkinCache.h
--- a/Engine/Source/Runtime/Engine/Public/GPUSkinCache.h
+++ b/Engine/Source/Runtime/Engine/Public/GPUSkinCache.h
@@ -45,6 +45,7 @@
 #include "CanvasTypes.h"
 #include "CachedGeometry.h"
 #include "DataDrivenShaderPlatformInfo.h"
+#include "Engine/SkinnedAssetCommon.h"
 
 class FGPUSkinPassthroughVertexFactory;
 class FGPUBaseSkinVertexFactory;
@@ -63,6 +64,8 @@ extern bool ShouldWeCompileGPUSkinVFShaders(EShaderPlatform Platform, ERHIFeatur
 
 extern ENGINE_API bool GPUSkinCacheNeedsDuplicatedVertices();
 
+extern ENGINE_API ESkinCacheDefaultBehavior GetSkinCacheDefaultBehavior();
+
 // Is it actually enabled?
 extern ENGINE_API int32 GEnableGPUSkinCache;
 extern int32 GSkinCacheRecomputeTangents;
