Substrate - fix for complex material not recognised.

#rnx
#tests none

diff --git a/Engine/Shaders/Private/Substrate/SubstrateExport.ush b/Engine/Shaders/Private/Substrate/SubstrateExport.ush
--- a/Engine/Shaders/Private/Substrate/SubstrateExport.ush
+++ b/Engine/Shaders/Private/Substrate/SubstrateExport.ush
@@ -201,14 +201,14 @@ void PackSubstrateOut(
 						#else
 						const bool bForceComplexSpecialPath = false;
 						#endif
-						const bool bForceComplexMaterial = BSDF_GETHASANISOTROPY(BSDF) || bForceComplexSpecialPath;
+						const bool bForceComplexMaterial = BSDF_GETHASANISOTROPY(BSDF) || bForceComplexSpecialPath || !LuminanceWeightEqualOne;
 
 						// Update simple material compatibility
-						bIsSimpleMaterial = bIsSimpleMaterial && IsSubstrateSlabFastPathCompatible(BSDF) && LuminanceWeightEqualOne && !bForceComplexMaterial;
+						bIsSimpleMaterial = bIsSimpleMaterial && IsSubstrateSlabFastPathCompatible(BSDF) && !bForceComplexMaterial;
 
 						// Update single material compatibility. For now, single materials don't support anisotropy, 
 						// they use the complex pass, as the toplayer data does not contain the full frame basis, only the top normal.
-						bIsSingleMaterial = bIsSingleMaterial && LuminanceWeightEqualOne && !bForceComplexMaterial;
+						bIsSingleMaterial = bIsSingleMaterial && !bForceComplexMaterial;
 
 						TopLayerTotalWeight += BSDF.TopLayerDataWeight;
 						float TopLayerRoughnessContribution = BSDF.TopLayerDataWeight * SLAB_ROUGHNESS(BSDF);
diff --git a/Engine/Source/Runtime/Engine/Private/Materials/HLSLMaterialTranslator.cpp b/Engine/Source/Runtime/Engine/Private/Materials/HLSLMaterialTranslator.cpp
--- a/Engine/Source/Runtime/Engine/Private/Materials/HLSLMaterialTranslator.cpp
+++ b/Engine/Source/Runtime/Engine/Private/Materials/HLSLMaterialTranslator.cpp
@@ -13142,8 +13142,17 @@ bool FHLSLMaterialTranslator::FSubstrateCompilationContext::SubstrateGenerateDer
 			SubstrateMaterialComplexity.bIsComplexSpecial = false;
 			bool bIsFastWaterPath = false;
 			bool bCustomEncoding = false;
+			bool bMayHaveCoverageLessThan1 = false;
 			for (const auto& It : SubstrateMaterialExpressionRegisteredOperators)
 			{
+				if (It.IsDiscarded() && It.OperatorType == SUBSTRATE_OPERATOR_WEIGHT)
+				{
+					// If a BSDF modified by a weight operator, its weight will be < 1.0f, and it won't be a "single" material anymore.
+					// This is also the case if the operator is "discarded" due to parameter blending, because the resulting BSDF might not have a coverage of 1 in the end.
+					// For instance with a BSDF => Weight => Horizonal node with parameter blending is encountered and weight is less than 1.
+					bMayHaveCoverageLessThan1 |= true;
+				}
+
 				if (It.IsDiscarded())
 				{
 					continue; // ignore discarded operations in sub tree using parameter blending
@@ -13214,6 +13223,13 @@ bool FHLSLMaterialTranslator::FSubstrateCompilationContext::SubstrateGenerateDer
 			}
 			SubstrateMaterialComplexity.bIsSingle = SubstrateMaterialComplexity.bIsSingle && !SubstrateMaterialComplexity.bIsSimple;
 
+			if (bMayHaveCoverageLessThan1)
+			{
+				// Material with coverage<1 can only be complex. It will be made simple or single at export time if possible.
+				SubstrateMaterialComplexity.bIsSimple = false;
+				SubstrateMaterialComplexity.bIsSingle = false;
+			}
+
 			// 2. Header
 
 			if (!SubstrateMaterialComplexity.bIsSimple && !SubstrateMaterialComplexity.bIsSingle && !bCustomEncoding && !bIsFastWaterPath) // header written later, 
