Fix relative render cost using slightly too much rendering budget.
#rnx

diff --git a/Engine/Source/Runtime/Engine/Private/AudioDevice.cpp b/Engine/Source/Runtime/Engine/Private/AudioDevice.cpp
--- a/Engine/Source/Runtime/Engine/Private/AudioDevice.cpp
+++ b/Engine/Source/Runtime/Engine/Private/AudioDevice.cpp
@@ -4182,23 +4182,31 @@ int32 FAudioDevice::GetSortedActiveWaveInstances(TArray<FWaveInstance*>& WaveIns
 		// Sort by priority (lowest priority first).
 		WaveInstances.Sort(FCompareFWaveInstanceByPlayPriority());
 
-		// Get the first index that will result in a active source voice
-		int32 CurrentMaxChannels = GetMaxChannels();
 
 		if (EnableRelativeRenderCostVoiceLimitCVar)
 		{
 			// Find the first active index based on the relative render cost estimate of the wave instances
-			FirstActiveIndex = WaveInstances.Num() - 1;
-			float NumActiveChannels = 0.0f;
-			while (NumActiveChannels < (float)CurrentMaxChannels && FirstActiveIndex > 0)
+			float RemainingRenderCost = (float)GetMaxChannels();
+
+			FirstActiveIndex = WaveInstances.Num();
+			for (const FWaveInstance* WaveInstance : ReverseIterate(WaveInstances))
 			{
-				FWaveInstance* WaveInstance = WaveInstances[FirstActiveIndex];
-				NumActiveChannels += WaveInstance->GetRelativeRenderCost();
-				FirstActiveIndex--;
+				RemainingRenderCost -= WaveInstance->GetRelativeRenderCost();
+				// Do not increment index if we've gone past our max render cost. 
+				if (RemainingRenderCost >= 0.f)
+				{
+					FirstActiveIndex--;
+				}
+				else
+				{
+					break;
+				}
 			}
 		}
 		else
 		{
+			// Get the first index that will result in a active source voice
+			int32 CurrentMaxChannels = GetMaxChannels();
 			FirstActiveIndex = FMath::Max(WaveInstances.Num() - CurrentMaxChannels, 0);
 		}
 	}
