BlendSpace: Expose option to sync anims. Apply time scaling using the "Axis Scaling" feature whether or not anims are being synced

diff --git a/Engine/Source/Runtime/Engine/Classes/Animation/BlendSpace.h b/Engine/Source/Runtime/Engine/Classes/Animation/BlendSpace.h
--- a/Engine/Source/Runtime/Engine/Classes/Animation/BlendSpace.h
+++ b/Engine/Source/Runtime/Engine/Classes/Animation/BlendSpace.h
@@ -808,6 +808,10 @@ public:
 	UPROPERTY(EditAnywhere, Category=Animation)
 	bool bLoop = true;
 
+	/** Whether to allow marker based sync between the samples (it won't force sync if the markers don't exist) */
+	UPROPERTY(EditAnywhere, Category = Animation)
+	bool bAllowMarkerBasedSync = true;
+
 #if WITH_EDITORONLY_DATA
 	/** Preview Base pose for additive BlendSpace **/
 	UPROPERTY(EditAnywhere, Category = AdditiveSettings)
@@ -866,7 +870,7 @@ protected:
 	/** Track index to get marker data from. Samples are tested for the suitability of marker based sync
 	    during load and if we can use marker based sync we cache an index to a representative sample here */
 	UPROPERTY()
-	int32 SampleIndexWithMarkers;
+	int32 SampleIndexWithMarkers = INDEX_NONE;
 
 	/** Sample animation data */
 	UPROPERTY(EditAnywhere, Category=BlendSamples)
diff --git a/Engine/Source/Runtime/Engine/Private/Animation/BlendSpace.cpp b/Engine/Source/Runtime/Engine/Private/Animation/BlendSpace.cpp
--- a/Engine/Source/Runtime/Engine/Private/Animation/BlendSpace.cpp
+++ b/Engine/Source/Runtime/Engine/Private/Animation/BlendSpace.cpp
@@ -447,354 +447,353 @@ void UBlendSpace::TickAssetPlayer(FAnimTickRecord& Instance, struct FAnimNotifyQ
 	// this happens even if MoveDelta == 0.f. This still should happen if it is being interpolated
 	// since we allow setting position of blendspace, we can't ignore MoveDelta == 0.f
 	// also now we don't have to worry about not following if DeltaTime = 0.f
+
+	// first filter input using blend filter
+	const FVector BlendSpacePosition(Instance.BlendSpace.BlendSpacePositionX, Instance.BlendSpace.BlendSpacePositionY, 0.f);
+	const FVector FilteredBlendInput = FilterInput(Instance.BlendSpace.BlendFilter, BlendSpacePosition, DeltaTime);
+
+	if (UpdateBlendSamples_Internal(FilteredBlendInput, DeltaTime, OldSampleDataList, SampleDataList, Instance.BlendSpace.TriangulationIndex))
 	{
-		// first filter input using blend filter
-		const FVector BlendSpacePosition(Instance.BlendSpace.BlendSpacePositionX, Instance.BlendSpace.BlendSpacePositionY, 0.f);
-		const FVector FilteredBlendInput = FilterInput(Instance.BlendSpace.BlendFilter, BlendSpacePosition, DeltaTime);
+		float NewAnimLength = 0.f;
+		float PreInterpAnimLength = 0.f;
 
-		if (UpdateBlendSamples_Internal(FilteredBlendInput, DeltaTime, OldSampleDataList, SampleDataList, Instance.BlendSpace.TriangulationIndex))
+		if (TargetWeightInterpolationSpeedPerSec > 0.f)
 		{
-			float NewAnimLength = 0.f;
-			float PreInterpAnimLength = 0.f;
-
-			if (TargetWeightInterpolationSpeedPerSec > 0.f)
-			{
-				// recalculate AnimLength based on weight of target animations - this is used for scaling animation later (change speed)
-				PreInterpAnimLength = GetAnimationLengthFromSampleData(*Instance.BlendSpace.BlendSampleDataCache);
-				UE_LOG(LogAnimation, Verbose, TEXT("BlendSpace(%s) - FilteredBlendInput(%s) : PreAnimLength(%0.5f) "), *GetName(), *FilteredBlendInput.ToString(), PreInterpAnimLength);
-			}
+			// recalculate AnimLength based on weight of target animations - this is used for scaling animation later (change speed)
+			PreInterpAnimLength = GetAnimationLengthFromSampleData(*Instance.BlendSpace.BlendSampleDataCache);
+			UE_LOG(LogAnimation, Verbose, TEXT("BlendSpace(%s) - FilteredBlendInput(%s) : PreAnimLength(%0.5f) "), *GetName(), *FilteredBlendInput.ToString(), PreInterpAnimLength);
+		}
 
-			EBlendSpaceAxis AxisToScale = GetAxisToScale();
-			if (AxisToScale != BSA_None)
+		EBlendSpaceAxis AxisToScale = GetAxisToScale();
+		if (AxisToScale != BSA_None)
+		{
+			float FilterMultiplier = 1.f;
+			// first use multiplier using new blendinput
+			// new filtered input is going to be used for sampling animation
+			// so we'll need to change playrate if you'd like to not slide foot
+			if (!BlendSpacePosition.Equals(FilteredBlendInput))
 			{
-				float FilterMultiplier = 1.f;
-				// first use multiplier using new blendinput
-				// new filtered input is going to be used for sampling animation
-				// so we'll need to change playrate if you'd like to not slide foot
-				if (!BlendSpacePosition.Equals(FilteredBlendInput))
+				// apply speed change if you want, 
+				if (AxisToScale == BSA_X)
 				{
-					// apply speed change if you want, 
-					if (AxisToScale == BSA_X)
+					if (FilteredBlendInput.X != 0.f)
 					{
-						if (FilteredBlendInput.X != 0.f)
-						{
-							FilterMultiplier = BlendSpacePosition.X / FilteredBlendInput.X;
-						}
+						FilterMultiplier = BlendSpacePosition.X / FilteredBlendInput.X;
 					}
-					else if (AxisToScale == BSA_Y)
+				}
+				else if (AxisToScale == BSA_Y)
+				{
+					if (FilteredBlendInput.Y != 0.f)
 					{
-						if (FilteredBlendInput.Y != 0.f)
-						{
-							FilterMultiplier = BlendSpacePosition.Y / FilteredBlendInput.Y;
-						}
+						FilterMultiplier = BlendSpacePosition.Y / FilteredBlendInput.Y;
 					}
 				}
+			}
 
-				// Now find if clamped input is different. If different, then apply scale to fit in. This allows
-				// "extrapolation" of the blend space outside of the range by time scaling the animation, which is
-				// appropriate when the specified axis is speed (for example).
-				FVector ClampedInput = GetClampedBlendInput(FilteredBlendInput);
-				if (!ClampedInput.Equals(FilteredBlendInput))
+			// Now find if clamped input is different. If different, then apply scale to fit in. This allows
+			// "extrapolation" of the blend space outside of the range by time scaling the animation, which is
+			// appropriate when the specified axis is speed (for example).
+			FVector ClampedInput = GetClampedBlendInput(FilteredBlendInput);
+			if (!ClampedInput.Equals(FilteredBlendInput))
+			{
+				// apply speed change if you want, 
+				if (AxisToScale == BSA_X && !BlendParameters[0].bWrapInput)
 				{
-					// apply speed change if you want, 
-					if (AxisToScale == BSA_X && !BlendParameters[0].bWrapInput)
+					if (ClampedInput.X != 0.f)
 					{
-						if (ClampedInput.X != 0.f)
-						{
-							FilterMultiplier *= FilteredBlendInput.X / ClampedInput.X;
-						}
+						FilterMultiplier *= FilteredBlendInput.X / ClampedInput.X;
 					}
-					else if (AxisToScale == BSA_Y)
+				}
+				else if (AxisToScale == BSA_Y)
+				{
+					if (ClampedInput.Y != 0.f && !BlendParameters[1].bWrapInput)
 					{
-						if (ClampedInput.Y != 0.f && !BlendParameters[1].bWrapInput)
-						{
-							FilterMultiplier *= FilteredBlendInput.Y / ClampedInput.Y;
-						}
+						FilterMultiplier *= FilteredBlendInput.Y / ClampedInput.Y;
 					}
 				}
-
-				Instance.DeltaTimeRecord->Delta *= FilterMultiplier;
-				UE_LOG(LogAnimation, Verbose, TEXT("BlendSpace(%s) - FilteredBlendInput(%s) : FilteredBlendInput(%s), FilterMultiplier(%0.2f)"), 
-					*GetName(), *BlendSpacePosition.ToString(), *FilteredBlendInput.ToString(), FilterMultiplier);
 			}
 
-			// We can use marker-based syncing when a valid sample with sync marker data exists.
-			bool bCanDoMarkerSync = (SampleIndexWithMarkers != INDEX_NONE) && (Context.IsSingleAnimationContext() || (Instance.bCanUseMarkerSync && Context.CanUseMarkerPosition()));
+			Instance.DeltaTimeRecord->Delta *= FilterMultiplier;
+			UE_LOG(LogAnimation, Verbose, TEXT("BlendSpace(%s) - FilteredBlendInput(%s) : FilteredBlendInput(%s), FilterMultiplier(%0.2f)"), 
+				*GetName(), *BlendSpacePosition.ToString(), *FilteredBlendInput.ToString(), FilterMultiplier);
+		}
+
+		// We can use marker-based syncing when a valid sample with sync marker data exists.
+		bool bCanDoMarkerSync = (SampleIndexWithMarkers != INDEX_NONE) && (Context.IsSingleAnimationContext() || (Instance.bCanUseMarkerSync && Context.CanUseMarkerPosition()));
 			
-			if (bCanDoMarkerSync)
+		if (bCanDoMarkerSync)
+		{
+			// Copy previous frame marker data to current frame
+			for (const FBlendSampleData& PrevBlendSampleItem : OldSampleDataList)
 			{
-				// Copy previous frame marker data to current frame
-				for (const FBlendSampleData& PrevBlendSampleItem : OldSampleDataList)
+				for (FBlendSampleData& CurrentBlendSampleItem : SampleDataList)
 				{
-					for (FBlendSampleData& CurrentBlendSampleItem : SampleDataList)
+					// it only can have one animation in the sample, make sure to copy Time
+					if (PrevBlendSampleItem.Animation && PrevBlendSampleItem.Animation == CurrentBlendSampleItem.Animation)
 					{
-						// it only can have one animation in the sample, make sure to copy Time
-						if (PrevBlendSampleItem.Animation && PrevBlendSampleItem.Animation == CurrentBlendSampleItem.Animation)
-						{
-							CurrentBlendSampleItem.Time = PrevBlendSampleItem.Time;
-							CurrentBlendSampleItem.PreviousTime = PrevBlendSampleItem.PreviousTime;
-							CurrentBlendSampleItem.MarkerTickRecord = PrevBlendSampleItem.MarkerTickRecord;
-						}
+						CurrentBlendSampleItem.Time = PrevBlendSampleItem.Time;
+						CurrentBlendSampleItem.PreviousTime = PrevBlendSampleItem.PreviousTime;
+						CurrentBlendSampleItem.MarkerTickRecord = PrevBlendSampleItem.MarkerTickRecord;
 					}
 				}
 			}
+		}
 
-			NewAnimLength = GetAnimationLengthFromSampleData(SampleDataList);
+		NewAnimLength = GetAnimationLengthFromSampleData(SampleDataList);
 
-			if (PreInterpAnimLength > 0.f && NewAnimLength > 0.f)
-			{
-				Instance.DeltaTimeRecord->Delta *= PreInterpAnimLength / NewAnimLength;
-			}
+		if (PreInterpAnimLength > 0.f && NewAnimLength > 0.f)
+		{
+			Instance.DeltaTimeRecord->Delta *= PreInterpAnimLength / NewAnimLength;
+		}
 
-			float& NormalizedCurrentTime = *(Instance.TimeAccumulator);
-			if (Context.ShouldResyncToSyncGroup() && !Instance.bIsEvaluator)
-			{
-				// Synchronize the asset player time to the other sync group members when (re)joining the group
-				NormalizedCurrentTime = Context.GetAnimationPositionRatio();
-			}
+		float& NormalizedCurrentTime = *(Instance.TimeAccumulator);
+		if (Context.ShouldResyncToSyncGroup() && !Instance.bIsEvaluator)
+		{
+			// Synchronize the asset player time to the other sync group members when (re)joining the group
+			NormalizedCurrentTime = Context.GetAnimationPositionRatio();
+		}
 
-			float NormalizedPreviousTime = NormalizedCurrentTime;
+		float NormalizedPreviousTime = NormalizedCurrentTime;
 
-			// @note for sync group vs non sync group
-			// in blendspace, it will still sync even if only one node in sync group
-			// so you're never non-sync group unless you have situation where some markers are relevant to one sync group but not all the time
-			// here we save NormalizedCurrentTime as Highest weighted samples' position in sync group
-			// if you're not in sync group, NormalizedCurrentTime is based on normalized length by sample weights
-			// if you move between sync to non sync within blendspace, you're going to see pop because we'll have to jump
-			// for now, our rule is to keep normalized time as highest weighted sample position within its own length
-			// also MoveDelta doesn't work if you're in sync group. It will move according to sync group position
-			// @todo consider using MoveDelta when  this is leader, but that can be scary because it's not matching with DeltaTime any more. 
-			// if you have interpolation delay, that value can be applied, but the output might be unpredictable. 
-			// 
-			// to fix this better in the future, we should use marker sync position from last tick
-			// but that still doesn't fix if you just join sync group, you're going to see pop since your animation doesn't fix
+		// @note for sync group vs non sync group
+		// in blendspace, it will still sync even if only one node in sync group
+		// so you're never non-sync group unless you have situation where some markers are relevant to one sync group but not all the time
+		// here we save NormalizedCurrentTime as Highest weighted samples' position in sync group
+		// if you're not in sync group, NormalizedCurrentTime is based on normalized length by sample weights
+		// if you move between sync to non sync within blendspace, you're going to see pop because we'll have to jump
+		// for now, our rule is to keep normalized time as highest weighted sample position within its own length
+		// also MoveDelta doesn't work if you're in sync group. It will move according to sync group position
+		// @todo consider using MoveDelta when  this is leader, but that can be scary because it's not matching with DeltaTime any more. 
+		// if you have interpolation delay, that value can be applied, but the output might be unpredictable. 
+		// 
+		// to fix this better in the future, we should use marker sync position from last tick
+		// but that still doesn't fix if you just join sync group, you're going to see pop since your animation doesn't fix
 
-			if (Context.IsLeader())
-			{
-				// advance current time - blend spaces hold normalized time as when dealing with changing anim length it would be possible to go backwards
-				UE_LOG(LogAnimation, Verbose, TEXT("BlendSpace(%s) - FilteredBlendInput(%s) : AnimLength(%0.5f) "), *GetName(), *FilteredBlendInput.ToString(), NewAnimLength);
+		if (Context.IsLeader())
+		{
+			// advance current time - blend spaces hold normalized time as when dealing with changing anim length it would be possible to go backwards
+			UE_LOG(LogAnimation, Verbose, TEXT("BlendSpace(%s) - FilteredBlendInput(%s) : AnimLength(%0.5f) "), *GetName(), *FilteredBlendInput.ToString(), NewAnimLength);
 
-				// Set context's data before updating time position.
-				Context.SetPreviousAnimationPositionRatio(NormalizedCurrentTime);
+			// Set context's data before updating time position.
+			Context.SetPreviousAnimationPositionRatio(NormalizedCurrentTime);
 
-				// Get highest weight sample with sync markers. This will become the leader for all other samples to follow.
-				const int32 HighestMarkerSyncWeightIndex = bCanDoMarkerSync ? FBlendSpaceUtilities::GetHighestWeightMarkerSyncSample(SampleDataList, SampleData) : INDEX_NONE;
+			// Get highest weight sample with sync markers. This will become the leader for all other samples to follow.
+			const int32 HighestMarkerSyncWeightIndex = bCanDoMarkerSync ? FBlendSpaceUtilities::GetHighestWeightMarkerSyncSample(SampleDataList, SampleData) : INDEX_NONE;
 
-				// Skip syncing, fallback to normal ticking.
-				if (HighestMarkerSyncWeightIndex == INDEX_NONE)
-				{
-					bCanDoMarkerSync = false;
-				}
+			// Skip syncing, fallback to normal ticking.
+			if (HighestMarkerSyncWeightIndex == INDEX_NONE)
+			{
+				bCanDoMarkerSync = false;
+			}
+
+			// Tick as leader using marked based syncing.
+			if (bCanDoMarkerSync)
+			{
+				FBlendSampleData& LeaderSampleData = SampleDataList[HighestMarkerSyncWeightIndex];
+				const FBlendSample& LeaderSample = SampleData[LeaderSampleData.SampleDataIndex];
 
-				// Tick as leader using marked based syncing.
-				if (bCanDoMarkerSync)
+				if (LeaderSample.Animation)
 				{
-					FBlendSampleData& LeaderSampleData = SampleDataList[HighestMarkerSyncWeightIndex];
-					const FBlendSample& LeaderSample = SampleData[LeaderSampleData.SampleDataIndex];
+					bool bResetMarkerDataOnFollowers = false;
 
-					if (LeaderSample.Animation)
+					// Invalidate sample followers' tick records if instance doesn't have any valid sync marker data. 
+					if (!Instance.MarkerTickRecord->IsValid(Instance.bLooping))
 					{
-						bool bResetMarkerDataOnFollowers = false;
-
-						// Invalidate sample followers' tick records if instance doesn't have any valid sync marker data. 
-						if (!Instance.MarkerTickRecord->IsValid(Instance.bLooping))
-						{
-							LeaderSampleData.MarkerTickRecord.Reset();
-							LeaderSampleData.Time = NormalizedCurrentTime * LeaderSample.Animation->GetPlayLength();
-							bResetMarkerDataOnFollowers = true;
-						}
-						// Re-compute marker indices since the leader sample's tick record is invalid. Get previous and next markers.
-						else if (!LeaderSampleData.MarkerTickRecord.IsValid(Instance.bLooping) && Context.MarkerTickContext.GetMarkerSyncStartPosition().IsValid())
-						{
-							// TODO: Look into the reason for not passing bLooping variable and just forcing the vale to be true. 
-							LeaderSample.Animation->GetMarkerIndicesForPosition(Context.MarkerTickContext.GetMarkerSyncStartPosition(), true, LeaderSampleData.MarkerTickRecord.PreviousMarker, LeaderSampleData.MarkerTickRecord.NextMarker, LeaderSampleData.Time, Instance.MirrorDataTable);
-						}
+						LeaderSampleData.MarkerTickRecord.Reset();
+						LeaderSampleData.Time = NormalizedCurrentTime * LeaderSample.Animation->GetPlayLength();
+						bResetMarkerDataOnFollowers = true;
+					}
+					// Re-compute marker indices since the leader sample's tick record is invalid. Get previous and next markers.
+					else if (!LeaderSampleData.MarkerTickRecord.IsValid(Instance.bLooping) && Context.MarkerTickContext.GetMarkerSyncStartPosition().IsValid())
+					{
+						// TODO: Look into the reason for not passing bLooping variable and just forcing the vale to be true. 
+						LeaderSample.Animation->GetMarkerIndicesForPosition(Context.MarkerTickContext.GetMarkerSyncStartPosition(), true, LeaderSampleData.MarkerTickRecord.PreviousMarker, LeaderSampleData.MarkerTickRecord.NextMarker, LeaderSampleData.Time, Instance.MirrorDataTable);
+					}
 
-						// Only tick samples if leader sample has any delta time to consume.
-						const float NewDeltaTime = Context.GetDeltaTime() * Instance.PlayRateMultiplier * LeaderSample.RateScale * LeaderSample.Animation->RateScale;
-						Context.SetLeaderDelta(NewDeltaTime);
+					// Only tick samples if leader sample has any delta time to consume.
+					const float NewDeltaTime = Instance.DeltaTimeRecord->Delta * LeaderSample.RateScale * LeaderSample.Animation->RateScale;
+					Context.SetLeaderDelta(NewDeltaTime);
 
-						if (!FMath::IsNearlyZero(NewDeltaTime))
-						{
-							// Tick leader sample
-							LeaderSample.Animation->TickByMarkerAsLeader(LeaderSampleData.MarkerTickRecord, Context.MarkerTickContext, LeaderSampleData.Time, LeaderSampleData.PreviousTime, NewDeltaTime, Instance.bLooping, Instance.MirrorDataTable);
+					if (!FMath::IsNearlyZero(NewDeltaTime))
+					{
+						// Tick leader sample
+						LeaderSample.Animation->TickByMarkerAsLeader(LeaderSampleData.MarkerTickRecord, Context.MarkerTickContext, LeaderSampleData.Time, LeaderSampleData.PreviousTime, NewDeltaTime, Instance.bLooping, Instance.MirrorDataTable);
 
-							check(!Instance.bLooping || Context.MarkerTickContext.IsMarkerSyncStartValid());
+						check(!Instance.bLooping || Context.MarkerTickContext.IsMarkerSyncStartValid());
 							
-							// Tick all the follower samples
-							TickFollowerSamples(SampleDataList, HighestMarkerSyncWeightIndex, Context, bResetMarkerDataOnFollowers, Instance.bLooping, Instance.MirrorDataTable);
-						}
-						else if (!Instance.MarkerTickRecord->IsValid(Instance.bLooping))
-						{
-							// Re-compute marker indices for leader sample's tick record. Get previous and next markers.
-							LeaderSample.Animation->GetMarkerIndicesForTime(LeaderSampleData.Time, Instance.bLooping, Context.MarkerTickContext.GetValidMarkerNames(), LeaderSampleData.MarkerTickRecord.PreviousMarker, LeaderSampleData.MarkerTickRecord.NextMarker);
+						// Tick all the follower samples
+						TickFollowerSamples(SampleDataList, HighestMarkerSyncWeightIndex, Context, bResetMarkerDataOnFollowers, Instance.bLooping, Instance.MirrorDataTable);
+					}
+					else if (!Instance.MarkerTickRecord->IsValid(Instance.bLooping))
+					{
+						// Re-compute marker indices for leader sample's tick record. Get previous and next markers.
+						LeaderSample.Animation->GetMarkerIndicesForTime(LeaderSampleData.Time, Instance.bLooping, Context.MarkerTickContext.GetValidMarkerNames(), LeaderSampleData.MarkerTickRecord.PreviousMarker, LeaderSampleData.MarkerTickRecord.NextMarker);
 
-							// Get sync position for followers to sync up to.
-							const FMarkerSyncAnimPosition SyncPosition = LeaderSample.Animation->GetMarkerSyncPositionFromMarkerIndicies(LeaderSampleData.MarkerTickRecord.PreviousMarker.MarkerIndex, LeaderSampleData.MarkerTickRecord.NextMarker.MarkerIndex, LeaderSampleData.Time, Instance.MirrorDataTable);
-							Context.MarkerTickContext.SetMarkerSyncStartPosition(SyncPosition);
-							Context.MarkerTickContext.SetMarkerSyncEndPosition(SyncPosition);
+						// Get sync position for followers to sync up to.
+						const FMarkerSyncAnimPosition SyncPosition = LeaderSample.Animation->GetMarkerSyncPositionFromMarkerIndicies(LeaderSampleData.MarkerTickRecord.PreviousMarker.MarkerIndex, LeaderSampleData.MarkerTickRecord.NextMarker.MarkerIndex, LeaderSampleData.Time, Instance.MirrorDataTable);
+						Context.MarkerTickContext.SetMarkerSyncStartPosition(SyncPosition);
+						Context.MarkerTickContext.SetMarkerSyncEndPosition(SyncPosition);
 							
-							// Make all follower samples match next sync position to equal that of the leader.
-							TickFollowerSamples(SampleDataList, HighestMarkerSyncWeightIndex, Context, true, Instance.bLooping, Instance.MirrorDataTable);
-						}
-						
-						NormalizedCurrentTime = LeaderSampleData.Time / LeaderSample.Animation->GetPlayLength();
-						*Instance.MarkerTickRecord = LeaderSampleData.MarkerTickRecord;
+						// Make all follower samples match next sync position to equal that of the leader.
+						TickFollowerSamples(SampleDataList, HighestMarkerSyncWeightIndex, Context, true, Instance.bLooping, Instance.MirrorDataTable);
 					}
+						
+					NormalizedCurrentTime = LeaderSampleData.Time / LeaderSample.Animation->GetPlayLength();
+					*Instance.MarkerTickRecord = LeaderSampleData.MarkerTickRecord;
 				}
-				else
-				{
-					// Advance time using current/new anim length
-					float CurrentTime = NormalizedCurrentTime * NewAnimLength;
-					FAnimationRuntime::AdvanceTime(Instance.bLooping, Instance.DeltaTimeRecord->Delta, /*inout*/ CurrentTime, NewAnimLength);
-					NormalizedCurrentTime = NewAnimLength ? (CurrentTime / NewAnimLength) : 0.0f;
-					UE_LOG(LogAnimMarkerSync, Log, 
-						TEXT("Leader (%s) (bCanDoMarkerSync == false)  - PreviousTime (%0.2f), CurrentTime (%0.2f), MoveDelta (%0.2f) "), 
-						*GetName(), NormalizedPreviousTime, NormalizedCurrentTime, Instance.DeltaTimeRecord->Delta);
-				}
-
-				// Update time position after it has undergone all side effects.
-				Context.SetAnimationPositionRatio(NormalizedCurrentTime);
 			}
 			else
 			{
-				// Skip syncing if leader doesn't have a valid sync start position.
-				if (!Context.MarkerTickContext.IsMarkerSyncStartValid())
-				{
-					bCanDoMarkerSync = false;
-				}
+				// Advance time using current/new anim length
+				float CurrentTime = NormalizedCurrentTime * NewAnimLength;
+				FAnimationRuntime::AdvanceTime(Instance.bLooping, Instance.DeltaTimeRecord->Delta, /*inout*/ CurrentTime, NewAnimLength);
+				NormalizedCurrentTime = NewAnimLength ? (CurrentTime / NewAnimLength) : 0.0f;
+				UE_LOG(LogAnimMarkerSync, Log, 
+					TEXT("Leader (%s) (bCanDoMarkerSync == false)  - PreviousTime (%0.2f), CurrentTime (%0.2f), MoveDelta (%0.2f) "), 
+					*GetName(), NormalizedPreviousTime, NormalizedCurrentTime, Instance.DeltaTimeRecord->Delta);
+			}
 
-				// Tick as follower using marked-based syncing.
-				if (bCanDoMarkerSync)
-				{
-					const int32 HighestWeightIndex = FBlendSpaceUtilities::GetHighestWeightSample(SampleDataList);
-					FBlendSampleData& SampleDataItem = SampleDataList[HighestWeightIndex];
-					const FBlendSample& Sample = SampleData[SampleDataItem.SampleDataIndex];
+			// Update time position after it has undergone all side effects.
+			Context.SetAnimationPositionRatio(NormalizedCurrentTime);
+		}
+		else
+		{
+			// Skip syncing if leader doesn't have a valid sync start position.
+			if (!Context.MarkerTickContext.IsMarkerSyncStartValid())
+			{
+				bCanDoMarkerSync = false;
+			}
 
-					if (Sample.Animation)
+			// Tick as follower using marked-based syncing.
+			if (bCanDoMarkerSync)
+			{
+				const int32 HighestWeightIndex = FBlendSpaceUtilities::GetHighestWeightSample(SampleDataList);
+				FBlendSampleData& SampleDataItem = SampleDataList[HighestWeightIndex];
+				const FBlendSample& Sample = SampleData[SampleDataItem.SampleDataIndex];
+
+				if (Sample.Animation)
+				{
+					// Only tick samples if sync group leader has any delta time to consume.
+					if (Context.GetDeltaTime() != 0.f)
 					{
-						// Only tick samples if sync group leader has any delta time to consume.
-						if (Context.GetDeltaTime() != 0.f)
+						if (!Instance.MarkerTickRecord->IsValid(Instance.bLooping))
 						{
-							if (!Instance.MarkerTickRecord->IsValid(Instance.bLooping))
-							{
-								SampleDataItem.Time = NormalizedCurrentTime * Sample.Animation->GetPlayLength();
-							}
-
-							// Tick all samples as followers
-							TickFollowerSamples(SampleDataList, INDEX_NONE, Context, false, Instance.bLooping, Instance.MirrorDataTable);
+							SampleDataItem.Time = NormalizedCurrentTime * Sample.Animation->GetPlayLength();
 						}
-						
-						*Instance.MarkerTickRecord = SampleDataItem.MarkerTickRecord;
-						NormalizedCurrentTime = SampleDataItem.Time / Sample.Animation->GetPlayLength();
+
+						// Tick all samples as followers
+						TickFollowerSamples(SampleDataList, INDEX_NONE, Context, false, Instance.bLooping, Instance.MirrorDataTable);
 					}
+						
+					*Instance.MarkerTickRecord = SampleDataItem.MarkerTickRecord;
+					NormalizedCurrentTime = SampleDataItem.Time / Sample.Animation->GetPlayLength();
 				}
-				else
-				{
-					// Fallback to length-based syncing. Match sync group leader position.
-					NormalizedPreviousTime = Context.GetPreviousAnimationPositionRatio();
-					NormalizedCurrentTime = Context.GetAnimationPositionRatio();
+			}
+			else
+			{
+				// Fallback to length-based syncing. Match sync group leader position.
+				NormalizedPreviousTime = Context.GetPreviousAnimationPositionRatio();
+				NormalizedCurrentTime = Context.GetAnimationPositionRatio();
 					
-					UE_LOG(LogAnimMarkerSync, Log, 
-						TEXT("Follower (%s) (bCanDoMarkerSync == false) - PreviousTime (%0.2f), CurrentTime (%0.2f), MoveDelta (%0.2f) "), 
-						*GetName(), NormalizedPreviousTime, NormalizedCurrentTime, Instance.DeltaTimeRecord->Delta);
-				}
+				UE_LOG(LogAnimMarkerSync, Log, 
+					TEXT("Follower (%s) (bCanDoMarkerSync == false) - PreviousTime (%0.2f), CurrentTime (%0.2f), MoveDelta (%0.2f) "), 
+					*GetName(), NormalizedPreviousTime, NormalizedCurrentTime, Instance.DeltaTimeRecord->Delta);
 			}
+		}
+
+		// Generate notifies and sets time.
+		{
+			FAnimNotifyContext NotifyContext(Instance);
+			float ClampedNormalizedPreviousTime = FMath::Clamp<float>(NormalizedPreviousTime, 0.f, 1.f);
+			float ClampedNormalizedCurrentTime = FMath::Clamp<float>(NormalizedCurrentTime, 0.f, 1.f);
 
-			// Generate notifies and sets time.
+			if (Instance.bIsEvaluator && !Instance.BlendSpace.bTeleportToTime)
 			{
-				FAnimNotifyContext NotifyContext(Instance);
-				float ClampedNormalizedPreviousTime = FMath::Clamp<float>(NormalizedPreviousTime, 0.f, 1.f);
-				float ClampedNormalizedCurrentTime = FMath::Clamp<float>(NormalizedCurrentTime, 0.f, 1.f);
+				// When running under an evaluator the time is being set explicitly and we want to add on the deltas.
+				ClampedNormalizedPreviousTime -= ExtraNormalizedDeltaTime;
+				// Note that ExtraNormalizedDeltaTime can be negative
+				ClampedNormalizedPreviousTime = FMath::Wrap<float>(ClampedNormalizedPreviousTime, 0.0f, 1.0f);
+
+				// Also when under an evaluator, since the time is explicitly set before the update is called, the desired 
+				// current time is actually what we recorded before advancing time (effectively ignoring whatever was added).
+				ClampedNormalizedCurrentTime = FMath::Clamp<float>(NormalizedPreviousTime, 0.f, 1.f);
+			}
 
-				if (Instance.bIsEvaluator && !Instance.BlendSpace.bTeleportToTime)
-				{
-					// When running under an evaluator the time is being set explicitly and we want to add on the deltas.
-					ClampedNormalizedPreviousTime -= ExtraNormalizedDeltaTime;
-					// Note that ExtraNormalizedDeltaTime can be negative
-					ClampedNormalizedPreviousTime = FMath::Wrap<float>(ClampedNormalizedPreviousTime, 0.0f, 1.0f);
-
-					// Also when under an evaluator, since the time is explicitly set before the update is called, the desired 
-					// current time is actually what we recorded before advancing time (effectively ignoring whatever was added).
-					ClampedNormalizedCurrentTime = FMath::Clamp<float>(NormalizedPreviousTime, 0.f, 1.f);
-				}
+			const bool bHasDeltaTime = (NormalizedCurrentTime != NormalizedPreviousTime);
+			const bool bGenerateNotifies = NotifyTriggerMode != ENotifyTriggerMode::None;
 
-				const bool bHasDeltaTime = (NormalizedCurrentTime != NormalizedPreviousTime);
-				const bool bGenerateNotifies = NotifyTriggerMode != ENotifyTriggerMode::None;
+			// Get the index of the highest weight, assuming that the first is the highest until we find otherwise
+			const bool bTriggerNotifyHighestWeightedAnim = NotifyTriggerMode == ENotifyTriggerMode::HighestWeightedAnimation && SampleDataList.Num() > 0;
+			const int32 HighestWeightIndex = (bGenerateNotifies && bTriggerNotifyHighestWeightedAnim) ? FBlendSpaceUtilities::GetHighestWeightSample(SampleDataList) : INDEX_NONE;
 
-				// Get the index of the highest weight, assuming that the first is the highest until we find otherwise
-				const bool bTriggerNotifyHighestWeightedAnim = NotifyTriggerMode == ENotifyTriggerMode::HighestWeightedAnimation && SampleDataList.Num() > 0;
-				const int32 HighestWeightIndex = (bGenerateNotifies && bTriggerNotifyHighestWeightedAnim) ? FBlendSpaceUtilities::GetHighestWeightSample(SampleDataList) : INDEX_NONE;
+			for (int32 I = 0; I < SampleDataList.Num(); ++I)
+			{
+				FBlendSampleData& SampleEntry = SampleDataList[I];
+				const int32 SampleDataIndex = SampleEntry.SampleDataIndex;
 
-				for (int32 I = 0; I < SampleDataList.Num(); ++I)
+				// Skip SamplesPoints that has no relevant weight
+				if (SampleData.IsValidIndex(SampleDataIndex) && (SampleEntry.TotalWeight > ZERO_ANIMWEIGHT_THRESH))
 				{
-					FBlendSampleData& SampleEntry = SampleDataList[I];
-					const int32 SampleDataIndex = SampleEntry.SampleDataIndex;
-
-					// Skip SamplesPoints that has no relevant weight
-					if (SampleData.IsValidIndex(SampleDataIndex) && (SampleEntry.TotalWeight > ZERO_ANIMWEIGHT_THRESH))
+					const FBlendSample& Sample = SampleData[SampleDataIndex];
+					if (Sample.Animation)
 					{
-						const FBlendSample& Sample = SampleData[SampleDataIndex];
-						if (Sample.Animation)
-						{
-							float PrevSampleDataTime;
-							float& CurrentSampleDataTime = SampleEntry.Time;
+						float PrevSampleDataTime;
+						float& CurrentSampleDataTime = SampleEntry.Time;
 
-							const float MultipliedSampleRateScale = Sample.Animation->RateScale * Sample.RateScale;
+						const float MultipliedSampleRateScale = Sample.Animation->RateScale * Sample.RateScale;
 
-							if (!bCanDoMarkerSync || Sample.Animation->AuthoredSyncMarkers.Num() == 0) //Have already updated time if we are doing marker sync
-							{
-								const float SampleNormalizedPreviousTime = MultipliedSampleRateScale >= 0.f ? ClampedNormalizedPreviousTime : 1.f - ClampedNormalizedPreviousTime;
-								const float SampleNormalizedCurrentTime = MultipliedSampleRateScale >= 0.f ? ClampedNormalizedCurrentTime : 1.f - ClampedNormalizedCurrentTime;
-								PrevSampleDataTime = SampleNormalizedPreviousTime * Sample.Animation->GetPlayLength();
-								CurrentSampleDataTime = SampleNormalizedCurrentTime * Sample.Animation->GetPlayLength();
-							}
-							else
-							{
-								PrevSampleDataTime = SampleEntry.PreviousTime;
-							}
+						if (!bCanDoMarkerSync || Sample.Animation->AuthoredSyncMarkers.Num() == 0) //Have already updated time if we are doing marker sync
+						{
+							const float SampleNormalizedPreviousTime = MultipliedSampleRateScale >= 0.f ? ClampedNormalizedPreviousTime : 1.f - ClampedNormalizedPreviousTime;
+							const float SampleNormalizedCurrentTime = MultipliedSampleRateScale >= 0.f ? ClampedNormalizedCurrentTime : 1.f - ClampedNormalizedCurrentTime;
+							PrevSampleDataTime = SampleNormalizedPreviousTime * Sample.Animation->GetPlayLength();
+							CurrentSampleDataTime = SampleNormalizedCurrentTime * Sample.Animation->GetPlayLength();
+						}
+						else
+						{
+							PrevSampleDataTime = SampleEntry.PreviousTime;
+						}
 
-							// Figure out delta time 
-							float DeltaTimePosition = CurrentSampleDataTime - PrevSampleDataTime;
-							const float SampleMoveDelta = Instance.DeltaTimeRecord->Delta * MultipliedSampleRateScale;
+						// Figure out delta time 
+						float DeltaTimePosition = CurrentSampleDataTime - PrevSampleDataTime;
+						const float SampleMoveDelta = Instance.DeltaTimeRecord->Delta * MultipliedSampleRateScale;
 
-							// if we went against play rate, then loop around.
-							if ((SampleMoveDelta * DeltaTimePosition) < 0.f)
-							{
-								DeltaTimePosition += FMath::Sign<float>(SampleMoveDelta) * Sample.Animation->GetPlayLength();
-							}
+						// if we went against play rate, then loop around.
+						if ((SampleMoveDelta * DeltaTimePosition) < 0.f)
+						{
+							DeltaTimePosition += FMath::Sign<float>(SampleMoveDelta) * Sample.Animation->GetPlayLength();
+						}
 
-							if (bGenerateNotifies && (!bTriggerNotifyHighestWeightedAnim || (I == HighestWeightIndex)))
-							{
-								// Harvest and record notifies
-								Sample.Animation->GetAnimNotifies(PrevSampleDataTime, DeltaTimePosition, NotifyContext);
-							}
+						if (bGenerateNotifies && (!bTriggerNotifyHighestWeightedAnim || (I == HighestWeightIndex)))
+						{
+							// Harvest and record notifies
+							Sample.Animation->GetAnimNotifies(PrevSampleDataTime, DeltaTimePosition, NotifyContext);
+						}
 
-							if (bHasDeltaTime)
+						if (bHasDeltaTime)
+						{
+							if (Context.RootMotionMode == ERootMotionMode::RootMotionFromEverything && Sample.Animation->bEnableRootMotion)
 							{
-								if (Context.RootMotionMode == ERootMotionMode::RootMotionFromEverything && Sample.Animation->bEnableRootMotion)
-								{
-									Context.RootMotionMovementParams.AccumulateWithBlend(Sample.Animation->ExtractRootMotion(PrevSampleDataTime, DeltaTimePosition, Instance.bLooping), SampleEntry.GetClampedWeight());
-								}
+								Context.RootMotionMovementParams.AccumulateWithBlend(Sample.Animation->ExtractRootMotion(PrevSampleDataTime, DeltaTimePosition, Instance.bLooping), SampleEntry.GetClampedWeight());
 							}
+						}
 
-							// Capture the final adjusted delta time and previous frame time as an asset player record
-							SampleEntry.DeltaTimeRecord.Set(PrevSampleDataTime, DeltaTimePosition);
+						// Capture the final adjusted delta time and previous frame time as an asset player record
+						SampleEntry.DeltaTimeRecord.Set(PrevSampleDataTime, DeltaTimePosition);
 
-							UE_LOG(LogAnimation, Verbose, TEXT("%d. Blending animation(%s) with %f weight at time %0.2f"), I + 1, *Sample.Animation->GetName(), SampleEntry.GetClampedWeight(), CurrentSampleDataTime);
-						}
+						UE_LOG(LogAnimation, Verbose, TEXT("%d. Blending animation(%s) with %f weight at time %0.2f"), I + 1, *Sample.Animation->GetName(), SampleEntry.GetClampedWeight(), CurrentSampleDataTime);
 					}
 				}
+			}
 
-				if (bGenerateNotifies && NotifyContext.ActiveNotifies.Num() > 0)
-				{
-					NotifyQueue.AddAnimNotifies(Context.ShouldGenerateNotifies(), NotifyContext.ActiveNotifies, Instance.EffectiveBlendWeight);
-				}
+			if (bGenerateNotifies && NotifyContext.ActiveNotifies.Num() > 0)
+			{
+				NotifyQueue.AddAnimNotifies(Context.ShouldGenerateNotifies(), NotifyContext.ActiveNotifies, Instance.EffectiveBlendWeight);
 			}
 		}
-
-		OldSampleDataList.Reset();
 	}
+
+	OldSampleDataList.Reset();
 }
 
 bool UBlendSpace::IsValidAdditive() const
@@ -1431,7 +1430,7 @@ void UBlendSpace::ValidateSampleData()
 						}
 					};
 
-					if (SampleWithMarkers == INDEX_NONE)
+					if (SampleWithMarkers == INDEX_NONE && bAllowMarkerBasedSync)
 					{
 						SampleWithMarkers = SampleIndex;
 					}
