[Engine]
* Wrapped FSceneView "ViewActor" with a struct that stores both actor and actor unique id. Actor will be removed once all code has been converted.

diff --git a/Engine/Plugins/Experimental/PostProcessMaterialChainGraph/Source/PostProcessMaterialChainGraph/Private/PPMChainGraphSceneViewExtension.cpp b/Engine/Plugins/Experimental/PostProcessMaterialChainGraph/Source/PostProcessMaterialChainGraph/Private/PPMChainGraphSceneViewExtension.cpp
--- a/Engine/Plugins/Experimental/PostProcessMaterialChainGraph/Source/PostProcessMaterialChainGraph/Private/PPMChainGraphSceneViewExtension.cpp
+++ b/Engine/Plugins/Experimental/PostProcessMaterialChainGraph/Source/PostProcessMaterialChainGraph/Private/PPMChainGraphSceneViewExtension.cpp
@@ -242,7 +242,7 @@ void FPPMChainGraphSceneViewExtension::GatherChainGraphProxies
 
 			if (PPMChainGraphComponent->CameraViewHandlingMode == ECameraViewHandling::IgnoreCameraViews)
 			{
-				if (Camera == InView.ViewActor)
+				if (Camera.Get() == InView.ViewActor)
 				{
 					bCameraOptionIsValid = false;
 					break;
@@ -250,7 +250,7 @@ void FPPMChainGraphSceneViewExtension::GatherChainGraphProxies
 			}
 			else
 			{
-				if (Camera == InView.ViewActor)
+				if (Camera.Get() == InView.ViewActor)
 				{
 					bCameraOptionIsValid = true;
 					break;
diff --git a/Engine/Plugins/VirtualProduction/CameraCalibrationCore/Source/CameraCalibrationCore/Private/LensDistortionSceneViewExtension.cpp b/Engine/Plugins/VirtualProduction/CameraCalibrationCore/Source/CameraCalibrationCore/Private/LensDistortionSceneViewExtension.cpp
--- a/Engine/Plugins/VirtualProduction/CameraCalibrationCore/Source/CameraCalibrationCore/Private/LensDistortionSceneViewExtension.cpp
+++ b/Engine/Plugins/VirtualProduction/CameraCalibrationCore/Source/CameraCalibrationCore/Private/LensDistortionSceneViewExtension.cpp
@@ -365,7 +365,7 @@ void FLensDistortionSceneViewExtension::InvertDistortionMap_RenderThread(FRDGBui
 
 void FLensDistortionSceneViewExtension::PreRenderView_RenderThread(FRDGBuilder& GraphBuilder, FSceneView& InView)
 {
-	if (const ACameraActor* CameraActor = Cast<const ACameraActor>(InView.ViewActor))
+	if (const ACameraActor* CameraActor = Cast<const ACameraActor>(InView.ViewActor.Get()))
 	{
 		FDisplacementMapBlendingParams BlendState;
 		{
diff --git a/Engine/Source/Runtime/Engine/Private/PrimitiveSceneProxy.cpp b/Engine/Source/Runtime/Engine/Private/PrimitiveSceneProxy.cpp
--- a/Engine/Source/Runtime/Engine/Private/PrimitiveSceneProxy.cpp
+++ b/Engine/Source/Runtime/Engine/Private/PrimitiveSceneProxy.cpp
@@ -330,7 +330,7 @@ void FPrimitiveSceneProxyDesc::InitializeFrom(const UPrimitiveComponent* InCompo
 			// Make a list of the actors which directly or indirectly own the InComponent.
 			for(const AActor* CurrentOwner = ActorOwner;CurrentOwner;CurrentOwner = CurrentOwner->GetOwner())
 			{
-				ActorOwners.Add(CurrentOwner);
+				ActorOwners.Add(CurrentOwner->GetUniqueID());
 			}
 		}
 		
@@ -1555,7 +1555,7 @@ bool FPrimitiveSceneProxy::IsShown(const FSceneView* View) const
 			return false;
 		}
 
-		const bool bOwnersContain = !Owners.IsEmpty() && Owners.Contains(View->ViewActor);
+		const bool bOwnersContain = View->ViewActor.IsPartOf(Owners);
 		if (bOnlyOwnerSee && !bOwnersContain)
 		{
 			return false;
@@ -1639,7 +1639,7 @@ bool FPrimitiveSceneProxy::IsShadowCast(const FSceneView* View) const
 
 		// In the OwnerSee cases, we still want to respect hidden shadows...
 		// This assumes that bCastHiddenShadow trumps the owner see flags.
-		const bool bOwnersContain = Owners.Contains(View->ViewActor);
+		const bool bOwnersContain = View->ViewActor.IsPartOf(Owners);
 		if (bOnlyOwnerSee && !bOwnersContain)
 		{
 			return false;
diff --git a/Engine/Source/Runtime/Engine/Private/SceneView.cpp b/Engine/Source/Runtime/Engine/Private/SceneView.cpp
--- a/Engine/Source/Runtime/Engine/Private/SceneView.cpp
+++ b/Engine/Source/Runtime/Engine/Private/SceneView.cpp
@@ -760,6 +760,13 @@ FViewMatrices::FViewMatrices(const FMinimalInitializer& Initializer) : FViewMatr
 	Init(Initializer);
 }
 
+FSceneViewOwner::FSceneViewOwner(const AActor* InActor)
+	: ActorUniqueId(InActor ? InActor->GetUniqueID() : 0)
+#if !UE_STATE_STREAM
+	, Actor(InActor)
+#endif
+{
+}
 
 static void SetupViewFrustum(FSceneView& View)
 {
diff --git a/Engine/Source/Runtime/Engine/Private/WorldPartition/WorldPartitionSubsystem.cpp b/Engine/Source/Runtime/Engine/Private/WorldPartition/WorldPartitionSubsystem.cpp
--- a/Engine/Source/Runtime/Engine/Private/WorldPartition/WorldPartitionSubsystem.cpp
+++ b/Engine/Source/Runtime/Engine/Private/WorldPartition/WorldPartitionSubsystem.cpp
@@ -1393,6 +1393,8 @@ void UWorldPartitionSubsystem::Draw(UCanvas* Canvas, class APlayerController* PC
 		return;
 	}
 
+#if !UE_STATE_STREAM // TODO: This needs to be fixed
+
 	// Filter out views that don't match our world
 	if (!GetWorld()->IsNetMode(NM_DedicatedServer) && !UWorldPartition::IsSimulating(false) &&
 		(Canvas->SceneView->ViewActor == nullptr || Canvas->SceneView->ViewActor->GetWorld() != GetWorld()))
@@ -1400,6 +1402,8 @@ void UWorldPartitionSubsystem::Draw(UCanvas* Canvas, class APlayerController* PC
 		return;
 	}
 
+#endif
+
 	// When GDrawWorldPartitionIndex < 0, draw all world partitions
 	const int32 WPDrawCount = (GDrawWorldPartitionIndex < 0) ? RegisteredWorldPartitions.Num() : 1;
 	UWorldPartition* SingleWorldPartition = (WPDrawCount == 1) ? RegisteredWorldPartitions[FMath::Max(GDrawWorldPartitionIndex, 0)] : nullptr;
diff --git a/Engine/Source/Runtime/Engine/Public/EngineDefines.h b/Engine/Source/Runtime/Engine/Public/EngineDefines.h
--- a/Engine/Source/Runtime/Engine/Public/EngineDefines.h
+++ b/Engine/Source/Runtime/Engine/Public/EngineDefines.h
@@ -25,6 +25,11 @@
 // The number of lights to consider for sky/atmospheric light scattering
 #define NUM_ATMOSPHERE_LIGHTS 2
 
+// Statestream is WIP and is the system that is going to enable full separation of game and render updates
+#ifndef UE_STATE_STREAM
+	#define UE_STATE_STREAM 0
+#endif
+
 /*-----------------------------------------------------------------------------
 	Size of the world.
 -----------------------------------------------------------------------------*/
diff --git a/Engine/Source/Runtime/Engine/Public/PrimitiveSceneProxy.h b/Engine/Source/Runtime/Engine/Public/PrimitiveSceneProxy.h
--- a/Engine/Source/Runtime/Engine/Public/PrimitiveSceneProxy.h
+++ b/Engine/Source/Runtime/Engine/Public/PrimitiveSceneProxy.h
@@ -7,6 +7,7 @@
 #pragma once
 
 #include "CoreMinimal.h"
+#include "SceneViewOwner.h"
 #include "Stats/Stats.h"
 #include "Misc/MemStack.h"
 #include "PrimitiveViewRelevance.h"
@@ -754,7 +755,7 @@ public:
 	inline bool SinglePassGDME() const { return bSinglePassGDME; }
 	inline bool IsCollisionEnabled() const { return bCollisionEnabled; }
 	inline bool IsHovered() const { return bHovered; }
-	inline bool IsOwnedBy(const AActor* Actor) const { return Owners.Find(Actor) != INDEX_NONE; }
+	inline bool IsOwnedBy(const FSceneViewOwner& Owner) const { return Owner.IsPartOf(Owners); }
 	inline bool HasViewDependentDPG() const { return bUseViewOwnerDepthPriorityGroup; }
 	inline bool HasStaticLighting() const { return bStaticLighting; }
 	inline bool NeedsUnbuiltPreviewLighting() const { return bNeedsUnbuiltPreviewLighting; }
@@ -1557,8 +1558,8 @@ protected:
 	TArray<ERuntimeVirtualTextureMaterialType> RuntimeVirtualTextureMaterialTypes;
 
 private:
-	/** The hierarchy of owners of this primitive.  These must not be dereferenced on the rendering thread, but the pointer values can be used for identification.  */
-	TArray<const AActor*> Owners;
+	/** The hierarchy of owners of this primitive. This array contains the unique id of actors.  */
+	TArray<uint32> Owners;
 
 	/** The primitive's local to world transform. */
 	FMatrix LocalToWorld;
diff --git a/Engine/Source/Runtime/Engine/Public/PrimitiveSceneProxyDesc.h b/Engine/Source/Runtime/Engine/Public/PrimitiveSceneProxyDesc.h
--- a/Engine/Source/Runtime/Engine/Public/PrimitiveSceneProxyDesc.h
+++ b/Engine/Source/Runtime/Engine/Public/PrimitiveSceneProxyDesc.h
@@ -177,7 +177,7 @@ struct FPrimitiveSceneProxyDesc
 	UWorld* World = nullptr;	
 
 	// Only used by actors for now, explicitly intended to be moved to the FPrimitiveSceneProxy
-	mutable TArray<const AActor*> ActorOwners;	
+	mutable TArray<uint32> ActorOwners;	
 	const FCustomPrimitiveData* CustomPrimitiveData = nullptr;
 	FSceneInterface* Scene = nullptr;
 	IPrimitiveComponent* PrimitiveComponentInterface = nullptr;
diff --git a/Engine/Source/Runtime/Engine/Public/SceneInterface.h b/Engine/Source/Runtime/Engine/Public/SceneInterface.h
--- a/Engine/Source/Runtime/Engine/Public/SceneInterface.h
+++ b/Engine/Source/Runtime/Engine/Public/SceneInterface.h
@@ -9,6 +9,7 @@
 #include "RenderGraphFwd.h"
 #include "Engine/EngineTypes.h"
 #include "PrimitiveComponentId.h"
+#include "SceneViewOwner.h"
 #include "ShowFlags.h"
 
 class AWorldSettings;
@@ -702,7 +703,7 @@ public:
 		TSet<FPrimitiveComponentId> HiddenPrimitives;
 		TOptional<TSet<FPrimitiveComponentId>> ShowOnlyPrimitives;
 		FEngineShowFlags EngineShowFlags;
-		const AActor* ViewActor = nullptr;
+		FSceneViewOwner ViewActor;
 		class FSceneViewStateInterface* ViewStateInterface = nullptr;
 		int32 PostVolumeUserFlags = 0;						// UserFlags from Post Process Volume, used if bOverridesPostVolumeUserFlags set, allows per-view material overrides
 		bool bIsSceneCapture = false;
diff --git a/Engine/Source/Runtime/Engine/Public/SceneView.h b/Engine/Source/Runtime/Engine/Public/SceneView.h
--- a/Engine/Source/Runtime/Engine/Public/SceneView.h
+++ b/Engine/Source/Runtime/Engine/Public/SceneView.h
@@ -14,6 +14,7 @@
 #include "PrimitiveComponentId.h"
 #include "RendererInterface.h"
 #include "RenderResource.h"
+#include "SceneViewOwner.h"
 #include "ShowFlags.h"
 #include "StereoRendering.h"
 #include "StereoRenderUtils.h"
@@ -174,7 +175,7 @@ struct FSceneViewInitOptions : public FSceneViewProjectionData
 {
 	const FSceneViewFamily* ViewFamily;
 	FSceneViewStateInterface* SceneViewStateInterface;
-	const AActor* ViewActor;
+	FSceneViewOwner ViewActor; // TODO: Rename to ViewOwner
 	int32 PlayerIndex;
 	FViewElementDrawer* ViewElementDrawer;
 
@@ -259,7 +260,6 @@ struct FSceneViewInitOptions : public FSceneViewProjectionData
 	FSceneViewInitOptions()
 		: ViewFamily(NULL)
 		, SceneViewStateInterface(NULL)
-		, ViewActor(NULL)
 		, PlayerIndex(INDEX_NONE)
 		, ViewElementDrawer(NULL)
 		, BackgroundColor(FLinearColor::Transparent)
@@ -1323,7 +1323,7 @@ public:
 	FSceneViewInitOptions SceneViewInitOptions;
 
 	/** The actor which is being viewed from. */
-	const AActor* ViewActor;
+	FSceneViewOwner ViewActor; // TODO: Rename to ViewOwner
 	 
 	/** Player index this view is associated with or INDEX_NONE. */
 	int32 PlayerIndex;
diff --git a/Engine/Source/Runtime/Engine/Public/SceneViewOwner.h b/Engine/Source/Runtime/Engine/Public/SceneViewOwner.h
new file mode 100644
--- /dev/null
+++ b/Engine/Source/Runtime/Engine/Public/SceneViewOwner.h
@@ -0,0 +1,26 @@
+// Copyright Epic Games, Inc. All Rights Reserved.
+
+#pragma once
+
+#include "Containers/Array.h"
+#include "EngineDefines.h"
+
+class AActor;
+
+// Owner of view
+struct FSceneViewOwner
+{
+	ENGINE_API FSceneViewOwner(const AActor* InActor = nullptr);
+	bool IsPartOf(const TArray<uint32>& InArray) const { return ActorUniqueId ? InArray.Contains(ActorUniqueId) : false; }
+	bool IsSet() const { return ActorUniqueId != 0; }
+	uint32 ActorUniqueId;
+
+	#if !UE_STATE_STREAM
+	const AActor* Get() const { return Actor; }
+	operator const AActor*() const { return Actor; }
+	const AActor* operator->() const { return Actor; }
+	void operator=(const AActor* InActor) { Actor = InActor; }
+	friend bool operator==(const FSceneViewOwner& A, const AActor* B) { return A.Actor == B; }
+	const AActor* Actor;
+	#endif
+};
diff --git a/Engine/Source/Runtime/GameplayDebugger/Private/GameplayDebuggerLocalController.cpp b/Engine/Source/Runtime/GameplayDebugger/Private/GameplayDebuggerLocalController.cpp
--- a/Engine/Source/Runtime/GameplayDebugger/Private/GameplayDebuggerLocalController.cpp
+++ b/Engine/Source/Runtime/GameplayDebugger/Private/GameplayDebuggerLocalController.cpp
@@ -232,7 +232,7 @@ void UGameplayDebuggerLocalController::OnDebugDraw(UCanvas* Canvas, APlayerContr
 		RebuildDataPackMap();
 	}
 
-	if (Canvas->SceneView->ViewActor == nullptr)
+	if (!Canvas->SceneView->ViewActor.IsSet())
 	{
 		CachedReplicator->SetViewPoint(Canvas->SceneView->ViewLocation, Canvas->SceneView->ViewRotation.Vector());
 	}
diff --git a/Engine/Source/Runtime/Renderer/Private/VirtualShadowMaps/VirtualShadowMapClipmap.cpp b/Engine/Source/Runtime/Renderer/Private/VirtualShadowMaps/VirtualShadowMapClipmap.cpp
--- a/Engine/Source/Runtime/Renderer/Private/VirtualShadowMaps/VirtualShadowMapClipmap.cpp
+++ b/Engine/Source/Runtime/Renderer/Private/VirtualShadowMaps/VirtualShadowMapClipmap.cpp
@@ -203,7 +203,13 @@ FVirtualShadowMapClipmap::FVirtualShadowMapClipmap(
 			{
 				FVector ViewForward = CameraViewMatrices.GetViewMatrix().GetColumn(2);
 				FCollisionObjectQueryParams ObjectParams = FCollisionObjectQueryParams(FCollisionObjectQueryParams::InitType::AllObjects);
-				FCollisionQueryParams CollisionParams = FCollisionQueryParams(FName(TEXT("OrthoCamera_VSMTrace")), true, DependentView->ViewActor);
+				
+				FCollisionQueryParams CollisionParams(FName(TEXT("OrthoCamera_VSMTrace")), true);
+				if (DependentView->ViewActor.IsSet())
+				{
+					CollisionParams.AddIgnoredSourceObject(DependentView->ViewActor.ActorUniqueId);
+				}
+
 				FHitResult Hit(ForceInit);
 
 				if (World->LineTraceSingleByObjectType(
