Fix material compilation failure with Substrate with Nanite displacement.

diff --git a/Engine/Source/Runtime/Engine/Private/Materials/HLSLMaterialTranslator.cpp b/Engine/Source/Runtime/Engine/Private/Materials/HLSLMaterialTranslator.cpp
--- a/Engine/Source/Runtime/Engine/Private/Materials/HLSLMaterialTranslator.cpp
+++ b/Engine/Source/Runtime/Engine/Private/Materials/HLSLMaterialTranslator.cpp
@@ -2118,7 +2118,10 @@ void FHLSLMaterialTranslator::TranslateMaterial()
 
 				// Reduce definition statements that don't contribute to the function's return value.
 				// @todo-lh: This should be expanded to a general reduction, but is currently only intended to fix an FXC internal compiler error reported in UE-117831
-				const bool bReduceAfterReturnValue = (PropertyId == MP_WorldPositionOffset || PropertyId == CompiledMP_PrevWorldPositionOffset || PropertyId == MP_Displacement);
+				// When Substrate is enabled, we do not reduce MP_Displacement as it causes shader compilation error, due to FrontMaterial expression being removed, causing local variables to be removed while still used.
+				const bool bReduceAfterReturnValue = Substrate::IsSubstrateEnabled() ? 
+					(PropertyId == MP_WorldPositionOffset || PropertyId == CompiledMP_PrevWorldPositionOffset) : 
+					(PropertyId == MP_WorldPositionOffset || PropertyId == CompiledMP_PrevWorldPositionOffset || PropertyId == MP_Displacement);
 
 				if (bSubstrateEnabled && PropertyId >= MP_FrontMaterial && PropertyShaderFrequency == FrontMaterialShaderFrequency)
 				{
