[Viewport Toolbar] Add current mode icon to "View Modes" menu toolbar entry

diff --git a/Engine/Source/Editor/UnrealEd/Private/ViewportToolbar/UnrealEdViewportToolbar.cpp b/Engine/Source/Editor/UnrealEd/Private/ViewportToolbar/UnrealEdViewportToolbar.cpp
--- a/Engine/Source/Editor/UnrealEd/Private/ViewportToolbar/UnrealEdViewportToolbar.cpp
+++ b/Engine/Source/Editor/UnrealEd/Private/ViewportToolbar/UnrealEdViewportToolbar.cpp
@@ -1273,6 +1273,19 @@ FToolMenuEntry GetLitWireframeEntryToSection(const TSharedPtr<SEditorViewport>&
 	);
 }
 
+FSlateIcon GetViewModesSubmenuIcon(const TWeakPtr<SEditorViewport>& InViewport)
+{
+	if (TSharedPtr<SEditorViewport> PinnedViewport = InViewport.Pin())
+	{
+		const TSharedPtr<FEditorViewportClient> ViewportClient = PinnedViewport->GetViewportClient();
+		check(ViewportClient.IsValid());
+
+		return UViewModeUtils::GetViewModeDisplaySlateIcon(ViewportClient->GetViewMode());
+	}
+
+	return FSlateIcon();
+}
+
 void PopulateViewModesMenu(UToolMenu* InMenu)
 {
 	UUnrealEdViewportToolbarContext* const Context = InMenu->FindContext<UUnrealEdViewportToolbarContext>();
@@ -1685,6 +1698,7 @@ FToolMenuEntry CreateViewportToolbarViewModesSubmenu()
 			{
 				// Base the label on the current view mode.
 				TAttribute<FText> LabelAttribute = UE::UnrealEd::GetViewModesSubmenuLabel(nullptr);
+				TAttribute<FSlateIcon> IconAttribute = TAttribute<FSlateIcon>();
 				if (UUnrealEdViewportToolbarContext* const Context =
 						InDynamicSection.FindContext<UUnrealEdViewportToolbarContext>())
 				{
@@ -1694,6 +1708,13 @@ FToolMenuEntry CreateViewportToolbarViewModesSubmenu()
 							return UE::UnrealEd::GetViewModesSubmenuLabel(WeakViewport);
 						}
 					);
+
+					IconAttribute = TAttribute<FSlateIcon>::CreateLambda(
+						[WeakViewport = Context->Viewport]()
+						{
+							return UE::UnrealEd::GetViewModesSubmenuIcon(WeakViewport);
+						}
+					);
 				}
 
 				InDynamicSection.AddSubMenu(
@@ -1705,7 +1726,9 @@ FToolMenuEntry CreateViewportToolbarViewModesSubmenu()
 						{
 							PopulateViewModesMenu(Submenu);
 						}
-					)
+					),
+					false,
+					IconAttribute
 				);
 			}
 		)
diff --git a/Engine/Source/Editor/UnrealEd/Public/ViewportToolbar/UnrealEdViewportToolbar.h b/Engine/Source/Editor/UnrealEd/Public/ViewportToolbar/UnrealEdViewportToolbar.h
--- a/Engine/Source/Editor/UnrealEd/Public/ViewportToolbar/UnrealEdViewportToolbar.h
+++ b/Engine/Source/Editor/UnrealEd/Public/ViewportToolbar/UnrealEdViewportToolbar.h
@@ -46,6 +46,8 @@ UNREALED_API FToolMenuEntry CreateViewportToolbarSnappingSubmenu();
 
 UNREALED_API FText GetViewModesSubmenuLabel(TWeakPtr<SEditorViewport> InViewport);
 
+UNREALED_API FSlateIcon GetViewModesSubmenuIcon(const TWeakPtr<SEditorViewport>& InViewport);
+
 /**
  * Populate a given UToolMenu with entries for a View Modes viewport toolbar submenu.
  *
diff --git a/Engine/Source/Runtime/Engine/Classes/Engine/EngineBaseTypes.h b/Engine/Source/Runtime/Engine/Classes/Engine/EngineBaseTypes.h
--- a/Engine/Source/Runtime/Engine/Classes/Engine/EngineBaseTypes.h
+++ b/Engine/Source/Runtime/Engine/Classes/Engine/EngineBaseTypes.h
@@ -20,6 +20,7 @@
 
 class UActorComponent;
 struct FSlateBrush;
+struct FSlateIcon;
 struct FTickContext;
 
 //
@@ -1015,6 +1016,7 @@ public:
 	 */
 	static ENGINE_API FText GetViewModeDisplayName(const EViewModeIndex ViewModeIndex);
 	static ENGINE_API const FSlateBrush* GetViewModeDisplayIcon(const EViewModeIndex ViewModeIndex);
+	static ENGINE_API FSlateIcon GetViewModeDisplaySlateIcon(const EViewModeIndex InViewModeIndex);
 };
 
 
diff --git a/Engine/Source/Runtime/Engine/Private/ViewModeNames.cpp b/Engine/Source/Runtime/Engine/Private/ViewModeNames.cpp
--- a/Engine/Source/Runtime/Engine/Private/ViewModeNames.cpp
+++ b/Engine/Source/Runtime/Engine/Private/ViewModeNames.cpp
@@ -3,6 +3,7 @@
 #include "Styling/AppStyle.h"
 #include "Engine/EngineBaseTypes.h"
 #include "Styling/SlateBrush.h"
+#include "Textures/SlateIcon.h"
 
 #define LOCTEXT_NAMESPACE "UViewModeUtils"
 
@@ -216,11 +217,157 @@ FText UViewModeUtils::GetViewModeDisplayName(const EViewModeIndex ViewModeIndex)
 	return ViewModeName;
 }
 
-// Icons 
-TArray<const FSlateBrush*> FillViewModeDisplayIcons()
+// Icons
+FName GetViewModeIconName(const EViewModeIndex InViewMode)
 {
-	const FSlateBrush* IconBrush = nullptr;
+	static FName WireframeIcon("EditorViewport.WireframeMode");
+	static FName UnlitIcon("EditorViewport.UnlitMode");
+	static FName LitIcon("EditorViewport.LitMode");
+	static FName DetailLightingIcon("EditorViewport.DetailLightingMode");
+	static FName LitWireframeIcon("EditorViewport.LitWireframeMode");
+	static FName LightingOnlyIcon("EditorViewport.LightingOnlyMode");
+	static FName LightComplexityIcon("EditorViewport.LightComplexityMode");
+	static FName ShaderComplexityIcon("EditorViewport.ShaderComplexityMode");
+	static FName LightmapDensityIcon("EditorViewport.LightmapDensityMode");
+	static FName ReflectionOverrideIcon("EditorViewport.ReflectionOverrideMode");
+	static FName VisualizeBufferIcon("EditorViewport.VisualizeBufferMode");
+	static FName VisualizeNaniteIcon("EditorViewport.VisualizeNaniteMode");
+	static FName VisualizeLumenIcon("EditorViewport.VisualizeLumenMode");
+	static FName VisualizeSubstrateIcon("EditorViewport.VisualizeSubstrateMode");
+	static FName VisualizeGroomIcon("EditorViewport.VisualizeGroomMode");
+	static FName VisualizeVirtualShadowMapIcon("EditorViewport.VisualizeVirtualShadowMapMode");
+	static FName StationaryLightOverlapIcon("EditorViewport.StationaryLightOverlapMode");
+	static FName CollisionPawnIcon("EditorViewport.CollisionPawn");
+	static FName CollisionVisibilityIcon("EditorViewport.CollisionVisibility");
+	static FName LODIcon("EditorViewport.LOD");
+	static FName QuadOverdrawIcon("EditorViewport.QuadOverdrawMode");
+	static FName TexStreamAccPrimitiveDistanceIcon("EditorViewport.TexStreamAccPrimitiveDistanceMode");
+	static FName TexStreamAccMeshUVDensityIcon("EditorViewport.TexStreamAccMeshUVDensityMode");
+	static FName ShaderComplexityWithQuadOverdrawIcon("EditorViewport.ShaderComplexityWithQuadOverdrawMode");
+	static FName TexStreamAccMaterialTextureScaleIcon("EditorViewport.TexStreamAccMaterialTextureScaleMode");
+	static FName RequiredTextureResolutionIcon("EditorViewport.RequiredTextureResolutionMode");
+	static FName VisualizeGPUSkinCacheIcon("EditorViewport.VisualizeGPUSkinCacheMode");
+	static FName LWCComplexityIcon("EditorViewport.LWCComplexityMode");
+	static FName PathTracingIcon("EditorViewport.PathTracingMode");
+	static FName RayTracingDebugIcon("EditorViewport.RayTracingDebugMode");
+	static FName VirtualTexturePendingMipsIcon("EditorViewport.VirtualTexturePendingMipsMode");
+	static FName VisualizeActorColorationIcon("EditorViewport.VisualizeActorColorationMode");
+
+	FName Icon = NAME_None;
+
+	switch (InViewMode)
+	{
+	case VMI_BrushWireframe:
+	case VMI_Wireframe:
+		Icon = WireframeIcon;
+		break;
+	case VMI_Unlit:
+		Icon = UnlitIcon;
+		break;
+	case VMI_Lit:
+		Icon = LitIcon;
+		break;
+	case VMI_Lit_DetailLighting:
+		Icon = DetailLightingIcon;
+		break;
+	case VMI_LightingOnly:
+		Icon = LightingOnlyIcon;
+		break;
+	case VMI_LightComplexity:
+		Icon = LightComplexityIcon;
+		break;
+	case VMI_ShaderComplexity:
+		Icon = ShaderComplexityIcon;
+		break;
+	case VMI_LightmapDensity:
+	case VMI_LitLightmapDensity:
+		Icon = LightmapDensityIcon;
+		break;
+	case VMI_ReflectionOverride:
+		Icon = ReflectionOverrideIcon;
+		break;
+	case VMI_VisualizeBuffer:
+		Icon = VisualizeBufferIcon;
+		break;
+	case VMI_StationaryLightOverlap:
+		Icon = StationaryLightOverlapIcon;
+		break;
+	case VMI_CollisionPawn:
+		Icon = CollisionPawnIcon;
+		break;
+	case VMI_CollisionVisibility:
+		Icon = CollisionVisibilityIcon;
+		break;
+	case VMI_LODColoration:
+	case VMI_HLODColoration:
+	case VMI_GroupLODColoration:
+		Icon = LODIcon;
+		break;
+	case VMI_QuadOverdraw:
+		Icon = QuadOverdrawIcon;
+		break;
+	case VMI_PrimitiveDistanceAccuracy:
+		Icon = TexStreamAccPrimitiveDistanceIcon;
+		break;
+	case VMI_MeshUVDensityAccuracy:
+		Icon = TexStreamAccMeshUVDensityIcon;
+		break;
+	case VMI_ShaderComplexityWithQuadOverdraw:
+		Icon = ShaderComplexityWithQuadOverdrawIcon;
+		break;
+	case VMI_MaterialTextureScaleAccuracy:
+		Icon = TexStreamAccMaterialTextureScaleIcon;
+		break;
+	case VMI_RequiredTextureResolution:
+		Icon = RequiredTextureResolutionIcon;
+		break;
+	case VMI_PathTracing:
+		Icon = PathTracingIcon;
+		break;
+	case VMI_RayTracingDebug:
+		Icon = RayTracingDebugIcon;
+		break;
+	case VMI_VisualizeNanite:
+		Icon = VisualizeNaniteIcon;
+		break;
+	case VMI_VirtualTexturePendingMips:
+		Icon = VirtualTexturePendingMipsIcon;
+		break;
+	case VMI_VisualizeLumen:
+		Icon = VisualizeLumenIcon;
+		break;
+	case VMI_VisualizeVirtualShadowMap:
+		Icon = VisualizeVirtualShadowMapIcon;
+		break;
+	case VMI_VisualizeGPUSkinCache:
+		Icon = VisualizeGPUSkinCacheIcon;
+		break;
+	case VMI_VisualizeSubstrate:
+		Icon = VisualizeSubstrateIcon;
+		break;
+	case VMI_VisualizeGroom:
+		Icon = VisualizeGroomIcon;
+		break;
+	case VMI_LWCComplexity:
+		Icon = LWCComplexityIcon;
+		break;
+	case VMI_Lit_Wireframe:
+		Icon = LitWireframeIcon;
+		break;
+	case VMI_VisualizeActorColoration:
+		Icon = VisualizeActorColorationIcon;
+		break;
+	case VMI_Max:
+		break;
+	case VMI_Unknown:
+		break;
+	}
 
+	return Icon;
+}
+
+TArray<const FSlateBrush*> FillViewModeDisplayIcons()
+{
 	// Allocate size
 	TArray<const FSlateBrush*> ViewModeDisplayIcons;
 	ViewModeDisplayIcons.Reserve(VMI_Unknown);
@@ -230,193 +377,7 @@ TArray<const FSlateBrush*> FillViewModeDisplayIcons()
 	{
 		const EViewModeIndex ViewModeIndex = (EViewModeIndex)Index;
 
-		if (ViewModeIndex == VMI_BrushWireframe)
-		{
-			ViewModeDisplayIcons.Emplace(FAppStyle::Get().GetBrush("EditorViewport.WireframeMode"));
-		}
-		// Wireframe w/ BSP
-		else if (ViewModeIndex == VMI_Wireframe)
-		{
-			ViewModeDisplayIcons.Emplace(FAppStyle::Get().GetBrush("EditorViewport.WireframeMode"));
-		}
-		// Unlit
-		else if (ViewModeIndex == VMI_Unlit)
-		{
-			ViewModeDisplayIcons.Emplace(FAppStyle::Get().GetBrush("EditorViewport.UnlitMode"));
-		}
-		// Lit
-		else if (ViewModeIndex == VMI_Lit)
-		{
-			ViewModeDisplayIcons.Emplace(FAppStyle::Get().GetBrush("EditorViewport.LitMode"));
-		}
-		else if (ViewModeIndex == VMI_Lit_DetailLighting)
-		{
-			ViewModeDisplayIcons.Emplace(FAppStyle::Get().GetBrush("EditorViewport.DetailLightingMode"));
-		}
-		else if (ViewModeIndex == VMI_Lit_Wireframe)
-		{
-			ViewModeDisplayIcons.Emplace(FAppStyle::Get().GetBrush("EditorViewport.LitWireframeMode"));
-		}
-		// Lit wo/ materials
-		else if (ViewModeIndex == VMI_LightingOnly)
-		{
-			ViewModeDisplayIcons.Emplace(FAppStyle::Get().GetBrush("EditorViewport.LightingOnlyMode"));
-		}
-		// Colored according to light count
-		else if (ViewModeIndex == VMI_LightComplexity)
-		{
-			ViewModeDisplayIcons.Emplace(FAppStyle::Get().GetBrush("EditorViewport.LightComplexityMode"));
-		}
-		// Colored according to shader complexity
-		else if (ViewModeIndex == VMI_ShaderComplexity)
-		{
-			ViewModeDisplayIcons.Emplace(FAppStyle::Get().GetBrush("EditorViewport.ShaderComplexityMode"));
-		}
-		// Colored according to world-space LightMap texture density
-		else if (ViewModeIndex == VMI_LightmapDensity)
-		{
-			ViewModeDisplayIcons.Emplace(FAppStyle::Get().GetBrush("EditorViewport.LightmapDensityMode"));
-		}
-		// Colored according to light count - showing lightmap texel density on texture mapped objects
-		else if (ViewModeIndex == VMI_LitLightmapDensity)
-		{
-			ViewModeDisplayIcons.Emplace(FAppStyle::Get().GetBrush("EditorViewport.LightmapDensityMode"));
-		}
-		else if (ViewModeIndex == VMI_ReflectionOverride)
-		{
-			ViewModeDisplayIcons.Emplace(FAppStyle::Get().GetBrush("EditorViewport.ReflectionOverrideMode"));
-		}
-		else if (ViewModeIndex == VMI_VisualizeBuffer)
-		{
-			ViewModeDisplayIcons.Emplace(FAppStyle::Get().GetBrush("EditorViewport.VisualizeBufferMode"));
-		}
-		else if (ViewModeIndex == VMI_VisualizeNanite)
-		{
-			ViewModeDisplayIcons.Emplace(FAppStyle::Get().GetBrush("EditorViewport.VisualizeNaniteMode"));
-		}
-		else if (ViewModeIndex == VMI_VisualizeLumen)
-		{
-			ViewModeDisplayIcons.Emplace(FAppStyle::Get().GetBrush("EditorViewport.VisualizeLumenMode"));
-		}
-		else if (ViewModeIndex == VMI_VisualizeSubstrate)
-		{
-			ViewModeDisplayIcons.Emplace(FAppStyle::Get().GetBrush("EditorViewport.VisualizeSubstrateMode"));
-		}
-		else if (ViewModeIndex == VMI_VisualizeGroom)
-		{
-			ViewModeDisplayIcons.Emplace(FAppStyle::Get().GetBrush("EditorViewport.VisualizeGroomMode"));
-		}
-		else if (ViewModeIndex == VMI_VisualizeVirtualShadowMap)
-		{
-			ViewModeDisplayIcons.Emplace(FAppStyle::Get().GetBrush("EditorViewport.VisualizeVirtualShadowMapMode"));
-		}
-
-		//	VMI_VoxelLighting = 13,
-	
-		// Colored according to stationary light overlap
-		else if (ViewModeIndex == VMI_StationaryLightOverlap)
-		{
-			ViewModeDisplayIcons.Emplace(FAppStyle::Get().GetBrush("EditorViewport.StationaryLightOverlapMode"));
-		}
-	
-		else if (ViewModeIndex == VMI_CollisionPawn)
-		{
-			ViewModeDisplayIcons.Emplace(FAppStyle::Get().GetBrush("EditorViewport.CollisionPawn"));
-		}
-		else if (ViewModeIndex == VMI_CollisionVisibility)
-		{
-			ViewModeDisplayIcons.Emplace(FAppStyle::Get().GetBrush("EditorViewport.CollisionVisibility"));
-		}
-		//VMI_UNUSED = 17,
-		// Colored according to the current LOD index
-		else if (ViewModeIndex == VMI_LODColoration)
-		{
-			ViewModeDisplayIcons.Emplace(FAppStyle::Get().GetBrush("EditorViewport.LOD"));
-		}
-		// Colored according to the quad coverage
-		else if (ViewModeIndex == VMI_QuadOverdraw)
-		{
-			ViewModeDisplayIcons.Emplace(FAppStyle::Get().GetBrush("EditorViewport.QuadOverdrawMode"));
-		}
-		// Visualize the accuracy of the primitive distance computed for texture streaming
-		else if (ViewModeIndex == VMI_PrimitiveDistanceAccuracy)
-		{
-			ViewModeDisplayIcons.Emplace(FAppStyle::Get().GetBrush("EditorViewport.TexStreamAccPrimitiveDistanceMode"));
-		}
-		// Visualize the accuracy of the mesh UV densities computed for texture streaming
-		else if (ViewModeIndex == VMI_MeshUVDensityAccuracy)
-		{
-			ViewModeDisplayIcons.Emplace(FAppStyle::Get().GetBrush("EditorViewport.TexStreamAccMeshUVDensityMode"));
-		}
-		// Colored according to shader complexity, including quad overdraw
-		else if (ViewModeIndex == VMI_ShaderComplexityWithQuadOverdraw)
-		{
-			ViewModeDisplayIcons.Emplace(FAppStyle::Get().GetBrush("EditorViewport.ShaderComplexityWithQuadOverdrawMode"));
-		}
-		// Colored according to the current HLOD index
-		else if (ViewModeIndex == VMI_HLODColoration)
-		{
-			ViewModeDisplayIcons.Emplace(FAppStyle::Get().GetBrush("EditorViewport.LOD"));
-		}
-		// Group item for LOD and HLOD coloration
-		else if (ViewModeIndex == VMI_GroupLODColoration)
-		{
-			ViewModeDisplayIcons.Emplace(FAppStyle::Get().GetBrush("EditorViewport.LOD"));
-		}
-		// Visualize the accuracy of the material texture scales used for texture streaming
-		else if (ViewModeIndex == VMI_MaterialTextureScaleAccuracy)
-		{
-			ViewModeDisplayIcons.Emplace(FAppStyle::Get().GetBrush("EditorViewport.TexStreamAccMaterialTextureScaleMode"));
-		}
-		// Compare the required texture resolution to the actual resolution
-		else if (ViewModeIndex == VMI_RequiredTextureResolution)
-		{
-			ViewModeDisplayIcons.Emplace(FAppStyle::Get().GetBrush("EditorViewport.RequiredTextureResolutionMode"));
-		}
-		// Visualize Skin Cache
-		else if (ViewModeIndex == VMI_VisualizeGPUSkinCache)
-		{
-			ViewModeDisplayIcons.Emplace(FAppStyle::Get().GetBrush("EditorViewport.VisualizeGPUSkinCacheMode"));
-		}
-		else if (ViewModeIndex == VMI_LWCComplexity)
-		{
-			ViewModeDisplayIcons.Emplace(FAppStyle::Get().GetBrush("EditorViewport.LWCComplexityMode"));
-		}
-
-		// Ray tracing modes
-		// Run path tracing pipeline
-		else if (ViewModeIndex == VMI_PathTracing)
-		{
-			ViewModeDisplayIcons.Emplace(FAppStyle::Get().GetBrush("EditorViewport.PathTracingMode"));
-		}
-		// Run ray tracing debug pipeline 
-		else if (ViewModeIndex == VMI_RayTracingDebug)
-		{
-			ViewModeDisplayIcons.Emplace(FAppStyle::Get().GetBrush("EditorViewport.RayTracingDebugMode"));
-		}
-
-		// Compare the required texture resolution to the actual resolution
-		else if (ViewModeIndex == VMI_VirtualTexturePendingMips)
-		{
-			ViewModeDisplayIcons.Emplace(FAppStyle::Get().GetBrush("EditorViewport.VirtualTexturePendingMipsMode"));
-		}
-
-		// Actor Coloration
-		else if (ViewModeIndex == VMI_VisualizeActorColoration)
-		{
-			ViewModeDisplayIcons.Emplace(FAppStyle::Get().GetBrush("EditorViewport.VisualizeActorColorationMode"));
-		}
-
-		// VMI_Max
-		else if (ViewModeIndex == VMI_Max)
-		{
-			ViewModeDisplayIcons.Emplace(FAppStyle::Get().GetBrush("EditorViewport.LitMode"));
-		}
-		// VMI_Unknown
-		else 
-		{
-			ViewModeDisplayIcons.Emplace(FStyleDefaults::GetNoBrush());
-		}
+		ViewModeDisplayIcons.Emplace(FAppStyle::Get().GetBrush(GetViewModeIconName(ViewModeIndex)));
 	}
 
 	return ViewModeDisplayIcons;
@@ -424,7 +385,7 @@ TArray<const FSlateBrush*> FillViewModeDisplayIcons()
 
 static TArray<const FSlateBrush*> GViewModeDisplayIcons;
 
-const FSlateBrush* UViewModeUtils::GetViewModeDisplayIcon(const EViewModeIndex ViewModeIndex) 
+const FSlateBrush* UViewModeUtils::GetViewModeDisplayIcon(const EViewModeIndex ViewModeIndex)
 {
 	static bool IconsInitialized = false;
 	if (!IconsInitialized)
@@ -436,4 +397,9 @@ const FSlateBrush* UViewModeUtils::GetViewModeDisplayIcon(const EViewModeIndex V
 	return GViewModeDisplayIcons[ViewModeIndex];
 }
 
+FSlateIcon UViewModeUtils::GetViewModeDisplaySlateIcon(const EViewModeIndex InViewModeIndex)
+{
+	return FSlateIcon(FAppStyle::GetAppStyleSetName(), GetViewModeIconName(InViewModeIndex));
+}
+
 #undef LOCTEXT_NAMESPACE
