Added possibility for ILevelInstanceInterface implemented to adjust the position of a newly created LevelInstanceActor before we the relative offsets of the containing actors are calculated.

diff --git a/Engine/Source/Runtime/Engine/Private/LevelInstance/LevelInstanceSubsystem.cpp b/Engine/Source/Runtime/Engine/Private/LevelInstance/LevelInstanceSubsystem.cpp
--- a/Engine/Source/Runtime/Engine/Private/LevelInstance/LevelInstanceSubsystem.cpp
+++ b/Engine/Source/Runtime/Engine/Private/LevelInstance/LevelInstanceSubsystem.cpp
@@ -1526,6 +1526,7 @@ ILevelInstanceInterface* ULevelInstanceSubsystem::CreateLevelInstanceFrom(const
 	ILevelInstanceInterface* NewLevelInstance = CastChecked<ILevelInstanceInterface>(NewLevelInstanceActor);
 	NewLevelInstance->SetWorldAsset(WorldPtr);
 	NewLevelInstanceActor->SetActorLocation(LevelInstanceLocation);
+	NewLevelInstance->AdjustPivotOnCreation();
 	NewLevelInstanceActor->SetActorLabel(WorldPtr.GetAssetName());
 
 	// Actors were moved and kept their World positions so when saving we want their positions to actually be relative to the LevelInstance Actor
diff --git a/Engine/Source/Runtime/Engine/Public/LevelInstance/LevelInstanceInterface.h b/Engine/Source/Runtime/Engine/Public/LevelInstance/LevelInstanceInterface.h
--- a/Engine/Source/Runtime/Engine/Public/LevelInstance/LevelInstanceInterface.h
+++ b/Engine/Source/Runtime/Engine/Public/LevelInstance/LevelInstanceInterface.h
@@ -58,6 +58,8 @@ class ILevelInstanceInterface
 	virtual ELevelInstanceRuntimeBehavior GetDefaultRuntimeBehavior() const = 0;
 
 	virtual TSubclassOf<AActor> GetEditorPivotClass() const { return nullptr; }
+	
+	virtual void AdjustPivotOnCreation() { }
 
 	ENGINE_API virtual bool SupportsPartialEditorLoading() const;
 #endif
