Fixed FSkinnedMeshSceneProxyDesc::GetUsedMaterials to use the same materials as the associated component.

diff --git a/Engine/Source/Runtime/Engine/Private/SkinnedMeshSceneProxyDesc.cpp b/Engine/Source/Runtime/Engine/Private/SkinnedMeshSceneProxyDesc.cpp
--- a/Engine/Source/Runtime/Engine/Private/SkinnedMeshSceneProxyDesc.cpp
+++ b/Engine/Source/Runtime/Engine/Private/SkinnedMeshSceneProxyDesc.cpp
@@ -8,6 +8,7 @@
 #include "SkeletalRenderStatic.h"
 #include "Engine/SkinnedAssetCommon.h"
 #include "Engine/SkinnedAsset.h"
+#include "PhysicsEngine/PhysicsAsset.h"
 #include "Components/SkinnedMeshComponent.h"
 #include "SceneInterface.h"
 #include "GPUSkinCache.h"
@@ -149,20 +150,30 @@ void FSkinnedMeshSceneProxyDesc::GetUsedMaterials(TArray<UMaterialInterface*>& O
 {
 	if (GetSkinnedAsset())
 	{
-		const TArray<FSkeletalMaterial>& Materials = GetSkinnedAsset()->GetMaterials();
-		OutMaterials.Reserve( Materials.Num());
-		for(int32 Idx = 0; Idx < Materials.Num(); ++Idx)
+		// The max number of materials used is the max of the materials on the skeletal mesh and the materials on the mesh component
+		const int32 NumMaterials = FMath::Max(GetSkinnedAsset()->GetMaterials().Num(), OverrideMaterials.Num() );
+		for( int32 MatIdx = 0; MatIdx < NumMaterials; ++MatIdx )
 		{
-			if(OverrideMaterials.IsValidIndex(Idx) && OverrideMaterials[Idx] != nullptr)
-			{
-				OutMaterials.Add(OverrideMaterials[Idx]);
-			}
-			else
-			{
-				const FSkeletalMaterial& Material = Materials[Idx];
-				OutMaterials.Add(Material.MaterialInterface);
-			}
+			// GetMaterial will determine the correct material to use for this index.
+			UMaterialInterface* MaterialInterface = GetMaterial( MatIdx );
+			OutMaterials.Add( MaterialInterface );
 		}
+
+		UMaterialInterface* OverlayMaterialInterface = GetOverlayMaterial();
+		if (OverlayMaterialInterface != nullptr)
+		{
+			OutMaterials.Add(OverlayMaterialInterface);
+		}
+	}
+
+	if (bGetDebugMaterials)
+	{
+#if WITH_EDITOR
+		if (UPhysicsAsset* PhysicsAssetForDebug = GetPhysicsAsset())
+		{
+			PhysicsAssetForDebug->GetUsedMaterials(OutMaterials);
+		}
+#endif
 	}
 }
 
