- Change memory mode capture to 1 as 2 is overwhelming the CSV stats

#rnx

diff --git a/Engine/Source/Runtime/Engine/Private/Particles/ParticlePerfStats.cpp b/Engine/Source/Runtime/Engine/Private/Particles/ParticlePerfStats.cpp
--- a/Engine/Source/Runtime/Engine/Private/Particles/ParticlePerfStats.cpp
+++ b/Engine/Source/Runtime/Engine/Private/Particles/ParticlePerfStats.cpp
@@ -1160,7 +1160,7 @@ static FAutoConsoleVariable CVarWriteDetailedCSVStats(
 	ECVF_Default | ECVF_RenderThreadSafe
 );
 
-int32 GFXDetailedCSVMemoryMode = 2;
+int32 GFXDetailedCSVMemoryMode = 1;
 static FAutoConsoleVariableRef CVarFXDetailedCSVMemoryStats(
 	TEXT("fx.DetailedCSVStats.MemoryMode"),
 	GFXDetailedCSVMemoryMode,
