Fixed missing engine materials in cloth painting skeletal mesh.

diff --git a/Engine/Source/Runtime/Engine/Private/SkinnedMeshSceneProxyDesc.cpp b/Engine/Source/Runtime/Engine/Private/SkinnedMeshSceneProxyDesc.cpp
--- a/Engine/Source/Runtime/Engine/Private/SkinnedMeshSceneProxyDesc.cpp
+++ b/Engine/Source/Runtime/Engine/Private/SkinnedMeshSceneProxyDesc.cpp
@@ -6,10 +6,13 @@
 #include "SkeletalRenderGPUSkin.h"
 #include "SkeletalRenderNanite.h"
 #include "SkeletalRenderStatic.h"
+#include "Materials/MaterialInstanceDynamic.h"
+#include "Engine/Engine.h"
 #include "Engine/SkinnedAssetCommon.h"
 #include "Engine/SkinnedAsset.h"
 #include "PhysicsEngine/PhysicsAsset.h"
 #include "Components/SkinnedMeshComponent.h"
+#include "Animation/DebugSkelMeshComponent.h"
 #include "SceneInterface.h"
 #include "GPUSkinCache.h"
 
@@ -35,6 +38,9 @@ void FSkinnedMeshSceneProxyDesc::InitializeFrom(const USkinnedMeshComponent* InC
 	bCPUSkinning = InComponent->ShouldCPUSkin();
 #if UE_ENABLE_DEBUG_DRAWING
 	DebugDrawColor = InComponent->GetDebugDrawColor();
+#endif
+#if WITH_EDITORONLY_DATA
+	bClothPainting = InComponent->IsA<UDebugSkelMeshComponent>();
 #endif
 	MaterialRelevance = InComponent->GetMaterialRelevance(ERHIFeatureLevel::SM6);
 	StreamingDistanceMultiplier = InComponent->StreamingDistanceMultiplier;
@@ -174,6 +180,14 @@ void FSkinnedMeshSceneProxyDesc::GetUsedMaterials(TArray<UMaterialInterface*>& O
 			PhysicsAssetForDebug->GetUsedMaterials(OutMaterials);
 		}
 #endif
+
+#if WITH_EDITORONLY_DATA
+		if (bClothPainting)
+		{
+			OutMaterials.Add(GEngine->ClothPaintMaterialInstance);
+			OutMaterials.Add(GEngine->ClothPaintMaterialWireframeInstance);
+		}
+#endif
 	}
 }
 
diff --git a/Engine/Source/Runtime/Engine/Public/SkinnedMeshSceneProxyDesc.h b/Engine/Source/Runtime/Engine/Public/SkinnedMeshSceneProxyDesc.h
--- a/Engine/Source/Runtime/Engine/Public/SkinnedMeshSceneProxyDesc.h
+++ b/Engine/Source/Runtime/Engine/Public/SkinnedMeshSceneProxyDesc.h
@@ -20,7 +20,7 @@ struct FMeshDeformerInstanceSet;
 struct ENGINE_API FSkinnedMeshSceneProxyDesc : public FPrimitiveSceneProxyDesc
 {
 	static FSkeletalMeshObject* CreateMeshObject(const FSkinnedMeshSceneProxyDesc& Desc);
-	
+
 	FSkinnedMeshSceneProxyDesc() = default;
 	FSkinnedMeshSceneProxyDesc(const USkinnedMeshComponent* Component);
 	void InitializeFrom(const USkinnedMeshComponent*);
@@ -74,7 +74,7 @@ struct ENGINE_API FSkinnedMeshSceneProxyDesc : public FPrimitiveSceneProxyDesc
 	float StreamingDistanceMultiplier = 1.0f;
 	float CapsuleIndirectShadowMinVisibility = 0.0f;
 	float OverlayMaterialMaxDrawDistance = 0.0f;
-				
+
 	uint8 bForceWireframe : 1 = false;
 	uint8 bCanHighlightSelectedSections : 1 = false;
 	uint8 bRenderStatic : 1 = false;
@@ -83,9 +83,12 @@ struct ENGINE_API FSkinnedMeshSceneProxyDesc : public FPrimitiveSceneProxyDesc
 	uint8 bCastCapsuleIndirectShadow : 1 = false;
 	uint8 bDrawDebugSkeleton : 1 = false;
 	uint8 bCPUSkinning : 1 = false;
+#if WITH_EDITORONLY_DATA
+	uint8 bClothPainting : 1 = false;
+#endif
 
 	/* SkeletalMesh Archetype Data End */
-	
+
 #if WITH_EDITORONLY_DATA
 	int32 GetSectionPreview() const { return SectionIndexPreview;  }
 	int32 GetMaterialPreview() const { return MaterialIndexPreview; }
