Fixed build break

diff --git a/Engine/Source/Runtime/Engine/Private/SkinnedMeshSceneProxyDesc.cpp b/Engine/Source/Runtime/Engine/Private/SkinnedMeshSceneProxyDesc.cpp
--- a/Engine/Source/Runtime/Engine/Private/SkinnedMeshSceneProxyDesc.cpp
+++ b/Engine/Source/Runtime/Engine/Private/SkinnedMeshSceneProxyDesc.cpp
@@ -12,10 +12,13 @@
 #include "Engine/SkinnedAsset.h"
 #include "PhysicsEngine/PhysicsAsset.h"
 #include "Components/SkinnedMeshComponent.h"
-#include "Animation/DebugSkelMeshComponent.h"
 #include "SceneInterface.h"
 #include "GPUSkinCache.h"
 
+#if WITH_EDITORONLY_DATA
+#include "Animation/DebugSkelMeshComponent.h"
+#endif
+
 extern ENGINE_API int32 GEnableGPUSkinCache;
 
 FSkinnedMeshSceneProxyDesc::FSkinnedMeshSceneProxyDesc(const USkinnedMeshComponent* Component)
