Workaround crashes caused by having 2 frames of overlap between the game and render threads.
 - The sync fixes in CL 36468181 have exposed an issue with how FDeferredCleanupInterface objects are deleted. They are deleted too soon, causing crashes in various FRenderResource destructors.
 - Temporarily disabling the 2 frame overlap to work around the problem while we investigate. This may affect performance as we'll be more agressively syncing the game and RHI threads to have only a single frame of overlap.

diff --git a/Engine/Source/Runtime/Engine/Private/UnrealEngine.cpp b/Engine/Source/Runtime/Engine/Private/UnrealEngine.cpp
--- a/Engine/Source/Runtime/Engine/Private/UnrealEngine.cpp
+++ b/Engine/Source/Runtime/Engine/Private/UnrealEngine.cpp
@@ -13371,7 +13371,8 @@ void FFrameEndSync::Sync()
 
 		switch (GTSyncType)
 		{
-		case 0:  SyncDepth = ESyncDepth::RHIThread; NumFramesOverlap = bHasRHIThread ? 2 : 1; break;
+		// HACK - only allow 1 frame of overlap while we investigate crashes caused by the GetPendingCleanupObjects() cleanup mechanism
+		case 0:  SyncDepth = ESyncDepth::RHIThread; NumFramesOverlap = 1/*bHasRHIThread ? 2 : 1*/; break;
 		case 1:  SyncDepth = ESyncDepth::RHIThread; NumFramesOverlap = 1; break;
 		case 2:  SyncDepth = ESyncDepth::Swapchain; NumFramesOverlap = 1; break;
 		default: SyncDepth = ESyncDepth::Swapchain; NumFramesOverlap = 1; break;
