[Engine]
* Fixed bug in assignment operator of SceneViewOwner.. it didn't set ActorUniqueId

diff --git a/Engine/Source/Runtime/Engine/Private/SceneView.cpp b/Engine/Source/Runtime/Engine/Private/SceneView.cpp
--- a/Engine/Source/Runtime/Engine/Private/SceneView.cpp
+++ b/Engine/Source/Runtime/Engine/Private/SceneView.cpp
@@ -768,6 +768,14 @@ FSceneViewOwner::FSceneViewOwner(const AActor* InActor)
 {
 }
 
+#if !UE_STATE_STREAM
+void FSceneViewOwner::operator=(const AActor* InActor)
+{
+	Actor = InActor;
+	ActorUniqueId = InActor ? InActor->GetUniqueID() : 0;
+}
+#endif
+
 static void SetupViewFrustum(FSceneView& View)
 {
 	if (View.SceneViewInitOptions.OverrideFarClippingPlaneDistance > 0.0f)
diff --git a/Engine/Source/Runtime/Engine/Public/SceneViewOwner.h b/Engine/Source/Runtime/Engine/Public/SceneViewOwner.h
--- a/Engine/Source/Runtime/Engine/Public/SceneViewOwner.h
+++ b/Engine/Source/Runtime/Engine/Public/SceneViewOwner.h
@@ -19,7 +19,7 @@ struct FSceneViewOwner
 	const AActor* Get() const { return Actor; }
 	operator const AActor*() const { return Actor; }
 	const AActor* operator->() const { return Actor; }
-	void operator=(const AActor* InActor) { Actor = InActor; }
+	ENGINE_API void operator=(const AActor* InActor);
 	friend bool operator==(const FSceneViewOwner& A, const AActor* B) { return A.Actor == B; }
 	const AActor* Actor;
 	#endif
