Add support for replicating entity destruction properly to the client.

diff --git a/Engine/Source/Runtime/Engine/Classes/Engine/Level.h b/Engine/Source/Runtime/Engine/Classes/Engine/Level.h
--- a/Engine/Source/Runtime/Engine/Classes/Engine/Level.h
+++ b/Engine/Source/Runtime/Engine/Classes/Engine/Level.h
@@ -842,6 +842,9 @@ public:
 
 	const TArray<FReplicatedStaticActorDestructionInfo>& GetDestroyedReplicatedStaticActors() const;
 
+	//** DO NOT USE - Part of internal networking implentation iteration. */
+	ENGINE_API void AddReplicatedDestructionInfo(AActor* const Actor);
+
 	/** Called when a level package has been dirtied. */
 	ENGINE_API static FSimpleMulticastDelegate LevelDirtiedEvent;
 
diff --git a/Engine/Source/Runtime/Engine/Classes/Interfaces/ReplicatedObjectInterface.h b/Engine/Source/Runtime/Engine/Classes/Interfaces/ReplicatedObjectInterface.h
new file mode 100644
--- /dev/null
+++ b/Engine/Source/Runtime/Engine/Classes/Interfaces/ReplicatedObjectInterface.h
@@ -0,0 +1,29 @@
+// Copyright Epic Games, Inc. All Rights Reserved.
+
+#pragma once
+
+#include "CoreMinimal.h"
+#include "UObject/ObjectMacros.h"
+#include "Templates/SubclassOf.h"
+#include "Templates/Casts.h"
+#include "UObject/Interface.h"
+#include "ReplicatedObjectInterface.generated.h"
+
+/** Interface for custom replicated objects. WARNING: Experimental. This interface may change or be removed without notice. **/
+UINTERFACE(MinimalApi, Experimental, meta = (CannotImplementInterfaceInBlueprint))
+class UReplicatedObjectInterface : public UInterface
+{
+	GENERATED_UINTERFACE_BODY()
+};
+
+class IReplicatedObjectInterface
+{
+	GENERATED_IINTERFACE_BODY()
+
+	ENGINE_API virtual bool ShouldClientDestroyWhenNotRelevant() { return false; }
+	ENGINE_API virtual void OnLostRelevancy() {}
+	ENGINE_API virtual bool IsStaticReplicator() { return false; }
+};
+
+
+
diff --git a/Engine/Source/Runtime/Engine/Private/DataChannel.cpp b/Engine/Source/Runtime/Engine/Private/DataChannel.cpp
--- a/Engine/Source/Runtime/Engine/Private/DataChannel.cpp
+++ b/Engine/Source/Runtime/Engine/Private/DataChannel.cpp
@@ -14,6 +14,7 @@
 #if UE_WITH_IRIS
 #include "Iris/ReplicationSystem/ReplicationSystem.h"
 #endif
+#include "Interfaces/ReplicatedObjectInterface.h"
 #include "Misc/MemStack.h"
 #include "Misc/ScopeExit.h"
 #include "Net/Core/Trace/Private/NetTraceInternal.h"
@@ -2612,7 +2613,23 @@ bool UActorChannel::CleanUp(const bool bForDestroy, EChannelCloseReason CloseRea
 		check(Actor == NULL || Actor->IsValidLowLevel());
 		checkSlow(Connection->IsValidLowLevel());
 		checkSlow(Connection->Driver->IsValidLowLevel());
-		if (Actor != NULL)
+
+		// Allow actors to customize behavior on relevancy changes.
+		bool bHandled = false;
+		if (IReplicatedObjectInterface* ReplicatedObject = Cast<IReplicatedObjectInterface>(Actor))
+		{
+			if (CloseReason == EChannelCloseReason::Relevancy)
+			{
+				ReplicatedObject->OnLostRelevancy();
+				if (!ReplicatedObject->ShouldClientDestroyWhenNotRelevant())
+				{
+					UE_LOG(LogNetTraffic, VeryVerbose, TEXT("UActorChannel::CleanUp: Closed for ReplicatedObject Relevancy. %s - [%s] FullNetGUIDPath: %s"), *Describe(), *ActorNetGUID.ToString(), *Connection->Driver->GuidCache->FullNetGUIDPath(ActorNetGUID));
+					bHandled = true;
+				}
+			}
+		}
+
+		if (!bHandled && Actor != nullptr)
 		{
 			if (CloseReason == EChannelCloseReason::TearOff && !Actor->GetTearOff())
 			{
diff --git a/Engine/Source/Runtime/Engine/Private/Level.cpp b/Engine/Source/Runtime/Engine/Private/Level.cpp
--- a/Engine/Source/Runtime/Engine/Private/Level.cpp
+++ b/Engine/Source/Runtime/Engine/Private/Level.cpp
@@ -863,17 +863,22 @@ void ULevel::CreateReplicatedDestructionInfo(AActor* const Actor)
 
 	if (bShouldCreateDestructionInfo)
 	{
-		FReplicatedStaticActorDestructionInfo NewInfo;
-		NewInfo.PathName = Actor->GetFName();
-		NewInfo.FullName = Actor->GetFullName();
-		NewInfo.DestroyedPosition = Actor->GetActorLocation();
-		NewInfo.ObjOuter = Actor->GetOuter();
-		NewInfo.ObjClass = Actor->GetClass();
-
-		DestroyedReplicatedStaticActors.Add(NewInfo);
+		AddReplicatedDestructionInfo(Actor);
 	}
 }
 
+void ULevel::AddReplicatedDestructionInfo(AActor* const Actor)
+{
+	FReplicatedStaticActorDestructionInfo NewInfo;
+	NewInfo.PathName = Actor->GetFName();
+	NewInfo.FullName = Actor->GetFullName();
+	NewInfo.DestroyedPosition = Actor->GetActorLocation();
+	NewInfo.ObjOuter = Actor->GetOuter();
+	NewInfo.ObjClass = Actor->GetClass();
+
+	DestroyedReplicatedStaticActors.Add(NewInfo);
+}
+
 const TArray<FReplicatedStaticActorDestructionInfo>& ULevel::GetDestroyedReplicatedStaticActors() const
 {
 	return DestroyedReplicatedStaticActors;
diff --git a/Engine/Source/Runtime/Engine/Private/Net/ReplicatedObjectInterface.cpp b/Engine/Source/Runtime/Engine/Private/Net/ReplicatedObjectInterface.cpp
new file mode 100644
--- /dev/null
+++ b/Engine/Source/Runtime/Engine/Private/Net/ReplicatedObjectInterface.cpp
@@ -0,0 +1,10 @@
+// Copyright Epic Games, Inc. All Rights Reserved.
+
+#include "Interfaces/ReplicatedObjectInterface.h"
+#include "UObject/Interface.h"
+
+
+UReplicatedObjectInterface::UReplicatedObjectInterface(const FObjectInitializer& ObjectInitializer)
+	: Super(ObjectInitializer)
+{
+}
diff --git a/Engine/Source/Runtime/Engine/Private/NetDriver.cpp b/Engine/Source/Runtime/Engine/Private/NetDriver.cpp
--- a/Engine/Source/Runtime/Engine/Private/NetDriver.cpp
+++ b/Engine/Source/Runtime/Engine/Private/NetDriver.cpp
@@ -67,6 +67,8 @@
 #include "Iris/ReplicationSystem/StringTokenStore.h"
 #include "GameplayTagTokenStore.h"
 #endif // UE_WITH_IRIS
+#include "Interfaces/ReplicatedObjectInterface.h"
+
 
 #if USE_SERVER_PERF_COUNTERS
 #endif
@@ -3977,8 +3979,13 @@ void UNetDriver::NotifyActorDestroyed(AActor* ThisActor, bool IsSeamlessTravel)
 	if (bIsServer)
 	{
 		FActorDestructionInfo* DestructionInfo = nullptr;
+		bool IsStaticReplicator = false;
+		if (IReplicatedObjectInterface* ReplicatedObject = Cast<IReplicatedObjectInterface>(ThisActor))
+		{
+			IsStaticReplicator = ReplicatedObject->IsStaticReplicator();
+		}
 
-		const bool bIsActorStatic = !GuidCache->IsDynamicObject( ThisActor );
+		const bool bIsActorStatic = !GuidCache->IsDynamicObject( ThisActor ) || IsStaticReplicator;
 		const bool bActorHasRole = ThisActor->GetRemoteRole() != ROLE_None;
 		const bool bShouldCreateDestructionInfo = bIsServer && bIsActorStatic && bActorHasRole && !IsSeamlessTravel && !GIsReconstructingBlueprintInstances && !UE::Net::ShouldIgnoreStaticActorDestruction();
 
@@ -5079,6 +5086,7 @@ void UNetDriver::ServerReplicateActors_BuildConsiderList( TArray<FNetworkObjectI
 
 		if ( Actor->GetRemoteRole() == ROLE_None )
 		{
+			UE_LOG(LogNetTraffic, VeryVerbose, TEXT("Skipping Actor %s - it has no remote role!"), *Actor->GetName());
 			ActorsToRemove.Add( Actor );
 			continue;
 		}
@@ -5096,6 +5104,7 @@ void UNetDriver::ServerReplicateActors_BuildConsiderList( TArray<FNetworkObjectI
 		// Verify the actor is actually initialized (it might have been intentionally spawn deferred until a later frame)
 		if ( !Actor->IsActorInitialized() )
 		{
+			UE_LOG(LogNetTraffic, VeryVerbose, TEXT("Skipping Actor %s - it not initialized!"), *Actor->GetName());
 			continue;
 		}
 
@@ -5103,6 +5112,7 @@ void UNetDriver::ServerReplicateActors_BuildConsiderList( TArray<FNetworkObjectI
 		ULevel* Level = Actor->GetLevel();
 		if ( Level->HasVisibilityChangeRequestPending() || Level->bIsAssociatingLevel )
 		{
+			UE_LOG(LogNetTraffic, VeryVerbose, TEXT("Skipping Actor %s - it has a pending visibility change!"), *Actor->GetName());
 			continue;
 		}
 
@@ -5113,7 +5123,7 @@ void UNetDriver::ServerReplicateActors_BuildConsiderList( TArray<FNetworkObjectI
 			SCOPE_CYCLE_COUNTER( STAT_NetInitialDormantCheckTime );
 			NumInitiallyDormant++;
 			ActorsToRemove.Add( Actor );
-			//UE_LOG(LogNetTraffic, Log, TEXT("Skipping Actor %s - its initially dormant!"), *Actor->GetName() );
+			UE_LOG(LogNetTraffic, VeryVerbose, TEXT("Skipping Actor %s - its initially dormant!"), *Actor->GetName() );
 			continue;
 		}
 
@@ -5123,6 +5133,7 @@ void UNetDriver::ServerReplicateActors_BuildConsiderList( TArray<FNetworkObjectI
 		// Set defaults if this actor is replicating for first time
 		if ( ActorInfo->LastNetReplicateTime == 0 )
 		{
+			UE_LOG(LogNetTraffic, VeryVerbose, TEXT("Setting Actor %s - initial replication update time!"), *Actor->GetName());
 			ActorInfo->LastNetReplicateTime = World->TimeSeconds;
 			ActorInfo->OptimalNetUpdateDelta = 1.0f / Actor->GetNetUpdateFrequency();
 		}
@@ -5673,6 +5684,8 @@ int64 UNetDriver::SendDestructionInfo(UNetConnection* Connection, FActorDestruct
 	checkf(Connection, TEXT("SendDestructionInfo called with invalid connection: %s"), *GetDescription());
 	checkf(DestructionInfo, TEXT("SendDestructionInfo called with invalid destruction info: %s"), *GetDescription());
 
+	UE_LOG(LogNetTraffic, VeryVerbose, TEXT("Sending Destruction Info %s"), *GetDescription());
+
 	FScopedRepContext RepContext(Connection, DestructionInfo->Level.Get());
 
 	if (GNetControlChannelDestructionInfo != 0 && !Connection->IsReplay())
