Allow for additional material shader tags to be applied to a material resource.  These tags can be used to limit the shader permutations generated by testing if the material contains a specific shader tag.
These tags can be added by either overriding GetShaderTags on a material expression or by overriding GetShaderTags on a custom FMaterialResource.

Overriding GetShaderTags on a material expression will add the tags to any material resource that adds this material expression.
Overriding GetShaderTags on a custom FMaterialResource allows for complete customization of the shader tags for that resource.

diff --git a/Engine/Source/Runtime/Engine/Private/Materials/Material.cpp b/Engine/Source/Runtime/Engine/Private/Materials/Material.cpp
--- a/Engine/Source/Runtime/Engine/Private/Materials/Material.cpp
+++ b/Engine/Source/Runtime/Engine/Private/Materials/Material.cpp
@@ -2310,6 +2310,18 @@ void UMaterial::SetRenderTracePhysicalMaterialOutputs(TArrayView<TObjectPtr<clas
 	RenderTracePhysicalMaterialOutputs = PhysicalMaterials;
 }
 
+#if WITH_EDITOR
+void UMaterial::GetShaderTags(TArray<FName>& OutShaderTags) const
+{
+	if (!CachedExpressionData || !CachedExpressionData->EditorOnlyData)
+	{
+		return;
+	}
+
+	OutShaderTags.Append(CachedExpressionData->EditorOnlyData->ShaderTags);
+}
+#endif
+
 UMaterialInterface* UMaterial::GetNaniteOverride(TMicRecursionGuard RecursionGuard) const
 {
 	return NaniteOverrideMaterial.GetOverrideMaterial();
diff --git a/Engine/Source/Runtime/Engine/Private/Materials/MaterialCachedData.cpp b/Engine/Source/Runtime/Engine/Private/Materials/MaterialCachedData.cpp
--- a/Engine/Source/Runtime/Engine/Private/Materials/MaterialCachedData.cpp
+++ b/Engine/Source/Runtime/Engine/Private/Materials/MaterialCachedData.cpp
@@ -342,6 +342,11 @@ void FMaterialCachedExpressionData::UpdateForExpressions(const FMaterialCachedEx
 		UObject* ReferencedTexture = nullptr;
 		UTextureCollection* ReferencedTextureCollection = nullptr;
 
+		// Add any expression specific custom shader tags
+		TArray<FName> ShaderTags;
+		Expression->GetShaderTags(ShaderTags);
+		EditorOnlyData->ShaderTags.Append(ShaderTags);
+
 		FMaterialParameterMetadata ParameterMeta;
 		if (Expression->GetParameterValue(ParameterMeta))
 		{
diff --git a/Engine/Source/Runtime/Engine/Private/Materials/MaterialShader.cpp b/Engine/Source/Runtime/Engine/Private/Materials/MaterialShader.cpp
--- a/Engine/Source/Runtime/Engine/Private/Materials/MaterialShader.cpp
+++ b/Engine/Source/Runtime/Engine/Private/Materials/MaterialShader.cpp
@@ -248,7 +248,10 @@ void RecordOrEmitMaterialShaderMapKey(FMaterialKeyGeneratorContext& Context,
 		Context.RecordObjectStart("ShaderMapId");
 		ShaderMapId.RecordAndEmit(Context);
 		Context.RecordObjectEnd();
-		Context.Record("ShaderParameters", ShaderParameters);
+
+		Context.RecordObjectStart("ShaderParameters");
+		ShaderParameters.Record(Context);
+		Context.RecordObjectEnd();
 	}
 	else
 	{
@@ -2849,6 +2852,19 @@ void FMaterialShaderMap::ProcessCompilationResults(const TArray<FShaderCommonCom
 
 #endif // WITH_EDITOR
 
+// TODO - Move to header before submit
+uint64 GetMaterialShaderParametersHash(const FMaterialShaderParameters& MaterialParameters, uint64 Seed)
+{
+	if (MaterialParameters.MaterialShaderTags.Num())
+	{
+		uint32 ShaderTagsHash = GetArrayHash<FName>(MaterialParameters.MaterialShaderTags.GetData(), MaterialParameters.MaterialShaderTags.Num());
+		Seed = Seed ^ (uint64)ShaderTagsHash;
+	}
+
+	return CityHash64WithSeed((char*)&static_cast<const FMaterialShaderParametersBase&>(MaterialParameters), sizeof(FMaterialShaderParametersBase), Seed);
+}
+
+
 class FMaterialShaderMapLayoutCache
 {
 public:
@@ -2860,7 +2876,8 @@ public:
 
 	const FMaterialShaderMapLayout& AcquireLayout(EShaderPlatform Platform, EShaderPermutationFlags Flags, const FMaterialShaderParameters& MaterialParameters)
 	{
-		const uint64 ParameterHash = CityHash64WithSeed((char*)&MaterialParameters, sizeof(MaterialParameters), (uint64)Platform | ((uint64)Flags << 32));
+		
+		const uint64 ParameterHash = GetMaterialShaderParametersHash(MaterialParameters, (uint64)Platform | ((uint64)Flags << 32));
 
 		FMaterialShaderMapLayout* Layout = nullptr;
 		{
diff --git a/Engine/Source/Runtime/Engine/Private/Materials/MaterialShared.cpp b/Engine/Source/Runtime/Engine/Private/Materials/MaterialShared.cpp
--- a/Engine/Source/Runtime/Engine/Private/Materials/MaterialShared.cpp
+++ b/Engine/Source/Runtime/Engine/Private/Materials/MaterialShared.cpp
@@ -53,6 +53,8 @@
 #include "Shader/PreshaderEvaluate.h"
 #include "Misc/AsciiSet.h"
 #include "PostProcess/PostProcessMaterialInputs.h"
+#include "MaterialKeyGeneratorContext.h"
+#include "Serialization/CompactBinarySerialization.h"
 #endif
 #if WITH_ODSC
 #include "ODSC/ODSCManager.h"
@@ -2336,6 +2338,12 @@ UMaterialInterface* FMaterialResource::GetMaterialInterface() const
 }
 
 #if WITH_EDITOR
+void FMaterialResource::GetShaderTags(TArray<FName>& OutShaderTags) const
+{
+	const FMaterialCachedExpressionData& CachedExpressionData = GetMaterialInterface()->GetCachedExpressionData();
+	OutShaderTags.Append(CachedExpressionData.EditorOnlyData->ShaderTags);
+}
+
 void FMaterialResource::NotifyCompilationFinished()
 {
 	UMaterial::NotifyCompilationFinished(MaterialInstance ? (UMaterialInterface*)MaterialInstance : (UMaterialInterface*)Material);
@@ -5578,7 +5586,9 @@ bool FMemoryImageMaterialParameterInfo::RemapLayerIndex(TArrayView<const int32>
 	return RemapParameterLayerIndex(IndexRemap, *this, OutResult);
 }
 
-FMaterialShaderParameters::FMaterialShaderParameters(const FMaterial* InMaterial)
+static_assert(TIsTrivial<FMaterialShaderParametersBase>::Value, "FMaterialShaderParametersBase - Must be trivial for serialization and hashing");
+
+FMaterialShaderParametersBase::FMaterialShaderParametersBase(const FMaterial* InMaterial)
 {
 	// zero-initialize is required even when InMaterial!=null so that all bytes in sizeof(*this) are initialized
 	// and we get consistent hashes.
@@ -5676,15 +5686,27 @@ FMaterialShaderParameters::FMaterialShaderParameters(const FMaterial* InMaterial
 	bAllowVariableRateShading = InMaterial->IsVariableRateShadingAllowed();
 }
 
+FMaterialShaderParameters::FMaterialShaderParameters(const FMaterial* InMaterial)
+	: FMaterialShaderParametersBase(InMaterial)
+{
+#if WITH_EDITOR
+	if (InMaterial)
+	{
+		InMaterial->GetShaderTags(MaterialShaderTags);
+	}
+#endif
+}
+
 #if WITH_EDITOR
-void FMaterialShaderParameters::Save(FCbWriter& Writer) const
+void FMaterialShaderParametersBase::Save(FCbWriter& Writer) const
 {
 	Writer.AddBinary(FMemoryView(this, sizeof(*this)));
 }
 
-bool FMaterialShaderParameters::TryLoad(FCbFieldView Field)
+bool FMaterialShaderParametersBase::TryLoad(FCbFieldView Field)
 {
-	*this = FMaterialShaderParameters();
+	*this = FMaterialShaderParametersBase();
+
 	FMemoryView MemoryView = Field.AsBinaryView();
 	if (Field.HasError() || MemoryView.GetSize() != sizeof(*this))
 	{
@@ -5694,11 +5716,17 @@ bool FMaterialShaderParameters::TryLoad(FCbFieldView Field)
 	return true;
 }
 
-bool LoadFromCompactBinary(FCbFieldView Field, FMaterialShaderParameters& OutValue)
+bool LoadFromCompactBinary(FCbFieldView Field, FMaterialShaderParametersBase& OutValue)
 {
 	return OutValue.TryLoad(Field);
 }
 
+void FMaterialShaderParameters::Record(FMaterialKeyGeneratorContext& Context)
+{
+	Context.Record("MaterialShaderParametersBase", static_cast<FMaterialShaderParametersBase&>(*this));
+	Context.Record("MaterialShaderTags", MaterialShaderTags);
+}
+
 int32 FMaterialCompilationOutput::FindOrAddUserSceneTexture(FName UserSceneTexture)
 {
 	FScriptName UserSceneTextureScript(UserSceneTexture);
diff --git a/Engine/Source/Runtime/Engine/Public/MaterialCachedData.h b/Engine/Source/Runtime/Engine/Public/MaterialCachedData.h
--- a/Engine/Source/Runtime/Engine/Public/MaterialCachedData.h
+++ b/Engine/Source/Runtime/Engine/Public/MaterialCachedData.h
@@ -180,6 +180,9 @@ struct FMaterialCachedExpressionEditorOnlyData
 
 	UPROPERTY()
 	TSet<FName> UserSceneTextureInputs;
+
+	UPROPERTY()
+	TArray<FName> ShaderTags;
 };
 
 USTRUCT()
diff --git a/Engine/Source/Runtime/Engine/Public/MaterialShared.h b/Engine/Source/Runtime/Engine/Public/MaterialShared.h
--- a/Engine/Source/Runtime/Engine/Public/MaterialShared.h
+++ b/Engine/Source/Runtime/Engine/Public/MaterialShared.h
@@ -979,7 +979,7 @@ namespace EMaterialShaderMapUsage
 	};
 }
 
-struct FMaterialShaderParameters
+struct FMaterialShaderParametersBase
 {
 	EMaterialDomain MaterialDomain;
 	FMaterialShadingModelField ShadingModels;
@@ -1061,22 +1061,35 @@ struct FMaterialShaderParameters
 		};
 	};
 
-	FMaterialShaderParameters(const FMaterial* InMaterial=nullptr);
+	FMaterialShaderParametersBase(const FMaterial* InMaterial=nullptr);
 
 private:
 #if WITH_EDITOR
 	// Compact binary API with hidden friend operator<<
 	ENGINE_API void Save(FCbWriter& Writer) const;
 	bool TryLoad(FCbFieldView Field);
-	friend inline FCbWriter& operator<<(FCbWriter& Writer, const FMaterialShaderParameters& Value)
+	friend inline FCbWriter& operator<<(FCbWriter& Writer, const FMaterialShaderParametersBase& Value)
 	{
 		Value.Save(Writer);
 		return Writer;
 	}
-	friend ENGINE_API bool LoadFromCompactBinary(FCbFieldView Field, FMaterialShaderParameters& OutValue);
+	friend ENGINE_API bool LoadFromCompactBinary(FCbFieldView Field, FMaterialShaderParametersBase& OutValue);
 #endif
 };
 
+// FMaterialShaderParametersBase split out to ease serializing of raw data
+struct FMaterialShaderParameters : public FMaterialShaderParametersBase
+{
+	FMaterialShaderParameters(const FMaterial* InMaterial = nullptr);
+
+public:
+#if WITH_EDITOR
+	void Record(FMaterialKeyGeneratorContext& Context);
+#endif
+
+	TArray<FName> MaterialShaderTags;
+};
+
 /** Contains all the information needed to uniquely identify a FMaterialShaderMap. */
 class FMaterialShaderMapId
 {
@@ -2253,6 +2266,8 @@ public:
 	virtual EMaterialShadingRate GetShadingRate() const { return MSR_1x1; }
 	virtual int32 GetNeuralProfileId() const { return INDEX_NONE; }
 	virtual bool IsVariableRateShadingAllowed() const { return true; }
+	virtual void GetShaderTags(TArray<FName>& OutShaderTags) const {};
+
 	/**
 	 * Should shaders compiled for this material be saved to disk?
 	 */
@@ -3018,6 +3033,7 @@ public:
 	ENGINE_API virtual uint16 GetPreshaderGap() const override;
 	ENGINE_API virtual UMaterialInterface* GetMaterialInterface() const override;
 	ENGINE_API virtual int32 GetNeuralProfileId() const override;
+
 	/**
 	 * Should shaders compiled for this material be saved to disk?
 	 */
@@ -3025,6 +3041,7 @@ public:
 	ENGINE_API virtual FGuid GetMaterialId() const override;
 
 #if WITH_EDITOR
+	ENGINE_API void GetShaderTags(TArray<FName>& OutShaderTags) const override;
 	ENGINE_API virtual void NotifyCompilationFinished() override;
 	/** Allows to associate the shader resources with the asset for load order. */
 	ENGINE_API virtual FName GetAssetPath() const override;
diff --git a/Engine/Source/Runtime/Engine/Public/Materials/Material.h b/Engine/Source/Runtime/Engine/Public/Materials/Material.h
--- a/Engine/Source/Runtime/Engine/Public/Materials/Material.h
+++ b/Engine/Source/Runtime/Engine/Public/Materials/Material.h
@@ -1923,6 +1923,12 @@ public:
 	 */
 	ENGINE_API void SetRenderTracePhysicalMaterialOutputs(TArrayView<TObjectPtr<UPhysicalMaterial>> PhysicalMaterials);
 
+#if WITH_EDITOR
+	/*
+	*/
+	ENGINE_API void GetShaderTags(TArray<FName>& OutShaderTags) const ;
+#endif
+
 private:
 #if WITH_EDITOR
 	static FMaterialCompilationFinished MaterialCompilationFinishedEvent;
diff --git a/Engine/Source/Runtime/Engine/Public/Materials/MaterialExpression.h b/Engine/Source/Runtime/Engine/Public/Materials/MaterialExpression.h
--- a/Engine/Source/Runtime/Engine/Public/Materials/MaterialExpression.h
+++ b/Engine/Source/Runtime/Engine/Public/Materials/MaterialExpression.h
@@ -336,6 +336,10 @@ class UMaterialExpression : public UObject
 	 */
 	ENGINE_API virtual void Build(UE::MIR::FEmitter& Emitter);
 
+	/* Material expressions can add shader tags that to a material.  Adding a shader tag helps to identify materials
+	   that use certain material expressions, and can be used to limit shader generation to only materials that have
+	   use that material expression*/
+	ENGINE_API virtual void GetShaderTags(TArray<FName>& ShaderTagsOut) {}
 #endif // WITH_EDITOR
 
 	/**
