Removing deprecated code from FStripDataFlags.
Moving member init to the header.

diff --git a/Engine/Source/Runtime/Engine/Private/EngineUtils.cpp b/Engine/Source/Runtime/Engine/Private/EngineUtils.cpp
--- a/Engine/Source/Runtime/Engine/Private/EngineUtils.cpp
+++ b/Engine/Source/Runtime/Engine/Private/EngineUtils.cpp
@@ -583,8 +583,6 @@ void FConsoleOutputDevice::Serialize(const TCHAR* Text, ELogVerbosity::Type Verb
 	Serialized data stripping.
 -----------------------------------------------------------------------------*/
 FStripDataFlags::FStripDataFlags( class FArchive& Ar, uint8 InClassFlags /*= 0*/, const FPackageFileVersion& InVersion /*= GOldestLoadablePackageFileUEVersion */ )
-	: GlobalStripFlags( 0 )
-	, ClassStripFlags( 0 )
 {
 	check(InVersion >= VER_UE4_OLDEST_LOADABLE_PACKAGE);
 	if (Ar.UEVer().IsCompatible(InVersion))
@@ -604,8 +602,6 @@ FStripDataFlags::FStripDataFlags( class FArchive& Ar, uint8 InClassFlags /*= 0*/
 }
 
 FStripDataFlags::FStripDataFlags( class FArchive& Ar, uint8 InGlobalFlags, uint8 InClassFlags, const FPackageFileVersion& InVersion /*= GOldestLoadablePackageFileUEVersion */)
-	: GlobalStripFlags( 0 )
-	, ClassStripFlags( 0 )
 {
 	check(InVersion >= VER_UE4_OLDEST_LOADABLE_PACKAGE);
 	if (Ar.UEVer().IsCompatible(InVersion))
@@ -625,8 +621,6 @@ FStripDataFlags::FStripDataFlags( class FArchive& Ar, uint8 InGlobalFlags, uint8
 Serialized data stripping.
 -----------------------------------------------------------------------------*/
 FStripDataFlags::FStripDataFlags(FStructuredArchive::FSlot Slot, uint8 InClassFlags /*= 0*/, const FPackageFileVersion& InVersion /*= GOldestLoadablePackageFileUEVersion */)
-	: GlobalStripFlags(0)
-	, ClassStripFlags(0)
 {
 	FArchive& UnderlyingArchive = Slot.GetUnderlyingArchive();
 	FStructuredArchive::FRecord Record = Slot.EnterRecord();
@@ -649,8 +643,6 @@ FStripDataFlags::FStripDataFlags(FStructuredArchive::FSlot Slot, uint8 InClassFl
 }
 
 FStripDataFlags::FStripDataFlags(FStructuredArchive::FSlot Slot, uint8 InGlobalFlags, uint8 InClassFlags, const FPackageFileVersion& InVersion /*= GOldestLoadablePackageFileUEVersion */)
-	: GlobalStripFlags(0)
-	, ClassStripFlags(0)
 {
 	FArchive& UnderlyingArchive = Slot.GetUnderlyingArchive();
 	FStructuredArchive::FRecord Record = Slot.EnterRecord();
diff --git a/Engine/Source/Runtime/Engine/Public/EngineUtils.h b/Engine/Source/Runtime/Engine/Public/EngineUtils.h
--- a/Engine/Source/Runtime/Engine/Public/EngineUtils.h
+++ b/Engine/Source/Runtime/Engine/Public/EngineUtils.h
@@ -834,9 +834,10 @@ namespace EngineUtils
 class FStripDataFlags
 {
 	/** Serialized engine strip flags (up to 8 flags). */
-	uint8 GlobalStripFlags;
+	uint8 GlobalStripFlags = 0;
+
 	/** Serialized per-class strip flags (user defined, up to 8 flags). */
-	uint8 ClassStripFlags;
+	uint8 ClassStripFlags = 0;
 
 public:
 
@@ -847,11 +848,9 @@ public:
 
 		/* This flag means that Editor-only data is stripped */
 		EditorOnly = 1,
-		Editor UE_DEPRECATED(5.4, "Use EditorOnly value instead of Editor") = EditorOnly,
 
 		/* This flag means that AudioVisual data is stripped (e.g. target is dedicated server). */
 		AudioVisual = 2,
-		Server UE_DEPRECATED(5.4, "Use AudioVisual value instead of Server") = AudioVisual,
 
 		/* This flag means that all data needed to cook packages will be stripped. What it means is specific to 
 		  an asset type, and it might be a subset of both EditorOnly and AudioVisual. This flag is normally set, except
@@ -874,12 +873,6 @@ public:
 	 * @param InVersion - Minimal strip version required to serialize strip flags
 	 */
 	ENGINE_API explicit FStripDataFlags(FArchive& Ar, uint8 InClassFlags = 0, const FPackageFileVersion& InVersion = GOldestLoadablePackageFileUEVersion);
-	UE_DEPRECATED(5.0, "Use the other overload that takes InVersion as a FPackageFileVersion. See the @FPackageFileVersion documentation for further details")
-	inline FStripDataFlags(FArchive& Ar, uint8 InClassFlags, int32 InVersion)
-		: FStripDataFlags(Ar, InClassFlags, FPackageFileVersion::CreateUE4Version(InVersion))
-	{
-
-	}
 
 	/** 
 	 * Constructor.
@@ -892,12 +885,6 @@ public:
 	 * @param InVersion - Minimal version required to serialize strip flags
 	 */
 	ENGINE_API FStripDataFlags( FArchive& Ar, uint8 InGlobalFlags, uint8 InClassFlags, const FPackageFileVersion& InVersion = GOldestLoadablePackageFileUEVersion);
-	UE_DEPRECATED(5.0, "Use the other overload that takes InVersion as a FPackageFileVersion. See the @FPackageFileVersion documentation for further details")
-	inline FStripDataFlags(FArchive& Ar, uint8 InGlobalFlags, uint8 InClassFlags, int32 InVersion)
-		: FStripDataFlags(Ar, InGlobalFlags, InClassFlags, FPackageFileVersion::CreateUE4Version(InVersion))
-	{
-
-	}
 
 	/**
 	* Constructor.
@@ -909,12 +896,6 @@ public:
 	* @param InVersion - Minimal strip version required to serialize strip flags
 	*/
 	ENGINE_API explicit FStripDataFlags(FStructuredArchive::FSlot Slot, uint8 InClassFlags = 0, const FPackageFileVersion& InVersion = GOldestLoadablePackageFileUEVersion);
-	UE_DEPRECATED(5.0, "Use the other overload that takes InVersion as a FPackageFileVersion. See the @FPackageFileVersion documentation for further details")
-	inline FStripDataFlags(FStructuredArchive::FSlot Slot, uint8 InClassFlags, int32 InVersion)
-		: FStripDataFlags(Slot, InClassFlags, FPackageFileVersion::CreateUE4Version(InVersion))
-	{
-
-	}
 
 	/**
 	* Constructor.
@@ -927,12 +908,6 @@ public:
 	* @param InVersion - Minimal version required to serialize strip flags
 	*/
 	ENGINE_API FStripDataFlags(FStructuredArchive::FSlot Slot, uint8 InGlobalFlags, uint8 InClassFlags, const FPackageFileVersion& InVersion = GOldestLoadablePackageFileUEVersion);
-	UE_DEPRECATED(5.0, "Use the other overload that takes InVersion as a FPackageFileVersion. See the @FPackageFileVersion documentation for further details")
-	inline FStripDataFlags(FStructuredArchive::FSlot Slot, uint8 InGlobalFlags, uint8 InClassFlags, int32 InVersion)
-		: FStripDataFlags(Slot, InGlobalFlags, InClassFlags, FPackageFileVersion::CreateUE4Version(InVersion))
-	{
-
-	}
 
 	/**
 	 * Checks if FStripDataFlags::Editor flag is set or not
@@ -954,17 +929,6 @@ public:
 		return (GlobalStripFlags & static_cast<uint8>(FStripDataFlags::EStrippedData::AudioVisual)) != 0;
 	}
 
-	/**
-	 * Checks if FStripDataFlags::Server flag is set or not
-	 *
-	 * @return true if FStripDataFlags::Server is set, false otherwise.
-	 */
-	UE_DEPRECATED(5.4, "Use IsAudioVisualDataStripped instead.")
-	bool IsDataStrippedForServer() const
-	{
-		return IsAudioVisualDataStripped();
-	}
-
 	/**
 	 * Checks if FStripDataFlags::NeededForCooking flag is set or not. It should be set for all non-content-worker targets
 	 *
