World Partition - Hashset: HLOD grids origin should always match the main layer origin.
-> Origin and bIs2D properties are now hidden when dealing with HLOD grids, as those  should always be infered from the main layer.

diff --git a/Engine/Source/Runtime/Engine/Private/WorldPartition/RuntimeHashSet/RuntimePartitionLHGrid.cpp b/Engine/Source/Runtime/Engine/Private/WorldPartition/RuntimeHashSet/RuntimePartitionLHGrid.cpp
--- a/Engine/Source/Runtime/Engine/Private/WorldPartition/RuntimeHashSet/RuntimePartitionLHGrid.cpp
+++ b/Engine/Source/Runtime/Engine/Private/WorldPartition/RuntimeHashSet/RuntimePartitionLHGrid.cpp
@@ -159,6 +159,7 @@ void URuntimePartitionLHGrid::InitHLODRuntimePartitionFrom(const URuntimePartiti
 	const URuntimePartitionLHGrid* RuntimePartitionLHGrid = CastChecked<const URuntimePartitionLHGrid>(InRuntimePartition);
 	CellSize = RuntimePartitionLHGrid->CellSize * 2;
 	bIs2D = RuntimePartitionLHGrid->bIs2D;
+	Origin = RuntimePartitionLHGrid->Origin;
 }
 
 void URuntimePartitionLHGrid::UpdateHLODRuntimePartitionFrom(const URuntimePartition* InRuntimePartition)
@@ -166,6 +167,7 @@ void URuntimePartitionLHGrid::UpdateHLODRuntimePartitionFrom(const URuntimeParti
 	Super::UpdateHLODRuntimePartitionFrom(InRuntimePartition);
 	const URuntimePartitionLHGrid* RuntimePartitionLHGrid = CastChecked<const URuntimePartitionLHGrid>(InRuntimePartition);
 	bIs2D = RuntimePartitionLHGrid->bIs2D;
+	Origin = RuntimePartitionLHGrid->Origin;
 }
 
 void URuntimePartitionLHGrid::SetDefaultValues()
diff --git a/Engine/Source/Runtime/Engine/Private/WorldPartition/RuntimeHashSet/WorldPartitionRuntimeHashSet.cpp b/Engine/Source/Runtime/Engine/Private/WorldPartition/RuntimeHashSet/WorldPartitionRuntimeHashSet.cpp
--- a/Engine/Source/Runtime/Engine/Private/WorldPartition/RuntimeHashSet/WorldPartitionRuntimeHashSet.cpp
+++ b/Engine/Source/Runtime/Engine/Private/WorldPartition/RuntimeHashSet/WorldPartitionRuntimeHashSet.cpp
@@ -1048,6 +1048,23 @@ void UWorldPartitionRuntimeHashSet::PostDuplicate(EDuplicateMode::Type Duplicate
 	}
 }
 
+void UWorldPartitionRuntimeHashSet::PostLoad()
+{
+	Super::PostLoad();
+
+	// Make sure that HLOD partitions settings are updated to match the main layer settings
+	for (const FRuntimePartitionDesc& RuntimePartition : RuntimePartitions)
+	{
+		for (const FRuntimePartitionHLODSetup& HLODSetup : RuntimePartition.HLODSetups)
+		{
+			if (HLODSetup.PartitionLayer)
+			{
+				HLODSetup.PartitionLayer->UpdateHLODRuntimePartitionFrom(RuntimePartition.MainLayer);
+			}
+		}
+	}
+}
+
 UWorldPartitionRuntimeHashSet::FCellUniqueId UWorldPartitionRuntimeHashSet::GetCellUniqueId(const URuntimePartition::FCellDescInstance& InCellDescInstance) const
 {
 	FCellUniqueId CellUniqueId;
diff --git a/Engine/Source/Runtime/Engine/Public/WorldPartition/RuntimeHashSet/RuntimePartitionLHGrid.h b/Engine/Source/Runtime/Engine/Public/WorldPartition/RuntimeHashSet/RuntimePartitionLHGrid.h
--- a/Engine/Source/Runtime/Engine/Public/WorldPartition/RuntimeHashSet/RuntimePartitionLHGrid.h
+++ b/Engine/Source/Runtime/Engine/Public/WorldPartition/RuntimeHashSet/RuntimePartitionLHGrid.h
@@ -46,10 +46,10 @@ private:
 	UPROPERTY(EditAnywhere, Category = RuntimeSettings)
 	uint32 CellSize = 25600;
 
-	UPROPERTY(EditAnywhere, Category = RuntimeSettings)
+	UPROPERTY(EditAnywhere, Category = RuntimeSettings, meta=(EditCondition="HLODIndex == INDEX_NONE", EditConditionHides))
 	FVector Origin = FVector::ZeroVector;
 
-	UPROPERTY(EditAnywhere, Category = RuntimeSettings)
+	UPROPERTY(EditAnywhere, Category = RuntimeSettings, meta=(EditCondition="HLODIndex == INDEX_NONE", EditConditionHides))
 	bool bIs2D = false;
 
 	UPROPERTY(EditAnywhere, Category = RuntimeSettings, Transient, SkipSerialization)
diff --git a/Engine/Source/Runtime/Engine/Public/WorldPartition/RuntimeHashSet/WorldPartitionRuntimeHashSet.h b/Engine/Source/Runtime/Engine/Public/WorldPartition/RuntimeHashSet/WorldPartitionRuntimeHashSet.h
--- a/Engine/Source/Runtime/Engine/Public/WorldPartition/RuntimeHashSet/WorldPartitionRuntimeHashSet.h
+++ b/Engine/Source/Runtime/Engine/Public/WorldPartition/RuntimeHashSet/WorldPartitionRuntimeHashSet.h
@@ -148,6 +148,7 @@ class UWorldPartitionRuntimeHashSet final : public UWorldPartitionRuntimeHash
 #if WITH_EDITOR
 	ENGINE_API virtual void PostEditChangeChainProperty(FPropertyChangedChainEvent& PropertyChangedEvent) override;
 	ENGINE_API virtual void PostDuplicate(EDuplicateMode::Type DuplicateMode) override;
+	ENGINE_API virtual void PostLoad() override;
 #endif
 	//~ End UObject Interface
 
