Integrated changes from CDPR to decouple decal scene proxy from its component.

diff --git a/Engine/Source/Runtime/Engine/Private/Components/DecalComponent.cpp b/Engine/Source/Runtime/Engine/Private/Components/DecalComponent.cpp
--- a/Engine/Source/Runtime/Engine/Private/Components/DecalComponent.cpp
+++ b/Engine/Source/Runtime/Engine/Private/Components/DecalComponent.cpp
@@ -15,6 +15,7 @@
 #include "LocalVertexFactory.h"
 #include "Materials/MaterialInstanceDynamic.h"
 #include "MarkActorRenderStateDirtyTask.h"
+#include "DeferredDecalSceneProxyDesc.h"
 
 #if WITH_EDITOR
 #include "Editor.h"
@@ -117,6 +118,28 @@ FDeferredDecalProxy::FDeferredDecalProxy(const USceneComponent* InComponent, UMa
 	}
 }
 
+FDeferredDecalProxy::FDeferredDecalProxy(const FDeferredDecalSceneProxyDesc& Desc)
+	: DrawInGame(Desc.bDrawInGame)
+	, DrawInEditor(Desc.bDrawInEditor)
+	, InvFadeDuration(-1.0f)
+	, InvFadeInDuration(1.0f)
+	, FadeStartDelayNormalized(1.0f)
+	, FadeInStartDelayNormalized(0.0f)
+	, FadeScreenSize(Desc.FadeScreenSize)
+	, DecalColor(Desc.DecalColor)
+{
+	Component = Desc.Component;
+	DecalMaterial = Desc.DecalMaterial;
+	SetTransformIncludingDecalSize(Desc.TransformWithDecalScale, Desc.Bounds);
+	SortOrder = Desc.SortOrder;
+
+	// We don't want to fade when we're editing, only in Simulate/PIE/Game
+	if (Desc.bShouldFade)
+	{
+		InitializeFadingParameters(Desc.InitializationWorldTimeSeconds, Desc.FadeDuration, Desc.FadeStartDelay, Desc.FadeInDuration, Desc.FadeInStartDelay);
+	}
+}
+
 void FDeferredDecalProxy::SetTransformIncludingDecalSize(const FTransform& InComponentToWorldIncludingDecalSize, const FBoxSphereBounds& InBounds)
 {
 	ComponentTrans = InComponentToWorldIncludingDecalSize;
@@ -137,6 +160,49 @@ void FDeferredDecalProxy::InitializeFadingParameters(float AbsSpawnTime, float F
 	}
 }
 
+FDeferredDecalSceneProxyDesc::FDeferredDecalSceneProxyDesc(const UDecalComponent* InComponent)
+{
+	Component = InComponent;
+	Bounds = InComponent->CalcBounds(InComponent->GetComponentTransform());
+	TransformWithDecalScale = InComponent->GetTransformIncludingDecalSize();
+	DecalColor = InComponent->DecalColor;
+	FadeScreenSize = InComponent->FadeScreenSize;
+	FadeDuration = InComponent->GetFadeDuration();
+	FadeStartDelay = InComponent->GetFadeStartDelay();
+	FadeInDuration = InComponent->GetFadeInDuration();
+	FadeInStartDelay = InComponent->GetFadeInStartDelay();
+	SortOrder = InComponent->SortOrder;
+	bDrawInGame = InComponent->GetVisibleFlag() && !InComponent->bHiddenInGame;
+	bDrawInEditor = InComponent->GetVisibleFlag();
+
+	DecalMaterial = [InComponent]
+	{
+		if(UMaterialInterface* ComponentMaterial = InComponent->GetDecalMaterial())
+		{
+			const UMaterial* BaseMaterial = ComponentMaterial->GetMaterial();
+			if (BaseMaterial->MaterialDomain == MD_DeferredDecal)
+			{
+				return ComponentMaterial;
+			}
+		}
+
+		return Cast<UMaterialInterface>(UMaterial::GetDefaultMaterial(MD_DeferredDecal));
+	}();
+
+#if WITH_EDITOR
+	// We don't want to fade when we're editing, only in Simulate/PIE/Game
+	bShouldFade = (!GIsEditor || (InComponent->GetWorld() && InComponent->GetWorld()->IsPlayInEditor()));
+#endif
+
+	if (InComponent->GetOwner())
+	{
+		bDrawInGame &= !(InComponent->GetOwner()->IsHidden());
+#if WITH_EDITOR
+		bDrawInEditor &= !InComponent->GetOwner()->IsHiddenEd();
+#endif
+	}
+}
+
 bool FDeferredDecalProxy::IsShown( const FSceneView* View ) const
 {
 	// Logic here should match FPrimitiveSceneProxy::IsShown for consistent behavior in editor and at runtime.
diff --git a/Engine/Source/Runtime/Engine/Public/DeferredDecalSceneProxyDesc.h b/Engine/Source/Runtime/Engine/Public/DeferredDecalSceneProxyDesc.h
new file mode 100644
--- /dev/null
+++ b/Engine/Source/Runtime/Engine/Public/DeferredDecalSceneProxyDesc.h
@@ -0,0 +1,37 @@
+﻿// Copyright Epic Games, Inc. All Rights Reserved.
+
+#pragma once
+
+#include "CoreTypes.h"
+#include "Math/Color.h"
+#include "Math/Transform.h"
+#include "Math/BoxSphereBounds.h"
+
+class UDecalComponent;
+class UMaterialInterface;
+
+struct ENGINE_API FDeferredDecalSceneProxyDesc
+{
+	FDeferredDecalSceneProxyDesc() = default;
+	FDeferredDecalSceneProxyDesc(const UDecalComponent* Decal);
+
+	UMaterialInterface* DecalMaterial = nullptr;
+	const UDecalComponent* Component = nullptr;
+	
+	FTransform TransformWithDecalScale = FTransform::Identity;
+	FBoxSphereBounds Bounds{};
+
+	FLinearColor DecalColor = FLinearColor::White;
+
+	float InitializationWorldTimeSeconds = 0.0f;
+	float FadeScreenSize = 0.0f;
+	float FadeDuration = 0.0f;
+	float FadeStartDelay = 0.0f;
+	float FadeInDuration = 0.0f;
+	float FadeInStartDelay = 0.0f;
+	int32 SortOrder = INDEX_NONE;
+		
+	uint8 bDrawInGame : 1 = false;
+	uint8 bDrawInEditor : 1 = false;
+	uint8 bShouldFade : 1 = true;
+};
\ No newline at end of file
diff --git a/Engine/Source/Runtime/Engine/Public/SceneManagement.h b/Engine/Source/Runtime/Engine/Public/SceneManagement.h
--- a/Engine/Source/Runtime/Engine/Public/SceneManagement.h
+++ b/Engine/Source/Runtime/Engine/Public/SceneManagement.h
@@ -1384,6 +1384,7 @@ public:
 	/** constructor */
 	ENGINE_API FDeferredDecalProxy(const UDecalComponent* InComponent);
 	ENGINE_API FDeferredDecalProxy(const USceneComponent* InComponent, UMaterialInterface* InMaterial);
+	ENGINE_API FDeferredDecalProxy(const struct FDeferredDecalSceneProxyDesc& Desc);
 
 	/**
 	 * Updates the decal proxy's cached transform and bounds.
@@ -3027,4 +3028,4 @@ extern ENGINE_API void InitializeSharedSamplerStates();
 
 #if UE_ENABLE_INCLUDE_ORDER_DEPRECATED_IN_5_5
 #include "MeshPaintVisualize.h"
-#endif
+#endif
\ No newline at end of file
