Add SetMetaData method to FRawStaticIndexBuffer16or32Interface.

- Allows to set the metadata to not serialized buffers.

#rnx

diff --git a/Engine/Source/Runtime/Engine/Public/RawIndexBuffer.h b/Engine/Source/Runtime/Engine/Public/RawIndexBuffer.h
--- a/Engine/Source/Runtime/Engine/Public/RawIndexBuffer.h
+++ b/Engine/Source/Runtime/Engine/Public/RawIndexBuffer.h
@@ -332,6 +332,8 @@ public:
 
 	virtual void SerializeMetaData(FArchive& Ar) = 0;
 
+	virtual void SetMetaData(int32 InCachedNumIndices) = 0;
+
 	/**
 	 * The following methods are basically just accessors that allow us
 	 * to hide the implementation of FRawStaticIndexBuffer16or32 by making
@@ -428,6 +430,11 @@ public:
 		Ar << CachedNumIndices;
 	}
 
+	virtual void SetMetaData(int32 InCachedNumIndices) override
+	{
+		CachedNumIndices = InCachedNumIndices;
+	}
+
 	/**
 	* Orders a triangle list for better vertex cache coherency.
 	*/
