[World Partition] Always load editor-only actors during cooking and move them in their corresponding streaming cell, to avoid referencing public objects from the level instance package for embedded actors.
In PIE, we continue to filter out editor-only actors and also null-out references to these objects using the instancing context.
In cook, the references will be filtered out by the cooker archive will be filtering editor-only objects, and will allow references from other cells because they all share the same outer.

diff --git a/Engine/Source/Runtime/Engine/Private/WorldPartition/WorldPartitionLevelHelper.cpp b/Engine/Source/Runtime/Engine/Private/WorldPartition/WorldPartitionLevelHelper.cpp
--- a/Engine/Source/Runtime/Engine/Private/WorldPartition/WorldPartitionLevelHelper.cpp
+++ b/Engine/Source/Runtime/Engine/Private/WorldPartition/WorldPartitionLevelHelper.cpp
@@ -225,7 +225,10 @@ void FWorldPartitionLevelHelper::MoveExternalActorsToLevel(const TArray<FWorldPa
 			continue;
 		}
 
-		if (PackageObjectMapping.bIsEditorOnly)
+		// Always load editor-only actors during cooking and move them in their corresponding streaming cell, to avoid referencing public objects from the level instance package for embedded actors.
+		// In PIE, we continue to filter out editor-only actors and also null-out references to these objects using the instancing context. In cook, the references will be filtered out by the cooker 
+		// archive will be filtering editor-only objects, and will allow references from other cells because they all share the same outer.
+		if (PackageObjectMapping.bIsEditorOnly && !IsRunningCookCommandlet())
 		{
 			continue;
 		}
@@ -647,16 +650,17 @@ bool FWorldPartitionLevelHelper::LoadActorsInternal(FLoadActorsParams&& InParams
 			}
 		
 			const FName ContainerPackageInstanceName = Context->RemapPackage(PackageObjectMapping.ContainerPackage);
+			const bool bConsiderActorEditorOnly = PackageObjectMapping.bIsEditorOnly && !IsRunningCookCommandlet(); // See relevant comment in MoveExternalActorsToLevel
 
-			if (PackageObjectMapping.bIsEditorOnly || PackageObjectMapping.ContainerPackage != ContainerPackageInstanceName)
+			if (bConsiderActorEditorOnly || PackageObjectMapping.ContainerPackage != ContainerPackageInstanceName)
 			{
 				const FName ActorPackageName = *FPackageName::ObjectPathToPackageName(PackageObjectMapping.Package.ToString());
-				const FName ActorPackageInstanceName = PackageObjectMapping.bIsEditorOnly ? NAME_None : FName(*ULevel::GetExternalActorPackageInstanceName(ContainerPackageInstanceName.ToString(), ActorPackageName.ToString()));
+				const FName ActorPackageInstanceName = bConsiderActorEditorOnly ? NAME_None : FName(*ULevel::GetExternalActorPackageInstanceName(ContainerPackageInstanceName.ToString(), ActorPackageName.ToString()));
 
 				Context->AddPackageMapping(ActorPackageName, ActorPackageInstanceName);
 			}
 
-			if (!PackageObjectMapping.bIsEditorOnly)
+			if (!bConsiderActorEditorOnly)
 			{
 				ActorPackagesToLoad.Add(&PackageObjectMapping);
 			}
