Silence static analysis error.

#rnx

diff --git a/Engine/Source/Runtime/Engine/Private/NetConnection.cpp b/Engine/Source/Runtime/Engine/Private/NetConnection.cpp
--- a/Engine/Source/Runtime/Engine/Private/NetConnection.cpp
+++ b/Engine/Source/Runtime/Engine/Private/NetConnection.cpp
@@ -1533,7 +1533,7 @@ void UChildConnection::AssignConnectionHandle()
 		int32 ChildId = UsedChildIds.Find(false);
 		if (ChildId == INDEX_NONE)
 		{
-			ensureMsgf(ChildId != INDEX_NONE, TEXT("Logic failed assigning child connection ID."));
+			ensureMsgf(false, TEXT("Logic failed assigning child connection ID."));
 			ChildId = InvalidChildId;
 		}
 		
