Fixed crash when scene proxy is null updating a skeletal mesh object due to deferred creation of scene proxies.

diff --git a/Engine/Source/Runtime/Engine/Private/SkeletalRenderGPUSkin.cpp b/Engine/Source/Runtime/Engine/Private/SkeletalRenderGPUSkin.cpp
--- a/Engine/Source/Runtime/Engine/Private/SkeletalRenderGPUSkin.cpp
+++ b/Engine/Source/Runtime/Engine/Private/SkeletalRenderGPUSkin.cpp
@@ -453,7 +453,7 @@ void FSkeletalMeshObjectGPUSkin::ReleaseResources()
 #endif // RHI_RAYTRACING
 }
 
-void FSkeletalMeshObjectGPUSkin::InitMorphResources(bool bInUsePerBoneMotionBlur, const TArray<float>& MorphTargetWeights)
+void FSkeletalMeshObjectGPUSkin::InitMorphResources(const TArray<float>& MorphTargetWeights)
 {
 	if (!bMorphResourcesInitialized)
 	{
@@ -469,7 +469,7 @@ void FSkeletalMeshObjectGPUSkin::InitMorphResources(bool bInUsePerBoneMotionBlur
 			{
 				// init any morph vertex buffers for each LOD
 				const FSkelMeshObjectLODInfo& MeshLODInfo = LODInfo[LODIndex];
-				SkelLOD.InitMorphResources(MeshLODInfo, bInUsePerBoneMotionBlur, FeatureLevel);
+				SkelLOD.InitMorphResources(MeshLODInfo, FeatureLevel);
 			}
 		}
 		bMorphResourcesInitialized = true;
@@ -491,10 +491,8 @@ void FSkeletalMeshObjectGPUSkin::Update(
 	// make sure morph data has been initialized for each LOD
 	if((!InActiveMorphTargets.IsEmpty() || !InExternalMorphWeightData.MorphSets.IsEmpty()))
 	{
-		const FSkeletalMeshSceneProxy* SkeletalMeshSceneProxy = static_cast<const FSkeletalMeshSceneProxy*>(InSceneProxy);
-		const bool bPerBoneMotionBlur = SkeletalMeshSceneProxy->AlwaysHasVelocity();
 		// initialized on-the-fly in order to avoid creating extra vertex streams for each skel mesh instance
-		InitMorphResources(bPerBoneMotionBlur, InMorphTargetWeights);
+		InitMorphResources(InMorphTargetWeights);
 	}
 
 	// create the new dynamic data for use by the rendering thread
@@ -2191,7 +2189,7 @@ void FSkeletalMeshObjectGPUSkin::FSkeletalMeshObjectLOD::ReleaseResources()
 	GPUSkinVertexFactories.ReleaseVertexFactories();
 }
 
-void FSkeletalMeshObjectGPUSkin::FSkeletalMeshObjectLOD::InitMorphResources(const FSkelMeshObjectLODInfo& MeshLODInfo, bool bInUsePerBoneMotionBlur, ERHIFeatureLevel::Type InFeatureLevel)
+void FSkeletalMeshObjectGPUSkin::FSkeletalMeshObjectLOD::InitMorphResources(const FSkelMeshObjectLODInfo& MeshLODInfo, ERHIFeatureLevel::Type InFeatureLevel)
 {
 	check(SkelMeshRenderData);
 	check(SkelMeshRenderData->LODRenderData.IsValidIndex(LODIndex));
diff --git a/Engine/Source/Runtime/Engine/Private/SkeletalRenderGPUSkin.h b/Engine/Source/Runtime/Engine/Private/SkeletalRenderGPUSkin.h
--- a/Engine/Source/Runtime/Engine/Private/SkeletalRenderGPUSkin.h
+++ b/Engine/Source/Runtime/Engine/Private/SkeletalRenderGPUSkin.h
@@ -600,7 +600,7 @@ protected:
 		 * @param MeshLODInfo - information about the state of the bone influence swapping
 		 * @param Chunks - relevant chunk information (either original or from swapped influence)
 		 */
-		void InitMorphResources(const FSkelMeshObjectLODInfo& MeshLODInfo, bool bInUsePerBoneMotionBlur, ERHIFeatureLevel::Type FeatureLevel);
+		void InitMorphResources(const FSkelMeshObjectLODInfo& MeshLODInfo, ERHIFeatureLevel::Type FeatureLevel);
 
 		/**
 		 * @return memory in bytes of size of the resources for this LOD
@@ -664,7 +664,7 @@ protected:
 	/** 
 	* Initialize morph rendering resources for each LOD 
 	*/
-	ENGINE_API void InitMorphResources(bool bInUsePerBoneMotionBlur, const TArray<float>& MorphTargetWeights);
+	ENGINE_API void InitMorphResources(const TArray<float>& MorphTargetWeights);
 
 	ENGINE_API void ProcessUpdatedDynamicData(EGPUSkinCacheEntryMode Mode, FGPUSkinCache* GPUSkinCache, FRHICommandList& RHICmdList, uint32 FrameNumberToPrepare, uint32 RevisionNumber, uint32 PreviousRevisionNumber, bool bMorphNeedsUpdate, int32 LODIndex, bool bRecreating);
 
