Fix new material being incorrect updated after adding new Substrate node.

diff --git a/Engine/Source/Runtime/Engine/Private/Materials/Material.cpp b/Engine/Source/Runtime/Engine/Private/Materials/Material.cpp
--- a/Engine/Source/Runtime/Engine/Private/Materials/Material.cpp
+++ b/Engine/Source/Runtime/Engine/Private/Materials/Material.cpp
@@ -3322,7 +3322,7 @@ static int32 GetSubstrateConversionVersion()
 	return 0;
 }
 
-bool UMaterial::ConvertMaterialToSubstrateMaterial()
+bool UMaterial::ConvertMaterialToSubstrateMaterial(bool bAllowEmptyMaterialUpdate)
 {
 	/*
 	* The data flow for legacy material conversion node that can be used in isolation is as such:
@@ -3499,6 +3499,9 @@ bool UMaterial::ConvertMaterialToSubstrateMaterial()
 	const bool bRequireSubsurfacePasses		= ShadingModels.HasShadingModel(MSM_SubsurfaceProfile) || ShadingModels.HasShadingModel(MSM_Subsurface) || ShadingModels.HasShadingModel(MSM_PreintegratedSkin) || ShadingModels.HasShadingModel(MSM_Eye);
 	const bool bRequireNoSubsurfaceProfile	= !bHasShadingModelMixture && (ShadingModel == MSM_Subsurface || ShadingModel == MSM_PreintegratedSkin); // Insure there is no profile, as this would take priority otherwise
 
+	// Empty material (i.e., material without any expression) can be 'patched' with UE4Legacy's MF in order to have correct visuals by default
+	const bool bUpdateEmptyMaterial = bAllowEmptyMaterialUpdate && !bUseMaterialAttributes && !EditorOnly->FrontMaterial.IsConnected() && GetExpressions().IsEmpty();
+
 	bool bInvalidateShader = false;
 	bool bEmptyShader = false;
 	bool bRelinkCustomOutputNodes = false;
@@ -3646,7 +3649,7 @@ bool UMaterial::ConvertMaterialToSubstrateMaterial()
 			RefractionCoverageMode = RCM_CoverageIgnored;
 			bInvalidateShader = true;
 		}
-		else if (!bUseMaterialAttributes && !EditorOnly->FrontMaterial.IsConnected() && GetExpressions().IsEmpty())
+		else if (bUpdateEmptyMaterial)
 		{
 			// Empty material: Create by default a slab node
 			UMaterialFunction* DefaultMF = LoadObject<UMaterialFunction>(nullptr, TEXT("/Engine/Functions/Substrate/SMF_UE4Legacy.SMF_UE4Legacy")); 
@@ -4317,7 +4320,7 @@ void UMaterial::PostLoad()
 	{
 		// Substrate materials conversion needs to be done after expressions are cached, otherwise material function won't have 
 		// valid inputs in certain cases
-		if (ConvertMaterialToSubstrateMaterial())
+		if (ConvertMaterialToSubstrateMaterial(true /*bAllowEmptyMaterialUpdate*/))
 		{
 			// Call PropagateDataToMaterialProxy in order to propagate Subsurface profiles and specular profiles data to the material proxy
 			PropagateDataToMaterialProxy();
@@ -4969,7 +4972,7 @@ void UMaterial::PostEditChangePropertyInternal(FPropertyChangedEvent& PropertyCh
 		// This late conversion ensures that imported materials are converted properly with Substrate.
 		if (!IsMaterialAlreadyConvertedToSubstrate(this, EditorOnly))
 		{
-			ConvertMaterialToSubstrateMaterial();
+			ConvertMaterialToSubstrateMaterial(false /*bAllowEmptyMaterialUpdate*/);
 		}
 	}
 
diff --git a/Engine/Source/Runtime/Engine/Public/Materials/Material.h b/Engine/Source/Runtime/Engine/Public/Materials/Material.h
--- a/Engine/Source/Runtime/Engine/Public/Materials/Material.h
+++ b/Engine/Source/Runtime/Engine/Public/Materials/Material.h
@@ -1406,7 +1406,7 @@ private:
 	void BackwardsCompatibilityInputConversion();
 	void BackwardsCompatibilityVirtualTextureOutputConversion();
 	void BackwardsCompatibilityDecalConversion();
-	bool ConvertMaterialToSubstrateMaterial();
+	bool ConvertMaterialToSubstrateMaterial(bool bAllowEmptyMaterialUpdate);
 
 	/** Marks the material's package dirty in order to make a material usage change set during map load persistent. 
 	  * This couldn't be done during map load as loading cannot mark packages dirty. Invoked manually by the user 
