PCG-generated LI's now replicated their assigned world and spawn on clients correctly. All of the rep properties are only sent once on initial spawn, so we can just use PostNetInit() to take care of business (vs. an OnRep for each of them)

diff --git a/Engine/Source/Runtime/Engine/Private/LevelInstance/LevelInstanceActor.cpp b/Engine/Source/Runtime/Engine/Private/LevelInstance/LevelInstanceActor.cpp
--- a/Engine/Source/Runtime/Engine/Private/LevelInstance/LevelInstanceActor.cpp
+++ b/Engine/Source/Runtime/Engine/Private/LevelInstance/LevelInstanceActor.cpp
@@ -52,14 +52,24 @@ void ALevelInstance::Serialize(FArchive& Ar)
 void ALevelInstance::GetLifetimeReplicatedProps(TArray<FLifetimeProperty>& OutLifetimeProps) const
 {
 	Super::GetLifetimeReplicatedProps(OutLifetimeProps);
+
 	FDoRepLifetimeParams Params;
 	Params.Condition = COND_InitialOnly;
 	DOREPLIFETIME_WITH_PARAMS_FAST(ALevelInstance, LevelInstanceSpawnGuid, Params);
+
+	// Generally the WorldAsset ref is set at edit time, but in cases where the LI is spawned during gameplay (ex: from PCG), the WorldAsset is set on spawn and needs to be replicated
+	// We can assume that these will never be changed during the lifetime of the LI and are assigned on spawn, so we only need the one initial rep
+#if WITH_EDITORONLY_DATA
+	DOREPLIFETIME_WITH_PARAMS_FAST(ALevelInstance, WorldAsset, Params);
+#endif
+	DOREPLIFETIME_WITH_PARAMS_FAST(ALevelInstance, CookedWorldAsset, Params);
 }
 
-void ALevelInstance::OnRep_LevelInstanceSpawnGuid()
+void ALevelInstance::PostNetInit()
 {
-	if (GetLocalRole() != ENetRole::ROLE_Authority && !LevelInstanceActorGuid.IsValid())
+	Super::PostNetInit();
+	
+	if (ensure(!LevelInstanceActorGuid.IsValid()))
 	{
 		check(LevelInstanceSpawnGuid.IsValid());
 		LevelInstanceActorGuid.ActorGuid = LevelInstanceSpawnGuid;
diff --git a/Engine/Source/Runtime/Engine/Public/LevelInstance/LevelInstanceActor.h b/Engine/Source/Runtime/Engine/Public/LevelInstance/LevelInstanceActor.h
--- a/Engine/Source/Runtime/Engine/Public/LevelInstance/LevelInstanceActor.h
+++ b/Engine/Source/Runtime/Engine/Public/LevelInstance/LevelInstanceActor.h
@@ -25,19 +25,20 @@ public:
 protected:
 #if WITH_EDITORONLY_DATA
 	/** Level LevelInstance */
-	UPROPERTY(EditAnywhere, Category = Level, Meta = (NoCreate, DisplayName="Level", DisableLevelInstancePropertyOverride))
+	UPROPERTY(EditAnywhere, Replicated, Category = Level, Meta = (NoCreate, DisplayName="Level", DisableLevelInstancePropertyOverride))
 	TSoftObjectPtr<UWorld> WorldAsset;
 
 	UPROPERTY(VisibleAnywhere, Category = Override, meta = (EditInline, NoResetToDefault, EditCondition="PropertyOverrides != nullptr", EditConditionHides))
 	TObjectPtr<ULevelInstancePropertyOverrideAsset> PropertyOverrides;
 #endif
+	
 	UPROPERTY(VisibleAnywhere, Category = Default)
 	TObjectPtr<ULevelInstanceComponent> LevelInstanceComponent;
 
-	UPROPERTY()
+	UPROPERTY(Replicated)
 	TSoftObjectPtr<UWorld> CookedWorldAsset;
 
-	UPROPERTY(Transient, ReplicatedUsing=OnRep_LevelInstanceSpawnGuid)
+	UPROPERTY(Transient, Replicated)
 	FGuid LevelInstanceSpawnGuid;
 
 	FLevelInstanceActorGuid LevelInstanceActorGuid;
@@ -63,9 +64,7 @@ public:
 	ENGINE_API virtual void PostUnregisterAllComponents() override;
 	ENGINE_API virtual void Serialize(FArchive& Ar) override;
 	ENGINE_API virtual void GetLifetimeReplicatedProps(TArray<FLifetimeProperty>& OutLifetimeProps) const override;
-		
-	UFUNCTION()
-	ENGINE_API void OnRep_LevelInstanceSpawnGuid();
+	ENGINE_API virtual void PostNetInit() override;
 
 #if WITH_EDITOR
 	// Begin ILevelInstanceInterface
