Split implementation of ExportProperties by backporting function implementation before OverrideSerialization support in T3D

* This change makes ExportProperties function shim to the legacy logic vs the new logic using the value of the CVar as the differentiator.

diff --git a/Engine/Source/Runtime/Engine/Private/UnrealExporter.cpp b/Engine/Source/Runtime/Engine/Private/UnrealExporter.cpp
--- a/Engine/Source/Runtime/Engine/Private/UnrealExporter.cpp
+++ b/Engine/Source/Runtime/Engine/Private/UnrealExporter.cpp
@@ -674,21 +674,8 @@ void UExporter::ExportObjectInner(const FExportObjectInnerContext* Context, UObj
 	TextIndent -= 3;
 }
 
-/**
- * Exports the property values for the specified object as text to the output device.
- *
- * @param	Context			Context from which the set of 'inner' objects is extracted.  If NULL, an object iterator will be used.
- * @param	Out				the output device to send the exported text to
- * @param	ObjectClass		the class of the object to dump properties for
- * @param	Object			the address of the object to dump properties for
- * @param	Indent			number of spaces to prepend to each line of output
- * @param	DiffClass		the class to use for comparing property values when delta export is desired.
- * @param	Diff			the address of the object to use for determining whether a property value should be exported.  If the value in Object matches the corresponding
- *							value in Diff, it is not exported.  Specify NULL to export all properties.
- * @param	Parent			the UObject corresponding to Object
- * @param	PortFlags		flags used for modifying the output and/or behavior of the export
- */
-void ExportProperties
+/** Shim ExportProperties when FOverridableSerializationLogic::ECapabilities::T3DSerialization capabilities enabled */
+void ExportProperties_Overrides
 (
 	const FExportObjectInnerContext* Context,
 	FOutputDevice&	Out,
@@ -1090,6 +1077,239 @@ void ExportProperties
 	Parent->ExportCustomProperties(Out, Indent);
 }
 
+/** Shim ExportProperties when FOverridableSerializationLogic::ECapabilities::T3DSerialization capabilities disabled */
+void ExportProperties_NoOverrides
+(
+	const FExportObjectInnerContext* Context,
+	FOutputDevice&	Out,
+	UClass*			ObjectClass,
+	uint8*			Object,
+	int32			Indent,
+	UClass*			DiffClass,
+	uint8*			Diff,
+	UObject*		Parent,
+	uint32			PortFlags,
+	UObject*		ExportRootScope
+)
+{
+	check(ObjectClass != NULL);
+	
+	for( FProperty* Property = ObjectClass->PropertyLink; Property; Property = Property->PropertyLinkNext )
+	{
+		if (!Property->ShouldPort(PortFlags))
+			continue;
+
+		FString SanitizedPropertyName = Property->GetName();
+		constexpr FAsciiSet Whitespace("\t ");
+		constexpr FAsciiSet SpecialCharacters("=()[].\"\'");
+		if (FAsciiSet::HasAny(SanitizedPropertyName, SpecialCharacters) || Whitespace.Contains(SanitizedPropertyName[0]))
+		{
+			// to increase frequency of forward compatibility, only sanitize property names that absolutely need it
+			SanitizedPropertyName = FString::Format(TEXT("\"{0}\""), {Property->GetName().ReplaceCharWithEscapedChar()});
+		}
+
+		FArrayProperty* ArrayProperty = CastField<FArrayProperty>(Property);
+		FObjectPropertyBase* ExportObjectProp = (Property->PropertyFlags & CPF_ExportObject) != 0 ? CastField<FObjectPropertyBase>(Property) : NULL;
+		const uint32 ExportFlags = PortFlags | PPF_Delimited;
+
+		if ( ArrayProperty != NULL )
+		{
+			// Export dynamic array.
+			FProperty* InnerProp = ArrayProperty->Inner;
+			ExportObjectProp = (Property->PropertyFlags & CPF_ExportObject) != 0 ? CastField<FObjectPropertyBase>(InnerProp) : NULL;
+			// This is used as the default value in the case of an array property that has
+			// fewer elements than the exported object.
+			uint8* StructDefaults = NULL;
+			FStructProperty* StructProperty = CastField<FStructProperty>(InnerProp);
+			if ( StructProperty != NULL )
+			{
+				checkSlow(StructProperty->Struct);
+				StructDefaults = (uint8*)FMemory::Malloc(StructProperty->Struct->GetStructureSize());
+				StructProperty->InitializeValue(StructDefaults);
+			}
+
+			for( int32 PropertyArrayIndex=0; PropertyArrayIndex<Property->ArrayDim; PropertyArrayIndex++ )
+			{
+				void* Arr = Property->ContainerPtrToValuePtr<void>(Object, PropertyArrayIndex);
+				FScriptArrayHelper ArrayHelper(ArrayProperty, Arr);
+
+				void*	DiffArr = NULL;
+				if( DiffClass )
+				{
+					DiffArr = Property->ContainerPtrToValuePtrForDefaults<void>(DiffClass, Diff, PropertyArrayIndex);
+				}
+				// we won't use this if DiffArr is NULL, but we have to set it up to something
+				FScriptArrayHelper DiffArrayHelper(ArrayProperty, DiffArr);
+
+				// If the current size of the array is 0 and the default one is not, add in an empty item so on import it will be empty
+				if( ArrayHelper.Num() == 0 && DiffArrayHelper.Num() != 0 )
+				{
+					Out.Logf(TEXT("%s%s=\r\n"), FCString::Spc(Indent), *SanitizedPropertyName);
+				}
+				else
+				{
+					// If the array sizes are different, we will need to export each index so on import we maintain the size
+					for (int32 DynamicArrayIndex = 0; DynamicArrayIndex < ArrayHelper.Num(); DynamicArrayIndex++)
+					{
+						FString	Value;
+
+						// compare each element's value manually so that elements which match the NULL value for the array's inner property type
+						// but aren't in the diff array are still exported
+						uint8* SourceData = ArrayHelper.GetRawPtr(DynamicArrayIndex);
+						bool bHasDiffData = DiffArr && DynamicArrayIndex < DiffArrayHelper.Num();
+						uint8* DiffData = bHasDiffData ? DiffArrayHelper.GetRawPtr(DynamicArrayIndex) : StructDefaults;
+
+						// Make sure to export the last element even if it is default value if the default data doesn't have an element at that index (!bHasDiffData && (DynamicArrayIndex == ArrayHelper.Num()-1) && !bHasDiffData) 
+						// because that will ensure the resulting imported array will be of proper size
+						bool bExportItem = DiffData == NULL || (!bHasDiffData && (DynamicArrayIndex == ArrayHelper.Num()-1)) || (DiffData != SourceData && !InnerProp->Identical(SourceData, DiffData, ExportFlags));
+						if (bExportItem)
+						{
+							InnerProp->ExportTextItem_Direct(Value, SourceData, DiffData, Parent, ExportFlags, ExportRootScope);
+							if (ExportObjectProp)
+							{
+								UObject* Obj = ExportObjectProp->GetObjectPropertyValue(ArrayHelper.GetRawPtr(DynamicArrayIndex));
+								check(!Obj || Obj->IsValidLowLevel());
+								if (Obj && !Obj->HasAnyMarks(OBJECTMARK_TagImp))
+								{
+									// only export the BEGIN OBJECT block for a component if Parent is the component's Outer....when importing subobject definitions,
+									// (i.e. BEGIN OBJECT), whichever BEGIN OBJECT block a component's BEGIN OBJECT block is located within is the object that will be
+									// used as the Outer to create the component
+
+									// Is this an array of components?
+									if (InnerProp->HasAnyPropertyFlags(CPF_InstancedReference))
+									{
+										if (Obj->GetOuter() == Parent)
+										{
+											// Don't export more than once.
+											Obj->Mark(OBJECTMARK_TagImp);
+											UExporter::ExportToOutputDevice(Context, Obj, NULL, Out, TEXT("T3D"), Indent, PortFlags);
+										}
+										else
+										{
+											// set the OBJECTMARK_TagExp flag so that the calling code knows we wanted to export this object
+											Obj->Mark(OBJECTMARK_TagExp);
+										}
+									}
+									else
+									{
+										// Don't export more than once.
+										Obj->Mark(OBJECTMARK_TagImp);
+										UExporter::ExportToOutputDevice(Context, Obj, NULL, Out, TEXT("T3D"), Indent, PortFlags);
+									}
+								}
+							}
+
+							Out.Logf(TEXT("%s%s(%i)=%s\r\n"), FCString::Spc(Indent), *SanitizedPropertyName, DynamicArrayIndex, *Value);
+						}
+					}
+					for (int32 DynamicArrayIndex = DiffArrayHelper.Num()-1; DynamicArrayIndex >= ArrayHelper.Num(); --DynamicArrayIndex)
+					{
+						Out.Logf(TEXT("%s%s.RemoveIndex(%d)\r\n"), FCString::Spc(Indent), *SanitizedPropertyName, DynamicArrayIndex);
+					}
+				}
+			}
+			if (StructDefaults)
+			{
+				StructProperty->DestroyValue(StructDefaults);
+				FMemory::Free(StructDefaults);
+			}
+		}
+		else
+		{
+			for( int32 PropertyArrayIndex=0; PropertyArrayIndex<Property->ArrayDim; PropertyArrayIndex++ )
+			{
+				FString	Value;
+				// Export single element.
+
+				uint8* DiffData = (DiffClass && Property->IsInContainer(DiffClass->GetPropertiesSize())) ? Diff : NULL;
+				if( Property->ExportText_InContainer( PropertyArrayIndex, Value, Object, DiffData, Parent, ExportFlags, ExportRootScope ) )
+				{
+					if ( ExportObjectProp )
+					{
+						UObject* Obj = ExportObjectProp->GetObjectPropertyValue(Property->ContainerPtrToValuePtr<void>(Object, PropertyArrayIndex));
+						if( Obj && !Obj->HasAnyMarks(OBJECTMARK_TagImp) )
+						{
+							// only export the BEGIN OBJECT block for a component if Parent is the component's Outer....when importing subobject definitions,
+							// (i.e. BEGIN OBJECT), whichever BEGIN OBJECT block a component's BEGIN OBJECT block is located within is the object that will be
+							// used as the Outer to create the component
+							if ( Property->HasAnyPropertyFlags(CPF_InstancedReference) )
+							{
+								if ( Obj->GetOuter() == Parent )
+								{
+									// Don't export more than once.
+									Obj->Mark(OBJECTMARK_TagImp);
+									UExporter::ExportToOutputDevice( Context, Obj, NULL, Out, TEXT("T3D"), Indent, PortFlags );
+								}
+								else
+								{
+									// set the OBJECTMARK_TagExp flag so that the calling code knows we wanted to export this object
+									Obj->Mark(OBJECTMARK_TagExp);
+								}
+							}
+							else
+							{
+								// Don't export more than once.
+								Obj->Mark(OBJECTMARK_TagImp);
+								UExporter::ExportToOutputDevice( Context, Obj, NULL, Out, TEXT("T3D"), Indent, PortFlags );
+							}
+						}
+					}
+
+					if( Property->ArrayDim == 1 )
+					{
+						Out.Logf( TEXT("%s%s=%s\r\n"), FCString::Spc(Indent), *SanitizedPropertyName, *Value );
+					}
+					else
+					{
+						Out.Logf( TEXT("%s%s(%i)=%s\r\n"), FCString::Spc(Indent), *SanitizedPropertyName, PropertyArrayIndex, *Value );
+					}
+				}
+			}
+		}
+	}
+
+	// Allows to import/export C++ properties in case the automatic unreal script mesh wouldn't work.
+	Parent->ExportCustomProperties(Out, Indent);
+}
+
+/**
+ * Exports the property values for the specified object as text to the output device.
+ *
+ * @param	Context			Context from which the set of 'inner' objects is extracted.  If NULL, an object iterator will be used.
+ * @param	Out				the output device to send the exported text to
+ * @param	ObjectClass		the class of the object to dump properties for
+ * @param	Object			the address of the object to dump properties for
+ * @param	Indent			number of spaces to prepend to each line of output
+ * @param	DiffClass		the class to use for comparing property values when delta export is desired.
+ * @param	Diff			the address of the object to use for determining whether a property value should be exported.  If the value in Object matches the corresponding
+ *							value in Diff, it is not exported.  Specify NULL to export all properties.
+ * @param	Parent			the UObject corresponding to Object
+ * @param	PortFlags		flags used for modifying the output and/or behavior of the export
+ */
+void ExportProperties
+(
+	const FExportObjectInnerContext* Context,
+	FOutputDevice&	Out,
+	UClass*			ObjectClass,
+	uint8*			Object,
+	int32			Indent,
+	UClass*			DiffClass,
+	uint8*			Diff,
+	UObject*		Parent,
+	uint32			PortFlags,
+	UObject*		ExportRootScope
+)
+{
+	if (FOverridableSerializationLogic::HasCapabilities(FOverridableSerializationLogic::ECapabilities::T3DSerialization))
+	{
+		ExportProperties_Overrides(Context, Out, ObjectClass, Object, Indent, DiffClass, Diff, Parent, PortFlags, ExportRootScope);
+	}
+	else
+	{
+		ExportProperties_NoOverrides(Context, Out, ObjectClass, Object, Indent, DiffClass, Diff, Parent, PortFlags, ExportRootScope);
+	}
+}
+
 /**
  * Debug spew for components
  * @param Object object to dump component spew for
