Post Process Volume User Flags exposed as a view property, allowing the value to be used in a shader without custom HLSL.

#rnx

diff --git a/Engine/Source/Runtime/Engine/Private/Materials/HLSLMaterialTranslator.cpp b/Engine/Source/Runtime/Engine/Private/Materials/HLSLMaterialTranslator.cpp
--- a/Engine/Source/Runtime/Engine/Private/Materials/HLSLMaterialTranslator.cpp
+++ b/Engine/Source/Runtime/Engine/Private/Materials/HLSLMaterialTranslator.cpp
@@ -5262,6 +5262,7 @@ int32 FHLSLMaterialTranslator::ViewProperty(EMaterialExposedViewProperty Propert
 		{MEVP_PreExposure, MCT_Float1, TEXT("View.PreExposure.x"), TEXT("View.OneOverPreExposure.x")},
 		{MEVP_RuntimeVirtualTextureMaxLevel, MCT_Float1, TEXT("View.RuntimeVirtualTextureMipLevel.y"), nullptr},
 		{MEVP_ResolutionFraction, MCT_Float1, TEXT("View.ResolutionFractionAndInv.x"), TEXT("View.ResolutionFractionAndInv.y")},
+		{MEVP_PostVolumeUserFlags, MCT_Float1, TEXT("View.PostVolumeUserFlags"), nullptr},
 	};
 	static_assert((sizeof(ViewPropertyMetaArray) / sizeof(ViewPropertyMetaArray[0])) == MEVP_MAX, "incoherency between EMaterialExposedViewProperty and ViewPropertyMetaArray");
 
diff --git a/Engine/Source/Runtime/Engine/Private/Materials/MaterialExpressionHLSL.cpp b/Engine/Source/Runtime/Engine/Private/Materials/MaterialExpressionHLSL.cpp
--- a/Engine/Source/Runtime/Engine/Private/Materials/MaterialExpressionHLSL.cpp
+++ b/Engine/Source/Runtime/Engine/Private/Materials/MaterialExpressionHLSL.cpp
@@ -1149,7 +1149,8 @@ bool UMaterialExpressionViewProperty::GenerateHLSLExpression(FMaterialHLSLGenera
 	case MEVP_RuntimeVirtualTextureOutputDerivative: Input = EExternalInput::RuntimeVirtualTextureOutputDerivative; bApplyRcp = bRequestedRcp; break;
 	case MEVP_PreExposure: Input = bRequestedRcp ? EExternalInput::RcpPreExposure : EExternalInput::PreExposure; break;
 	case MEVP_RuntimeVirtualTextureMaxLevel: Input = EExternalInput::RuntimeVirtualTextureMaxLevel; bApplyRcp = bRequestedRcp; break;
-	case MEVP_ResolutionFraction: Input = bRequestedRcp ? EExternalInput::RcpResolutionFraction : EExternalInput::ResolutionFraction;
+	case MEVP_ResolutionFraction: Input = bRequestedRcp ? EExternalInput::RcpResolutionFraction : EExternalInput::ResolutionFraction; break;
+	case MEVP_PostVolumeUserFlags: Input = EExternalInput::PostVolumeUserFlags; bApplyRcp = bRequestedRcp; break;
 	default: checkNoEntry(); break;
 	}
 
diff --git a/Engine/Source/Runtime/Engine/Private/Materials/MaterialHLSLTree.cpp b/Engine/Source/Runtime/Engine/Private/Materials/MaterialHLSLTree.cpp
--- a/Engine/Source/Runtime/Engine/Private/Materials/MaterialHLSLTree.cpp
+++ b/Engine/Source/Runtime/Engine/Private/Materials/MaterialHLSLTree.cpp
@@ -118,6 +118,7 @@ FExternalInputDescription GetExternalInputDescription(EExternalInput Input)
 	case EExternalInput::RuntimeVirtualTextureMaxLevel: return FExternalInputDescription(TEXT("RuntimeVirtualTextureMaxLevel"), Shader::EValueType::Float1);
 	case EExternalInput::ResolutionFraction: return FExternalInputDescription(TEXT("ResolutionFraction"), Shader::EValueType::Float1);
 	case EExternalInput::RcpResolutionFraction: return FExternalInputDescription(TEXT("RcpResolutionFraction"), Shader::EValueType::Float1);
+	case EExternalInput::PostVolumeUserFlags: return FExternalInputDescription(TEXT("PostVolumeUserFlags"), Shader::EValueType::Float1);
 
 	case EExternalInput::CameraVector: return FExternalInputDescription(TEXT("CameraVector"), Shader::EValueType::Float3);
 	case EExternalInput::LightVector: return FExternalInputDescription(TEXT("LightVector"), Shader::EValueType::Float3);
@@ -442,6 +443,7 @@ void FExpressionExternalInput::EmitValueShader(FEmitContext& Context, FEmitScope
 		case EExternalInput::RuntimeVirtualTextureMaxLevel:  Code = TEXT("View.RuntimeVirtualTextureMipLevel.y"); break;
 		case EExternalInput::ResolutionFraction: Code = TEXT("View.ResolutionFractionAndInv.x"); break;
 		case EExternalInput::RcpResolutionFraction: Code = TEXT("View.ResolutionFractionAndInv.y"); break;
+		case EExternalInput::PostVolumeUserFlags: Code = TEXT("((float)View.PostVolumeUserFlags)"); break;
 
 		case EExternalInput::CameraVector: Code = TEXT("Parameters.CameraVector"); break;
 		case EExternalInput::LightVector: Code = TEXT("Parameters.LightVector"); break;
diff --git a/Engine/Source/Runtime/Engine/Public/MaterialHLSLTree.h b/Engine/Source/Runtime/Engine/Public/MaterialHLSLTree.h
--- a/Engine/Source/Runtime/Engine/Public/MaterialHLSLTree.h
+++ b/Engine/Source/Runtime/Engine/Public/MaterialHLSLTree.h
@@ -151,6 +151,7 @@ enum class EExternalInput : uint8
 	RuntimeVirtualTextureMaxLevel,
 	ResolutionFraction,
 	RcpResolutionFraction,
+	PostVolumeUserFlags,
 
 	CameraVector,
 	LightVector,
diff --git a/Engine/Source/Runtime/Engine/Public/Materials/MaterialExpressionViewProperty.h b/Engine/Source/Runtime/Engine/Public/Materials/MaterialExpressionViewProperty.h
--- a/Engine/Source/Runtime/Engine/Public/Materials/MaterialExpressionViewProperty.h
+++ b/Engine/Source/Runtime/Engine/Public/Materials/MaterialExpressionViewProperty.h
@@ -41,6 +41,8 @@ enum EMaterialExposedViewProperty : int
 	MEVP_RuntimeVirtualTextureMaxLevel UMETA(DisplayName = "Virtual Texture Max Level"),
 	/** Screen percentage at which the rendering resolution happens, to allow tech-art to remain consistent with dynamic resolution. */
 	MEVP_ResolutionFraction UMETA(DisplayName = "ScreenPercentage / 100"),
+	/** Post process volume user flags, useful for varying the behavior of a material per view */
+	MEVP_PostVolumeUserFlags UMETA(DisplayName = "Post Process Volume User Flags"),
 
 	MEVP_MAX,
 };
