EngineTest
* Moved NetTestHelpers files from EngineTest to Engine\Tests so they they can be reused by other modules wanting to add replication unit tests

#rnx

diff --git a/Engine/Source/Runtime/Engine/Private/Tests/NetTestHelpers.cpp b/Engine/Source/Runtime/Engine/Private/Tests/NetTestHelpers.cpp
new file mode 100644
--- /dev/null
+++ b/Engine/Source/Runtime/Engine/Private/Tests/NetTestHelpers.cpp
@@ -0,0 +1,436 @@
+// Copyright Epic Games, Inc. All Rights Reserved.
+
+#include "Tests/NetTestHelpers.h"
+
+#include "CoreGlobals.h"
+#include "Engine/NetDriver.h"
+#include "Engine/NetworkObjectList.h"
+#include "Engine/Engine.h"
+#include "Engine/GameInstance.h"
+#include "Engine/World.h"
+#include "Engine/GameViewportClient.h"
+#include "EngineUtils.h"
+#include "GameFramework/PlayerController.h"
+#include "GameFramework/PlayerState.h"
+#include "HAL/IConsoleManager.h"
+#include "HAL/PlatformMath.h"
+#include "Kismet/GameplayStatics.h"
+#include "Net/UnrealNetwork.h"
+
+#if WITH_EDITOR
+#include "Settings/LevelEditorPlaySettings.h"
+#endif
+
+#include "UObject/UObjectGlobals.h"
+#include "UObject/Package.h"
+
+#if WITH_EDITOR
+
+namespace UE::Net
+{
+
+FTestWorldInstance FTestWorldInstance::CreateServer(const TCHAR* InURL)
+{
+	FGameInstancePIEParameters ServerParams;
+	ServerParams.bSimulateInEditor = false;
+	ServerParams.bAnyBlueprintErrors = false;
+	ServerParams.bStartInSpectatorMode = false;
+	ServerParams.bRunAsDedicated = true;
+	ServerParams.bIsPrimaryPIEClient = false;
+	ServerParams.WorldFeatureLevel = ERHIFeatureLevel::SM5;
+	ServerParams.EditorPlaySettings = nullptr;
+	ServerParams.NetMode = PIE_ListenServer;
+
+	FTestWorldInstance NewInstance(ServerParams);
+
+	FURL LocalURL(nullptr, InURL, TRAVEL_Absolute);
+	FString BrowseError;
+	GEngine->Browse(*NewInstance.GameInstance->GetWorldContext(), LocalURL, BrowseError);
+
+	return NewInstance;
+}
+
+FTestWorldInstance FTestWorldInstance::CreateClient(int32 ServerPort)
+{
+	FGameInstancePIEParameters ClientParams;
+	ClientParams.bSimulateInEditor = false;
+	ClientParams.bAnyBlueprintErrors = false;
+	ClientParams.bStartInSpectatorMode = false;
+	ClientParams.bRunAsDedicated = false;
+	ClientParams.bIsPrimaryPIEClient = false;
+	ClientParams.WorldFeatureLevel = ERHIFeatureLevel::SM5;
+	ClientParams.EditorPlaySettings = nullptr;
+	ClientParams.NetMode = PIE_Client;
+
+	FTestWorldInstance NewInstance(ClientParams);
+
+	FWorldContext* ClientWorldContext = NewInstance.GameInstance->GetWorldContext();
+
+	UGameViewportClient* ViewportClient = NewObject<UGameViewportClient>(GEngine, GEngine->GameViewportClientClass);
+	ViewportClient->Init(*ClientWorldContext, NewInstance.GameInstance);
+	ClientWorldContext->GameViewport = ViewportClient;
+
+	FString OutCreatePlayerError;
+	ViewportClient->SetupInitialLocalPlayer(OutCreatePlayerError);
+	GEngine->BrowseToDefaultMap(*ClientWorldContext);
+
+	FString ClientURLString = FString::Format(TEXT("127.0.0.1:{0}"), {ServerPort});
+
+	FURL ClientURL(nullptr, *ClientURLString, TRAVEL_Absolute);
+	FString ClientBrowseError;
+	GEngine->Browse(*ClientWorldContext, ClientURL, ClientBrowseError);
+
+	return NewInstance;
+}
+
+FTestWorldInstance::FTestWorldInstance(const FGameInstancePIEParameters& InstanceParams)
+{
+	GameInstance = NewObject<UGameInstance>(GEngine, UGameInstance::StaticClass());
+	GameInstance->AddToRoot();
+
+	GameInstance->InitializeForPlayInEditor(FindUnusedPIEInstance(), InstanceParams);
+
+	FWorldContext* WorldContext = GameInstance->GetWorldContext();
+}
+
+FTestWorldInstance::FTestWorldInstance(FTestWorldInstance&& Other)
+	: GameInstance(Other.GameInstance)
+{
+	Other.GameInstance = nullptr;
+}
+
+FTestWorldInstance& FTestWorldInstance::operator=(FTestWorldInstance&& Other)
+{
+	if (this != &Other)
+	{
+		Shutdown();
+
+		GameInstance = Other.GameInstance;
+		Other.GameInstance = nullptr;
+	}
+
+	return *this;
+}
+
+FTestWorldInstance::~FTestWorldInstance()
+{
+	Shutdown();
+}
+
+void FTestWorldInstance::Shutdown()
+{
+	UWorld* World = GetWorld();
+
+	if (World)
+	{
+		for (FActorIterator ActorIt(World); ActorIt; ++ActorIt)
+		{
+			ActorIt->RouteEndPlay(EEndPlayReason::EndPlayInEditor);
+		}
+	}
+
+	if (GameInstance)
+	{
+		GameInstance->Shutdown();
+		GameInstance->RemoveFromRoot();
+	}
+	
+	if (World)
+	{
+		World->BeginTearingDown();
+		GEngine->ShutdownWorldNetDriver(World);
+		GEngine->DestroyWorldContext(World);
+		World->CleanupWorld();
+	}
+}
+
+int32 FTestWorldInstance::FindUnusedPIEInstance()
+{
+	if (!GEngine)
+	{
+		return INDEX_NONE;
+	}
+
+	int32 MaxUsedPIEInstance = INDEX_NONE;
+	for (const FWorldContext& Context : GEngine->GetWorldContexts())
+	{
+		MaxUsedPIEInstance = FMath::Max(MaxUsedPIEInstance, Context.PIEInstance);
+	}
+
+	return MaxUsedPIEInstance + 1;
+}
+
+UWorld* FTestWorldInstance::GetWorld() const
+{
+	return GameInstance ? GameInstance->GetWorld() : nullptr;
+}
+
+FWorldContext* FTestWorldInstance::GetWorldContext() const
+{
+	return GameInstance->GetWorldContext();
+}
+
+UNetDriver* FTestWorldInstance::GetNetDriver() const
+{
+	UWorld* World = GetWorld();
+	return World ? World->GetNetDriver() : nullptr;
+}
+
+void FTestWorldInstance::Tick(float DeltaSeconds)
+{
+	GEngine->TickWorldTravel(*GetWorldContext(), DeltaSeconds);
+	if (UWorld* World = GetWorld())
+	{
+		World->Tick(LEVELTICK_All, DeltaSeconds);
+		World->UpdateLevelStreaming();
+	}
+}
+
+int32 FTestWorldInstance::GetPort()
+{
+	if (const UWorld* const World = GetWorld())
+	{
+		if (UNetDriver* const NetDriver = World->GetNetDriver())
+		{
+			if (NetDriver->GetLocalAddr().IsValid())
+			{
+				return NetDriver->GetLocalAddr()->GetPort();
+			}
+		}
+	}
+
+	return 0;
+}
+
+void FTestWorldInstance::LoadStreamingLevel(FName LevelName)
+{
+	FLatentActionInfo LoadLatentInfo;
+	LoadLatentInfo.UUID = LevelStreamRequestUUID++;
+	constexpr bool bMakeVisibleAfterLoad = true;
+	constexpr bool bShouldBlockOnLoad = true;
+	UGameplayStatics::LoadStreamLevel(GetWorld(), LevelName, bMakeVisibleAfterLoad, bShouldBlockOnLoad, LoadLatentInfo);
+}
+
+void FTestWorldInstance::UnloadStreamingLevel(FName LevelName)
+{
+	FLatentActionInfo UnloadLatentInfo;
+	UnloadLatentInfo.UUID = LevelStreamRequestUUID++;
+	constexpr bool bShouldBlockOnUnload = true;
+	UGameplayStatics::UnloadStreamLevel(GetWorld(), LevelName, UnloadLatentInfo, bShouldBlockOnUnload);
+}
+
+//============================================================================
+// FTestWorlds
+//============================================================================
+
+FTestWorlds::FTestWorlds(const TCHAR* ServerURL, float DeltaSeconds)
+	: Server(true)
+	, TickDeltaSeconds(DeltaSeconds)
+{
+	NetDriverCreatedHandle = FWorldDelegates::OnNetDriverCreated.AddRaw(this, &FTestWorlds::OnNetDriverCreated);
+
+	Server = FTestWorldInstance::CreateServer(ServerURL);
+}
+
+FTestWorlds::~FTestWorlds()
+{
+	FWorldDelegates::OnNetDriverCreated.Remove(NetDriverCreatedHandle);
+}
+
+void FTestWorlds::OnNetDriverCreated(UWorld* InWorld, UNetDriver* InNetDriver)
+{
+	// Make sure NetDriver will tick every engine frame
+	InNetDriver->MaxNetTickRate = 0;
+}
+
+bool FTestWorlds::CreateAndConnectClient()
+{
+	Clients.Emplace(FTestWorldInstance::CreateClient(Server.GetPort()));
+	const bool bConnected = WaitForClientConnect(Clients.Last());
+	return bConnected;
+}
+
+bool FTestWorlds::WaitForClientConnect(FTestWorldInstance& Client)
+{
+	return TickAllUntil([&Client]()
+	{
+		if (UWorld* World = Client.GetWorld())
+		{
+			APlayerController* PC = World->GetFirstPlayerController();
+			return IsValid(PC) && (PC->GetLocalRole() == ROLE_AutonomousProxy);
+		}
+		return false;
+	});
+}
+
+void FTestWorlds::TickAll(int32 NumTicks)
+{
+	for (int32 i = 0; i < NumTicks; ++i)
+	{
+		TickServer();
+		TickClients();
+		GFrameCounter++;
+	}
+}
+
+void FTestWorlds::TickServer()
+{
+	Server.Tick(TickDeltaSeconds);
+}
+
+void FTestWorlds::TickClients()
+{
+	for (FTestWorldInstance& Client : Clients)
+	{
+		Client.Tick(TickDeltaSeconds);
+	}
+}
+
+void FTestWorlds::TickServerAndDrop()
+{
+#if DO_ENABLE_NET_TEST
+	UNetDriver* NetDriver = Server.GetNetDriver();
+	NetDriver->PacketSimulationSettings.PktLoss = 100;
+	NetDriver->OnPacketSimulationSettingsChanged();
+	
+	Server.Tick(TickDeltaSeconds);
+
+	NetDriver->PacketSimulationSettings.PktLoss = 0;
+	NetDriver->OnPacketSimulationSettingsChanged();
+#else
+	UE_LOG(LogNet, Error, TEXT("FTestWorlds::TickServerAndDrop does not work called without NetDriver Simulation Settings"));
+#endif
+}
+
+void FTestWorlds::TickClientsAndDrop()
+{
+#if DO_ENABLE_NET_TEST
+	for (FTestWorldInstance& Client : Clients)
+	{
+		UNetDriver* NetDriver = Client.GetNetDriver();
+		NetDriver->PacketSimulationSettings.PktLoss = 100;
+		NetDriver->OnPacketSimulationSettingsChanged();
+
+		Client.Tick(TickDeltaSeconds);
+		
+		NetDriver->PacketSimulationSettings.PktLoss = 0;
+		NetDriver->OnPacketSimulationSettingsChanged();
+	}
+#else
+	UE_LOG(LogNet, Error, TEXT("FTestWorlds::TickClientsAndDrop does not work without NetDriver Simulation Settings"));
+#endif
+}
+
+void FTestWorlds::TickServerAndDelay(uint32 NumFramesToDelay)
+{
+#if DO_ENABLE_NET_TEST
+	UNetDriver* NetDriver = Server.GetNetDriver();
+	NetDriver->PacketSimulationSettings.PktFrameDelay = NumFramesToDelay;
+	NetDriver->OnPacketSimulationSettingsChanged();
+
+	Server.Tick(TickDeltaSeconds);
+
+	NetDriver->PacketSimulationSettings.PktFrameDelay = 0;
+	NetDriver->OnPacketSimulationSettingsChanged();
+#else
+	UE_LOG(LogNet, Error, TEXT("FTestWorlds::TickServerAndDelay does not work without NetDriver Simulation Settings"));
+#endif
+}
+
+void FTestWorlds::TickClientsAndDelay(uint32 NumFramesToDelay)
+{
+#if DO_ENABLE_NET_TEST
+	for (FTestWorldInstance& Client : Clients)
+	{
+		UNetDriver* NetDriver = Client.GetNetDriver();
+		NetDriver->PacketSimulationSettings.PktFrameDelay = NumFramesToDelay;
+		NetDriver->OnPacketSimulationSettingsChanged();
+
+		Client.Tick(TickDeltaSeconds);
+
+		NetDriver->PacketSimulationSettings.PktFrameDelay = 0;
+		NetDriver->OnPacketSimulationSettingsChanged();
+	}
+#else
+	UE_LOG(LogNet, Error, TEXT("FTestWorlds::TickClientsAndDelay does not work without NetDriver Simulation Settings"));
+#endif
+}
+
+ APlayerController* FTestWorlds::GetServerPlayerControllerOfClient(uint32 ClientIndex)
+{
+	if (Clients.IsValidIndex(ClientIndex) == false)
+	{
+		// No client exists for that index
+		return nullptr;
+	}
+
+	int32 PlayerId = -1;
+	// Get the unique info from the PlayerController on the client world
+	{
+		APlayerController* ClientPC = Clients[ClientIndex].GetWorld()->GetFirstPlayerController();
+		if (!ClientPC || !ClientPC->PlayerState)
+		{
+			return nullptr;
+		}
+
+		PlayerId = ClientPC->PlayerState->GetPlayerId();
+	}
+
+	// Find the PlayerController on the server related to to this client
+	UWorld* ServerWorld = Server.GetWorld();
+	for (FConstPlayerControllerIterator PCIter = ServerWorld->GetPlayerControllerIterator(); PCIter; ++PCIter)
+	{
+		if (APlayerController* PC = PCIter->Get())
+		{
+			if (PC->PlayerState && PC->PlayerState->GetPlayerId() == PlayerId)
+			{
+				return PC;
+			}
+		}
+	}
+
+	return nullptr;
+}
+
+//------------------------------------------------------------------------
+
+FScopedCVarOverrideInt::FScopedCVarOverrideInt(const TCHAR* VariableName, int32 Value)
+	: Variable(IConsoleManager::Get().FindConsoleVariable(VariableName))
+{
+	if (Variable)
+	{
+		SavedValue = Variable->GetInt();
+		Variable->Set(Value);
+	}
+}
+
+FScopedCVarOverrideInt::~FScopedCVarOverrideInt()
+{
+	if (Variable)
+	{
+		Variable->Set(SavedValue);
+	}
+}
+
+FScopedTestSettings::FScopedTestSettings()
+	: AddressResolutionDisabled(TEXT("net.IpConnectionDisableResolution"), 1)
+	, BandwidthThrottlingDisabled(TEXT("net.DisableBandwithThrottling"), 1)
+	, RepGraphBandwidthThrottlingDisabled(TEXT("Net.RepGraph.DisableBandwithLimit"), 1)
+	, RandomNetUpdateDelayDisabled(TEXT("net.DisableRandomNetUpdateDelay"), 1)
+	, GameplayDebuggerDisabled(TEXT("GameplayDebugger.AutoCreateGameplayDebuggerManager"), 0)
+	, OldGWorld(GWorld)
+	, OldPIEID(UE::GetPlayInEditorID())
+	, OldGIsPlayInEditorWorld(GIsPlayInEditorWorld)
+{
+}
+
+FScopedTestSettings::~FScopedTestSettings()
+{
+	GWorld = OldGWorld;
+	UE::SetPlayInEditorID(OldPIEID);
+	GIsPlayInEditorWorld = OldGIsPlayInEditorWorld;
+}
+
+}
+
+#endif // WITH_EDITOR
\ No newline at end of file
diff --git a/Engine/Source/Runtime/Engine/Public/Tests/NetTestHelpers.h b/Engine/Source/Runtime/Engine/Public/Tests/NetTestHelpers.h
new file mode 100644
--- /dev/null
+++ b/Engine/Source/Runtime/Engine/Public/Tests/NetTestHelpers.h
@@ -0,0 +1,188 @@
+// Copyright Epic Games, Inc. All Rights Reserved.
+
+#pragma once
+
+#include "HAL/Platform.h"
+#include "Net/UnrealNetwork.h"
+#include "Engine/ActorChannel.h"
+#include "Engine/NetDriver.h"
+#include "Engine/NetworkObjectList.h"
+
+class APlayerController;
+class IConsoleVariable;
+class UGameInstance;
+class UWorld;
+
+struct FGameInstancePIEParameters;
+
+#if WITH_EDITOR
+
+namespace UE::Net
+{
+
+/**
+ * Properly scoped/RAII wrapper around a GameInstance/WorldContext/World that makes it easier to write tests
+ * involving full UWorld functionality within the scope of one function.
+ */
+struct FTestWorldInstance
+{
+	static FTestWorldInstance CreateServer(const TCHAR* InURL);
+	static FTestWorldInstance CreateClient(int32 ServerPort);
+
+	~FTestWorldInstance();
+
+	FTestWorldInstance(const FTestWorldInstance&) = delete;
+	FTestWorldInstance& operator=(const FTestWorldInstance&) = delete;
+
+	explicit FTestWorldInstance(bool bDelayedInit) : GameInstance(nullptr) {}
+	
+	FTestWorldInstance(FTestWorldInstance&& Other);
+	FTestWorldInstance& operator=(FTestWorldInstance&& Other);
+	
+	UGameInstance* GameInstance = nullptr;
+
+	UWorld* GetWorld() const;
+	FWorldContext* GetWorldContext() const;
+	UNetDriver* GetNetDriver() const;
+
+	int32 GetPort();
+
+	void Tick(float DeltaSeconds = 0.0166f);
+
+	void LoadStreamingLevel(FName LevelName);
+	void UnloadStreamingLevel(FName LevelName);
+
+private:
+	explicit FTestWorldInstance(const FGameInstancePIEParameters& InstanceParams);
+
+	void Shutdown();
+
+	static int32 FindUnusedPIEInstance();
+
+	int32 LevelStreamRequestUUID = 0;
+};
+
+/**
+ * Stores FTestWorldInstances for a server and clients and allows synchronously ticking them.
+ * Can be used within a single function to make automated tests that use the whole world & net driver flow.
+ */
+struct FTestWorlds
+{
+	/** Creates a server world using the given URL. */
+	explicit FTestWorlds(const TCHAR* ServerURL, float DeltaSeconds = 0.0166f);
+	~FTestWorlds();
+
+	bool CreateAndConnectClient();
+
+	/** Ticks all server & client worlds NumTick times synchronously. */
+	void TickAll(int32 NumTicks=1);
+	void TickServer();
+	void TickClients();
+
+	/** Tick the world and drop all outgoing packets */
+	void TickServerAndDrop();
+	void TickClientsAndDrop();
+
+	/** Tick the world but delay the packets that would be sent */
+	void TickServerAndDelay(uint32 NumFramesToDelay = 1);
+	void TickClientsAndDelay(uint32 NumFramesToDelay = 1);
+
+	/**
+	 * Ticks all server & client worlds until Predicate returns true, or MaxTicks is reached.
+	 * Returns true if Predicate did, false if it didn't within MaxTicks.
+	 */
+	template<class PredicateT>
+	bool TickAllUntil(const PredicateT& Predicate, float DeltaSeconds = 0.0166f, int32 MaxTicks = 60);
+
+	/** Ticks all server & client worlds until the passed in client world has a valid client PlayerController. */
+	bool WaitForClientConnect(FTestWorldInstance& Client);
+
+	/** Return the Server's player state corresponding to a specific client */
+	APlayerController* GetServerPlayerControllerOfClient(uint32 ClientIndex);
+
+	/** Server and Client Worlds */
+	FTestWorldInstance Server;
+	TArray<FTestWorldInstance> Clients;
+
+private:
+
+	void OnNetDriverCreated(UWorld* InWorld, UNetDriver* InNetDriver);
+	FDelegateHandle NetDriverCreatedHandle;
+
+	float TickDeltaSeconds = 0.0166f;
+};
+
+//------------------------------------------------------------------------
+// Inline functions
+//------------------------------------------------------------------------
+
+template<class PredicateT>
+inline bool FTestWorlds::TickAllUntil(const PredicateT& Predicate, float DeltaSeconds, int32 MaxTicks)
+{
+	int32 TickCount = 0;
+	bool bPredicateResult = Predicate();
+		
+	while (!bPredicateResult && TickCount < MaxTicks)
+	{
+		Server.Tick(DeltaSeconds);
+		for (FTestWorldInstance& Client : Clients)
+		{
+			Client.Tick(DeltaSeconds);
+		}
+		TickCount++;
+		GFrameCounter++;
+		bPredicateResult = Predicate();
+	}
+
+	return bPredicateResult;
+}
+
+//------------------------------------------------------------------------
+// FScopedCVarOverrideInt
+//------------------------------------------------------------------------
+class FScopedCVarOverrideInt
+{
+public:
+	FScopedCVarOverrideInt(const TCHAR* VariableName, int32 Value);
+	~FScopedCVarOverrideInt();
+
+	FScopedCVarOverrideInt(FScopedCVarOverrideInt&&) = delete;
+	FScopedCVarOverrideInt(const FScopedCVarOverrideInt&) = delete;
+	FScopedCVarOverrideInt& operator=(FScopedCVarOverrideInt&&) = delete;
+	FScopedCVarOverrideInt& operator=(const FScopedCVarOverrideInt&) = delete;
+
+private:
+	IConsoleVariable* Variable = nullptr;
+	int32 SavedValue = 0;
+};
+
+/**
+ * Sets and restores globals and cvars needed to use FNetTestWorldInstances within a scope.
+ * Meant to be used within a single function.
+ */
+class FScopedTestSettings
+{
+public:
+	FScopedTestSettings();
+	~FScopedTestSettings();
+
+	FScopedTestSettings(FScopedTestSettings&&) = delete;
+	FScopedTestSettings(const FScopedTestSettings&) = delete;
+	FScopedTestSettings& operator=(FScopedTestSettings&&) = delete;
+	FScopedTestSettings& operator=(const FScopedTestSettings&) = delete;
+
+private:
+	FScopedCVarOverrideInt AddressResolutionDisabled;
+	FScopedCVarOverrideInt BandwidthThrottlingDisabled;
+	FScopedCVarOverrideInt RepGraphBandwidthThrottlingDisabled;
+	FScopedCVarOverrideInt RandomNetUpdateDelayDisabled;
+	FScopedCVarOverrideInt GameplayDebuggerDisabled;
+
+	UWorld* OldGWorld;
+	int32 OldPIEID;
+	bool OldGIsPlayInEditorWorld;
+};
+
+} // end namespace UE::Net
+
+#endif //WITH_EDITOR
\ No newline at end of file
