Fix missing ENGINE_API

#rnx

diff --git a/Engine/Source/Runtime/Engine/Public/Tests/NetTestHelpers.h b/Engine/Source/Runtime/Engine/Public/Tests/NetTestHelpers.h
--- a/Engine/Source/Runtime/Engine/Public/Tests/NetTestHelpers.h
+++ b/Engine/Source/Runtime/Engine/Public/Tests/NetTestHelpers.h
@@ -26,31 +26,33 @@ namespace UE::Net
  */
 struct FTestWorldInstance
 {
-	static FTestWorldInstance CreateServer(const TCHAR* InURL);
-	static FTestWorldInstance CreateClient(int32 ServerPort);
+	ENGINE_API static FTestWorldInstance CreateServer(const TCHAR* InURL);
+	ENGINE_API static FTestWorldInstance CreateClient(int32 ServerPort);
 
-	~FTestWorldInstance();
+	ENGINE_API ~FTestWorldInstance();
 
 	FTestWorldInstance(const FTestWorldInstance&) = delete;
 	FTestWorldInstance& operator=(const FTestWorldInstance&) = delete;
 
 	explicit FTestWorldInstance(bool bDelayedInit) : GameInstance(nullptr) {}
 	
-	FTestWorldInstance(FTestWorldInstance&& Other);
-	FTestWorldInstance& operator=(FTestWorldInstance&& Other);
+	ENGINE_API FTestWorldInstance(FTestWorldInstance&& Other);
+	ENGINE_API FTestWorldInstance& operator=(FTestWorldInstance&& Other);
 	
-	UGameInstance* GameInstance = nullptr;
+	ENGINE_API UWorld* GetWorld() const;
+	ENGINE_API FWorldContext* GetWorldContext() const;
+	ENGINE_API UNetDriver* GetNetDriver() const;
+
+	ENGINE_API int32 GetPort();
 
-	UWorld* GetWorld() const;
-	FWorldContext* GetWorldContext() const;
-	UNetDriver* GetNetDriver() const;
+	ENGINE_API void Tick(float DeltaSeconds = 0.0166f);
 
-	int32 GetPort();
+	ENGINE_API void LoadStreamingLevel(FName LevelName);
+	ENGINE_API void UnloadStreamingLevel(FName LevelName);
 
-	void Tick(float DeltaSeconds = 0.0166f);
+public:
 
-	void LoadStreamingLevel(FName LevelName);
-	void UnloadStreamingLevel(FName LevelName);
+	UGameInstance* GameInstance = nullptr;
 
 private:
 	explicit FTestWorldInstance(const FGameInstancePIEParameters& InstanceParams);
@@ -69,23 +71,23 @@ private:
 struct FTestWorlds
 {
 	/** Creates a server world using the given URL. */
-	explicit FTestWorlds(const TCHAR* ServerURL, float DeltaSeconds = 0.0166f);
-	~FTestWorlds();
+	ENGINE_API explicit FTestWorlds(const TCHAR* ServerURL, float DeltaSeconds = 0.0166f);
+	ENGINE_API ~FTestWorlds();
 
-	bool CreateAndConnectClient();
+	ENGINE_API bool CreateAndConnectClient();
 
 	/** Ticks all server & client worlds NumTick times synchronously. */
-	void TickAll(int32 NumTicks=1);
-	void TickServer();
-	void TickClients();
+	ENGINE_API void TickAll(int32 NumTicks=1);
+	ENGINE_API void TickServer();
+	ENGINE_API void TickClients();
 
 	/** Tick the world and drop all outgoing packets */
-	void TickServerAndDrop();
-	void TickClientsAndDrop();
+	ENGINE_API void TickServerAndDrop();
+	ENGINE_API void TickClientsAndDrop();
 
 	/** Tick the world but delay the packets that would be sent */
-	void TickServerAndDelay(uint32 NumFramesToDelay = 1);
-	void TickClientsAndDelay(uint32 NumFramesToDelay = 1);
+	ENGINE_API void TickServerAndDelay(uint32 NumFramesToDelay = 1);
+	ENGINE_API void TickClientsAndDelay(uint32 NumFramesToDelay = 1);
 
 	/**
 	 * Ticks all server & client worlds until Predicate returns true, or MaxTicks is reached.
@@ -95,10 +97,12 @@ struct FTestWorlds
 	bool TickAllUntil(const PredicateT& Predicate, float DeltaSeconds = 0.0166f, int32 MaxTicks = 60);
 
 	/** Ticks all server & client worlds until the passed in client world has a valid client PlayerController. */
-	bool WaitForClientConnect(FTestWorldInstance& Client);
+	ENGINE_API bool WaitForClientConnect(FTestWorldInstance& Client);
 
 	/** Return the Server's player state corresponding to a specific client */
-	APlayerController* GetServerPlayerControllerOfClient(uint32 ClientIndex);
+	ENGINE_API APlayerController* GetServerPlayerControllerOfClient(uint32 ClientIndex);
+
+public:
 
 	/** Server and Client Worlds */
 	FTestWorldInstance Server;
