More missing ENGINE_API

#rnx

diff --git a/Engine/Source/Runtime/Engine/Public/Tests/NetTestHelpers.h b/Engine/Source/Runtime/Engine/Public/Tests/NetTestHelpers.h
--- a/Engine/Source/Runtime/Engine/Public/Tests/NetTestHelpers.h
+++ b/Engine/Source/Runtime/Engine/Public/Tests/NetTestHelpers.h
@@ -147,8 +147,8 @@ inline bool FTestWorlds::TickAllUntil(const PredicateT& Predicate, float DeltaSe
 class FScopedCVarOverrideInt
 {
 public:
-	FScopedCVarOverrideInt(const TCHAR* VariableName, int32 Value);
-	~FScopedCVarOverrideInt();
+	ENGINE_API FScopedCVarOverrideInt(const TCHAR* VariableName, int32 Value);
+	ENGINE_API ~FScopedCVarOverrideInt();
 
 	FScopedCVarOverrideInt(FScopedCVarOverrideInt&&) = delete;
 	FScopedCVarOverrideInt(const FScopedCVarOverrideInt&) = delete;
@@ -167,8 +167,8 @@ private:
 class FScopedTestSettings
 {
 public:
-	FScopedTestSettings();
-	~FScopedTestSettings();
+	ENGINE_API FScopedTestSettings();
+	ENGINE_API ~FScopedTestSettings();
 
 	FScopedTestSettings(FScopedTestSettings&&) = delete;
 	FScopedTestSettings(const FScopedTestSettings&) = delete;
