Duplicated static mesh components duplicate their MeshPaintTexture.
This prevents a single mesh paint texture from being referenced by multiple components.

diff --git a/Engine/Source/Runtime/Engine/Private/Components/StaticMeshComponent.cpp b/Engine/Source/Runtime/Engine/Private/Components/StaticMeshComponent.cpp
--- a/Engine/Source/Runtime/Engine/Private/Components/StaticMeshComponent.cpp
+++ b/Engine/Source/Runtime/Engine/Private/Components/StaticMeshComponent.cpp
@@ -760,6 +760,22 @@ void UStaticMeshComponent::InitializeComponent()
 
 void UStaticMeshComponent::PostDuplicate(bool bDuplicateForPIE)
 {
+	if (!bDuplicateForPIE && MeshPaintTexture)
+	{
+		UMeshPaintVirtualTexture* NewTexture = nullptr;
+		
+		FImage Image;
+		if (MeshPaintTexture->Source.GetMipImage(Image, 0))
+		{
+			NewTexture = NewObject<UMeshPaintVirtualTexture>(GetOutermost());
+			NewTexture->Source.Init(Image);
+			NewTexture->OwningComponent = MakeWeakObjectPtr(this);
+			NewTexture->UpdateResource();
+		}
+
+		MeshPaintTexture = NewTexture;
+	}
+
 	NotifyIfStaticMeshChanged();
 
 	Super::PostDuplicate(bDuplicateForPIE);
