Exposed EAlphaBlendOption to blueprint

diff --git a/Engine/Source/Runtime/Engine/Public/AlphaBlend.h b/Engine/Source/Runtime/Engine/Public/AlphaBlend.h
--- a/Engine/Source/Runtime/Engine/Public/AlphaBlend.h
+++ b/Engine/Source/Runtime/Engine/Public/AlphaBlend.h
@@ -8,7 +8,7 @@
 
 class UCurveFloat;
 
-UENUM()
+UENUM(BlueprintType)
 enum class EAlphaBlendOption : uint8
 {
 	Linear = 0 UMETA(Grouping = Linear, DisplayName = "Linear", ToolTip = "Linear interpolation"),
