[Audio] when deleting an ActiveSound, make sure it isn't removed from tracking until after Notify is sent.
+ this ordering was swapped in 33902729, now changing it back
+ IsPendingStop should answer "true" if it ever somehow gets called during the Notify
+ this change is not guaranteed to fix the problem reported in the ticket, but might be related?

#rnx

diff --git a/Engine/Source/Runtime/Engine/Private/AudioDevice.cpp b/Engine/Source/Runtime/Engine/Private/AudioDevice.cpp
--- a/Engine/Source/Runtime/Engine/Private/AudioDevice.cpp
+++ b/Engine/Source/Runtime/Engine/Private/AudioDevice.cpp
@@ -5531,9 +5531,6 @@ void FAudioDevice::ProcessingPendingActiveSoundStops(bool bForceDelete)
 				}
 				ActiveSound->ClearTransmitter();
 
-				// Remove from the list of pending sounds to stop
-				PendingSoundsToStop.Remove(ActiveSound);
-	
 				if (bWasVirtualized)
 				{
 					NotifySubsystemsVirtualizedSoundDeleting(*ActiveSound);
@@ -5542,6 +5539,11 @@ void FAudioDevice::ProcessingPendingActiveSoundStops(bool bForceDelete)
 				{
 					NotifySubsystemsActiveSoundDeleting(*ActiveSound);
 				}
+
+				// Remove from the list of pending sounds to stop.
+				// MUST be done last to avoid risk of double delete.
+				PendingSoundsToStop.Remove(ActiveSound);
+
 				delete ActiveSound;
 			}
 			else
