Fix EngineTest build by adding missing ENGINE_API tags at NetTestHelpers

#rnx

diff --git a/Engine/Source/Runtime/Engine/Public/Tests/NetTestHelpers.h b/Engine/Source/Runtime/Engine/Public/Tests/NetTestHelpers.h
--- a/Engine/Source/Runtime/Engine/Public/Tests/NetTestHelpers.h
+++ b/Engine/Source/Runtime/Engine/Public/Tests/NetTestHelpers.h
@@ -150,10 +150,10 @@ public:
 	ENGINE_API FScopedCVarOverrideInt(const TCHAR* VariableName, int32 Value);
 	ENGINE_API ~FScopedCVarOverrideInt();
 
-	FScopedCVarOverrideInt(FScopedCVarOverrideInt&&) = delete;
-	FScopedCVarOverrideInt(const FScopedCVarOverrideInt&) = delete;
-	FScopedCVarOverrideInt& operator=(FScopedCVarOverrideInt&&) = delete;
-	FScopedCVarOverrideInt& operator=(const FScopedCVarOverrideInt&) = delete;
+	ENGINE_API FScopedCVarOverrideInt(FScopedCVarOverrideInt&&) = delete;
+	ENGINE_API FScopedCVarOverrideInt(const FScopedCVarOverrideInt&) = delete;
+	ENGINE_API FScopedCVarOverrideInt& operator=(FScopedCVarOverrideInt&&) = delete;
+	ENGINE_API FScopedCVarOverrideInt& operator=(const FScopedCVarOverrideInt&) = delete;
 
 private:
 	IConsoleVariable* Variable = nullptr;
@@ -170,10 +170,10 @@ public:
 	ENGINE_API FScopedTestSettings();
 	ENGINE_API ~FScopedTestSettings();
 
-	FScopedTestSettings(FScopedTestSettings&&) = delete;
-	FScopedTestSettings(const FScopedTestSettings&) = delete;
-	FScopedTestSettings& operator=(FScopedTestSettings&&) = delete;
-	FScopedTestSettings& operator=(const FScopedTestSettings&) = delete;
+	ENGINE_API FScopedTestSettings(FScopedTestSettings&&) = delete;
+	ENGINE_API FScopedTestSettings(const FScopedTestSettings&) = delete;
+	ENGINE_API FScopedTestSettings& operator=(FScopedTestSettings&&) = delete;
+	ENGINE_API FScopedTestSettings& operator=(const FScopedTestSettings&) = delete;
 
 private:
 	FScopedCVarOverrideInt AddressResolutionDisabled;
