Fixed memory leak when switching skeletal meshes and setting animation mode during montages

Make sure to call UninitializeAnimation before we release the current anim instance so it can release its resources, including active montage instances

[RN] Fixed memory leak when switching skeletal meshes and setting animation mode during montages

diff --git a/Engine/Source/Runtime/Engine/Private/Components/SkeletalMeshComponent.cpp b/Engine/Source/Runtime/Engine/Private/Components/SkeletalMeshComponent.cpp
--- a/Engine/Source/Runtime/Engine/Private/Components/SkeletalMeshComponent.cpp
+++ b/Engine/Source/Runtime/Engine/Private/Components/SkeletalMeshComponent.cpp
@@ -1179,6 +1179,7 @@ void USkeletalMeshComponent::ClearAnimScriptInstance()
 		HandleExistingParallelEvaluationTask(bBlockOnTask, bPerformPostAnimEvaluation);
 
 		AnimScriptInstance->EndNotifyStates();
+		AnimScriptInstance->UninitializeAnimation();
 	}
 	AnimScriptInstance = nullptr;
 	ResetLinkedAnimInstances();
