Fix potential race condition when building skinned asset linkup data

diff --git a/Engine/Source/Runtime/Engine/Private/Animation/Skeleton.cpp b/Engine/Source/Runtime/Engine/Private/Animation/Skeleton.cpp
--- a/Engine/Source/Runtime/Engine/Private/Animation/Skeleton.cpp
+++ b/Engine/Source/Runtime/Engine/Private/Animation/Skeleton.cpp
@@ -785,9 +785,19 @@ const FSkeletonToMeshLinkup& USkeleton::AddMeshLinkupData(const USkinnedAsset* I
 	TUniquePtr<FSkeletonToMeshLinkup> TmpLinkup = MakeUnique<FSkeletonToMeshLinkup>();
 	BuildLinkupData(InSkinnedAsset, *TmpLinkup.Get());
 
+	const TObjectKey<USkinnedAsset> SkinnedAssetKey(InSkinnedAsset);
+	const uint32 SkinnedAssetHash = GetTypeHash(SkinnedAssetKey);
+
 	{
 		FTransactionallySafeWriteScopeLock Lock(SkinnedAssetLinkupCacheLock);
-		return *SkinnedAssetLinkupCache.Add(TObjectKey<USkinnedAsset>(InSkinnedAsset), MoveTemp(TmpLinkup)).Get();
+		if (const TUniquePtr<FSkeletonToMeshLinkup>* SkeletonToMeshLinkupPtr = SkinnedAssetLinkupCache.FindByHash(SkinnedAssetHash, SkinnedAssetKey))
+		{
+			// While we were building the linkup data, another thread beat us to it
+			// Discard the work we did to avoid freeing memory other threads might now be using
+			return *SkeletonToMeshLinkupPtr->Get();
+		}
+
+		return *SkinnedAssetLinkupCache.AddByHash(SkinnedAssetHash, SkinnedAssetKey, MoveTemp(TmpLinkup)).Get();
 	}
 }
 
