Added kismet function overloads for BoxOverlap, CapsuleOverlap, and the corresponding Actor functions to take an orientation.

diff --git a/Engine/Source/Runtime/Engine/Classes/Kismet/KismetSystemLibrary.h b/Engine/Source/Runtime/Engine/Classes/Kismet/KismetSystemLibrary.h
--- a/Engine/Source/Runtime/Engine/Classes/Kismet/KismetSystemLibrary.h
+++ b/Engine/Source/Runtime/Engine/Classes/Kismet/KismetSystemLibrary.h
@@ -1062,6 +1062,21 @@ class UKismetSystemLibrary : public UBlueprintFunctionLibrary
 	UFUNCTION(BlueprintCallable, Category="Collision", meta=(WorldContext="WorldContextObject", AutoCreateRefTerm="ActorsToIgnore", DisplayName="Box Overlap Actors"))
 	static ENGINE_API bool BoxOverlapActors(const UObject* WorldContextObject, const FVector BoxPos, FVector BoxExtent, const TArray<TEnumAsByte<EObjectTypeQuery> > & ObjectTypes, UClass* ActorClassFilter, const TArray<AActor*>& ActorsToIgnore, TArray<class AActor*>& OutActors);
 
+	/**
+	 * Returns an array of actors that overlap the given axis-aligned box.
+	 * @param WorldContext	World context
+	 * @param BoxPos		Center of box.
+	 * @param BoxExtent		Extents of box.
+	 * @param Orientation	Orientation of the box.
+	 * @param Filter		Option to restrict results to only static or only dynamic.  For efficiency.
+	 * @param ClassFilter	If set, will only return results of this class or subclasses of it.
+	 * @param ActorsToIgnore		Ignore these actors in the list
+	 * @param OutActors		Returned array of actors. Unsorted.
+	 * @return				true if there was an overlap that passed the filters, false otherwise.
+	 */
+	UFUNCTION(BlueprintCallable, Category="Collision", meta=(WorldContext="WorldContextObject", AutoCreateRefTerm="ActorsToIgnore", DisplayName="Box Overlap Actors With Orientation"))
+	static ENGINE_API bool BoxOverlapActorsWithOrientation(const UObject* WorldContextObject, const FVector BoxPos, FVector BoxExtent, FRotator Orientation, const TArray<TEnumAsByte<EObjectTypeQuery> > & ObjectTypes, UClass* ActorClassFilter, const TArray<AActor*>& ActorsToIgnore, TArray<class AActor*>& OutActors);
+
 	/**
 	 * Returns an array of components that overlap the given axis-aligned box.
 	 * @param WorldContext	World context
@@ -1076,6 +1091,21 @@ class UKismetSystemLibrary : public UBlueprintFunctionLibrary
 	UFUNCTION(BlueprintCallable, Category="Collision", meta=(WorldContext="WorldContextObject", AutoCreateRefTerm="ActorsToIgnore", DisplayName="Box Overlap Components"))
 	static ENGINE_API bool BoxOverlapComponents(const UObject* WorldContextObject, const FVector BoxPos, FVector Extent, const TArray<TEnumAsByte<EObjectTypeQuery> > & ObjectTypes, UClass* ComponentClassFilter, const TArray<AActor*>& ActorsToIgnore, TArray<class UPrimitiveComponent*>& OutComponents);
 
+	/**
+	 * Returns an array of components that overlap the given box.
+	 * @param WorldContext	World context
+	 * @param BoxPos		Center of box.
+	 * @param Extent		Extents of box.
+	 * @param Orientation	Orientation of the box.
+	 * @param Filter		Option to restrict results to only static or only dynamic.  For efficiency.
+	 * @param ClassFilter	If set, will only return results of this class or subclasses of it.
+	 * @param ActorsToIgnore		Ignore these actors in the list
+	 * @param OutActors		Returned array of actors. Unsorted.
+	 * @return				true if there was an overlap that passed the filters, false otherwise.
+	 */
+	UFUNCTION(BlueprintCallable, Category="Collision", meta=(WorldContext="WorldContextObject", AutoCreateRefTerm="ActorsToIgnore", DisplayName="Box Overlap Components With Rotation"))
+	static ENGINE_API bool BoxOverlapComponentsWithOrientation(const UObject* WorldContextObject, const FVector BoxPos, FVector Extent, FRotator Orientation, const TArray<TEnumAsByte<EObjectTypeQuery> > & ObjectTypes, UClass* ComponentClassFilter, const TArray<AActor*>& ActorsToIgnore, TArray<class UPrimitiveComponent*>& OutComponents);
+
 
 	/**
 	 * Returns an array of actors that overlap the given capsule.
@@ -1092,6 +1122,22 @@ class UKismetSystemLibrary : public UBlueprintFunctionLibrary
 	UFUNCTION(BlueprintCallable, Category="Collision", meta=(WorldContext="WorldContextObject", AutoCreateRefTerm="ActorsToIgnore", DisplayName="Capsule Overlap Actors"))
 	static ENGINE_API bool CapsuleOverlapActors(const UObject* WorldContextObject, const FVector CapsulePos, float Radius, float HalfHeight, const TArray<TEnumAsByte<EObjectTypeQuery> > & ObjectTypes, UClass* ActorClassFilter, const TArray<AActor*>& ActorsToIgnore, TArray<class AActor*>& OutActors);
 
+	/**
+	 * Returns an array of actors that overlap the given capsule.
+	 * @param WorldContext	World context
+	 * @param CapsulePos	Center of the capsule.
+	 * @param Radius		Radius of capsule hemispheres and radius of center cylinder portion.
+	 * @param HalfHeight	Half-height of the capsule (from center of capsule to tip of hemisphere.
+	 * @param Orientation	Orientation of the box.
+	 * @param Filter		Option to restrict results to only static or only dynamic.  For efficiency.
+	 * @param ClassFilter	If set, will only return results of this class or subclasses of it.
+	 * @param ActorsToIgnore		Ignore these actors in the list
+	 * @param OutActors		Returned array of actors. Unsorted.
+	 * @return				true if there was an overlap that passed the filters, false otherwise.
+	 */
+	UFUNCTION(BlueprintCallable, Category="Collision", meta=(WorldContext="WorldContextObject", AutoCreateRefTerm="ActorsToIgnore", DisplayName="Capsule Overlap Actors With Orientation"))
+	static ENGINE_API bool CapsuleOverlapActorsWithOrientation(const UObject* WorldContextObject, const FVector CapsulePos, float Radius, float HalfHeight, FRotator Orientation, const TArray<TEnumAsByte<EObjectTypeQuery> > & ObjectTypes, UClass* ActorClassFilter, const TArray<AActor*>& ActorsToIgnore, TArray<class AActor*>& OutActors);
+
 	/**
 	 * Returns an array of components that overlap the given capsule.
 	 * @param WorldContext	World context
@@ -1107,6 +1153,21 @@ class UKismetSystemLibrary : public UBlueprintFunctionLibrary
 	UFUNCTION(BlueprintCallable, Category="Collision", meta=(WorldContext="WorldContextObject", AutoCreateRefTerm="ActorsToIgnore", DisplayName="Capsule Overlap Components") )
 	static ENGINE_API bool CapsuleOverlapComponents(const UObject* WorldContextObject, const FVector CapsulePos, float Radius, float HalfHeight, const TArray<TEnumAsByte<EObjectTypeQuery> > & ObjectTypes, UClass* ComponentClassFilter, const TArray<AActor*>& ActorsToIgnore, TArray<class UPrimitiveComponent*>& OutComponents);
 
+	/**
+	 * Returns an array of components that overlap the given capsule.
+	 * @param WorldContext	World context
+	 * @param CapsulePos	Center of the capsule.
+	 * @param Radius		Radius of capsule hemispheres and radius of center cylinder portion.
+	 * @param HalfHeight	Half-height of the capsule (from center of capsule to tip of hemisphere.
+	 * @param Orientation	Orientation of the box.
+	 * @param Filter		Option to restrict results to only static or only dynamic.  For efficiency.
+	 * @param ClassFilter	If set, will only return results of this class or subclasses of it.
+	 * @param ActorsToIgnore		Ignore these actors in the list
+	 * @param OutActors		Returned array of actors. Unsorted.
+	 * @return				true if there was an overlap that passed the filters, false otherwise.
+	 */
+	UFUNCTION(BlueprintCallable, Category="Collision", meta=(WorldContext="WorldContextObject", AutoCreateRefTerm="ActorsToIgnore", DisplayName="Capsule Overlap Components With Orietnation") )
+	static ENGINE_API bool CapsuleOverlapComponentsWithOrientation(const UObject* WorldContextObject, const FVector CapsulePos, float Radius, float HalfHeight, FRotator Orientation, const TArray<TEnumAsByte<EObjectTypeQuery> > & ObjectTypes, UClass* ComponentClassFilter, const TArray<AActor*>& ActorsToIgnore, TArray<class UPrimitiveComponent*>& OutComponents);
 
 	/**
 	 * Returns an array of actors that overlap the given component.
diff --git a/Engine/Source/Runtime/Engine/Private/KismetSystemLibrary.cpp b/Engine/Source/Runtime/Engine/Private/KismetSystemLibrary.cpp
--- a/Engine/Source/Runtime/Engine/Private/KismetSystemLibrary.cpp
+++ b/Engine/Source/Runtime/Engine/Private/KismetSystemLibrary.cpp
@@ -1670,11 +1670,16 @@ bool UKismetSystemLibrary::SphereOverlapComponents(const UObject* WorldContextOb
 
 
 bool UKismetSystemLibrary::BoxOverlapActors(const UObject* WorldContextObject, const FVector BoxPos, FVector BoxExtent, const TArray<TEnumAsByte<EObjectTypeQuery> > & ObjectTypes, UClass* ActorClassFilter, const TArray<AActor*>& ActorsToIgnore, TArray<AActor*>& OutActors)
+{
+	return BoxOverlapActorsWithOrientation(WorldContextObject, BoxPos, BoxExtent, FRotator::ZeroRotator, ObjectTypes, ActorClassFilter, ActorsToIgnore, OutActors);
+}
+
+bool UKismetSystemLibrary::BoxOverlapActorsWithOrientation(const UObject* WorldContextObject, const FVector BoxPos, FVector BoxExtent, FRotator Orientation, const TArray<TEnumAsByte<EObjectTypeQuery> > & ObjectTypes, UClass* ActorClassFilter, const TArray<AActor*>& ActorsToIgnore, TArray<AActor*>& OutActors)
 {
 	OutActors.Empty();
 
 	TArray<UPrimitiveComponent*> OverlapComponents;
-	bool bOverlapped = BoxOverlapComponents(WorldContextObject, BoxPos, BoxExtent, ObjectTypes, NULL, ActorsToIgnore, OverlapComponents);
+	bool bOverlapped = BoxOverlapComponentsWithOrientation(WorldContextObject, BoxPos, BoxExtent, Orientation, ObjectTypes, NULL, ActorsToIgnore, OverlapComponents);
 	if (bOverlapped)
 	{
 		GetActorListFromComponentList(OverlapComponents, ActorClassFilter, OutActors);
@@ -1684,6 +1689,11 @@ bool UKismetSystemLibrary::BoxOverlapActors(const UObject* WorldContextObject, c
 }
 
 bool UKismetSystemLibrary::BoxOverlapComponents(const UObject* WorldContextObject, const FVector BoxPos, FVector BoxExtent, const TArray<TEnumAsByte<EObjectTypeQuery> > & ObjectTypes, UClass* ComponentClassFilter, const TArray<AActor*>& ActorsToIgnore, TArray<UPrimitiveComponent*>& OutComponents)
+{
+	return BoxOverlapComponentsWithOrientation(WorldContextObject, BoxPos, BoxExtent, FRotator::ZeroRotator, ObjectTypes, ComponentClassFilter, ActorsToIgnore, OutComponents);
+}
+
+bool UKismetSystemLibrary::BoxOverlapComponentsWithOrientation(const UObject* WorldContextObject, const FVector BoxPos, FVector BoxExtent, FRotator Orientation, const TArray<TEnumAsByte<EObjectTypeQuery> >& ObjectTypes, UClass* ComponentClassFilter, const TArray<AActor*>& ActorsToIgnore, TArray<UPrimitiveComponent*>& OutComponents)
 {
 	OutComponents.Empty();
 
@@ -1695,22 +1705,22 @@ bool UKismetSystemLibrary::BoxOverlapComponents(const UObject* WorldContextObjec
 	FCollisionObjectQueryParams ObjectParams;
 	for (auto Iter = ObjectTypes.CreateConstIterator(); Iter; ++Iter)
 	{
-		const ECollisionChannel & Channel = UCollisionProfile::Get()->ConvertToCollisionChannel(false, *Iter);
+		const ECollisionChannel& Channel = UCollisionProfile::Get()->ConvertToCollisionChannel(false, *Iter);
 		ObjectParams.AddObjectTypesToQuery(Channel);
 	}
 
 	UWorld* World = GEngine->GetWorldFromContextObject(WorldContextObject, EGetWorldErrorMode::LogAndReturnNull);
 	if (World != nullptr)
 	{
-		World->OverlapMultiByObjectType(Overlaps, BoxPos, FQuat::Identity, ObjectParams, FCollisionShape::MakeBox(BoxExtent), Params);
+		World->OverlapMultiByObjectType(Overlaps, BoxPos, Orientation.Quaternion(), ObjectParams, FCollisionShape::MakeBox(BoxExtent), Params);
 	}
 
-	for (int32 OverlapIdx=0; OverlapIdx<Overlaps.Num(); ++OverlapIdx)
+	for (int32 OverlapIdx = 0; OverlapIdx < Overlaps.Num(); ++OverlapIdx)
 	{
 		FOverlapResult const& O = Overlaps[OverlapIdx];
 		if (O.Component.IsValid())
-		{ 
-			if ( !ComponentClassFilter || O.Component.Get()->IsA(ComponentClassFilter) )
+		{
+			if (!ComponentClassFilter || O.Component.Get()->IsA(ComponentClassFilter))
 			{
 				OutComponents.Add(O.Component.Get());
 			}
@@ -1721,12 +1731,17 @@ bool UKismetSystemLibrary::BoxOverlapComponents(const UObject* WorldContextObjec
 }
 
 
-bool UKismetSystemLibrary::CapsuleOverlapActors(const UObject* WorldContextObject, const FVector CapsulePos, float Radius, float HalfHeight, const TArray<TEnumAsByte<EObjectTypeQuery> > & ObjectTypes, UClass* ActorClassFilter, const TArray<AActor*>& ActorsToIgnore, TArray<AActor*>& OutActors)
+bool UKismetSystemLibrary::CapsuleOverlapActors(const UObject* WorldContextObject, const FVector CapsulePos, float Radius, float HalfHeight, const TArray<TEnumAsByte<EObjectTypeQuery> >& ObjectTypes, UClass* ActorClassFilter, const TArray<AActor*>& ActorsToIgnore, TArray<AActor*>& OutActors)
+{
+	return CapsuleOverlapActorsWithOrientation(WorldContextObject, CapsulePos, Radius, HalfHeight, FRotator::ZeroRotator, ObjectTypes, ActorClassFilter, ActorsToIgnore, OutActors);
+}
+
+bool UKismetSystemLibrary::CapsuleOverlapActorsWithOrientation(const UObject* WorldContextObject, const FVector CapsulePos, float Radius, float HalfHeight, FRotator Orientation, const TArray<TEnumAsByte<EObjectTypeQuery> > & ObjectTypes, UClass* ActorClassFilter, const TArray<AActor*>& ActorsToIgnore, TArray<AActor*>& OutActors)
 {
 	OutActors.Empty();
 
 	TArray<UPrimitiveComponent*> OverlapComponents;
-	bool bOverlapped = CapsuleOverlapComponents(WorldContextObject, CapsulePos, Radius, HalfHeight, ObjectTypes, NULL, ActorsToIgnore, OverlapComponents);
+	bool bOverlapped = CapsuleOverlapComponentsWithOrientation(WorldContextObject, CapsulePos, Radius, HalfHeight, Orientation, ObjectTypes, NULL, ActorsToIgnore, OverlapComponents);
 	if (bOverlapped)
 	{
 		GetActorListFromComponentList(OverlapComponents, ActorClassFilter, OutActors);
@@ -1735,7 +1750,12 @@ bool UKismetSystemLibrary::CapsuleOverlapActors(const UObject* WorldContextObjec
 	return (OutActors.Num() > 0);
 }
 
-bool UKismetSystemLibrary::CapsuleOverlapComponents(const UObject* WorldContextObject, const FVector CapsulePos, float Radius, float HalfHeight, const TArray<TEnumAsByte<EObjectTypeQuery> > & ObjectTypes, UClass* ComponentClassFilter, const TArray<AActor*>& ActorsToIgnore, TArray<UPrimitiveComponent*>& OutComponents)
+bool UKismetSystemLibrary::CapsuleOverlapComponents(const UObject* WorldContextObject, const FVector CapsulePos, float Radius, float HalfHeight, const TArray<TEnumAsByte<EObjectTypeQuery> >& ObjectTypes, UClass* ComponentClassFilter, const TArray<AActor*>& ActorsToIgnore, TArray<UPrimitiveComponent*>& OutComponents)
+{
+	return CapsuleOverlapComponentsWithOrientation(WorldContextObject, CapsulePos, Radius, HalfHeight, FRotator::ZeroRotator, ObjectTypes, ComponentClassFilter, ActorsToIgnore, OutComponents);
+}
+
+bool UKismetSystemLibrary::CapsuleOverlapComponentsWithOrientation(const UObject* WorldContextObject, const FVector CapsulePos, float Radius, float HalfHeight, FRotator Orientation, const TArray<TEnumAsByte<EObjectTypeQuery> > & ObjectTypes, UClass* ComponentClassFilter, const TArray<AActor*>& ActorsToIgnore, TArray<UPrimitiveComponent*>& OutComponents)
 {
 	OutComponents.Empty();
 
@@ -1754,7 +1774,7 @@ bool UKismetSystemLibrary::CapsuleOverlapComponents(const UObject* WorldContextO
 	UWorld* World = GEngine->GetWorldFromContextObject(WorldContextObject, EGetWorldErrorMode::LogAndReturnNull);
 	if (World != nullptr)
 	{
-		World->OverlapMultiByObjectType(Overlaps, CapsulePos, FQuat::Identity, ObjectParams, FCollisionShape::MakeCapsule(Radius, HalfHeight), Params);
+		World->OverlapMultiByObjectType(Overlaps, CapsulePos, Orientation.Quaternion(), ObjectParams, FCollisionShape::MakeCapsule(Radius, HalfHeight), Params);
 	}
 
 	for (int32 OverlapIdx=0; OverlapIdx<Overlaps.Num(); ++OverlapIdx)
