Remove unecessary log from ApplyImpact

diff --git a/Engine/Source/Runtime/Engine/Private/Components/PawnMovementComponent.cpp b/Engine/Source/Runtime/Engine/Private/Components/PawnMovementComponent.cpp
--- a/Engine/Source/Runtime/Engine/Private/Components/PawnMovementComponent.cpp
+++ b/Engine/Source/Runtime/Engine/Private/Components/PawnMovementComponent.cpp
@@ -194,14 +194,6 @@ void UPawnMovementComponent::ExecuteAsyncPhysicsStateAction(const UPrimitiveComp
 
 	if (FBodyInstance* BI = ActionComponent->GetBodyInstance(BoneName))
 	{
-		if (GetOwner()->IsNetMode(NM_Client) || GetOwner()->IsNetMode(NM_Standalone))
-		{
-			UE_LOG(LogTemp, Warning, TEXT("ApplyImpactAtLocationImp CLIENT Force = %s | Location = %s | LocalFrame = %d | ServerFrame = %d | ComponentName = %s | ComponentPtr = %d"), *ActionDatas.ToString(), *ActionPosition.ToString(), Timestamp.LocalFrame, Timestamp.ServerFrame, *ActionComponent->GetPathName(), ActionComponent);
-		}
-		else
-		{
-			UE_LOG(LogTemp, Warning, TEXT("ApplyImpactAtLocationImp SERVER Force = %s | Location = %s | LocalFrame = %d | ServerFrame = %d | ComponentName = %s | ComponentPtr = %d"), *ActionDatas.ToString(), *ActionPosition.ToString(), Timestamp.LocalFrame, Timestamp.ServerFrame, *ActionComponent->GetPathName(), ActionComponent);
-		}
 		APawn* LocalPawn = Cast<APawn>(GetOwner());
 		APlayerController* PlayerController = (LocalPawn && LocalPawn->GetController()) ?
 			Cast<APlayerController>(LocalPawn->GetController()) : GetWorld()->GetFirstPlayerController();
