- Added bounds cache in UInstancedStaticMeshComponent to avoid recomputing the bounds.
- Replaced NavigationBounds with this new cache
- Saves ~43% of the time taken by UInstancedStaticMeshComponent CalcBounds/CalcAndCacheNavigationBounds in CitySample Big_City_LVL.

diff --git a/Engine/Source/Runtime/Engine/Classes/Components/InstancedStaticMeshComponent.h b/Engine/Source/Runtime/Engine/Classes/Components/InstancedStaticMeshComponent.h
--- a/Engine/Source/Runtime/Engine/Classes/Components/InstancedStaticMeshComponent.h
+++ b/Engine/Source/Runtime/Engine/Classes/Components/InstancedStaticMeshComponent.h
@@ -113,10 +113,6 @@ class UInstancedStaticMeshComponent : public UStaticMeshComponent, public ISMIns
 	UPROPERTY(Transient)
 	TArray<FMatrix> PerInstancePrevTransform;
 
-	/** Bounds are calculated and cached on component registration. */
-	UPROPERTY(Transient)
-	FBox NavigationBounds;
-	
 	/** Main transform stored to be able to send updates when component's transform changed. */
 	UPROPERTY(Transient)
 	FTransform PreviousComponentTransform;
@@ -184,6 +180,9 @@ private:
 	/** Flag for whether we are using conservative bounds. */
 	uint8 bUseConservativeBounds : 1;
 
+	/** Internally used to know when CalcBounds is called by UpdateBounds. */
+	uint8 bIsUpdatingBounds : 1;
+
 public:
 	/** Mapping from PerInstanceSMData order to instance render buffer order. If empty, the PerInstanceSMData order is used. */
 	UPROPERTY()
@@ -647,6 +646,8 @@ protected:
 
 	FPrimitiveInstanceDataManager PrimitiveInstanceDataManager;
 
+	void InvalidateCachedBounds();
+
 	ENGINE_API void CalcAndCacheNavigationBounds();
 
 	/** Creates body instances for all instances owned by this component. */
@@ -658,11 +659,15 @@ protected:
 	/** Request to navigation system to update for the bounds of the ISM. */
 	ENGINE_API virtual void FullNavigationUpdate();
 
+	UE_DEPRECATED(5.5, "Use version with FBox")
+	FBoxSphereBounds CalcBoundsImpl(const FTransform& BoundTransform, bool bForNavigation) const;
+
 	/**
 	 * Calculates bounds from all instances.
-	 * @param bForNavigation Indicates if NavCollision must be used if available, using static mesh bounds otherwise.
+	 * @param BoundTransform The bounds transform
+	 * @param InstanceBounds The instance bounds
 	 */
-	FBoxSphereBounds CalcBoundsImpl(const FTransform& BoundTransform, bool bForNavigation) const;
+	FBoxSphereBounds CalcBoundsImpl(const FTransform& BoundTransform, const FBox& InstanceBounds) const;
 	
 	/** Does this component support partial navigation updates */
 	virtual bool SupportsPartialNavigationUpdate() const { return true; }
@@ -729,6 +734,73 @@ protected:
 	/** Used for dynamic stats */
 	TStatId StatId;
 #endif
+
+private:
+
+	struct FBoundsCacheElement
+	{
+		FBoundsCacheElement() 
+			: bIsValid(false)
+			, Hash(0)
+			, Value(ForceInit)
+		{}
+
+		void Invalidate()
+		{
+			bIsValid = false;
+		}
+
+		bool IsValid() const
+		{
+			return bIsValid;
+		}
+
+		void CacheValue(const FBoxSphereBounds& InValue, uint32 InHash = 0)
+		{
+			bIsValid = true;
+			Value = InValue;
+			Hash = InHash;
+		}
+
+		const FBoxSphereBounds& GetValue(bool bCheckValid = true) const
+		{
+			ensure(!bCheckValid || bIsValid);
+			return Value;
+		}
+
+		bool operator == (uint32 InHash) const
+		{
+			return bIsValid && (Hash == InHash);
+		}
+
+		bool operator != (uint32 InHash) const
+		{
+			return !(*this == InHash);
+		}
+
+		static uint32 GetHash(const FTransform& InBoundTransform, const FBox& InInstanceBounds)
+		{
+			return HashCombineFast(HashCombineFast(GetTypeHash(InBoundTransform), GetTypeHash(InInstanceBounds.Min)), GetTypeHash(InInstanceBounds.Max));
+		}
+
+	private:
+		bool bIsValid;
+		uint32 Hash;
+		FBoxSphereBounds Value;
+	};
+
+	enum class EBoundsCacheType : uint8
+	{
+		Component,
+		Navigation,
+		Local,
+		Count,
+		Other
+	};
+
+	FBoxSphereBounds GetOrCalcBounds(const FTransform& InBoundTransform, EBoundsCacheType InType) const;
+
+	mutable FBoundsCacheElement CachedBounds[(int)EBoundsCacheType::Count];
 };
 
 /** InstancedStaticMeshInstance hit proxy */
diff --git a/Engine/Source/Runtime/Engine/Private/InstancedStaticMesh.cpp b/Engine/Source/Runtime/Engine/Private/InstancedStaticMesh.cpp
--- a/Engine/Source/Runtime/Engine/Private/InstancedStaticMesh.cpp
+++ b/Engine/Source/Runtime/Engine/Private/InstancedStaticMesh.cpp
@@ -2054,6 +2054,7 @@ UInstancedStaticMeshComponent::UInstancedStaticMeshComponent(const FObjectInitia
 
 	bUseGpuLodSelection = true;
 	bInheritPerInstanceData = false;
+	bIsUpdatingBounds = false;
 
 #if STATS
 	{
@@ -2704,41 +2705,99 @@ bool UInstancedStaticMeshComponent::CanEditSimulatePhysics()
 	return false;
 }
 
-FBoxSphereBounds UInstancedStaticMeshComponent::CalcBounds(const FTransform& BoundTransform) const
+void UInstancedStaticMeshComponent::InvalidateCachedBounds()
+{
+	static_assert(UE_ARRAY_COUNT(CachedBounds) == int32(EBoundsCacheType::Count));
+	for (int i = 0; i < (int)EBoundsCacheType::Count; ++i)
+	{
+		CachedBounds[i].Invalidate();
+	}
+}
+
+FBoxSphereBounds UInstancedStaticMeshComponent::CalcBounds(const FTransform& InBoundTransform) const
 {
 	TRACE_CPUPROFILER_EVENT_SCOPE(UInstancedStaticMeshComponent::CalcBounds);
-	return CalcBoundsImpl(BoundTransform, /*bForNavigation*/false);
+	EBoundsCacheType Type = bIsUpdatingBounds ? EBoundsCacheType::Component : InBoundTransform.Equals(FTransform::Identity) ? EBoundsCacheType::Local : EBoundsCacheType::Other;
+	return GetOrCalcBounds(InBoundTransform, Type);
 }
 
 void UInstancedStaticMeshComponent::CalcAndCacheNavigationBounds()
 {
 	TRACE_CPUPROFILER_EVENT_SCOPE(UInstancedStaticMeshComponent::CalcAndCacheNavigationBounds);
-	NavigationBounds = CalcBoundsImpl(GetComponentTransform(), /*bForNavigation*/true).GetBox();
+	GetOrCalcBounds(GetComponentTransform(), EBoundsCacheType::Navigation).GetBox();
 }
 
-FBoxSphereBounds UInstancedStaticMeshComponent::CalcBoundsImpl(const FTransform& BoundTransform, const bool bForNavigation) const
+FBoxSphereBounds UInstancedStaticMeshComponent::GetOrCalcBounds(const FTransform& InBoundTransform, EBoundsCacheType InType) const
 {
 	if (GetStaticMesh() && PerInstanceSMData.Num() > 0)
 	{
-		const FBox InstanceBounds = bForNavigation ? GetInstanceNavigationBounds() : GetStaticMesh()->GetBounds().GetBox();
+		const FBox InstanceBounds = (InType == EBoundsCacheType::Navigation) ? GetInstanceNavigationBounds() : GetStaticMesh()->GetBounds().GetBox();
 		if (InstanceBounds.IsValid)
 		{
-			const FMatrix BoundTransformMatrix = BoundTransform.ToMatrixWithScale();
-			FBoxSphereBounds::Builder BoundsBuilder;
-			for (int32 InstanceIndex = 0; InstanceIndex < PerInstanceSMData.Num(); InstanceIndex++)
+			const UWorld* const World = GetWorld();
+			const bool bIsGameWorld = World && World->IsGameWorld();
+			if (!bIsGameWorld || InType == EBoundsCacheType::Other)
 			{
-				BoundsBuilder += InstanceBounds.TransformBy(PerInstanceSMData[InstanceIndex].Transform * BoundTransformMatrix);
+				FBoxSphereBounds Value = CalcBoundsImpl(InBoundTransform, InstanceBounds);
+				if (InType == EBoundsCacheType::Navigation)
+				{
+					// Store result for Navigation Cache as GetNavigationBounds() relies on it.
+					CachedBounds[(int32)EBoundsCacheType::Navigation].CacheValue(Value);
+				}
+				return Value;
 			}
 
-			return BoundsBuilder;
+			const uint32 Hash = FBoundsCacheElement::GetHash(InBoundTransform, InstanceBounds);
+			int32 TypeAsInt = (int32)InType;
+			FBoundsCacheElement& TargetBounds = CachedBounds[TypeAsInt];
+			if (TargetBounds != Hash)
+			{
+				for (int32 TypeIndex = 0; TypeIndex < (int32)EBoundsCacheType::Count; ++TypeIndex)
+				{
+					if ((TypeIndex != TypeAsInt) && CachedBounds[TypeIndex] == Hash)
+					{
+						TargetBounds = CachedBounds[TypeIndex];
+						return TargetBounds.GetValue();
+					}
+				}
+				TargetBounds.CacheValue(CalcBoundsImpl(InBoundTransform, InstanceBounds), Hash);
+			}
+			return TargetBounds.GetValue();
+		}
+	}
+	return FBoxSphereBounds(InBoundTransform.GetLocation(), FVector::ZeroVector, 0.f);
+}
+
+// Deprecated
+FBoxSphereBounds UInstancedStaticMeshComponent::CalcBoundsImpl(const FTransform& BoundTransform, bool bForNavigation) const
+{
+	if (GetStaticMesh() && PerInstanceSMData.Num() > 0)
+	{
+		const FBox InstanceBounds = bForNavigation ? GetInstanceNavigationBounds() : GetStaticMesh()->GetBounds().GetBox();
+		if (InstanceBounds.IsValid)
+		{
+			return CalcBoundsImpl(BoundTransform, InstanceBounds);
 		}
 	}
-	
 	return FBoxSphereBounds(BoundTransform.GetLocation(), FVector::ZeroVector, 0.f);
 }
 
+FBoxSphereBounds UInstancedStaticMeshComponent::CalcBoundsImpl(const FTransform& BoundTransform, const FBox& InstanceBounds) const
+{
+	check(PerInstanceSMData.Num() > 0);
+	const FMatrix BoundTransformMatrix = BoundTransform.ToMatrixWithScale();
+	FBoxSphereBounds::Builder BoundsBuilder;
+	for (int32 InstanceIndex = 0; InstanceIndex < PerInstanceSMData.Num(); InstanceIndex++)
+	{
+		BoundsBuilder += InstanceBounds.TransformBy(PerInstanceSMData[InstanceIndex].Transform * BoundTransformMatrix);
+	}
+	return BoundsBuilder;
+}
+
 void UInstancedStaticMeshComponent::UpdateBounds()
 {
+	FGuardValue_Bitfield(bIsUpdatingBounds, true);
+	
 	if (GetNumRenderInstances() == 0)
 	{
 		Bounds = FBoxSphereBounds(GetComponentTransform().GetLocation(), FVector::ZeroVector, 0.f);
@@ -3248,6 +3307,7 @@ int32 UInstancedStaticMeshComponent::AddInstanceInternal(int32 InstanceIndex, FI
 		FInstancedStaticMeshDelegates::OnInstanceIndexUpdated.Broadcast(this, MakeArrayView(&IndexUpdate, 1));
 	}
 
+	InvalidateCachedBounds();
 	return InstanceIndex;
 }
 
@@ -3346,6 +3406,7 @@ TArray<int32> UInstancedStaticMeshComponent::AddInstancesInternal(TConstArrayVie
 		}
 	}
 
+	InvalidateCachedBounds();
 	return NewInstanceIndices;
 }
 
@@ -3664,6 +3725,8 @@ bool UInstancedStaticMeshComponent::RemoveInstanceInternal(int32 InstanceIndex,
 #if WITH_EDITOR
 	DeletionState = EInstanceDeletionReason::NotDeleting;
 #endif
+
+	InvalidateCachedBounds();
 	return true;
 }
 
@@ -3851,6 +3914,7 @@ bool UInstancedStaticMeshComponent::UpdateInstanceTransform(const int32 Instance
 	}
 
 	Modify();
+	InvalidateCachedBounds();
 
 	FInstancedStaticMeshInstanceData& InstanceData = PerInstanceSMData[InstanceIndex];
 
@@ -3913,6 +3977,7 @@ bool UInstancedStaticMeshComponent::BatchUpdateInstancesTransforms(int32 StartIn
 	check(PerInstancePrevTransform.IsValidIndex(StartInstanceIndex) && PerInstancePrevTransform.IsValidIndex(StartInstanceIndex + NewInstancesPrevTransforms.Num() - 1));
 
 	Modify();
+	InvalidateCachedBounds();
 
 	// Navigation update required if the component is relevant to navigation and registered.
 	const bool bNavigationUpdateRequired = bNavigationRelevant && IsRegistered();
@@ -4028,6 +4093,7 @@ PRAGMA_DISABLE_DEPRECATION_WARNINGS
 	}
 
 	Modify();
+	InvalidateCachedBounds();
 
 #if (CSV_PROFILER_STATS)
 	int32 TotalSizeUpdateBytes = 0;
@@ -4250,6 +4316,7 @@ bool UInstancedStaticMeshComponent::BatchUpdateInstancesTransformsInternal(int32
 	}
 
 	Modify();
+	InvalidateCachedBounds();
 
 	// Navigation update required if the component is relevant to navigation and registered.
 	const bool bNavigationUpdateRequired = bNavigationRelevant && IsRegistered();
@@ -4323,6 +4390,7 @@ bool UInstancedStaticMeshComponent::BatchUpdateInstancesTransform(int32 StartIns
 	}
 
 	Modify();
+	InvalidateCachedBounds();
 
 	// Navigation update required if the component is relevant to navigation and registered.
 	const bool bNavigationUpdateRequired = bNavigationRelevant && IsRegistered();
@@ -4393,6 +4461,7 @@ bool UInstancedStaticMeshComponent::BatchUpdateInstancesData(int32 StartInstance
 	}
 
 	Modify();
+	InvalidateCachedBounds();
 
 	// Navigation update required if the component is relevant to navigation and registered.
 	const bool bNavigationUpdateRequired = bNavigationRelevant && IsRegistered();
@@ -4569,6 +4638,8 @@ void UInstancedStaticMeshComponent::GetStreamingRenderAssetInfo(FStreamingTextur
 
 void UInstancedStaticMeshComponent::ClearInstances()
 {
+	InvalidateCachedBounds();
+
 #if WITH_EDITOR
 	DeletionState = EInstanceDeletionReason::Clearing;
 #endif
@@ -5429,8 +5500,9 @@ void UInstancedStaticMeshComponent::GetNavigationData(FNavigationRelevantData& D
 
 FBox UInstancedStaticMeshComponent::GetNavigationBounds() const
 {
-	ensureMsgf(NavigationBounds.IsValid, TEXT("Navigation bounds should be calculated before being requested."));
-	return NavigationBounds;
+	const UWorld* const World = GetWorld();
+	const bool bCheckValid = World && World->IsGameWorld();
+	return CachedBounds[(int)EBoundsCacheType::Navigation].GetValue(bCheckValid).GetBox();
 }
 
 bool UInstancedStaticMeshComponent::IsNavigationRelevant() const
@@ -5558,6 +5630,7 @@ void UInstancedStaticMeshComponent::PostEditChangeChainProperty(FPropertyChanged
 		}
 		else if (PropertyChangedEvent.Property->GetFName() == GET_MEMBER_NAME_CHECKED(FInstancedStaticMeshInstanceData, Transform))
 		{
+			InvalidateCachedBounds();
 			// No need to update when using partial update since OnRegister will take care of it.
 			if (!SupportsPartialNavigationUpdate())
 			{
@@ -5592,6 +5665,8 @@ void UInstancedStaticMeshComponent::PostEditChangeChainProperty(FPropertyChanged
 
 void UInstancedStaticMeshComponent::PostEditUndo()
 {
+	InvalidateCachedBounds();
+
 	// Need to do this before, because the proxy is recreated up this stack, which ends up hitting a PrimitiveInstanceDataManager.FlushChanges.
 	PrimitiveInstanceDataManager.Invalidate(PerInstanceSMData.Num());
 	Super::PostEditUndo();
