Small optimizations when calling AActor::IncrementalPreRegisterComponents(). Saves 9+% of time spent in AActor::IncrementalPreRegisterComponents in Big_City_LVL.

diff --git a/Engine/Source/Runtime/Engine/Classes/Components/PrimitiveComponent.h b/Engine/Source/Runtime/Engine/Classes/Components/PrimitiveComponent.h
--- a/Engine/Source/Runtime/Engine/Classes/Components/PrimitiveComponent.h
+++ b/Engine/Source/Runtime/Engine/Classes/Components/PrimitiveComponent.h
@@ -3075,6 +3075,9 @@ public:
 	ENGINE_API virtual TArray<Chaos::FPhysicsObject*> GetAllPhysicsObjects() const override;
 	ENGINE_API virtual Chaos::FPhysicsObjectId GetIdFromGTParticle(Chaos::FGeometryParticle* Particle) const override;
 	//~ End IPhysicsComponent Interface.
+
+	/** Whether the physics is static for this component. */
+	FORCEINLINE bool IsStaticPhysics() const { return Mobility != EComponentMobility::Movable && bStaticWhenNotMoveable; }
 };
 
 /** 
diff --git a/Engine/Source/Runtime/Engine/Private/Actor.cpp b/Engine/Source/Runtime/Engine/Private/Actor.cpp
--- a/Engine/Source/Runtime/Engine/Private/Actor.cpp
+++ b/Engine/Source/Runtime/Engine/Private/Actor.cpp
@@ -5603,13 +5603,13 @@ void AActor::IncrementalPreRegisterComponents()
 
 	UWorld* const World = GetWorld();
 	TInlineComponentArray<UActorComponent*> Components;
-	ForEachComponent<UActorComponent>(false, [World, &Components](UActorComponent* Component)
+	for (UActorComponent* Component : GetComponents())
 	{
 		if (Component->ShouldIncrementalPreRegister(World))
 		{
 			Components.Add(Component);
 		}
-	});
+	}
 
 	// No component found to pre-register
 	if (Components.IsEmpty())
diff --git a/Engine/Source/Runtime/Engine/Private/Components/StaticMeshComponent.cpp b/Engine/Source/Runtime/Engine/Private/Components/StaticMeshComponent.cpp
--- a/Engine/Source/Runtime/Engine/Private/Components/StaticMeshComponent.cpp
+++ b/Engine/Source/Runtime/Engine/Private/Components/StaticMeshComponent.cpp
@@ -2318,7 +2318,7 @@ bool UStaticMeshComponent::ShouldCreatePhysicsState() const
 bool UStaticMeshComponent::AllowsAsyncPhysicsStateCreationAndDestruction() const
 {
 	// Async physics state creation is supported for static collision or primitive component not using welding
-	return (FInitBodySpawnParams(this).bStaticPhysics || !BodyInstance.bAutoWeld || !GetAttachParent());
+	return !BodyInstance.bAutoWeld || !GetAttachParent() || IsStaticPhysics();
 }
 
 void UStaticMeshComponent::SetStaticMeshInternal(UStaticMesh* NewMesh)
diff --git a/Engine/Source/Runtime/Engine/Private/PhysicsEngine/BodyInstance.cpp b/Engine/Source/Runtime/Engine/Private/PhysicsEngine/BodyInstance.cpp
--- a/Engine/Source/Runtime/Engine/Private/PhysicsEngine/BodyInstance.cpp
+++ b/Engine/Source/Runtime/Engine/Private/PhysicsEngine/BodyInstance.cpp
@@ -1640,9 +1640,7 @@ void TInitBodiesHelperBase<InAllocatorType>::InitBodies()
 
 FInitBodySpawnParams::FInitBodySpawnParams(const UPrimitiveComponent* PrimComp)
 {
-	bStaticPhysics = PrimComp == nullptr || (
-		PrimComp->Mobility != EComponentMobility::Movable &&
-		PrimComp->GetStaticWhenNotMoveable());
+	bStaticPhysics = PrimComp == nullptr || PrimComp->IsStaticPhysics();
 
 	if(const USkeletalMeshComponent* SKOwner = Cast<USkeletalMeshComponent>(PrimComp))
 	{
