Add MINIMAL_API to UMaterialExpressionDeriveNormalZ.

- This matches other MaterialExpression class so it can be used by plugins

#rnx

diff --git a/Engine/Source/Runtime/Engine/Public/Materials/MaterialExpressionDeriveNormalZ.h b/Engine/Source/Runtime/Engine/Public/Materials/MaterialExpressionDeriveNormalZ.h
--- a/Engine/Source/Runtime/Engine/Public/Materials/MaterialExpressionDeriveNormalZ.h
+++ b/Engine/Source/Runtime/Engine/Public/Materials/MaterialExpressionDeriveNormalZ.h
@@ -9,7 +9,7 @@
 #include "Materials/MaterialExpression.h"
 #include "MaterialExpressionDeriveNormalZ.generated.h"
 
-UCLASS(collapsecategories, hidecategories=Object)
+UCLASS(collapsecategories, hidecategories=Object, MinimalAPI)
 class UMaterialExpressionDeriveNormalZ : public UMaterialExpression
 {
 	GENERATED_UCLASS_BODY()
@@ -17,7 +17,6 @@ class UMaterialExpressionDeriveNormalZ : public UMaterialExpression
 	UPROPERTY()
 	FExpressionInput InXY;
 
-
 	//~ Begin UMaterialExpression Interface
 #if WITH_EDITOR
 	virtual int32 Compile(class FMaterialCompiler* Compiler, int32 OutputIndex) override;
