Fixed crash with BSPs & StaticLighting

diff --git a/Engine/Source/Runtime/Engine/Private/ModelLight.cpp b/Engine/Source/Runtime/Engine/Private/ModelLight.cpp
--- a/Engine/Source/Runtime/Engine/Private/ModelLight.cpp
+++ b/Engine/Source/Runtime/Engine/Private/ModelLight.cpp
@@ -107,7 +107,7 @@ FBSPSurfaceStaticLighting::FBSPSurfaceStaticLighting(
 		InComponent,
 		InNodeGroup->BoundingBox, 
 		InModel->LightingGuid,
-		InComponent->GetElements()[InModel->Nodes[InNodeGroup->Nodes[0]].ComponentElementIndex].MapBuildDataId
+		FGuid()
 		),
 	NodeGroup(InNodeGroup),
 	bComplete(false),
