[World Partition Runtime HashSet] Only apply specific bounds method for the lowest levels of the grid.

diff --git a/Engine/Source/Runtime/Engine/Private/WorldPartition/RuntimeHashSet/WorldPartitionRuntimeHashSetStreamingGeneration.cpp b/Engine/Source/Runtime/Engine/Private/WorldPartition/RuntimeHashSet/WorldPartitionRuntimeHashSetStreamingGeneration.cpp
--- a/Engine/Source/Runtime/Engine/Private/WorldPartition/RuntimeHashSet/WorldPartitionRuntimeHashSetStreamingGeneration.cpp
+++ b/Engine/Source/Runtime/Engine/Private/WorldPartition/RuntimeHashSet/WorldPartitionRuntimeHashSetStreamingGeneration.cpp
@@ -50,7 +50,7 @@ bool UWorldPartitionRuntimeHashSet::GenerateRuntimePartitionsStreamingDescs(cons
 	// Generate runtime partitions streaming data
 	//
 	TMap<URuntimePartition*, TArray<URuntimePartition::FCellDesc>> RuntimePartitionsStreamingDescs;
-	for (auto [RuntimePartition, ActorSetInstances] : RuntimePartitionsToActorSetMap)
+	for (auto& [RuntimePartition, ActorSetInstances] : RuntimePartitionsToActorSetMap)
 	{
 		// Gather runtime partition cell descs
 		URuntimePartition::FGenerateStreamingParams GenerateStreamingParams;
@@ -185,18 +185,21 @@ bool UWorldPartitionRuntimeHashSet::GenerateStreaming(UWorldPartitionStreamingPo
 
 				if (CellDescInstance.CellBounds.IsSet())
 				{
-					switch (CellDescInstance.SourcePartition->BoundsMethod)
+					if (!RuntimeCell->RuntimeCellData->HierarchicalLevel)
 					{
-					case ERuntimePartitionCellBoundsMethod::UseCellBounds:
-						RuntimeCell->RuntimeCellData->ContentBounds = CellDescInstance.CellBounds.GetValue();
-						break;
-					case ERuntimePartitionCellBoundsMethod::UseMinContentCellBounds:
-						if (RuntimeCell->RuntimeCellData->ContentBounds.IsValid)
+						switch (CellDescInstance.SourcePartition->BoundsMethod)
 						{
-							RuntimeCell->RuntimeCellData->ContentBounds = RuntimeCell->RuntimeCellData->ContentBounds.Overlap(CellDescInstance.CellBounds.GetValue());
-							check(CellDescInstance.CellBounds.GetValue().IsValid);
+						case ERuntimePartitionCellBoundsMethod::UseCellBounds:
+							RuntimeCell->RuntimeCellData->ContentBounds = CellDescInstance.CellBounds.GetValue();
+							break;
+						case ERuntimePartitionCellBoundsMethod::UseMinContentCellBounds:
+							if (RuntimeCell->RuntimeCellData->ContentBounds.IsValid)
+							{
+								RuntimeCell->RuntimeCellData->ContentBounds = RuntimeCell->RuntimeCellData->ContentBounds.Overlap(CellDescInstance.CellBounds.GetValue());
+								check(CellDescInstance.CellBounds.GetValue().IsValid);
+							}
+							break;
 						}
-						break;
 					}
 				}
 
