Revert use of SCOPE_CYCLE_COUNTER_STATID; it currently crashes when StatId.StatString is nullptr in the ENABLE_STATNAMEDEVENTS case.
#rnx

diff --git a/Engine/Source/Runtime/Engine/Private/Tickable.cpp b/Engine/Source/Runtime/Engine/Private/Tickable.cpp
--- a/Engine/Source/Runtime/Engine/Private/Tickable.cpp
+++ b/Engine/Source/Runtime/Engine/Private/Tickable.cpp
@@ -192,7 +192,7 @@ void FTickableGameObject::TickObjects(UWorld* World, const ELevelTick LevelTickT
 					if ((GIsEditor && TickableObject->IsTickableInEditor()) ||
 						(bIsGameWorld && ((!bIsPaused && LevelTickType != LEVELTICK_TimeOnly) || (bIsPaused && TickableObject->IsTickableWhenPaused()))))
 					{
-						SCOPE_CYCLE_COUNTER_STATID(TickableObject->GetStatId());
+						FScopeCycleCounter Context(TickableObject->GetStatId());
 						TickableObject->Tick(DeltaSeconds);
 					}
 				}
