Passing -m to gamever will display the gameversion components on separate lines to make it easier to view the information you want

diff --git a/Engine/Source/Runtime/Engine/Private/UnrealEngine.cpp b/Engine/Source/Runtime/Engine/Private/UnrealEngine.cpp
--- a/Engine/Source/Runtime/Engine/Private/UnrealEngine.cpp
+++ b/Engine/Source/Runtime/Engine/Private/UnrealEngine.cpp
@@ -5525,8 +5525,15 @@ bool UEngine::HandleFlushLogCommand( const TCHAR* Cmd, FOutputDevice& Ar )
 
 bool UEngine::HandleGameVerCommand( const TCHAR* Cmd, FOutputDevice& Ar )
 {
-	FString VersionString = FString::Printf( TEXT( "GameVersion Branch: %s, Configuration: %s, Build: %s, CommandLine: %s" ),
-		*FApp::GetBranchName(), LexToString( FApp::GetBuildConfiguration() ), FApp::GetBuildVersion(), FCommandLine::Get() );
+	const TArray<FString, TInlineAllocator<4> > VersionDetails = {
+		FString::Printf(TEXT("GameVersion Branch: %s"), *FApp::GetBranchName()),
+		FString::Printf(TEXT("Configuration: %s"), LexToString( FApp::GetBuildConfiguration() )),
+		FString::Printf(TEXT("Build: %s"), FApp::GetBuildVersion()),
+		FString::Printf(TEXT("CommandLine: %s"), FCommandLine::Get()),
+	};
+
+	const bool bMultiline = FCString::Stristr(Cmd, TEXT("-m")) != nullptr;
+	const FString VersionString = FString::Join(VersionDetails, bMultiline ? TEXT("\n") : TEXT(", "));
 
 	Ar.Logf( TEXT("%s"), *VersionString );
 	FPlatformApplicationMisc::ClipboardCopy( *VersionString );
