Add MeshPaintTextureCoordinateIndex material expression node.
This outputs the coordinate index for the primitive being shaded.

diff --git a/Engine/Source/Runtime/Engine/Private/Materials/HLSLMaterialTranslator.cpp b/Engine/Source/Runtime/Engine/Private/Materials/HLSLMaterialTranslator.cpp
--- a/Engine/Source/Runtime/Engine/Private/Materials/HLSLMaterialTranslator.cpp
+++ b/Engine/Source/Runtime/Engine/Private/Materials/HLSLMaterialTranslator.cpp
@@ -8881,6 +8881,11 @@ int32 FHLSLMaterialTranslator::VertexColor()
 	}
 }
 
+int32 FHLSLMaterialTranslator::MeshPaintTextureCoordinateIndex()
+{
+	return AddInlinedCodeChunkZeroDeriv(MCT_Float, TEXT("GetMeshPaintTextureCoordinateIndex(GetPrimitiveData(Parameters))"));
+}
+
 int32 FHLSLMaterialTranslator::MeshPaintTextureDescriptor()
 {
 	return AddInlinedCodeChunkZeroDeriv(EMaterialValueType(MCT_TextureVirtual | MCT_TextureMeshPaint), TEXT("GetMeshPaintTextureDescriptor(GetPrimitiveData(Parameters))"));
diff --git a/Engine/Source/Runtime/Engine/Private/Materials/HLSLMaterialTranslator.h b/Engine/Source/Runtime/Engine/Private/Materials/HLSLMaterialTranslator.h
--- a/Engine/Source/Runtime/Engine/Private/Materials/HLSLMaterialTranslator.h
+++ b/Engine/Source/Runtime/Engine/Private/Materials/HLSLMaterialTranslator.h
@@ -1122,6 +1122,7 @@ protected:
 	virtual int32 FontSignedDistanceData() override;
 
 	virtual int32 VertexColor() override;
+	virtual int32 MeshPaintTextureCoordinateIndex() override;
 	virtual int32 MeshPaintTextureDescriptor() override;
 	virtual int32 MeshPaintTextureReplace(int32 Invalid, int32 Valid) override;
 
diff --git a/Engine/Source/Runtime/Engine/Private/Materials/MaterialExpressions.cpp b/Engine/Source/Runtime/Engine/Private/Materials/MaterialExpressions.cpp
--- a/Engine/Source/Runtime/Engine/Private/Materials/MaterialExpressions.cpp
+++ b/Engine/Source/Runtime/Engine/Private/Materials/MaterialExpressions.cpp
@@ -125,6 +125,7 @@
 #include "Materials/MaterialExpressionIf.h"
 #include "Materials/MaterialExpressionInverseLinearInterpolate.h"
 #include "Materials/MaterialExpressionLightmapUVs.h"
+#include "Materials/MaterialExpressionMeshPaintTextureCoordinateIndex.h"
 #include "Materials/MaterialExpressionMeshPaintTextureObject.h"
 #include "Materials/MaterialExpressionMeshPaintTextureReplace.h"
 #include "Materials/MaterialExpressionPrecomputedAOMask.h"
@@ -3923,6 +3924,47 @@ void UMaterialExpressionVirtualTextureFeatureSwitch::GetCaption(TArray<FString>&
 
 #endif // WITH_EDITOR
 
+UMaterialExpressionMeshPaintTextureCoordinateIndex::UMaterialExpressionMeshPaintTextureCoordinateIndex(const FObjectInitializer& ObjectInitializer)
+	: Super(ObjectInitializer)
+{
+	// Structure to hold one-time initialization
+	struct FConstructorStatics
+	{
+		FText NAME_MeshPaintTexture;
+		FConstructorStatics()
+			: NAME_MeshPaintTexture(LOCTEXT("MeshPaintTexture", "MeshPaintTexture"))
+		{
+		}
+	};
+	static FConstructorStatics ConstructorStatics;
+
+#if WITH_EDITORONLY_DATA
+	MenuCategories.Add(ConstructorStatics.NAME_MeshPaintTexture);
+#endif
+
+	Outputs.Reset();
+	Outputs.Add(FExpressionOutput(TEXT("")));
+}
+
+#if WITH_EDITOR
+
+void UMaterialExpressionMeshPaintTextureCoordinateIndex::GetCaption(TArray<FString>& OutCaptions) const
+{
+	OutCaptions.Add(TEXT("Mesh Paint Texture Coordinate Index"));
+}
+
+void UMaterialExpressionMeshPaintTextureCoordinateIndex::GetExpressionToolTip(TArray<FString>& OutToolTip)
+{
+	ConvertToMultilineToolTip(TEXT("Get the Mesh Paint Texture UV coordinate index."), 40, OutToolTip);
+}
+
+int32 UMaterialExpressionMeshPaintTextureCoordinateIndex::Compile(class FMaterialCompiler* Compiler, int32 OutputIndex)
+{
+	return Compiler->MeshPaintTextureCoordinateIndex();
+}
+
+#endif // WITH_EDITOR
+
 UMaterialExpressionMeshPaintTextureObject::UMaterialExpressionMeshPaintTextureObject(const FObjectInitializer& ObjectInitializer)
 	: Super(ObjectInitializer)
 {
@@ -3949,7 +3991,7 @@ UMaterialExpressionMeshPaintTextureObject::UMaterialExpressionMeshPaintTextureOb
 
 void UMaterialExpressionMeshPaintTextureObject::GetCaption(TArray<FString>& OutCaptions) const
 {
-	OutCaptions.Add(TEXT("MeshPaintTextureObject"));
+	OutCaptions.Add(TEXT("Mesh Paint Texture Object"));
 }
 
 void UMaterialExpressionMeshPaintTextureObject::GetExpressionToolTip(TArray<FString>& OutToolTip)
@@ -3995,7 +4037,7 @@ UMaterialExpressionMeshPaintTextureReplace::UMaterialExpressionMeshPaintTextureR
 
 void UMaterialExpressionMeshPaintTextureReplace::GetCaption(TArray<FString>& OutCaptions) const
 {
-	OutCaptions.Add(TEXT("MeshPaintTextureReplace"));
+	OutCaptions.Add(TEXT("Mesh Paint Texture Replace"));
 }
 
 void UMaterialExpressionMeshPaintTextureReplace::GetExpressionToolTip(TArray<FString>& OutToolTip)
diff --git a/Engine/Source/Runtime/Engine/Public/MaterialCompiler.h b/Engine/Source/Runtime/Engine/Public/MaterialCompiler.h
--- a/Engine/Source/Runtime/Engine/Public/MaterialCompiler.h
+++ b/Engine/Source/Runtime/Engine/Public/MaterialCompiler.h
@@ -453,6 +453,7 @@ public:
 	virtual int32 FontSignedDistanceData() = 0;
 
 	virtual int32 VertexColor() = 0;
+	virtual int32 MeshPaintTextureCoordinateIndex() = 0;
 	virtual int32 MeshPaintTextureDescriptor() = 0;
 	virtual int32 MeshPaintTextureReplace(int32 Invalid, int32 Valid) = 0;
 
@@ -1009,6 +1010,7 @@ public:
 	virtual int32 FontSignedDistanceData() override { return Compiler->FontSignedDistanceData(); }
 
 	virtual int32 VertexColor() override { return Compiler->VertexColor(); }
+	virtual int32 MeshPaintTextureCoordinateIndex() override { return Compiler->MeshPaintTextureCoordinateIndex(); }
 	virtual int32 MeshPaintTextureDescriptor() override { return Compiler->MeshPaintTextureDescriptor(); }
 	virtual int32 MeshPaintTextureReplace(int32 Invalid, int32 Valid) override { return Compiler->MeshPaintTextureReplace(Invalid, Valid); }
 
diff --git a/Engine/Source/Runtime/Engine/Public/Materials/MaterialExpressionMeshPaintTextureCoordinateIndex.h b/Engine/Source/Runtime/Engine/Public/Materials/MaterialExpressionMeshPaintTextureCoordinateIndex.h
new file mode 100644
--- /dev/null
+++ b/Engine/Source/Runtime/Engine/Public/Materials/MaterialExpressionMeshPaintTextureCoordinateIndex.h
@@ -0,0 +1,21 @@
+// Copyright Epic Games, Inc. All Rights Reserved.
+
+#pragma once
+
+#include "Materials/MaterialExpression.h"
+
+#include "MaterialExpressionMeshPaintTextureCoordinateIndex.generated.h"
+
+UCLASS(collapsecategories, hidecategories=Object, MinimalAPI)
+class UMaterialExpressionMeshPaintTextureCoordinateIndex : public UMaterialExpression
+{
+	GENERATED_UCLASS_BODY()
+
+	//~ Begin UMaterialExpression Interface
+#if WITH_EDITOR
+	virtual void GetCaption(TArray<FString>& OutCaptions) const override;
+	virtual void GetExpressionToolTip(TArray<FString>& OutToolTip) override;
+	virtual int32 Compile(class FMaterialCompiler* Compiler, int32 OutputIndex) override;
+#endif
+	//~ End UMaterialExpression Interface
+};
