Audited existing usage of FText::FindText

#rnx

diff --git a/Engine/Source/Editor/BlueprintGraph/Private/K2Node_VariableGet.cpp b/Engine/Source/Editor/BlueprintGraph/Private/K2Node_VariableGet.cpp
--- a/Engine/Source/Editor/BlueprintGraph/Private/K2Node_VariableGet.cpp
+++ b/Engine/Source/Editor/BlueprintGraph/Private/K2Node_VariableGet.cpp
@@ -219,7 +219,7 @@ FText UK2Node_VariableGet::GetPropertyTooltip(FProperty const* VariableProperty)
 				// See if the native property has a tooltip
 				SubTooltip = PropertyTooltip;
 				FString TooltipName = FString::Printf(TEXT("%s.%s"), *VarName.ToString(), *FBlueprintMetadata::MD_Tooltip.ToString());
-				FText::FindText(*VariableProperty->GetFullGroupName(true), *TooltipName, SubTooltip);
+				FText::FindTextInLiveTable_Advanced(*VariableProperty->GetFullGroupName(true), *TooltipName, SubTooltip);
 			}
 		}
 		else if (SourceClass)
diff --git a/Engine/Source/Editor/BlueprintGraph/Private/K2Node_VariableSet.cpp b/Engine/Source/Editor/BlueprintGraph/Private/K2Node_VariableSet.cpp
--- a/Engine/Source/Editor/BlueprintGraph/Private/K2Node_VariableSet.cpp
+++ b/Engine/Source/Editor/BlueprintGraph/Private/K2Node_VariableSet.cpp
@@ -169,7 +169,7 @@ FText UK2Node_VariableSet::GetPropertyTooltip(FProperty const* VariableProperty)
 				// See if the native property has a tooltip
 				SubTooltip = PropertyTooltip;
 				FString TooltipName = FString::Printf(TEXT("%s.%s"), *VarName.ToString(), *FBlueprintMetadata::MD_Tooltip.ToString());
-				FText::FindText(*VariableProperty->GetFullGroupName(true), *TooltipName, SubTooltip);
+				FText::FindTextInLiveTable_Advanced(*VariableProperty->GetFullGroupName(true), *TooltipName, SubTooltip);
 			}
 		}
 		else if (SourceClass)
diff --git a/Engine/Source/Editor/LandscapeEditor/Private/LandscapeEditorDetails.cpp b/Engine/Source/Editor/LandscapeEditor/Private/LandscapeEditorDetails.cpp
--- a/Engine/Source/Editor/LandscapeEditor/Private/LandscapeEditorDetails.cpp
+++ b/Engine/Source/Editor/LandscapeEditor/Private/LandscapeEditorDetails.cpp
@@ -218,55 +218,52 @@ void FLandscapeEditorDetails::CustomizeDetails(IDetailLayoutBuilder& DetailBuild
 
 FText FLandscapeEditorDetails::GetLocalizedName(FString Name)
 {
-	static bool bInitialized = false;
-	if (!bInitialized)
-	{
-		FEdModeLandscape* LandscapeEdMode = GetEditorMode();
-
-		bInitialized = true;
-		LOCTEXT("ToolSet_NewLandscape", "New Landscape");
-		LOCTEXT("ToolSet_ResizeLandscape", "Change Component Size");
-		LOCTEXT("ToolSet_Sculpt", "Sculpt");
-		LOCTEXT("ToolSet_Erase", "Erase");
-		LOCTEXT("ToolSet_Paint", "Paint");
-		LOCTEXT("ToolSet_Smooth", "Smooth");
-		LOCTEXT("ToolSet_Flatten", "Flatten");
-		LOCTEXT("ToolSet_Ramp", "Ramp");
-		LOCTEXT("ToolSet_Erosion", "Erosion");
-		LOCTEXT("ToolSet_HydraErosion", "HydroErosion");
-		LOCTEXT("ToolSet_Noise", "Noise");
-		LOCTEXT("ToolSet_Retopologize", "Retopologize");
-		LOCTEXT("ToolSet_Visibility", "Visibility");
-		LOCTEXT("ToolSet_BlueprintBrush", "Blueprint Brushes");
-		
-		LOCTEXT("ToolSet_Select", "Selection");
-		LOCTEXT("ToolSet_AddComponent", "Add");
-		LOCTEXT("ToolSet_DeleteComponent", "Delete");
-		LOCTEXT("ToolSet_MoveToLevel", "Move to Level");
-
-		LOCTEXT("ToolSet_Mask", "Selection");
-		LOCTEXT("ToolSet_CopyPaste", "Copy/Paste");
-		LOCTEXT("ToolSet_Mirror", "Mirror");
-
-		LOCTEXT("ToolSet_Splines", "Edit Splines");
-
-		LOCTEXT("BrushSet_Circle", "Circle");
-		LOCTEXT("BrushSet_Alpha", "Alpha");
-		LOCTEXT("BrushSet_Pattern", "Pattern");
-		LOCTEXT("BrushSet_Component", "Component");
-		LOCTEXT("BrushSet_Gizmo", "Gizmo");
-		LOCTEXT("BrushSet_Dummy", "NoBrush");
-		LOCTEXT("BrushSet_Splines", "Splines");
-
-		LOCTEXT("Circle_Smooth", "Smooth");
-		LOCTEXT("Circle_Linear", "Linear");
-		LOCTEXT("Circle_Spherical", "Spherical");
-		LOCTEXT("Circle_Tip", "Tip");
-		LOCTEXT("Circle_Dummy", "NoBrush");
-	}
-
-	FText Result;
-	ensure(FText::FindText(TEXT(LOCTEXT_NAMESPACE), Name, Result));
+#define LANDSCAPE_TOOL_LOCTEXT(KEY, SOURCE) { TEXT(KEY), LOCTEXT(KEY, SOURCE) }
+	static const TSortedMap<FString, FText, FDefaultAllocator, FLocKeySortedMapLess> ToolNames = {
+		LANDSCAPE_TOOL_LOCTEXT("ToolSet_NewLandscape", "New Landscape"),
+		LANDSCAPE_TOOL_LOCTEXT("ToolSet_ResizeLandscape", "Change Component Size"),
+		LANDSCAPE_TOOL_LOCTEXT("ToolSet_Sculpt", "Sculpt"),
+		LANDSCAPE_TOOL_LOCTEXT("ToolSet_Erase", "Erase"),
+		LANDSCAPE_TOOL_LOCTEXT("ToolSet_Paint", "Paint"),
+		LANDSCAPE_TOOL_LOCTEXT("ToolSet_Smooth", "Smooth"),
+		LANDSCAPE_TOOL_LOCTEXT("ToolSet_Flatten", "Flatten"),
+		LANDSCAPE_TOOL_LOCTEXT("ToolSet_Ramp", "Ramp"),
+		LANDSCAPE_TOOL_LOCTEXT("ToolSet_Erosion", "Erosion"),
+		LANDSCAPE_TOOL_LOCTEXT("ToolSet_HydraErosion", "HydroErosion"),
+		LANDSCAPE_TOOL_LOCTEXT("ToolSet_Noise", "Noise"),
+		LANDSCAPE_TOOL_LOCTEXT("ToolSet_Retopologize", "Retopologize"),
+		LANDSCAPE_TOOL_LOCTEXT("ToolSet_Visibility", "Visibility"),
+		LANDSCAPE_TOOL_LOCTEXT("ToolSet_BlueprintBrush", "Blueprint Brushes"),
+
+		LANDSCAPE_TOOL_LOCTEXT("ToolSet_Select", "Selection"),
+		LANDSCAPE_TOOL_LOCTEXT("ToolSet_AddComponent", "Add"),
+		LANDSCAPE_TOOL_LOCTEXT("ToolSet_DeleteComponent", "Delete"),
+		LANDSCAPE_TOOL_LOCTEXT("ToolSet_MoveToLevel", "Move to Level"),
+
+		LANDSCAPE_TOOL_LOCTEXT("ToolSet_Mask", "Selection"),
+		LANDSCAPE_TOOL_LOCTEXT("ToolSet_CopyPaste", "Copy/Paste"),
+		LANDSCAPE_TOOL_LOCTEXT("ToolSet_Mirror", "Mirror"),
+
+		LANDSCAPE_TOOL_LOCTEXT("ToolSet_Splines", "Edit Splines"),
+
+		LANDSCAPE_TOOL_LOCTEXT("BrushSet_Circle", "Circle"),
+		LANDSCAPE_TOOL_LOCTEXT("BrushSet_Alpha", "Alpha"),
+		LANDSCAPE_TOOL_LOCTEXT("BrushSet_Pattern", "Pattern"),
+		LANDSCAPE_TOOL_LOCTEXT("BrushSet_Component", "Component"),
+		LANDSCAPE_TOOL_LOCTEXT("BrushSet_Gizmo", "Gizmo"),
+		LANDSCAPE_TOOL_LOCTEXT("BrushSet_Dummy", "NoBrush"),
+		LANDSCAPE_TOOL_LOCTEXT("BrushSet_Splines", "Splines"),
+
+		LANDSCAPE_TOOL_LOCTEXT("Circle_Smooth", "Smooth"),
+		LANDSCAPE_TOOL_LOCTEXT("Circle_Linear", "Linear"),
+		LANDSCAPE_TOOL_LOCTEXT("Circle_Spherical", "Spherical"),
+		LANDSCAPE_TOOL_LOCTEXT("Circle_Tip", "Tip"),
+		LANDSCAPE_TOOL_LOCTEXT("Circle_Dummy", "NoBrush"),
+	};
+#undef LANDSCAPE_TOOL_LOCTEXT
+	
+	FText Result = ToolNames.FindRef(Name);
+	ensure(!Result.IsEmpty());
 	return Result;
 }
 
diff --git a/Engine/Source/Editor/UnrealEd/Private/EditorEngine.cpp b/Engine/Source/Editor/UnrealEd/Private/EditorEngine.cpp
--- a/Engine/Source/Editor/UnrealEd/Private/EditorEngine.cpp
+++ b/Engine/Source/Editor/UnrealEd/Private/EditorEngine.cpp
@@ -4689,7 +4689,7 @@ FString UEditorEngine::GetFriendlyName( const FProperty* Property, UStruct* Owne
 	{
 		FString PropertyPathName = Property->GetPathName(CurrentStruct);
 
-		DidFindText = FText::FindText(*CurrentStruct->GetName(), *(PropertyPathName + TEXT(".FriendlyName")), /*OUT*/FoundText );
+		DidFindText = FText::FindTextInLiveTable_Advanced(*CurrentStruct->GetName(), *(PropertyPathName + TEXT(".FriendlyName")), /*OUT*/FoundText );
 		CurrentStruct = CurrentStruct->GetSuperStruct();
 	} while( CurrentStruct != NULL && CurrentStruct->IsChildOf(RealOwnerStruct) && !DidFindText );
 
diff --git a/Engine/Source/Runtime/CoreUObject/Private/UObject/Class.cpp b/Engine/Source/Runtime/CoreUObject/Private/UObject/Class.cpp
--- a/Engine/Source/Runtime/CoreUObject/Private/UObject/Class.cpp
+++ b/Engine/Source/Runtime/CoreUObject/Private/UObject/Class.cpp
@@ -334,7 +334,7 @@ FText UField::GetToolTipText(bool bShortTooltip) const
 
 	const FString Namespace = bFoundShortTooltip ? TEXT("UObjectShortTooltips") : TEXT("UObjectToolTips");
 	const FString Key = GetFullGroupName(false);
-	if ( !FText::FindText( Namespace, Key, /*OUT*/LocalizedToolTip, &NativeToolTip ) )
+	if ( !FText::FindTextInLiveTable_Advanced( Namespace, Key, /*OUT*/LocalizedToolTip, &NativeToolTip ) )
 	{
 		if (NativeToolTip.IsEmpty())
 		{
diff --git a/Engine/Source/Runtime/CoreUObject/Private/UObject/Enum.cpp b/Engine/Source/Runtime/CoreUObject/Private/UObject/Enum.cpp
--- a/Engine/Source/Runtime/CoreUObject/Private/UObject/Enum.cpp
+++ b/Engine/Source/Runtime/CoreUObject/Private/UObject/Enum.cpp
@@ -919,7 +919,7 @@ FText UEnum::GetToolTipTextByIndex(int32 NameIndex) const
 	static const FString Namespace = TEXT("UObjectToolTips");
 	FString Key = GetFullGroupName(false) + TEXT(".") + GetNameStringByIndex(NameIndex);
 		
-	if ( !FText::FindText( Namespace, Key, /*OUT*/LocalizedToolTip, &NativeToolTip ) )
+	if ( !FText::FindTextInLiveTable_Advanced( Namespace, Key, /*OUT*/LocalizedToolTip, &NativeToolTip ) )
 	{
 		static const FString DoxygenSee(TEXT("@see"));
 		static const FString TooltipSee(TEXT("See:"));
diff --git a/Engine/Source/Runtime/CoreUObject/Private/UObject/Field.cpp b/Engine/Source/Runtime/CoreUObject/Private/UObject/Field.cpp
--- a/Engine/Source/Runtime/CoreUObject/Private/UObject/Field.cpp
+++ b/Engine/Source/Runtime/CoreUObject/Private/UObject/Field.cpp
@@ -954,7 +954,7 @@ FText FField::GetToolTipText(bool bShortTooltip) const
 
 	const FString Namespace = bFoundShortTooltip ? TEXT("UObjectShortTooltips") : TEXT("UObjectToolTips");
 	const FString Key = GetFullGroupName(false);
-	if (!FText::FindText(Namespace, Key, /*OUT*/LocalizedToolTip, &NativeToolTip))
+	if (!FText::FindTextInLiveTable_Advanced(Namespace, Key, /*OUT*/LocalizedToolTip, &NativeToolTip))
 	{
 		if (!NativeToolTip.IsEmpty())
 		{
diff --git a/Engine/Source/Runtime/Engine/Private/DialogueWave.cpp b/Engine/Source/Runtime/Engine/Private/DialogueWave.cpp
--- a/Engine/Source/Runtime/Engine/Private/DialogueWave.cpp
+++ b/Engine/Source/Runtime/Engine/Private/DialogueWave.cpp
@@ -350,7 +350,7 @@ namespace
 			}
 
 			// Gather the Subtitle Override for each context
-			if (!DialogueWave->SubtitleOverride.IsEmpty())
+			if (DialogueWave->bOverride_SubtitleOverride && !DialogueWave->SubtitleOverride.IsEmpty())
 			{
 				FGatherableTextData& GatherableTextData = FindOrAddDialogueTextData(DialogueWave->SubtitleOverride, FDialogueConstants::DialogueNamespace);
 
@@ -641,12 +641,7 @@ USoundBase* UDialogueWave::GetWaveFromContext(const FDialogueContextMapping& Con
 FText UDialogueWave::GetLocalizedSpokenText() const
 {
 	// Try and find a general dialogue wave localization
-	FText LocalizedSpokenText;
-	if (!FText::FindText(FDialogueConstants::DialogueNamespace, LocalizationGUID.ToString(), /*OUT*/LocalizedSpokenText, &SpokenText))
-	{
-		LocalizedSpokenText = FText::AsCultureInvariant(SpokenText);
-	}
-	return LocalizedSpokenText;
+	return FText::AsLocalizable_Advanced(FDialogueConstants::DialogueNamespace, LocalizationGUID.ToString(), SpokenText);
 }
 
 FText UDialogueWave::GetLocalizedSpokenText(const FDialogueContext& Context) const
@@ -664,33 +659,32 @@ FText UDialogueWave::GetLocalizedSpokenText(const FDialogueContext& Context) con
 FText UDialogueWave::GetLocalizedSpokenText(const FDialogueContextMapping& ContextMapping) const
 {
 	// First try and find a context specific localization
-	FText LocalizedSpokenText;
-	if (!FText::FindText(FDialogueConstants::DialogueNamespace, GetContextLocalizationKey(ContextMapping), /*OUT*/LocalizedSpokenText, &SpokenText))
+	if (ContextMapping.Context.Speaker)
 	{
-		// Failing that, try and find a general dialogue wave localization
-		LocalizedSpokenText = GetLocalizedSpokenText();
+		FText ContextSpecificText;
+		if (FText::FindTextInLiveTable_Advanced(FDialogueConstants::DialogueNamespace, GetContextLocalizationKey(ContextMapping), ContextSpecificText, &SpokenText))
+		{
+			return ContextSpecificText;
+		}
 	}
-	return LocalizedSpokenText;
+
+	// Failing that, try and find a general dialogue wave localization
+	return GetLocalizedSpokenText();
 }
 
 FText UDialogueWave::GetLocalizedSubtitle() const
 {
-	const FString* SubtitleSourceString = &SpokenText;
+	FStringView SubtitleSourceString = SpokenText;
 
 	// Try and find a general dialogue wave localization
 	FString Key = LocalizationGUID.ToString();
 	if (bOverride_SubtitleOverride)
 	{
-		SubtitleSourceString = &SubtitleOverride;
+		SubtitleSourceString = SubtitleOverride;
 		Key += FDialogueConstants::SubtitleKeySuffix;
 	}
 
-	FText LocalizedSubtitle;
-	if (!FText::FindText(FDialogueConstants::DialogueNamespace, Key, /*OUT*/LocalizedSubtitle, SubtitleSourceString))
-	{
-		LocalizedSubtitle = bOverride_SubtitleOverride ? FText::AsCultureInvariant(SubtitleOverride) : FText::AsCultureInvariant(SpokenText);
-	}
-	return LocalizedSubtitle;
+	return FText::AsLocalizable_Advanced(FDialogueConstants::DialogueNamespace, Key, SubtitleSourceString);
 }
 
 FText UDialogueWave::GetLocalizedSubtitle(const FDialogueContext& Context) const
@@ -707,24 +701,28 @@ FText UDialogueWave::GetLocalizedSubtitle(const FDialogueContext& Context) const
 
 FText UDialogueWave::GetLocalizedSubtitle(const FDialogueContextMapping& ContextMapping) const
 {
-	const FString* SubtitleSourceString = &SpokenText;
-
-	// Do we have a subtitle override?
-	FString Key = GetContextLocalizationKey(ContextMapping);
-	if (bOverride_SubtitleOverride)
-	{
-		SubtitleSourceString = &SubtitleOverride;
-		Key += FDialogueConstants::SubtitleKeySuffix;
-	}
-
 	// First try and find a context specific localization
-	FText LocalizedSubtitle;
-	if (!FText::FindText(FDialogueConstants::DialogueNamespace, Key, /*OUT*/LocalizedSubtitle, SubtitleSourceString))
+	if (ContextMapping.Context.Speaker)
 	{
-		// Failing that, try and find a general dialogue wave localization
-		LocalizedSubtitle = GetLocalizedSubtitle();
+		const FString* SubtitleSourceString = &SpokenText;
+
+		// Do we have a subtitle override?
+		FString Key = GetContextLocalizationKey(ContextMapping);
+		if (bOverride_SubtitleOverride)
+		{
+			SubtitleSourceString = &SubtitleOverride;
+			Key += FDialogueConstants::SubtitleKeySuffix;
+		}
+
+		FText ContextSpecificText;
+		if (FText::FindTextInLiveTable_Advanced(FDialogueConstants::DialogueNamespace, Key, ContextSpecificText, SubtitleSourceString))
+		{
+			return ContextSpecificText;
+		}
 	}
-	return LocalizedSubtitle;
+
+	// Failing that, try and find a general dialogue wave localization
+	return GetLocalizedSubtitle();
 }
 
 FString UDialogueWave::GetContextLocalizationKey(const FDialogueContext& Context) const
diff --git a/Engine/Source/Runtime/Engine/Private/ObjectEditorUtils.cpp b/Engine/Source/Runtime/Engine/Private/ObjectEditorUtils.cpp
--- a/Engine/Source/Runtime/Engine/Private/ObjectEditorUtils.cpp
+++ b/Engine/Source/Runtime/Engine/Private/ObjectEditorUtils.cpp
@@ -18,7 +18,7 @@ namespace FObjectEditorUtils
 		FText DisplayName;
 
 		const FString NativeCategory = InName.ToString();
-		if (FText::FindText(CategoryLocalizationNamespace, NativeCategory, /*OUT*/DisplayName, &NativeCategory))
+		if (FText::FindTextInLiveTable_Advanced(CategoryLocalizationNamespace, NativeCategory, /*OUT*/DisplayName, &NativeCategory))
 		{
 			// Category names in English are typically gathered in their non-pretty form (eg "UserInterface" rather than "User Interface"), so skip 
 			// applying the localized variant if the text matches the raw category name, as in this case the pretty printer will do a better job
@@ -48,7 +48,7 @@ namespace FObjectEditorUtils
 			if (!NativeCategory.IsEmpty())
 			{
 				FText LocalizedCategory;
-				if (!FText::FindText(CategoryLocalizationNamespace, NativeCategory, /*OUT*/LocalizedCategory, &NativeCategory))
+				if (!FText::FindTextInLiveTable_Advanced(CategoryLocalizationNamespace, NativeCategory, /*OUT*/LocalizedCategory, &NativeCategory))
 				{
 					LocalizedCategory = FText::AsCultureInvariant(NativeCategory);
 				}
@@ -70,7 +70,7 @@ namespace FObjectEditorUtils
 			if (!NativeCategory.IsEmpty())
 			{
 				FText LocalizedCategory;
-				if (!FText::FindText(CategoryLocalizationNamespace, NativeCategory, /*OUT*/LocalizedCategory, &NativeCategory))
+				if (!FText::FindTextInLiveTable_Advanced(CategoryLocalizationNamespace, NativeCategory, /*OUT*/LocalizedCategory, &NativeCategory))
 				{
 					LocalizedCategory = FText::AsCultureInvariant(NativeCategory);
 				}
diff --git a/Engine/Source/Runtime/Engine/Private/Tests/Internationalization/TextPropertyTest.cpp b/Engine/Source/Runtime/Engine/Private/Tests/Internationalization/TextPropertyTest.cpp
--- a/Engine/Source/Runtime/Engine/Private/Tests/Internationalization/TextPropertyTest.cpp
+++ b/Engine/Source/Runtime/Engine/Private/Tests/Internationalization/TextPropertyTest.cpp
@@ -82,7 +82,7 @@ bool FTextPropertyTest::RunTest (const FString& Parameters)
 		
 		auto SavedObject = NewObject<UTextPropertyTestObject>();
 
-		FText::FindText( TEXT("TextPropertyTest"), TEXT("DefaultedText"), /*OUT*/SavedObject->DefaultedText );
+		SavedObject->DefaultedText = FText::AsLocalizable_Advanced( TEXT("TextPropertyTest"), TEXT("DefaultedText"), TEXT("DefaultValue") );
 		SavedObject->UndefaultedText = LOCTEXT("ModifiedUndefaultedText", "Modified UndefaultedText Value");
 		const FText TransientText = FText::Format( LOCTEXT("TransientTest", "{0}"), LOCTEXT("TransientTestMessage", "Testing Transient serialization detection") );
 		SavedObject->TransientText = TransientText;
diff --git a/Engine/Source/Runtime/PreLoadScreen/Private/PreLoadSettingsContainer.cpp b/Engine/Source/Runtime/PreLoadScreen/Private/PreLoadSettingsContainer.cpp
--- a/Engine/Source/Runtime/PreLoadScreen/Private/PreLoadSettingsContainer.cpp
+++ b/Engine/Source/Runtime/PreLoadScreen/Private/PreLoadSettingsContainer.cpp
@@ -418,17 +418,8 @@ void FPreLoadSettingsContainerBase::ParseLocalizedTextConfigString(const FString
         LocInitialValue.RemoveFromEnd(")");
         LocInitialValue.RemoveFromEnd("\"");
 
-        //Actually try to add the FText to our list by finding it in the FText collection (should already be in there due to Localization system)
-        FText FoundText = FText::GetEmpty();
-        if (FText::FindText(LocNameSpace, LocIdentifier, FoundText))
-        {
-            AddLocalizedText(Identifier, FoundText);
-        }
-        //We couldn't find it already, so go ahead and add a version to FText with an initial value. This one won't be localized, but that may be intended
-        else
-        {
-            AddLocalizedText(Identifier, FText::FromString(LocInitialValue));
-        }
+        //Add the FText to our list
+        AddLocalizedText(Identifier, FText::AsLocalizable_Advanced(MoveTemp(LocNameSpace), MoveTemp(LocIdentifier), MoveTemp(LocInitialValue)));
     }
 }
 
