Remove check if InGameThread because TStrongPtr can now be created from async threads as well

diff --git a/Engine/Source/Runtime/Engine/Private/Components/PrimitiveComponent.cpp b/Engine/Source/Runtime/Engine/Private/Components/PrimitiveComponent.cpp
--- a/Engine/Source/Runtime/Engine/Private/Components/PrimitiveComponent.cpp
+++ b/Engine/Source/Runtime/Engine/Private/Components/PrimitiveComponent.cpp
@@ -4785,7 +4785,7 @@ void UPrimitiveComponent::PrecachePSOs()
 	// Also only request PSOs from game thread because TStrongObjectPtr is used on the material to make
 	// it's not deleted via garbage collection when PSO precaching is still busy. TStrongObjectPtr can only
 	// be constructed on the GameThread
-	if (!FApp::CanEverRender() || !IsComponentPSOPrecachingEnabled() || !IsInGameThread())
+	if (!FApp::CanEverRender() || !IsComponentPSOPrecachingEnabled())
 	{
 		return;
 	}
