Deprecated USkeletalMeshComponent::SetAnimClass function

Replaced with SetAnimInstanceClass. Fixed up deprecation warnings and places thast were callling both (!)

PR #11100: Deprecatead the 'SetAnimClass' function. Redundant call to 'SetAnimIn?
[RN] Deprecated USkeletalMeshComponent::SetAnimClass function

diff --git a/Engine/Plugins/Animation/ControlRig/Source/ControlRig/Public/ControlRigComponent.h b/Engine/Plugins/Animation/ControlRig/Source/ControlRig/Public/ControlRigComponent.h
--- a/Engine/Plugins/Animation/ControlRig/Source/ControlRig/Public/ControlRigComponent.h
+++ b/Engine/Plugins/Animation/ControlRig/Source/ControlRig/Public/ControlRigComponent.h
@@ -630,10 +630,9 @@ private:
 
 		void Apply(USkeletalMeshComponent* InComponent)
 		{
-				InComponent->SetAnimClass(AnimClass);
 				InComponent->SetAnimInstanceClass(AnimInstanceClass);
 				InComponent->PrimaryComponentTick.bCanEverTick = bCanEverTick;
-			}
+		}
 	};
 
 	UControlRig* SetupControlRigIfRequired();
diff --git a/Engine/Plugins/Animation/PerformanceCaptureCore/Source/PerformanceCaptureCore/Private/PerformerComponent.cpp b/Engine/Plugins/Animation/PerformanceCaptureCore/Source/PerformanceCaptureCore/Private/PerformerComponent.cpp
--- a/Engine/Plugins/Animation/PerformanceCaptureCore/Source/PerformanceCaptureCore/Private/PerformerComponent.cpp
+++ b/Engine/Plugins/Animation/PerformanceCaptureCore/Source/PerformanceCaptureCore/Private/PerformerComponent.cpp
@@ -21,7 +21,7 @@ void UPerformerComponent::DestroyComponent(bool bPromoteChildren)
 	for (USkeletalMeshComponent* SkeletalMeshComponent : SkeletalMeshComponents)
 	{
 		SkeletalMeshComponent->SetUpdateAnimationInEditor(false/* NewUpdateState */);
-		SkeletalMeshComponent->SetAnimClass(nullptr);
+		SkeletalMeshComponent->SetAnimInstanceClass(nullptr);
 		SkeletalMeshComponent->InitAnim(true/* bForceReinit */);
 	}
 	
@@ -140,7 +140,7 @@ void UPerformerComponent::InitiateAnimation()
 	{
 		for (USkeletalMeshComponent* OwnerSkeletalMeshComponent : OwnerSkeletalMeshComponents)
 		{
-			OwnerSkeletalMeshComponent->SetAnimClass(nullptr);
+			OwnerSkeletalMeshComponent->SetAnimInstanceClass(nullptr);
 			OwnerSkeletalMeshComponent->InitAnim(true /*bForceReinit*/);
 			OwnerSkeletalMeshComponent->SetUpdateAnimationInEditor(true);
 		}
@@ -156,7 +156,7 @@ void UPerformerComponent::InitiateAnimation()
 		//Skip any skeletal meshes that already have anim instance set that's not a LiveLinkInstance
 		if(OwnerSkeletalMeshComponent->GetAnimClass()==ULiveLinkInstance::StaticClass())
 		{
-			OwnerSkeletalMeshComponent->SetAnimClass(nullptr);
+			OwnerSkeletalMeshComponent->SetAnimInstanceClass(nullptr);
 			OwnerSkeletalMeshComponent->InitAnim(true /*bForceReinit*/);
 			OwnerSkeletalMeshComponent->SetUpdateAnimationInEditor(true);
 		}
@@ -175,7 +175,7 @@ void UPerformerComponent::InitiateAnimation()
 
 	if(IsValid(ControlledMesh))
 	{
-		ControlledMesh->SetAnimClass(ULiveLinkInstance::StaticClass());
+		ControlledMesh->SetAnimInstanceClass(ULiveLinkInstance::StaticClass());
 		ControlledMesh->InitAnim(true /*bForceReinit*/);
 		ControlledMesh->SetUpdateAnimationInEditor(true);
 		ControlledMesh->bPropagateCurvesToFollowers = true;
diff --git a/Engine/Plugins/Animation/PerformanceCaptureCore/Source/PerformanceCaptureCore/Private/RetargetComponent.cpp b/Engine/Plugins/Animation/PerformanceCaptureCore/Source/PerformanceCaptureCore/Private/RetargetComponent.cpp
--- a/Engine/Plugins/Animation/PerformanceCaptureCore/Source/PerformanceCaptureCore/Private/RetargetComponent.cpp
+++ b/Engine/Plugins/Animation/PerformanceCaptureCore/Source/PerformanceCaptureCore/Private/RetargetComponent.cpp
@@ -123,7 +123,7 @@ void URetargetComponent::DestroyComponent(bool bPromoteChildren)
 		//If the anim class is LiveLinkInstance assume user does not want to reset that skeletal mesh component
 		if(OwnerSkeletalMeshComponent->GetAnimClass()!=ULiveLinkInstance::StaticClass())
 		{
-			OwnerSkeletalMeshComponent->SetAnimClass(nullptr);
+			OwnerSkeletalMeshComponent->SetAnimInstanceClass(nullptr);
 			OwnerSkeletalMeshComponent->SetLeaderPoseComponent(nullptr);
 			OwnerSkeletalMeshComponent->InitAnim(true /*bForceReinit*/);
 			OwnerSkeletalMeshComponent->SetUpdateAnimationInEditor(false);
@@ -192,7 +192,7 @@ void URetargetComponent::SetForceOtherMeshesToFollowControlledMesh(bool bInBool)
 				&& OwnerSkeletalMeshComponent!=ControlledMesh && OwnerSkeletalMeshComponent!=SourceMesh;
 				if(bShouldResetMesh)
 				{
-					OwnerSkeletalMeshComponent->SetAnimClass(nullptr);
+					OwnerSkeletalMeshComponent->SetAnimInstanceClass(nullptr);
 					OwnerSkeletalMeshComponent->InitAnim(true /*bForceReinit*/);
 					OwnerSkeletalMeshComponent->SetUpdateAnimationInEditor(true);
 					OwnerSkeletalMeshComponent->SetLeaderPoseComponent(nullptr);
@@ -229,7 +229,7 @@ void URetargetComponent::InitiateAnimation()
 	if(bShouldBeReinitialized)
 	{
 		//Set the anim instance class on the controlled mesh to use RetargetAnimInstance
-		ControlledMesh->SetAnimClass(URetargetAnimInstance::StaticClass());
+		ControlledMesh->SetAnimInstanceClass(URetargetAnimInstance::StaticClass());
 
 		TObjectPtr<URetargetAnimInstance> AnimInstance = Cast<URetargetAnimInstance>(ControlledMesh->GetAnimInstance());
 		
diff --git a/Engine/Source/Editor/UnrealEd/Classes/Animation/DebugSkelMeshComponent.h b/Engine/Source/Editor/UnrealEd/Classes/Animation/DebugSkelMeshComponent.h
--- a/Engine/Source/Editor/UnrealEd/Classes/Animation/DebugSkelMeshComponent.h
+++ b/Engine/Source/Editor/UnrealEd/Classes/Animation/DebugSkelMeshComponent.h
@@ -397,7 +397,7 @@ class UDebugSkelMeshComponent : public USkeletalMeshComponent
 	//~ Begin SkeletalMeshComponent Interface
 	UNREALED_API virtual void InitAnim(bool bForceReinit) override;
 	virtual bool IsWindEnabled() const override { return true; }
-	UNREALED_API virtual void SetAnimClass(class UClass* NewClass) override;
+	UNREALED_API virtual void SetAnimInstanceClass(class UClass* NewClass) override;
 	UNREALED_API virtual void OnClearAnimScriptInstance() override;
 	UNREALED_API virtual void SetSkeletalMesh(USkeletalMesh* InSkelMesh, bool bReinitPose = true) override;
 	//~ End SkeletalMeshComponent Interface
diff --git a/Engine/Source/Editor/UnrealEd/Private/Animation/DebugSkelMeshComponent.cpp b/Engine/Source/Editor/UnrealEd/Private/Animation/DebugSkelMeshComponent.cpp
--- a/Engine/Source/Editor/UnrealEd/Private/Animation/DebugSkelMeshComponent.cpp
+++ b/Engine/Source/Editor/UnrealEd/Private/Animation/DebugSkelMeshComponent.cpp
@@ -612,7 +612,7 @@ void UDebugSkelMeshComponent::InitAnim(bool bForceReinit)
 	}
 }
 
-void UDebugSkelMeshComponent::SetAnimClass(class UClass* NewClass)
+void UDebugSkelMeshComponent::SetAnimInstanceClass(class UClass* NewClass)
 {
 	// Override this to do nothing and warn the user
 	UE_LOG(LogAnimation, Warning, TEXT("Attempting to destroy an animation preview actor, skipping."));
diff --git a/Engine/Source/Runtime/Engine/Classes/Components/SkeletalMeshComponent.h b/Engine/Source/Runtime/Engine/Classes/Components/SkeletalMeshComponent.h
--- a/Engine/Source/Runtime/Engine/Classes/Components/SkeletalMeshComponent.h
+++ b/Engine/Source/Runtime/Engine/Classes/Components/SkeletalMeshComponent.h
@@ -955,11 +955,12 @@ public:
 	/** Array of physical interactions for the frame. This is a temporary solution for a more permanent force system and should not be used directly*/
 	TArray<FPendingRadialForces> PendingRadialForces;
 
-	UE_DEPRECATED(4.23, "This function is deprecated. Please use SetAnimClass instead. ")
+	UE_DEPRECATED(4.23, "This function is deprecated. Please use SetAnimInstanceClass instead. ")
 	ENGINE_API virtual void K2_SetAnimInstanceClass(class UClass* NewClass);
 
 	/** Set the anim instance class. Clears and re-initializes the anim instance with the new class and sets animation mode to 'AnimationBlueprint' */
-	UFUNCTION(BlueprintCallable, Category = "Components|SkeletalMesh", meta = (Keywords = "AnimBlueprint", DisplayName = "Set Anim Instance Class"))
+	UE_DEPRECATED(5.5, "This function is deprecated. Please use 'SetAnimInstanceClass' instead. ")
+	UFUNCTION(BlueprintCallable, Category = "Components|SkeletalMesh", meta = (Keywords = "AnimBlueprint", DisplayName = "Set Anim Class", DeprecatedFunction, DeprecationMessage = "This function is deprecated. Please use 'SetAnimInstanceClass' instead. "))
 	ENGINE_API virtual void SetAnimClass(class UClass* NewClass);
 
 	/** Get the anim instance class via getter callable by sequencer.  */
@@ -967,7 +968,8 @@ public:
 	ENGINE_API class UClass*  GetAnimClass();
 
 	/** Set the anim instance class. Clears and re-initializes the anim instance with the new class and sets animation mode to 'AnimationBlueprint' */
-	ENGINE_API void SetAnimInstanceClass(class UClass* NewClass);
+	UFUNCTION(BlueprintCallable, Category = "Components|SkeletalMesh", meta = (Keywords = "AnimBlueprint", DisplayName = "Set Anim Instance Class"))
+	ENGINE_API virtual void SetAnimInstanceClass(class UClass* NewClass);
 
 	/** 
 	 * Returns the animation instance that is driving the class (if available). This is typically an instance of
@@ -1029,7 +1031,7 @@ public:
 	 * Returns all tagged linked instance nodes that match the tag.
 	 */
 	UE_DEPRECATED(5.0, "Tags are unique so this funciton is no longer supported. Please use GetLinkedAnimGraphInstanceByTag instead")
-	UFUNCTION(BlueprintPure, Category = "Components|SkeletalMesh|Animation Blueprint Linking", meta = (Keywords = "AnimBlueprint", DeprecatedFunction, DeprecationMessage="Tags are unique so this funciton is no longer supported. Please use GetLinkedAnimGraphInstanceByTag instead"))
+	UFUNCTION(BlueprintPure, Category = "Components|SkeletalMesh|Animation Blueprint Linking", meta = (Keywords = "AnimBlueprint", DeprecatedFunction, DeprecationMessage="Tags are unique so this function is no longer supported. Please use GetLinkedAnimGraphInstanceByTag instead"))
 	ENGINE_API void GetLinkedAnimGraphInstancesByTag(FName InTag, TArray<UAnimInstance*>& OutLinkedInstances) const;
 
 	UE_DEPRECATED(4.24, "Function renamed, please use LinkAnimGraphByTag")
