Fixed OnlyTickMontagesAndRefreshBonesWhenPlayingMontages option updating the anim instance multiple times

[RN] Fixed OnlyTickMontagesAndRefreshBonesWhenPlayingMontages option updating the anim instance multiple times

diff --git a/Engine/Source/Runtime/Engine/Classes/Components/SkeletalMeshComponent.h b/Engine/Source/Runtime/Engine/Classes/Components/SkeletalMeshComponent.h
--- a/Engine/Source/Runtime/Engine/Classes/Components/SkeletalMeshComponent.h
+++ b/Engine/Source/Runtime/Engine/Classes/Components/SkeletalMeshComponent.h
@@ -1621,6 +1621,8 @@ public:
 	*/
 	ENGINE_API bool ShouldOnlyTickMontages(const float DeltaTime) const;
 
+	ENGINE_API bool ShouldOnlyTickMontagesAndRefreshBones(const float DeltaTime) const;
+
 	/** @return whether we should tick animation (we may want to skip it due to URO) */
 	ENGINE_API bool ShouldTickAnimation() const;
 
diff --git a/Engine/Source/Runtime/Engine/Private/Animation/AnimInstance.cpp b/Engine/Source/Runtime/Engine/Private/Animation/AnimInstance.cpp
--- a/Engine/Source/Runtime/Engine/Private/Animation/AnimInstance.cpp
+++ b/Engine/Source/Runtime/Engine/Private/Animation/AnimInstance.cpp
@@ -532,7 +532,7 @@ void UAnimInstance::UpdateAnimation(float DeltaSeconds, bool bNeedsValidRootMoti
 			If we're set to OnlyTickMontagesWhenNotRendered and we haven't been recently rendered,
 			then only update montages and skip everything else. 
 		*/
-		if (SkelMeshComp->ShouldOnlyTickMontages(DeltaSeconds))
+		if (SkelMeshComp->ShouldOnlyTickMontages(DeltaSeconds) || (SkelMeshComp->ShouldOnlyTickMontagesAndRefreshBones(DeltaSeconds) && !IsAnyMontagePlaying()))
 		{
 			/**
 				Clear NotifyQueue prior to ticking montages.
diff --git a/Engine/Source/Runtime/Engine/Private/Components/CharacterMovementComponent.cpp b/Engine/Source/Runtime/Engine/Private/Components/CharacterMovementComponent.cpp
--- a/Engine/Source/Runtime/Engine/Private/Components/CharacterMovementComponent.cpp
+++ b/Engine/Source/Runtime/Engine/Private/Components/CharacterMovementComponent.cpp
@@ -10440,7 +10440,7 @@ void UCharacterMovementComponent::MoveAutonomous
 		else
 		{
 			// Revert back to old behavior if wanted/needed.
-			if (OwnerMesh->ShouldOnlyTickMontages(DeltaTime))
+			if (OwnerMesh->ShouldOnlyTickMontages(DeltaTime) || OwnerMesh->ShouldOnlyTickMontagesAndRefreshBones(DeltaTime))
 			{
 				OwnerMesh->ConditionallyDispatchQueuedAnimEvents();
 			}
diff --git a/Engine/Source/Runtime/Engine/Private/Components/SkeletalMeshComponent.cpp b/Engine/Source/Runtime/Engine/Private/Components/SkeletalMeshComponent.cpp
--- a/Engine/Source/Runtime/Engine/Private/Components/SkeletalMeshComponent.cpp
+++ b/Engine/Source/Runtime/Engine/Private/Components/SkeletalMeshComponent.cpp
@@ -1412,8 +1412,16 @@ bool USkeletalMeshComponent::ShouldOnlyTickMontages(const float DeltaTime) const
 {
 	// Ignore DeltaSeconds == 0.f, as that is used when we want to force an update followed by RefreshBoneTransforms.
 	// RefreshBoneTransforms will need an updated graph.
-	return (VisibilityBasedAnimTickOption == EVisibilityBasedAnimTickOption::OnlyTickMontagesWhenNotRendered
-		|| VisibilityBasedAnimTickOption == EVisibilityBasedAnimTickOption::OnlyTickMontagesAndRefreshBonesWhenPlayingMontages)
+	return VisibilityBasedAnimTickOption == EVisibilityBasedAnimTickOption::OnlyTickMontagesWhenNotRendered
+		&& !bRecentlyRendered
+		&& (DeltaTime > 0.f);
+}
+
+bool USkeletalMeshComponent::ShouldOnlyTickMontagesAndRefreshBones(const float DeltaTime) const
+{
+	// Ignore DeltaSeconds == 0.f, as that is used when we want to force an update followed by RefreshBoneTransforms.
+	// RefreshBoneTransforms will need an updated graph.
+	return VisibilityBasedAnimTickOption == EVisibilityBasedAnimTickOption::OnlyTickMontagesAndRefreshBonesWhenPlayingMontages
 		&& !bRecentlyRendered
 		&& (DeltaTime > 0.f);
 }
@@ -1452,7 +1460,7 @@ void USkeletalMeshComponent::TickAnimation(float DeltaTime, bool bNeedsValidRoot
 			If we're called directly for autonomous proxies, TickComponent is not guaranteed to get called.
 			So dispatch all queued events here if we're doing MontageOnly ticking.
 		*/
-		if (ShouldOnlyTickMontages(DeltaTime))
+		if (ShouldOnlyTickMontages(DeltaTime) || ShouldOnlyTickMontagesAndRefreshBones(DeltaTime))
 		{
 			ConditionallyDispatchQueuedAnimEvents();
 		}
