[Material Editor] Fixed handling of stray % characters by FStringTemplate.

diff --git a/Engine/Source/Runtime/Engine/Private/Materials/StringTemplate.cpp b/Engine/Source/Runtime/Engine/Private/Materials/StringTemplate.cpp
--- a/Engine/Source/Runtime/Engine/Private/Materials/StringTemplate.cpp
+++ b/Engine/Source/Runtime/Engine/Private/Materials/StringTemplate.cpp
@@ -116,16 +116,6 @@ bool FStringTemplate::Load(FString TemplateStringArg, FErrorInfo& OutErrorinfo)
 				PushChunk(true);
 				NumNamedParameters += 1;
 			}
-			else if (Accept('%')) // e.g. %%
-			{
-				ChunkLength -= 1; // take one '%'
-				NewChunks.Add({ TEXT("%"), false });
-			}
-			else
-			{
-				OutErrorinfo.Message = TEXT("unrecognized % sequence. Are you missing an additional '%' for the modulo sequence?");
-				return false;
-			}
 		}
 		else
 		{
