PR #8098: Fix FAnimMontageInstance::GetNextSection to return next section of instance rather than the montage default next section.
* If the next section of a montage instance is overridden, e.g. via SetNextSectionName, then GetNextSection should return this new value.

diff --git a/Engine/Source/Runtime/Engine/Private/Animation/AnimMontage.cpp b/Engine/Source/Runtime/Engine/Private/Animation/AnimMontage.cpp
--- a/Engine/Source/Runtime/Engine/Private/Animation/AnimMontage.cpp
+++ b/Engine/Source/Runtime/Engine/Private/Animation/AnimMontage.cpp
@@ -1824,8 +1824,11 @@ FName FAnimMontageInstance::GetNextSection() const
 		const int32 CurrentSectionIndex = Montage->GetAnimCompositeSectionIndexFromPos(Position, CurrentPosition);
 		if (Montage->IsValidSectionIndex(CurrentSectionIndex))
 		{
-			FCompositeSection& CurrentSection = Montage->GetAnimCompositeSection(CurrentSectionIndex);
-			return CurrentSection.NextSectionName;
+			const int32 NextSectionIndex = GetNextSectionID(CurrentSectionIndex);
+			if (Montage->IsValidSectionIndex(NextSectionIndex))
+			{
+				return GetSectionNameFromID(NextSectionIndex);
+			}
 		}
 	}
 
