Fixes and improvements to how the content worker uses the ShaderCompile worker when testing locally
* Fix the Metal SDK errors that have been occuring when running with embedded SDKs
* Use the locally built ShaderCompiler worker when running in "SourceWithCookedContent" mode (when base dir is set)
#rnx

diff --git a/Engine/Source/Runtime/Engine/Private/ShaderCompiler/ShaderCompiler.cpp b/Engine/Source/Runtime/Engine/Private/ShaderCompiler/ShaderCompiler.cpp
--- a/Engine/Source/Runtime/Engine/Private/ShaderCompiler/ShaderCompiler.cpp
+++ b/Engine/Source/Runtime/Engine/Private/ShaderCompiler/ShaderCompiler.cpp
@@ -703,6 +703,22 @@ static FString GetBuildMachineArtifactBasePath()
 	return FPaths::Combine(*FPaths::EngineDir(), TEXT("Programs"), TEXT("AutomationTool"), TEXT("Saved"), TEXT("Logs"));
 }
 
+static bool FindShaderCompileWorkerExecutableInLaunchDir(const FString& ExecutableName, FString& OutFilename)
+{
+	FString LocalShaderCompileWorkerName = FPaths::Combine(FPaths::LaunchDir(), ExecutableName);
+	if (!IFileManager::Get().FileExists(*LocalShaderCompileWorkerName))
+	{
+		LocalShaderCompileWorkerName = FPaths::Combine(FPaths::LaunchDir(), TEXT("../../../Engine/Binaries"), FPlatformProcess::GetBinariesSubdirectory(), ExecutableName);
+
+		if (!IFileManager::Get().FileExists(*LocalShaderCompileWorkerName))
+		{			
+			return false;
+		}
+	}
+
+	return true;
+}
+
 FShaderCompilingManager::FShaderCompilingManager() :
 	bCompilingDuringGame(false),
 	NumExternalJobs(0),
@@ -745,6 +761,21 @@ FShaderCompilingManager::FShaderCompilingManager() :
 		ShaderCompileWorkerName = FPaths::Combine(FPaths::EngineDir(), TEXT("Binaries"), FPlatformProcess::GetBinariesSubdirectory(), ExecutableName);
 	}
 
+	// Optionally allow the shader worker path to use the launch directory, this allows the engine to use a locally built shader compile worker when running with the -basedir argument
+	bool bUseShaderCompilerFromLaunchDir = false;
+	if (GConfig->GetBool(TEXT("DevOptions.Shaders"), TEXT("bUseShaderCompilerFromLaunchDir"), bUseShaderCompilerFromLaunchDir, GEngineIni) && bUseShaderCompilerFromLaunchDir)
+	{
+		FString LocalShaderCompileWorkerName;
+		if (FindShaderCompileWorkerExecutableInLaunchDir(ExecutableName, LocalShaderCompileWorkerName))
+		{
+			ShaderCompileWorkerName = LocalShaderCompileWorkerName;
+		}
+		else
+		{
+			UE_LOG(LogShaderCompilers, Warning, TEXT("Using bUseShaderCompilerFromLaunchDir but could not find shader compile worker in LaunchDir - '%s'."), *FPaths::LaunchDir());
+		}
+	}
+
 	// Threads must use absolute paths on Windows in case the current directory is changed on another thread!
 	ShaderCompileWorkerName = FPaths::ConvertRelativePathToFull(ShaderCompileWorkerName);
 
@@ -1517,6 +1548,19 @@ FProcHandle FShaderCompilingManager::LaunchWorker(const FString& WorkingDirector
 	int32 PriorityModifier = -1; // below normal
 	GConfig->GetInt(TEXT("DevOptions.Shaders"), TEXT("WorkerProcessPriority"), PriorityModifier, GEngineIni);
 
+	//Inherit the base directory from the engine process
+	FString BaseDirOverride;
+	if (FParse::Value(FCommandLine::Get(), TEXT("basedir="), BaseDirOverride))
+	{
+		WorkerParameters += FString::Printf(TEXT("-basedir=%s"), *BaseDirOverride);
+	}
+
+	FString BaseFromWorkingDirOverride;
+	if (FParse::Value(FCommandLine::Get(), TEXT("BaseFromWorkingDir="), BaseFromWorkingDirOverride))
+	{
+		WorkerParameters += FString::Printf(TEXT("-BaseFromWorkingDir=%s"), *BaseFromWorkingDirOverride);
+	}
+
 	if (DEBUG_SHADERCOMPILEWORKER)
 	{
 		// Note: Set breakpoint here and launch the ShaderCompileWorker with WorkerParameters a cmd-line
