Fixed crashes when re-instancing animation blueprints with custom proxies

[RN] Fixed crashes when re-instancing animation blueprints with custom proxies

diff --git a/Engine/Source/Runtime/Engine/Private/Animation/AnimNode_LinkedAnimGraph.cpp b/Engine/Source/Runtime/Engine/Private/Animation/AnimNode_LinkedAnimGraph.cpp
--- a/Engine/Source/Runtime/Engine/Private/Animation/AnimNode_LinkedAnimGraph.cpp
+++ b/Engine/Source/Runtime/Engine/Private/Animation/AnimNode_LinkedAnimGraph.cpp
@@ -522,6 +522,11 @@ void FAnimNode_LinkedAnimGraph::HandleObjectsReinstanced_Impl(UObject* InSourceO
 		// re-instancing, and they dont call the constructor that takes a UAnimInstance*)
 		SourceProxy.Initialize(SourceAnimInstance);
 
+		// Similarly call Initialize here to catch any custom target proxies
+		UAnimInstance* TargetAnimInstance = CastChecked<UAnimInstance>(InTargetObject);
+		FAnimInstanceProxy& TargetProxy = TargetAnimInstance->GetProxyOnAnyThread<FAnimInstanceProxy>();
+		TargetProxy.Initialize(TargetAnimInstance);
+
 		InitializeProperties(SourceAnimInstance, GetTargetClass());
 		DynamicUnlink(SourceAnimInstance);
 		DynamicLink(SourceAnimInstance);
