[Engine]
* Renamed UE_STATE_STREAM to WITH_STATE_STREAM and moved define to Build.h since it is used in more places than under Engine

diff --git a/Engine/Source/Runtime/Core/Public/Misc/Build.h b/Engine/Source/Runtime/Core/Public/Misc/Build.h
--- a/Engine/Source/Runtime/Core/Public/Misc/Build.h
+++ b/Engine/Source/Runtime/Core/Public/Misc/Build.h
@@ -167,6 +167,13 @@
 	#define WITH_TEXT_ARCHIVE_SUPPORT WITH_EDITORONLY_DATA
 #endif
 
+// Statestream is WIP and is the system that is going to enable full decoupling of game and render updates
+// Set in .target.cs to enable
+#ifndef WITH_STATE_STREAM
+	#define WITH_STATE_STREAM 0
+#endif
+
+
 /*----------------------------------------------------------------------------
 	Optional bridge options coming from UBT, do not modify directly!
 	If UBT doesn't set the value, it is assumed to be 0, and we set that here.
diff --git a/Engine/Source/Runtime/Engine/Private/SceneView.cpp b/Engine/Source/Runtime/Engine/Private/SceneView.cpp
--- a/Engine/Source/Runtime/Engine/Private/SceneView.cpp
+++ b/Engine/Source/Runtime/Engine/Private/SceneView.cpp
@@ -762,13 +762,13 @@ FViewMatrices::FViewMatrices(const FMinimalInitializer& Initializer) : FViewMatr
 
 FSceneViewOwner::FSceneViewOwner(const AActor* InActor)
 	: ActorUniqueId(InActor ? InActor->GetUniqueID() : 0)
-#if !UE_STATE_STREAM
+#if !WITH_STATE_STREAM
 	, Actor(InActor)
 #endif
 {
 }
 
-#if !UE_STATE_STREAM
+#if !WITH_STATE_STREAM
 void FSceneViewOwner::operator=(const AActor* InActor)
 {
 	Actor = InActor;
diff --git a/Engine/Source/Runtime/Engine/Private/WorldPartition/WorldPartitionSubsystem.cpp b/Engine/Source/Runtime/Engine/Private/WorldPartition/WorldPartitionSubsystem.cpp
--- a/Engine/Source/Runtime/Engine/Private/WorldPartition/WorldPartitionSubsystem.cpp
+++ b/Engine/Source/Runtime/Engine/Private/WorldPartition/WorldPartitionSubsystem.cpp
@@ -1393,7 +1393,7 @@ void UWorldPartitionSubsystem::Draw(UCanvas* Canvas, class APlayerController* PC
 		return;
 	}
 
-#if !UE_STATE_STREAM // TODO: This needs to be fixed
+#if !WITH_STATE_STREAM // TODO: This needs to be fixed
 
 	// Filter out views that don't match our world
 	if (!GetWorld()->IsNetMode(NM_DedicatedServer) && !UWorldPartition::IsSimulating(false) &&
diff --git a/Engine/Source/Runtime/Engine/Public/EngineDefines.h b/Engine/Source/Runtime/Engine/Public/EngineDefines.h
--- a/Engine/Source/Runtime/Engine/Public/EngineDefines.h
+++ b/Engine/Source/Runtime/Engine/Public/EngineDefines.h
@@ -25,9 +25,10 @@
 // The number of lights to consider for sky/atmospheric light scattering
 #define NUM_ATMOSPHERE_LIGHTS 2
 
-// Statestream is WIP and is the system that is going to enable full separation of game and render updates
-#ifndef UE_STATE_STREAM
-	#define UE_STATE_STREAM 0
+// Enables StateStream for actor/component system
+// Set in .target.cs to enable
+#ifndef WITH_STATE_STREAM_ACTOR
+	#define WITH_STATE_STREAM_ACTOR 0
 #endif
 
 /*-----------------------------------------------------------------------------
diff --git a/Engine/Source/Runtime/Engine/Public/SceneViewOwner.h b/Engine/Source/Runtime/Engine/Public/SceneViewOwner.h
--- a/Engine/Source/Runtime/Engine/Public/SceneViewOwner.h
+++ b/Engine/Source/Runtime/Engine/Public/SceneViewOwner.h
@@ -15,7 +15,7 @@ struct FSceneViewOwner
 	bool IsSet() const { return ActorUniqueId != 0; }
 	uint32 ActorUniqueId;
 
-	#if !UE_STATE_STREAM
+	#if !WITH_STATE_STREAM
 	const AActor* Get() const { return Actor; }
 	operator const AActor*() const { return Actor; }
 	const AActor* operator->() const { return Actor; }
