[VisualLogger] made VisualLoggerKismetLibrary functions BlueprintThreadSafe.
For message logging we use FMessageLog only when functions are called on the game thread, otherwise we use FMsg for thread safety

diff --git a/Engine/Source/Runtime/Engine/Classes/VisualLogger/VisualLoggerKismetLibrary.h b/Engine/Source/Runtime/Engine/Classes/VisualLogger/VisualLoggerKismetLibrary.h
--- a/Engine/Source/Runtime/Engine/Classes/VisualLogger/VisualLoggerKismetLibrary.h
+++ b/Engine/Source/Runtime/Engine/Classes/VisualLogger/VisualLoggerKismetLibrary.h
@@ -8,56 +8,93 @@
 UCLASS(MinimalAPI, meta=(ScriptName="VisualLoggerLibrary"))
 class UVisualLoggerKismetLibrary : public UBlueprintFunctionLibrary
 {
-	GENERATED_UCLASS_BODY()
+	GENERATED_BODY()
 
+public:
 	UFUNCTION(BlueprintCallable, Category = "Debug|VisualLogger", meta = (DisplayName = "Enable VisLog Recording", CallableWithoutWorldContext, DevelopmentOnly))
 	static void EnableRecording(bool bEnabled);
 
 	/** Makes SourceOwner log to DestinationOwner's vislog*/
-	UFUNCTION(BlueprintCallable, Category = "Debug|VisualLogger", meta = (CallableWithoutWorldContext, DevelopmentOnly))
+	UFUNCTION(BlueprintCallable, Category = "Debug|VisualLogger", meta = (BlueprintThreadSafe, CallableWithoutWorldContext, DevelopmentOnly))
 	static void RedirectVislog(UObject* SourceOwner, UObject* DestinationOwner);
 
-	/** Logs simple text string with Visual Logger - recording for Visual Logs has to be enabled to record this data */
-	UFUNCTION(BlueprintCallable, Category = "Debug|VisualLogger", meta = (DisplayName = "VisLog Text", AdvancedDisplay = "WorldContextObject, bAddToMessageLog", CallableWithoutWorldContext, DevelopmentOnly, DefaultToSelf = "WorldContextObject"))
+	/** Logs simple text string with Visual Logger - recording for Visual Logs has to be enabled to record this data
+	 * @param bAddToMessageLog Indicates whether a log should also be produced.
+	 *  Message log is used for calls on game thread, otherwise regular output log is used for thread safety.
+	 */
+	UFUNCTION(BlueprintCallable, Category = "Debug|VisualLogger", meta = (BlueprintThreadSafe, DisplayName = "VisLog Text", AdvancedDisplay = "WorldContextObject, bAddToMessageLog", CallableWithoutWorldContext, DevelopmentOnly, DefaultToSelf = "WorldContextObject"))
 	static void LogText(UObject* WorldContextObject, FString Text, FName LogCategory = TEXT("VisLogBP"), bool bAddToMessageLog = false);
 
-	/** Logs location as sphere with given radius - recording for Visual Logs has to be enabled to record this data */
-	UFUNCTION(BlueprintCallable, Category = "Debug|VisualLogger", meta = (DisplayName = "VisLog Location", AdvancedDisplay = "WorldContextObject, bAddToMessageLog", CallableWithoutWorldContext, DevelopmentOnly, DefaultToSelf = "WorldContextObject"))
+	/** Logs location as sphere with given radius - recording for Visual Logs has to be enabled to record this data
+	 * @param bAddToMessageLog Indicates whether a log should also be produced.
+	 *  Message log is used for calls on game thread, otherwise regular output log is used for thread safety.
+	 */
+	UFUNCTION(BlueprintCallable, Category = "Debug|VisualLogger", meta = (BlueprintThreadSafe, DisplayName = "VisLog Location", AdvancedDisplay = "WorldContextObject, bAddToMessageLog", CallableWithoutWorldContext, DevelopmentOnly, DefaultToSelf = "WorldContextObject"))
 	static void LogLocation(UObject* WorldContextObject, FVector Location, FString Text, FLinearColor ObjectColor = FLinearColor::Blue, float Radius = 10, FName LogCategory = TEXT("VisLogBP"), bool bAddToMessageLog = false);
 
-	/** Logs sphere shape - recording for Visual Logs has to be enabled to record this data */
-	UFUNCTION(BlueprintCallable, Category = "Debug|VisualLogger", meta = (DisplayName = "VisLog Sphere Shape", AdvancedDisplay = "WorldContextObject, bAddToMessageLog", CallableWithoutWorldContext, DevelopmentOnly, DefaultToSelf = "WorldContextObject"))
+	/** Logs sphere shape - recording for Visual Logs has to be enabled to record this data
+	 * @param bAddToMessageLog Indicates whether a log should also be produced.
+	 *  Message log is used for calls on game thread, otherwise regular output log is used for thread safety.
+	 */
+	UFUNCTION(BlueprintCallable, Category = "Debug|VisualLogger", meta = (BlueprintThreadSafe, DisplayName = "VisLog Sphere Shape", AdvancedDisplay = "WorldContextObject, bAddToMessageLog", CallableWithoutWorldContext, DevelopmentOnly, DefaultToSelf = "WorldContextObject"))
 	static void LogSphere(UObject* WorldContextObject, FVector Center, float Radius, FString Text, FLinearColor ObjectColor = FLinearColor::Blue, FName LogCategory = TEXT("VisLogBP"), bool bAddToMessageLog = false, bool bWireframe = false);
 
-	/** Logs cone shape - recording for Visual Logs has to be enabled to record this data */
-	UFUNCTION(BlueprintCallable, Category = "Debug|VisualLogger", meta = (DisplayName = "VisLog Cone Shape", AdvancedDisplay = "WorldContextObject, bAddToMessageLog", CallableWithoutWorldContext, DevelopmentOnly, DefaultToSelf = "WorldContextObject"))
+	/** Logs cone shape - recording for Visual Logs has to be enabled to record this data
+	 * @param bAddToMessageLog Indicates whether a log should also be produced.
+	 *  Message log is used for calls on game thread, otherwise regular output log is used for thread safety.
+	 */
+	UFUNCTION(BlueprintCallable, Category = "Debug|VisualLogger", meta = (BlueprintThreadSafe, DisplayName = "VisLog Cone Shape", AdvancedDisplay = "WorldContextObject, bAddToMessageLog", CallableWithoutWorldContext, DevelopmentOnly, DefaultToSelf = "WorldContextObject"))
 	static void LogCone(UObject* WorldContextObject, FVector Origin, FVector Direction, float Length, float Angle, FString Text, FLinearColor ObjectColor = FLinearColor::Blue, FName LogCategory = TEXT("VisLogBP"), bool bAddToMessageLog = false, bool bWireframe = false);
 
-	/** Logs cylinder shape - recording for Visual Logs has to be enabled to record this data */
-	UFUNCTION(BlueprintCallable, Category = "Debug|VisualLogger", meta = (DisplayName = "VisLog Cylinder Shape", AdvancedDisplay = "WorldContextObject, bAddToMessageLog", CallableWithoutWorldContext, DevelopmentOnly, DefaultToSelf = "WorldContextObject"))
+	/** Logs cylinder shape - recording for Visual Logs has to be enabled to record this data
+	 * @param bAddToMessageLog Indicates whether a log should also be produced.
+	 *  Message log is used for calls on game thread, otherwise regular output log is used for thread safety.
+	 */
+	UFUNCTION(BlueprintCallable, Category = "Debug|VisualLogger", meta = (BlueprintThreadSafe, DisplayName = "VisLog Cylinder Shape", AdvancedDisplay = "WorldContextObject, bAddToMessageLog", CallableWithoutWorldContext, DevelopmentOnly, DefaultToSelf = "WorldContextObject"))
 	static void LogCylinder(UObject* WorldContextObject, FVector Start, FVector End, float Radius, FString Text, FLinearColor ObjectColor = FLinearColor::Blue, FName LogCategory = TEXT("VisLogBP"), bool bAddToMessageLog = false, bool bWireframe = false);
 
-	/** Logs capsule shape - recording for Visual Logs has to be enabled to record this data */
-	UFUNCTION(BlueprintCallable, Category = "Debug|VisualLogger", meta = (DisplayName = "VisLog Capsule Shape", AdvancedDisplay = "WorldContextObject, bAddToMessageLog", CallableWithoutWorldContext, DevelopmentOnly, DefaultToSelf = "WorldContextObject"))
+	/** Logs capsule shape - recording for Visual Logs has to be enabled to record this data
+	 * @param bAddToMessageLog Indicates whether a log should also be produced.
+	 *  Message log is used for calls on game thread, otherwise regular output log is used for thread safety.
+	 */
+	UFUNCTION(BlueprintCallable, Category = "Debug|VisualLogger", meta = (BlueprintThreadSafe, DisplayName = "VisLog Capsule Shape", AdvancedDisplay = "WorldContextObject, bAddToMessageLog", CallableWithoutWorldContext, DevelopmentOnly, DefaultToSelf = "WorldContextObject"))
 	static void LogCapsule(UObject* WorldContextObject, FVector Base, float HalfHeight, float Radius, FQuat Rotation, FString Text, FLinearColor ObjectColor = FLinearColor::Blue, FName LogCategory = TEXT("VisLogBP"), bool bAddToMessageLog = false, bool bWireframe = false);
 
-	/** Logs box shape - recording for Visual Logs has to be enabled to record this data */
-	UFUNCTION(BlueprintCallable, Category = "Debug|VisualLogger", meta = (DisplayName = "VisLog Box Shape", AdvancedDisplay = "WorldContextObject, bAddToMessageLog", CallableWithoutWorldContext, DevelopmentOnly, DefaultToSelf = "WorldContextObject"))
+	/** Logs box shape - recording for Visual Logs has to be enabled to record this data
+	 * @param bAddToMessageLog Indicates whether a log should also be produced.
+	 *  Message log is used for calls on game thread, otherwise regular output log is used for thread safety.
+	 */
+	UFUNCTION(BlueprintCallable, Category = "Debug|VisualLogger", meta = (BlueprintThreadSafe, DisplayName = "VisLog Box Shape", AdvancedDisplay = "WorldContextObject, bAddToMessageLog", CallableWithoutWorldContext, DevelopmentOnly, DefaultToSelf = "WorldContextObject"))
 	static void LogBox(UObject* WorldContextObject, FBox BoxShape, FString Text, FLinearColor ObjectColor = FLinearColor::Blue, FName LogCategory = TEXT("VisLogBP"), bool bAddToMessageLog = false, bool bWireframe = false);
 
-	/** Logs oriented box shape - recording for Visual Logs has to be enabled to record this data */
-	UFUNCTION(BlueprintCallable, Category = "Debug|VisualLogger", meta = (DisplayName = "VisLog Oriented Box Shape", AdvancedDisplay = "WorldContextObject, bAddToMessageLog", CallableWithoutWorldContext, DevelopmentOnly, DefaultToSelf = "WorldContextObject"))
+	/** Logs oriented box shape - recording for Visual Logs has to be enabled to record this data
+	 * @param bAddToMessageLog Indicates whether a log should also be produced.
+	 *  Message log is used for calls on game thread, otherwise regular output log is used for thread safety.
+	 */
+	UFUNCTION(BlueprintCallable, Category = "Debug|VisualLogger", meta = (BlueprintThreadSafe, DisplayName = "VisLog Oriented Box Shape", AdvancedDisplay = "WorldContextObject, bAddToMessageLog", CallableWithoutWorldContext, DevelopmentOnly, DefaultToSelf = "WorldContextObject"))
 	static void LogOrientedBox(UObject* WorldContextObject, FBox BoxShape, FTransform Transform, FString Text, FLinearColor ObjectColor = FLinearColor::Blue, FName LogCategory = TEXT("VisLogBP"), bool bAddToMessageLog = false, bool bWireframe = false);
 
-	/** Logs arrow - recording for Visual Logs has to be enabled to record this data */
-	UFUNCTION(BlueprintCallable, Category = "Debug|VisualLogger", meta = (DisplayName = "VisLog Arrow", AdvancedDisplay = "WorldContextObject, bAddToMessageLog, Thickness", CallableWithoutWorldContext, DevelopmentOnly, DefaultToSelf = "WorldContextObject"))
+	/** Logs arrow - recording for Visual Logs has to be enabled to record this data
+	 * @param bAddToMessageLog Indicates whether a log should also be produced.
+	 *  Message log is used for calls on game thread, otherwise regular output log is used for thread safety.
+	 */
+	UFUNCTION(BlueprintCallable, Category = "Debug|VisualLogger", meta = (BlueprintThreadSafe, DisplayName = "VisLog Arrow", AdvancedDisplay = "WorldContextObject, bAddToMessageLog, Thickness", CallableWithoutWorldContext, DevelopmentOnly, DefaultToSelf = "WorldContextObject"))
 	static void LogArrow(UObject* WorldContextObject, const FVector SegmentStart, const FVector SegmentEnd, FString Text, FLinearColor ObjectColor = FLinearColor::Blue, FName CategoryName = TEXT("VisLogBP"), bool bAddToMessageLog = false, float ArrowHeadSize = 8.0f);
 
-	/** Logs circle - recording for Visual Logs has to be enabled to record this data */
-	UFUNCTION(BlueprintCallable, Category = "Debug|VisualLogger", meta = (DisplayName = "VisLog Circle", AdvancedDisplay = "WorldContextObject, bAddToMessageLog, Thickness", CallableWithoutWorldContext, DevelopmentOnly, DefaultToSelf = "WorldContextObject"))
+	/** Logs circle - recording for Visual Logs has to be enabled to record this data
+	 * @param bAddToMessageLog Indicates whether a log should also be produced.
+	 *  Message log is used for calls on game thread, otherwise regular output log is used for thread safety.
+	 */
+	UFUNCTION(BlueprintCallable, Category = "Debug|VisualLogger", meta = (BlueprintThreadSafe, DisplayName = "VisLog Circle", AdvancedDisplay = "WorldContextObject, bAddToMessageLog, Thickness", CallableWithoutWorldContext, DevelopmentOnly, DefaultToSelf = "WorldContextObject"))
 	static void LogCircle(UObject* WorldContextObject, FVector Center, FVector UpAxis, float Radius, FString Text, FLinearColor ObjectColor = FLinearColor::Blue, const float Thickness = 0.f, FName CategoryName = TEXT("VisLogBP"), bool bAddToMessageLog = false, bool bWireframe = false);
 
-	/** Logs segment - recording for Visual Logs has to be enabled to record this data */
-	UFUNCTION(BlueprintCallable, Category = "Debug|VisualLogger", meta = (DisplayName = "VisLog Segment", AdvancedDisplay = "WorldContextObject, bAddToMessageLog, Thickness", CallableWithoutWorldContext, DevelopmentOnly, DefaultToSelf = "WorldContextObject"))
+	/** Logs segment - recording for Visual Logs has to be enabled to record this data
+	 * @param bAddToMessageLog Indicates whether a log should also be produced.
+	 *  Message log is used for calls on game thread, otherwise regular output log is used for thread safety.
+	 */
+	UFUNCTION(BlueprintCallable, Category = "Debug|VisualLogger", meta = (BlueprintThreadSafe, DisplayName = "VisLog Segment", AdvancedDisplay = "WorldContextObject, bAddToMessageLog, Thickness", CallableWithoutWorldContext, DevelopmentOnly, DefaultToSelf = "WorldContextObject"))
 	static void LogSegment(UObject* WorldContextObject, const FVector SegmentStart, const FVector SegmentEnd, FString Text, FLinearColor ObjectColor = FLinearColor::Blue, const float Thickness = 0.f, FName CategoryName = TEXT("VisLogBP"), bool bAddToMessageLog = false);
+
+private:
+	static constexpr ELogVerbosity::Type DefaultVerbosity = ELogVerbosity::Log;
 };
diff --git a/Engine/Source/Runtime/Engine/Private/VisualLogger/VisualLoggerKismetLibrary.cpp b/Engine/Source/Runtime/Engine/Private/VisualLogger/VisualLoggerKismetLibrary.cpp
--- a/Engine/Source/Runtime/Engine/Private/VisualLogger/VisualLoggerKismetLibrary.cpp
+++ b/Engine/Source/Runtime/Engine/Private/VisualLogger/VisualLoggerKismetLibrary.cpp
@@ -6,16 +6,15 @@
 
 #include UE_INLINE_GENERATED_CPP_BY_NAME(VisualLoggerKismetLibrary)
 
-namespace
-{
-	const FName NAME_Empty;
-	const FName NAME_BlueprintVLog = TEXT("VisLogBP");
-}
-
-UVisualLoggerKismetLibrary::UVisualLoggerKismetLibrary(const FObjectInitializer& ObjectInitializer)
-	: Super(ObjectInitializer)
-{
-}
+#define VLOG_BP_LIBRARY_ADD_TO_LOG(CategoryName, Format, ...) \
+	if (IsInGameThread())\
+	{\
+		FMessageLog(CategoryName).Info(FText::FromString(FString::Printf((Format), __VA_ARGS__)));\
+	}\
+	else\
+	{\
+		FMsg::Logf(__FILE__, __LINE__, CategoryName, DefaultVerbosity, (Format), __VA_ARGS__);\
+	}
 
 void UVisualLoggerKismetLibrary::RedirectVislog(UObject* SourceOwner, UObject* DestinationOwner)
 {
@@ -35,149 +34,130 @@ void UVisualLoggerKismetLibrary::EnableRecording(bool bEnabled)
 void UVisualLoggerKismetLibrary::LogText(UObject* WorldContextObject, FString Text, FName CategoryName, bool bAddToMessageLog)
 {
 #if ENABLE_VISUAL_LOG
-	const ELogVerbosity::Type DefaultVerbosity = ELogVerbosity::Log;
 	FVisualLogger::CategorizedLogf(WorldContextObject, CategoryName, DefaultVerbosity
 		, TEXT("%s"), *Text);
 #endif
 	if (bAddToMessageLog)
 	{
-		FMessageLog(CategoryName).Info(FText::FromString(FString::Printf(TEXT("LogText: '%s'"), *Text)));
+		VLOG_BP_LIBRARY_ADD_TO_LOG(CategoryName, TEXT("LogText: '%s'"), *Text);
 	}
 }
 
 void UVisualLoggerKismetLibrary::LogLocation(UObject* WorldContextObject, FVector Location, FString Text, FLinearColor ObjectColor, float Radius, FName CategoryName, bool bAddToMessageLog)
 {
 #if ENABLE_VISUAL_LOG
-	const ELogVerbosity::Type DefaultVerbosity = ELogVerbosity::Log;
 	FVisualLogger::LocationLogf(WorldContextObject, CategoryName, DefaultVerbosity
 		, Location, static_cast<uint16>(Radius), ObjectColor.ToFColor(true), TEXT("%s"), *Text);
 #endif
 	if (bAddToMessageLog)
 	{
-		FMessageLog(CategoryName).Info(FText::FromString(FString::Printf(TEXT("LogLocation: '%s' - Location: (%s)")
-			, *Text, *Location.ToString())));
+		VLOG_BP_LIBRARY_ADD_TO_LOG(CategoryName, TEXT("LogLocation: '%s' - Location: (%s)"), *Text, *Location.ToString());
 	}
 }
 
 void UVisualLoggerKismetLibrary::LogSphere(UObject* WorldContextObject, FVector Center, float Radius, FString Text, FLinearColor ObjectColor, FName CategoryName, bool bAddToMessageLog, bool bWireframe)
 {
 #if ENABLE_VISUAL_LOG
-	const ELogVerbosity::Type DefaultVerbosity = ELogVerbosity::Log;
 	FVisualLogger::SphereLogf(WorldContextObject, CategoryName, DefaultVerbosity
 		, Center, Radius, ObjectColor.ToFColor(true), bWireframe, TEXT("%s"), *Text);
 #endif
 	if (bAddToMessageLog)
 	{
-		FMessageLog(CategoryName).Info(FText::FromString(FString::Printf(TEXT("LogSphere: '%s' - Center: (%s) | Radius: %f")
-			, *Text, *Center.ToString(), Radius)));
+		VLOG_BP_LIBRARY_ADD_TO_LOG(CategoryName, TEXT("LogSphere: '%s' - Center: (%s) | Radius: %f"), *Text, *Center.ToString(), Radius);
 	}
 }
 
 void UVisualLoggerKismetLibrary::LogCone(UObject* WorldContextObject, FVector Origin, FVector Direction, float Length, float Angle, FString Text, FLinearColor ObjectColor, FName CategoryName, bool bAddToMessageLog, bool bWireframe)
 {
 #if ENABLE_VISUAL_LOG
-	const ELogVerbosity::Type DefaultVerbosity = ELogVerbosity::Log;
 	FVisualLogger::ConeLogf(WorldContextObject, CategoryName, DefaultVerbosity
 		, Origin, Direction, Length, Angle, ObjectColor.ToFColor(true), bWireframe, TEXT("%s"), *Text);
 #endif
 	if (bAddToMessageLog)
 	{
-		FMessageLog(CategoryName).Info(FText::FromString(FString::Printf(TEXT("LogCone: '%s' - Origin: (%s) | Direction: (%s) | Length: %f | Angle: %f")
-			, *Text, *Origin.ToString(), *Direction.ToString(), Length, Angle)));
+		VLOG_BP_LIBRARY_ADD_TO_LOG(CategoryName, TEXT("LogCone: '%s' - Origin: (%s) | Direction: (%s) | Length: %f | Angle: %f"), *Text, *Origin.ToString(), *Direction.ToString(), Length, Angle);
 	}
 }
 
 void UVisualLoggerKismetLibrary::LogCylinder(UObject* WorldContextObject, FVector Start, FVector End, float Radius, FString Text, FLinearColor ObjectColor, FName CategoryName, bool bAddToMessageLog, bool bWireframe)
 {
 #if ENABLE_VISUAL_LOG
-	const ELogVerbosity::Type DefaultVerbosity = ELogVerbosity::Log;
 	FVisualLogger::CylinderLogf(WorldContextObject, CategoryName, DefaultVerbosity, Start, End, Radius, ObjectColor.ToFColor(true), bWireframe, TEXT("%s"), *Text);
 #endif
 	if (bAddToMessageLog)
 	{
-		FMessageLog(CategoryName).Info(FText::FromString(FString::Printf(TEXT("LogCylinder: '%s' - Start: (%s) | End: (%s) | Radius: %f")
-			, *Text, *Start.ToString(), *End.ToString(), Radius)));
+		VLOG_BP_LIBRARY_ADD_TO_LOG(CategoryName, TEXT("LogCylinder: '%s' - Start: (%s) | End: (%s) | Radius: %f"), *Text, *Start.ToString(), *End.ToString(), Radius);
 	}
 }
 
 void UVisualLoggerKismetLibrary::LogCapsule(UObject* WorldContextObject, FVector Base, float HalfHeight, float Radius, FQuat Rotation, FString Text, FLinearColor ObjectColor, FName CategoryName, bool bAddToMessageLog, bool bWireframe)
 {
 #if ENABLE_VISUAL_LOG
-	const ELogVerbosity::Type DefaultVerbosity = ELogVerbosity::Log;
 	FVisualLogger::CapsuleLogf(WorldContextObject, CategoryName, DefaultVerbosity
 		, Base, HalfHeight, Radius, Rotation, ObjectColor.ToFColor(true), bWireframe, TEXT("%s"), *Text);
 #endif
 	if (bAddToMessageLog)
 	{
-		FMessageLog(CategoryName).Info(FText::FromString(FString::Printf(TEXT("LogCapsule: '%s' - Base: (%s) | HalfHeight: %f | Radius: %f | Rotation: (%s)")
-			, *Text, *Base.ToString(), HalfHeight, Radius, *Rotation.ToString())));
+		VLOG_BP_LIBRARY_ADD_TO_LOG(CategoryName, TEXT("LogCapsule: '%s' - Base: (%s) | HalfHeight: %f | Radius: %f | Rotation: (%s)"), *Text, *Base.ToString(), HalfHeight, Radius, *Rotation.ToString());
 	}
 }
 
 void UVisualLoggerKismetLibrary::LogBox(UObject* WorldContextObject, FBox Box, FString Text, FLinearColor ObjectColor, FName CategoryName, bool bAddToMessageLog, bool bWireframe)
 {
 #if ENABLE_VISUAL_LOG
-	const ELogVerbosity::Type DefaultVerbosity = ELogVerbosity::Log;
 	FVisualLogger::BoxLogf(WorldContextObject, CategoryName, DefaultVerbosity
 		, Box, FMatrix::Identity, ObjectColor.ToFColor(true), bWireframe, TEXT("%s"), *Text);
 #endif
 	if (bAddToMessageLog)
 	{
-		FMessageLog(CategoryName).Info(FText::FromString(FString::Printf(TEXT("LogBox: '%s' - BoxMin: (%s) | BoxMax: (%s)")
-			, *Text, *Box.Min.ToString(), *Box.Max.ToString())));
+		VLOG_BP_LIBRARY_ADD_TO_LOG(CategoryName, TEXT("LogBox: '%s' - BoxMin: (%s) | BoxMax: (%s)"), *Text, *Box.Min.ToString(), *Box.Max.ToString());
 	}
 }
 
 void UVisualLoggerKismetLibrary::LogOrientedBox(UObject* WorldContextObject, FBox Box, FTransform Transform, FString Text, FLinearColor ObjectColor, FName CategoryName, bool bAddToMessageLog, bool bWireframe)
 {
 #if ENABLE_VISUAL_LOG
-	const ELogVerbosity::Type DefaultVerbosity = ELogVerbosity::Log;
 	FVisualLogger::BoxLogf(WorldContextObject, CategoryName, DefaultVerbosity
 		, Box, Transform.ToMatrixWithScale(), ObjectColor.ToFColor(true), bWireframe, TEXT("%s"), *Text);
 #endif
 	if (bAddToMessageLog)
 	{
-		FMessageLog(CategoryName).Info(FText::FromString(FString::Printf(TEXT("LogOrientedBox: '%s' - BoxMin: (%s) | BoxMax: (%s) | Transform: (%s)")
-			, *Text, *Box.Min.ToString(), *Box.Max.ToString(), *Transform.ToString())));
+		VLOG_BP_LIBRARY_ADD_TO_LOG(CategoryName, TEXT("LogOrientedBox: '%s' - BoxMin: (%s) | BoxMax: (%s) | Transform: (%s)"), *Text, *Box.Min.ToString(), *Box.Max.ToString(), *Transform.ToString());
 	}
 }
 
 void UVisualLoggerKismetLibrary::LogArrow(UObject* WorldContextObject, const FVector SegmentStart, const FVector SegmentEnd, FString Text, FLinearColor ObjectColor, FName CategoryName, bool bAddToMessageLog, float ArrowHeadSize)
 {
 #if ENABLE_VISUAL_LOG
-	const ELogVerbosity::Type DefaultVerbosity = ELogVerbosity::Log;
 	FVisualLogger::ArrowLineLogf(WorldContextObject, CategoryName, DefaultVerbosity, SegmentStart, SegmentEnd, ObjectColor.ToFColor(true), ArrowHeadSize, TEXT("%s"), *Text);
 #endif
 	if (bAddToMessageLog)
 	{
-		FMessageLog(CategoryName).Info(FText::FromString(FString::Printf(TEXT("LogArrow: '%s' - SegmentStart: (%s) | SegmentEnd:(%s)")
-			, *Text, *SegmentStart.ToString(), *SegmentStart.ToString())));
+		VLOG_BP_LIBRARY_ADD_TO_LOG(CategoryName, TEXT("LogArrow: '%s' - SegmentStart: (%s) | SegmentEnd:(%s)"), *Text, *SegmentStart.ToString(), *SegmentStart.ToString());
 	}
 }
 
 void UVisualLoggerKismetLibrary::LogCircle(UObject* WorldContextObject, FVector Center, FVector UpAxis, float Radius, FString Text, FLinearColor ObjectColor, const float Thickness, FName CategoryName, bool bAddToMessageLog, bool bWireframe)
 {
 #if ENABLE_VISUAL_LOG
-	const ELogVerbosity::Type DefaultVerbosity = ELogVerbosity::Log;
-	FVisualLogger::DiscLogf(WorldContextObject, CategoryName, DefaultVerbosity, Center, UpAxis, Radius, ObjectColor.ToFColor(true), (uint16)Thickness, bWireframe, TEXT("%s"), *Text);
+	FVisualLogger::DiscLogf(WorldContextObject, CategoryName, DefaultVerbosity, Center, UpAxis, Radius, ObjectColor.ToFColor(true), static_cast<uint16>(Thickness), bWireframe, TEXT("%s"), *Text);
 #endif
 	if (bAddToMessageLog)
 	{
-		FMessageLog(CategoryName).Info(FText::FromString(FString::Printf(TEXT("LogCircle: '%s' - Center: (%s) | UpAxis: (%s) | Radius: %f")
-			, *Text, *Center.ToString(), *UpAxis.ToString(), Radius)));
+		VLOG_BP_LIBRARY_ADD_TO_LOG(CategoryName, TEXT("LogCircle: '%s' - Center: (%s) | UpAxis: (%s) | Radius: %f")	, *Text, *Center.ToString(), *UpAxis.ToString(), Radius);
 	}
 }
 
 void UVisualLoggerKismetLibrary::LogSegment(UObject* WorldContextObject, const FVector SegmentStart, const FVector SegmentEnd, FString Text, FLinearColor ObjectColor, const float Thickness, FName CategoryName, bool bAddToMessageLog)
 {
 #if ENABLE_VISUAL_LOG
-	const ELogVerbosity::Type DefaultVerbosity = ELogVerbosity::Log;
-	FVisualLogger::SegmentLogf(WorldContextObject, CategoryName, DefaultVerbosity, SegmentStart, SegmentEnd, ObjectColor.ToFColor(true), (uint16)Thickness, TEXT("%s"), *Text);
+	FVisualLogger::SegmentLogf(WorldContextObject, CategoryName, DefaultVerbosity, SegmentStart, SegmentEnd, ObjectColor.ToFColor(true), static_cast<uint16>(Thickness), TEXT("%s"), *Text);
 #endif
 	if (bAddToMessageLog)
 	{
-		FMessageLog(CategoryName).Info(FText::FromString(FString::Printf(TEXT("LogSegment: '%s' - SegmentStart: (%s) | SegmentEnd: (%s)")
-			, *Text, *SegmentStart.ToString(), *SegmentStart.ToString())));
+		VLOG_BP_LIBRARY_ADD_TO_LOG(CategoryName, TEXT("LogSegment: '%s' - SegmentStart: (%s) | SegmentEnd: (%s)")
+			, *Text, *SegmentStart.ToString(), *SegmentStart.ToString());
 	}
 }
 
+#undef VLOG_BP_LIBRARY_ADD_TO_LOG
