Fix uninitialized chunk revision in sound wave data
#rnx

diff --git a/Engine/Source/Runtime/Engine/Classes/Sound/SoundWave.h b/Engine/Source/Runtime/Engine/Classes/Sound/SoundWave.h
--- a/Engine/Source/Runtime/Engine/Classes/Sound/SoundWave.h
+++ b/Engine/Source/Runtime/Engine/Classes/Sound/SoundWave.h
@@ -1669,7 +1669,7 @@ private:
 #endif //WITH_EDITORONLY_DATA
 	
 #if WITH_EDITOR
-	int32 CurrentChunkRevision;
+	int32 CurrentChunkRevision = 0;
 #endif // #if WITH_EDITOR
 
 	FName NameCached;
diff --git a/Engine/Source/Runtime/Engine/Private/SoundWave.cpp b/Engine/Source/Runtime/Engine/Private/SoundWave.cpp
--- a/Engine/Source/Runtime/Engine/Private/SoundWave.cpp
+++ b/Engine/Source/Runtime/Engine/Private/SoundWave.cpp
@@ -231,6 +231,7 @@ void FSoundWaveData::InitializeDataFromSoundWave(USoundWave& InWave)
 
 #if WITH_EDITOR
 	bLoadedFromCookedData = InWave.IsLoadedFromCookedData();
+	CurrentChunkRevision = InWave.CurrentChunkRevision;
 
 	// only necessary to set cue points here in editor 
 	// when the sample rate can change due to platform settings changing
