PR #12268: Fixing bug in FAudioDevice::CreateComponent outer object when no actor is passed in
#rnx

diff --git a/Engine/Source/Runtime/Engine/Private/AudioDevice.cpp b/Engine/Source/Runtime/Engine/Private/AudioDevice.cpp
--- a/Engine/Source/Runtime/Engine/Private/AudioDevice.cpp
+++ b/Engine/Source/Runtime/Engine/Private/AudioDevice.cpp
@@ -6334,7 +6334,7 @@ UAudioComponent* FAudioDevice::CreateComponent(USoundBase* Sound, const FCreateC
 				// Let engine pick the outer (transient package).
 				else
 				{
-					AudioComponent = NewObject<UAudioComponent>((Params.AudioComponentClass != nullptr) ? (UClass*)Params.AudioComponentClass : UAudioComponent::StaticClass());
+					AudioComponent = NewObject<UAudioComponent>(GetTransientPackage(), (Params.AudioComponentClass != nullptr) ? (UClass*)Params.AudioComponentClass : UAudioComponent::StaticClass());
 				}
 
 				check(AudioComponent);
