Fixed TSAN race condition around bForceRefPose flag

Flag is now copied into the evaluation context for use on the worker thread
Reworked some old deprecations around PerformAnimationProcessing/EvaluatePostProcessMeshInstance

[RN] Fixed TSAN race condition around bForceRefPose flag

diff --git a/Engine/Source/Runtime/Engine/Classes/Components/SkeletalMeshComponent.h b/Engine/Source/Runtime/Engine/Classes/Components/SkeletalMeshComponent.h
--- a/Engine/Source/Runtime/Engine/Classes/Components/SkeletalMeshComponent.h
+++ b/Engine/Source/Runtime/Engine/Classes/Components/SkeletalMeshComponent.h
@@ -106,6 +106,9 @@ struct FAnimationEvaluationContext
 	// duplicate the cached attributes
 	bool bDuplicateToCachedAttributes;
 
+	// Force reference pose
+	bool bForceRefPose;
+
 	// Curve data, swapped in from the component when we are running parallel eval
 	FBlendedHeapCurve	Curve;
 	FBlendedHeapCurve	CachedCurve;
@@ -140,6 +143,7 @@ struct FAnimationEvaluationContext
 		bDuplicateToCacheBones = Other.bDuplicateToCacheBones;
 		bDuplicateToCacheCurve = Other.bDuplicateToCacheCurve;
 		bDuplicateToCachedAttributes = Other.bDuplicateToCachedAttributes;
+		bForceRefPose = Other.bForceRefPose;
 
 		CustomAttributes.CopyFrom(Other.CustomAttributes);
 		CachedCustomAttributes.CopyFrom(Other.CachedCustomAttributes);
@@ -2029,18 +2033,19 @@ public:
 	UE_DEPRECATED(4.26, "Please use PerformAnimationEvaluation with different signature")
 	ENGINE_API void PerformAnimationEvaluation(const USkeletalMesh* InSkeletalMesh, UAnimInstance* InAnimInstance, TArray<FTransform>& OutSpaceBases, TArray<FTransform>& OutBoneSpaceTransforms, FVector& OutRootBoneTranslation, FBlendedHeapCurve& OutCurve);
 #endif
-	ENGINE_API void PerformAnimationProcessing(const USkeletalMesh* InSkeletalMesh, UAnimInstance* InAnimInstance, bool bInDoEvaluation, TArray<FTransform>& OutSpaceBases, TArray<FTransform>& OutBoneSpaceTransforms, FVector& OutRootBoneTranslation, FBlendedHeapCurve& OutCurve, UE::Anim::FMeshAttributeContainer& OutAttributes);
 
-	UE_DEPRECATED(4.26, "Please use PerformAnimationEvaluation with different signature")
-	ENGINE_API void PerformAnimationProcessing(const USkeletalMesh* InSkeletalMesh, UAnimInstance* InAnimInstance, bool bInDoEvaluation, TArray<FTransform>& OutSpaceBases, TArray<FTransform>& OutBoneSpaceTransforms, FVector& OutRootBoneTranslation, FBlendedHeapCurve& OutCurve);
+	ENGINE_API void PerformAnimationProcessing(const USkeletalMesh* InSkeletalMesh, UAnimInstance* InAnimInstance, bool bInDoEvaluation, bool bInForceRefPose, TArray<FTransform>& OutSpaceBases, TArray<FTransform>& OutBoneSpaceTransforms, FVector& OutRootBoneTranslation, FBlendedHeapCurve& OutCurve, UE::Anim::FMeshAttributeContainer& OutAttributes);
+
+	UE_DEPRECATED(5.5, "Please use PerformAnimationEvaluation with different signature")
+	ENGINE_API void PerformAnimationProcessing(const USkeletalMesh* InSkeletalMesh, UAnimInstance* InAnimInstance, bool bInDoEvaluation, TArray<FTransform>& OutSpaceBases, TArray<FTransform>& OutBoneSpaceTransforms, FVector& OutRootBoneTranslation, FBlendedHeapCurve& OutCurve, UE::Anim::FMeshAttributeContainer& OutAttributes);
 
 	/**
 	 * Evaluates the post process instance from the skeletal mesh this component is using.
 	 */
-	ENGINE_API void EvaluatePostProcessMeshInstance(TArray<FTransform>& OutBoneSpaceTransforms, FCompactPose& InOutPose, FBlendedHeapCurve& OutCurve, const USkeletalMesh* InSkeletalMesh, FVector& OutRootBoneTranslation, UE::Anim::FHeapAttributeContainer& OutAttributes) const;
+	ENGINE_API void EvaluatePostProcessMeshInstance(TArray<FTransform>& OutBoneSpaceTransforms, FCompactPose& InOutPose, FBlendedHeapCurve& OutCurve, const USkeletalMesh* InSkeletalMesh, FVector& OutRootBoneTranslation, UE::Anim::FHeapAttributeContainer& OutAttributes, bool bInForceRefPose) const;
 
-	UE_DEPRECATED(4.26, "Please use EvaluatePostProcessMeshInstance with different signature")
-	ENGINE_API void EvaluatePostProcessMeshInstance(TArray<FTransform>& OutBoneSpaceTransforms, FCompactPose& InOutPose, FBlendedHeapCurve& OutCurve, const USkeletalMesh* InSkeletalMesh, FVector& OutRootBoneTranslation) const;
+	UE_DEPRECATED(5.5, "Please use EvaluatePostProcessMeshInstance with different signature")
+	ENGINE_API void EvaluatePostProcessMeshInstance(TArray<FTransform>& OutBoneSpaceTransforms, FCompactPose& InOutPose, FBlendedHeapCurve& OutCurve, const USkeletalMesh* InSkeletalMesh, FVector& OutRootBoneTranslation, UE::Anim::FHeapAttributeContainer& OutAttributes) const;
 
 	ENGINE_API void PostAnimEvaluation(FAnimationEvaluationContext& EvaluationContext);
 
@@ -2435,7 +2440,7 @@ private:
 	ENGINE_API void EndPhysicsTickComponent(FSkeletalMeshComponentEndPhysicsTickFunction& ThisTickFunction);
 
 	/** Evaluate Anim System **/
-	ENGINE_API void EvaluateAnimation(const USkeletalMesh* InSkeletalMesh, UAnimInstance* InAnimInstance, FVector& OutRootBoneTranslation, FBlendedHeapCurve& OutCurve, FCompactPose& OutPose, UE::Anim::FHeapAttributeContainer& OutAttributes) const;
+	ENGINE_API void EvaluateAnimation(const USkeletalMesh* InSkeletalMesh, UAnimInstance* InAnimInstance, bool bInForceRefPose, FVector& OutRootBoneTranslation, FBlendedHeapCurve& OutCurve, FCompactPose& OutPose, UE::Anim::FHeapAttributeContainer& OutAttributes) const;
 
 	/** Queues up tasks for parallel update/evaluation, as well as the chained game thread completion task */
 	ENGINE_API void DispatchParallelEvaluationTasks(FActorComponentTickFunction* TickFunction);
diff --git a/Engine/Source/Runtime/Engine/Private/Components/SkeletalMeshComponent.cpp b/Engine/Source/Runtime/Engine/Private/Components/SkeletalMeshComponent.cpp
--- a/Engine/Source/Runtime/Engine/Private/Components/SkeletalMeshComponent.cpp
+++ b/Engine/Source/Runtime/Engine/Private/Components/SkeletalMeshComponent.cpp
@@ -2122,7 +2122,7 @@ bool USkeletalMeshComponent::AreRequiredCurvesUpToDate() const
 	return true;
 }
 
-void USkeletalMeshComponent::EvaluateAnimation(const USkeletalMesh* InSkeletalMesh, UAnimInstance* InAnimInstance, FVector& OutRootBoneTranslation, FBlendedHeapCurve& OutCurve, FCompactPose& OutPose, UE::Anim::FHeapAttributeContainer& OutAttributes) const
+void USkeletalMeshComponent::EvaluateAnimation(const USkeletalMesh* InSkeletalMesh, UAnimInstance* InAnimInstance, bool bInForceRefPose, FVector& OutRootBoneTranslation, FBlendedHeapCurve& OutCurve, FCompactPose& OutPose, UE::Anim::FHeapAttributeContainer& OutAttributes) const
 {
 	ANIM_MT_SCOPE_CYCLE_COUNTER(SkeletalComponentAnimEvaluate, !IsInGameThread());
 
@@ -2137,7 +2137,7 @@ void USkeletalMeshComponent::EvaluateAnimation(const USkeletalMesh* InSkeletalMe
 		InAnimInstance->ParallelCanEvaluate(InSkeletalMesh))
 	{
 		FParallelEvaluationData EvaluationData = { OutCurve, OutPose, OutAttributes };
-		InAnimInstance->ParallelEvaluateAnimation(bForceRefpose, InSkeletalMesh, EvaluationData);
+		InAnimInstance->ParallelEvaluateAnimation(bInForceRefPose, InSkeletalMesh, EvaluationData);
 	}
 }
 
@@ -2182,7 +2182,7 @@ void USkeletalMeshComponent::UpdateFollowerComponent()
 
 void USkeletalMeshComponent::PerformAnimationEvaluation(const USkeletalMesh* InSkeletalMesh, UAnimInstance* InAnimInstance, TArray<FTransform>& OutSpaceBases, TArray<FTransform>& OutBoneSpaceTransforms, FVector& OutRootBoneTranslation, FBlendedHeapCurve& OutCurve, UE::Anim::FMeshAttributeContainer& OutAttributes)
 {
-	PerformAnimationProcessing(InSkeletalMesh, InAnimInstance, true, OutSpaceBases, OutBoneSpaceTransforms, OutRootBoneTranslation, OutCurve, OutAttributes);
+	PerformAnimationProcessing(InSkeletalMesh, InAnimInstance, true, false, OutSpaceBases, OutBoneSpaceTransforms, OutRootBoneTranslation, OutCurve, OutAttributes);
 }
 
 void USkeletalMeshComponent::PerformAnimationEvaluation(const USkeletalMesh* InSkeletalMesh, UAnimInstance* InAnimInstance, TArray<FTransform>& OutSpaceBases, TArray<FTransform>& OutBoneSpaceTransforms, FVector& OutRootBoneTranslation, FBlendedHeapCurve& OutCurve)
@@ -2193,7 +2193,7 @@ void USkeletalMeshComponent::PerformAnimationEvaluation(const USkeletalMesh* InS
 
 #endif
 
-void USkeletalMeshComponent::PerformAnimationProcessing(const USkeletalMesh* InSkeletalMesh, UAnimInstance* InAnimInstance, bool bInDoEvaluation, TArray<FTransform>& OutSpaceBases, TArray<FTransform>& OutBoneSpaceTransforms, FVector& OutRootBoneTranslation, FBlendedHeapCurve& OutCurve, UE::Anim::FMeshAttributeContainer& OutAttributes)
+void USkeletalMeshComponent::PerformAnimationProcessing(const USkeletalMesh* InSkeletalMesh, UAnimInstance* InAnimInstance, bool bInDoEvaluation, bool bInForceRefPose, TArray<FTransform>& OutSpaceBases, TArray<FTransform>& OutBoneSpaceTransforms, FVector& OutRootBoneTranslation, FBlendedHeapCurve& OutCurve, UE::Anim::FMeshAttributeContainer& OutAttributes)
 {
 	CSV_SCOPED_TIMING_STAT(Animation, WorkerThreadTickTime);
 	ANIM_MT_SCOPE_CYCLE_COUNTER(PerformAnimEvaluation, !IsInGameThread());
@@ -2225,8 +2225,8 @@ void USkeletalMeshComponent::PerformAnimationProcessing(const USkeletalMesh* InS
 		UE::Anim::FHeapAttributeContainer Attributes;		
 
 		// evaluate pure animations, and fill up BoneSpaceTransforms
-		EvaluateAnimation(InSkeletalMesh, InAnimInstance, OutRootBoneTranslation, OutCurve, EvaluatedPose, Attributes);
-		EvaluatePostProcessMeshInstance(OutBoneSpaceTransforms, EvaluatedPose, OutCurve, InSkeletalMesh, OutRootBoneTranslation, Attributes);
+		EvaluateAnimation(InSkeletalMesh, InAnimInstance, bInForceRefPose, OutRootBoneTranslation, OutCurve, EvaluatedPose, Attributes);
+		EvaluatePostProcessMeshInstance(OutBoneSpaceTransforms, EvaluatedPose, OutCurve, InSkeletalMesh, OutRootBoneTranslation, Attributes, bInForceRefPose);
 
 		// Finalize the transforms from the evaluation
 		FinalizePoseEvaluationResult(InSkeletalMesh, OutBoneSpaceTransforms, OutRootBoneTranslation, EvaluatedPose);
@@ -2242,19 +2242,17 @@ void USkeletalMeshComponent::PerformAnimationProcessing(const USkeletalMesh* InS
 }
 
 
-void USkeletalMeshComponent::PerformAnimationProcessing(const USkeletalMesh* InSkeletalMesh, UAnimInstance* InAnimInstance, bool bInDoEvaluation, TArray<FTransform>& OutSpaceBases, TArray<FTransform>& OutBoneSpaceTransforms, FVector& OutRootBoneTranslation, FBlendedHeapCurve& OutCurve)
+void USkeletalMeshComponent::PerformAnimationProcessing(const USkeletalMesh* InSkeletalMesh, UAnimInstance* InAnimInstance, bool bInDoEvaluation, TArray<FTransform>& OutSpaceBases, TArray<FTransform>& OutBoneSpaceTransforms, FVector& OutRootBoneTranslation, FBlendedHeapCurve& OutCurve, UE::Anim::FMeshAttributeContainer& OutAttributes)
 {
-	UE::Anim::FMeshAttributeContainer Attributes;	
-	PerformAnimationProcessing(InSkeletalMesh, InAnimInstance, bInDoEvaluation, OutSpaceBases, OutBoneSpaceTransforms, OutRootBoneTranslation, OutCurve, Attributes);
+	PerformAnimationProcessing(InSkeletalMesh, InAnimInstance, bInDoEvaluation, false, OutSpaceBases, OutBoneSpaceTransforms, OutRootBoneTranslation, OutCurve, OutAttributes);
 }
 
-void USkeletalMeshComponent::EvaluatePostProcessMeshInstance(TArray<FTransform>& OutBoneSpaceTransforms, FCompactPose& InOutPose, FBlendedHeapCurve& OutCurve, const USkeletalMesh* InSkeletalMesh, FVector& OutRootBoneTranslation) const
+void USkeletalMeshComponent::EvaluatePostProcessMeshInstance(TArray<FTransform>& OutBoneSpaceTransforms, FCompactPose& InOutPose, FBlendedHeapCurve& OutCurve, const USkeletalMesh* InSkeletalMesh, FVector& OutRootBoneTranslation, UE::Anim::FHeapAttributeContainer& OutAttributes) const
 {
-	UE::Anim::FHeapAttributeContainer Attributes;
-	EvaluatePostProcessMeshInstance(OutBoneSpaceTransforms, InOutPose, OutCurve, InSkeletalMesh, OutRootBoneTranslation, Attributes);
+	EvaluatePostProcessMeshInstance(OutBoneSpaceTransforms, InOutPose, OutCurve, InSkeletalMesh, OutRootBoneTranslation, OutAttributes, false);
 }
 
-void USkeletalMeshComponent::EvaluatePostProcessMeshInstance(TArray<FTransform>& OutBoneSpaceTransforms, FCompactPose& InOutPose, FBlendedHeapCurve& OutCurve, const USkeletalMesh* InSkeletalMesh, FVector& OutRootBoneTranslation, UE::Anim::FHeapAttributeContainer& OutAttributes) const
+void USkeletalMeshComponent::EvaluatePostProcessMeshInstance(TArray<FTransform>& OutBoneSpaceTransforms, FCompactPose& InOutPose, FBlendedHeapCurve& OutCurve, const USkeletalMesh* InSkeletalMesh, FVector& OutRootBoneTranslation, UE::Anim::FHeapAttributeContainer& OutAttributes, bool bInForceRefPose) const
 {
 	if (ShouldEvaluatePostProcessInstance())
 	{
@@ -2277,7 +2275,7 @@ void USkeletalMeshComponent::EvaluatePostProcessMeshInstance(TArray<FTransform>&
 			}
 		}
 
-		EvaluateAnimation(InSkeletalMesh, PostProcessAnimInstance, OutRootBoneTranslation, OutCurve, InOutPose, OutAttributes);
+		EvaluateAnimation(InSkeletalMesh, PostProcessAnimInstance, bInForceRefPose, OutRootBoneTranslation, OutCurve, InOutPose, OutAttributes);
 	}
 }
 
@@ -2540,6 +2538,7 @@ void USkeletalMeshComponent::RefreshBoneTransforms(FActorComponentTickFunction*
 	AnimEvaluationContext.bDuplicateToCacheCurve = bInvalidCachedCurve || (bDoEvaluationRateOptimization && AnimEvaluationContext.bDoEvaluation && !AnimEvaluationContext.bDoInterpolation);
 
 	AnimEvaluationContext.bDuplicateToCachedAttributes = bInvalidCachedAttributes || (bDoEvaluationRateOptimization && AnimEvaluationContext.bDoEvaluation && !AnimEvaluationContext.bDoInterpolation);
+	AnimEvaluationContext.bForceRefPose = bForceRefpose;
 
 	if (!bDoEvaluationRateOptimization)
 	{
@@ -4225,11 +4224,11 @@ void USkeletalMeshComponent::ParallelAnimationEvaluation()
 {
 	if (AnimEvaluationContext.bDoInterpolation)
 	{
-		PerformAnimationProcessing(AnimEvaluationContext.SkeletalMesh, AnimEvaluationContext.AnimInstance, AnimEvaluationContext.bDoEvaluation, AnimEvaluationContext.CachedComponentSpaceTransforms, AnimEvaluationContext.CachedBoneSpaceTransforms, AnimEvaluationContext.RootBoneTranslation, AnimEvaluationContext.CachedCurve, AnimEvaluationContext.CachedCustomAttributes);
+		PerformAnimationProcessing(AnimEvaluationContext.SkeletalMesh, AnimEvaluationContext.AnimInstance, AnimEvaluationContext.bDoEvaluation, AnimEvaluationContext.bForceRefPose, AnimEvaluationContext.CachedComponentSpaceTransforms, AnimEvaluationContext.CachedBoneSpaceTransforms, AnimEvaluationContext.RootBoneTranslation, AnimEvaluationContext.CachedCurve, AnimEvaluationContext.CachedCustomAttributes);
 	}
 	else
 	{
-		PerformAnimationProcessing(AnimEvaluationContext.SkeletalMesh, AnimEvaluationContext.AnimInstance, AnimEvaluationContext.bDoEvaluation, AnimEvaluationContext.ComponentSpaceTransforms, AnimEvaluationContext.BoneSpaceTransforms, AnimEvaluationContext.RootBoneTranslation, AnimEvaluationContext.Curve, AnimEvaluationContext.CustomAttributes);
+		PerformAnimationProcessing(AnimEvaluationContext.SkeletalMesh, AnimEvaluationContext.AnimInstance, AnimEvaluationContext.bDoEvaluation, AnimEvaluationContext.bForceRefPose, AnimEvaluationContext.ComponentSpaceTransforms, AnimEvaluationContext.BoneSpaceTransforms, AnimEvaluationContext.RootBoneTranslation, AnimEvaluationContext.Curve, AnimEvaluationContext.CustomAttributes);
 	}
 
 	ParallelDuplicateAndInterpolate(AnimEvaluationContext);
