[Shaders] Allow to dump all shader compilation timing statistics to CSV files.

This adds CVar "r.ShaderCompiler.DumpShaderTimeStats" which dumps all timing statistics from the shader compilation to CSV files,
one sorted by the average time and another one sorted by the total time.

#rnx

diff --git a/Engine/Source/Runtime/Engine/Private/ShaderCompiler/ShaderCompilerStats.cpp b/Engine/Source/Runtime/Engine/Private/ShaderCompiler/ShaderCompilerStats.cpp
--- a/Engine/Source/Runtime/Engine/Private/ShaderCompiler/ShaderCompilerStats.cpp
+++ b/Engine/Source/Runtime/Engine/Private/ShaderCompiler/ShaderCompilerStats.cpp
@@ -18,6 +18,7 @@
 #include "ProfilingDebugging/LoadTimeTracker.h"
 #include "ProfilingDebugging/StallDetector.h"
 #include "Serialization/CompactBinaryWriter.h"
+#include "Misc/FileHelper.h"
 
 static int32 GLogShaderCompilerStats = 0;
 static FAutoConsoleVariableRef CVarLogShaderCompilerStats(
@@ -25,13 +26,19 @@ static FAutoConsoleVariableRef CVarLogShaderCompilerStats(
 	GLogShaderCompilerStats,
 	TEXT("When set to 1, Log detailed shader compiler stats."));
 
-int32 GMaxShaderStatsToLog = 5;
+static int32 GMaxShaderStatsToLog = 5;
 static FAutoConsoleVariableRef CVarShaderCompilerMaxShaderStatsToPrint(
 	TEXT("r.ShaderCompiler.MaxShaderStatsToPrint"),
 	GMaxShaderStatsToLog,
 	TEXT("Max number of shaders FShaderJobCache stats logs out"),
 	ECVF_Default);
 
+static bool GDumpShaderTimeStats = false;
+static FAutoConsoleVariableRef CVarDumpShaderTimeStats(
+	TEXT("r.ShaderCompiler.DumpShaderTimeStats"),
+	GDumpShaderTimeStats,
+	TEXT("When set to true, dump shader compiler timing statistics to a CSV file."));
+
 
 FShaderCompilerStats* GShaderCompilerStats = nullptr;
 
@@ -203,6 +210,34 @@ static FString PrintJobsCompletedPercentageToString(int64 JobsAssigned, int64 Jo
 	return FString::Printf(TEXT("%.2f%%"), FMath::Min(JobsCompletedPercentage, 99.99));
 }
 
+static void DumpShaderTimingsToCSVFile(const TCHAR* Filename, const TMap<FString, FShaderTimings>& ShaderTimings)
+{
+	// Ensure output folder exists
+	const FString& OutputDirectory = GShaderCompilingManager->GetAbsoluteShaderDebugInfoDirectory();
+	if (!IFileManager::Get().DirectoryExists(*OutputDirectory))
+	{
+		IFileManager::Get().MakeDirectory(*OutputDirectory, true);
+	}
+
+	// Write CSV table
+	FString Table;
+
+	Table += FString::Printf(TEXT("SHADER,NUMBER OF COMPILATIONS,TOTAL COMPILE TIME,TOTAL PREPROCESS TIME,AVERAGE TIME,MAX TIME,MIN TIME%s"), LINE_TERMINATOR);
+
+	for (TMap<FString, FShaderTimings>::TConstIterator Iter(ShaderTimings); Iter; ++Iter)
+	{
+		const FShaderTimings& Timings = Iter.Value();
+		Table += FString::Printf(
+			TEXT("%s,%d,%4.2fs,%4.2fs,%4.2fs,%4.2fs,%4.2fs%s"),
+			*Iter.Key(), Timings.NumCompiled, Timings.TotalCompileTime, Timings.TotalPreprocessTime, Timings.AverageCompileTime, Timings.MaxCompileTime, Timings.MinCompileTime, LINE_TERMINATOR
+		);
+	}
+
+	// Dump CSV table to file
+	const FString OutputFilename = FPaths::Combine(OutputDirectory, Filename);
+	FFileHelper::SaveStringToFile(Table, *OutputFilename);
+}
+
 void FShaderCompilerStats::WriteStatSummary()
 {
 	const uint32 TotalCompiled = GetTotalShadersCompiled();
@@ -372,6 +407,11 @@ void FShaderCompilerStats::WriteStatSummary()
 			}
 		}
 
+		if (GDumpShaderTimeStats)
+		{
+			DumpShaderTimingsToCSVFile(TEXT("ShaderTimings.SortedByAverageTime.csv"), ShaderTimings);
+		}
+
 		// sort by total time
 		ShaderTimings.ValueSort([](const FShaderTimings& A, const FShaderTimings& B) { return A.TotalCompileTime > B.TotalCompileTime; });
 
@@ -389,6 +429,11 @@ void FShaderCompilerStats::WriteStatSummary()
 				break;
 			}
 		}
+
+		if (GDumpShaderTimeStats)
+		{
+			DumpShaderTimingsToCSVFile(TEXT("ShaderTimings.SortedByTotalTime.csv"), ShaderTimings);
+		}
 	}
 
 	MaterialCounters.WriteStatSummary(AggregatedSuffix);
diff --git a/Engine/Source/Runtime/Engine/Public/ShaderCompiler.h b/Engine/Source/Runtime/Engine/Public/ShaderCompiler.h
--- a/Engine/Source/Runtime/Engine/Public/ShaderCompiler.h
+++ b/Engine/Source/Runtime/Engine/Public/ShaderCompiler.h
@@ -602,7 +602,8 @@ struct FShaderCompilerMaterialCounters
 	void GatherAnalytics(TArray<FAnalyticsEventAttribute>& Attributes);
 };
 
-/** Structure used to describe compiling time of a shader type (for all the instances of it that we have seen). */
+/** Structure used to describe compiling time of a shader type (for all the instances of it that we have seen).
+	Can be dumped to CSV file via 'r.ShaderCompiler.DumpShaderTimeStats' CVar. */
 USTRUCT()
 struct FShaderTimings
 {
