Add support for color and monospace to visual logger text drawn through the rewind debugger

diff --git a/Engine/Plugins/Animation/GameplayInsights/Source/RewindDebuggerVLog/Private/RewindDebuggerVLog.cpp b/Engine/Plugins/Animation/GameplayInsights/Source/RewindDebuggerVLog/Private/RewindDebuggerVLog.cpp
--- a/Engine/Plugins/Animation/GameplayInsights/Source/RewindDebuggerVLog/Private/RewindDebuggerVLog.cpp
+++ b/Engine/Plugins/Animation/GameplayInsights/Source/RewindDebuggerVLog/Private/RewindDebuggerVLog.cpp
@@ -29,7 +29,7 @@ FRewindDebuggerVLog::FRewindDebuggerVLog()
 
 void FRewindDebuggerVLog::OnShowDebugInfo(UCanvas* Canvas, APlayerController* Player)
 {
-	ScreenTextY = 100;
+	ScreenTextY = 60;
 	if (IRewindDebugger* RewindDebugger = IRewindDebugger::Instance())
 	{
 		if (RewindDebugger->IsPIESimulating())
@@ -145,7 +145,7 @@ void FRewindDebuggerVLog::RenderLogEntry(const FVisualLogEntry& Entry, UCanvas*
 	else
 	{
 		UWorld* World = IRewindDebugger::Instance()->GetWorldToVisualize();
-		FVisualLogEntryRenderer::RenderLogEntry(World,Entry, &MatchCategoryFilters, Canvas, MonospaceFont, ScreenTextY);
+		FVisualLogEntryRenderer::RenderLogEntry(World,Entry, &MatchCategoryFilters, Canvas, GEngine->GetMediumFont(), MonospaceFont, ScreenTextY);
 
 	}
 }
diff --git a/Engine/Plugins/Animation/GameplayInsights/Source/RewindDebuggerVLog/Private/VisualLogEntryRenderer.cpp b/Engine/Plugins/Animation/GameplayInsights/Source/RewindDebuggerVLog/Private/VisualLogEntryRenderer.cpp
--- a/Engine/Plugins/Animation/GameplayInsights/Source/RewindDebuggerVLog/Private/VisualLogEntryRenderer.cpp
+++ b/Engine/Plugins/Animation/GameplayInsights/Source/RewindDebuggerVLog/Private/VisualLogEntryRenderer.cpp
@@ -115,13 +115,13 @@ namespace
 
 
 void FVisualLogEntryRenderer::RenderLogEntry(class UWorld* World, const FVisualLogEntry& Entry, TFunctionRef<bool (const FName&, ELogVerbosity::Type)> MatchCategoryFilters,
-	UCanvas* Canvas, UFont* Font, int32& ScreenTextY)
+	UCanvas* Canvas, UFont* Font, UFont* MonospaceFont, int32& ScreenTextY)
 {
 	for(const FVisualLogLine& LogLine : Entry.LogLines)
 	{
 		if (MatchCategoryFilters(LogLine.Category, LogLine.Verbosity))
 		{
-			Canvas->Canvas->DrawShadowedString(100, ScreenTextY+=20, LogLine.Line, Font, FColor::White);
+			Canvas->Canvas->DrawShadowedString(20, ScreenTextY+=16, LogLine.Line, (LogLine.bMonospace && MonospaceFont) ? MonospaceFont : Font, LogLine.Color);
 		}
 	}
 	
diff --git a/Engine/Plugins/Animation/GameplayInsights/Source/RewindDebuggerVLog/Private/VisualLogEntryRenderer.h b/Engine/Plugins/Animation/GameplayInsights/Source/RewindDebuggerVLog/Private/VisualLogEntryRenderer.h
--- a/Engine/Plugins/Animation/GameplayInsights/Source/RewindDebuggerVLog/Private/VisualLogEntryRenderer.h
+++ b/Engine/Plugins/Animation/GameplayInsights/Source/RewindDebuggerVLog/Private/VisualLogEntryRenderer.h
@@ -7,5 +7,5 @@ class FVisualLogEntryRenderer
 {
 public:
 	static void RenderLogEntry(class UWorld* World, const struct FVisualLogEntry& Entry, TFunctionRef<bool(const FName&, ELogVerbosity::Type)> MatchCategoryFilters,
-		class UCanvas* Canvas, class UFont* Font, int32& ScreenTextY);
+		class UCanvas* Canvas, class UFont* Font, class UFont* MonospaceFont, int32& ScreenTextY);
 };
diff --git a/Engine/Source/Runtime/Engine/Private/VisualLogger/VisualLoggerTypes.cpp b/Engine/Source/Runtime/Engine/Private/VisualLogger/VisualLoggerTypes.cpp
--- a/Engine/Source/Runtime/Engine/Private/VisualLogger/VisualLoggerTypes.cpp
+++ b/Engine/Source/Runtime/Engine/Private/VisualLogger/VisualLoggerTypes.cpp
@@ -572,6 +572,8 @@ FArchive& operator<<(FArchive& Ar, FVisualLogLine& LogLine)
 	Ar << LogLine.Verbosity;
 	Ar << LogLine.UniqueId;
 	Ar << LogLine.UserData;
+	Ar << LogLine.Color;
+	Ar << LogLine.bMonospace;
 	Ar << LogLine.Line;
 	return Ar;
 }
diff --git a/Engine/Source/Runtime/Engine/Public/VisualLogger/VisualLoggerTypes.h b/Engine/Source/Runtime/Engine/Public/VisualLogger/VisualLoggerTypes.h
--- a/Engine/Source/Runtime/Engine/Public/VisualLogger/VisualLoggerTypes.h
+++ b/Engine/Source/Runtime/Engine/Public/VisualLogger/VisualLoggerTypes.h
@@ -102,10 +102,12 @@ struct FVisualLogLine
 	int32 UniqueId;
 	int64 UserData;
 	FName TagName;
-
+	FColor Color;
+	bool bMonospace;
 	FVisualLogLine() { /* Empty */ }
 	ENGINE_API FVisualLogLine(const FName& InCategory, ELogVerbosity::Type InVerbosity, const FString& InLine);
 	ENGINE_API FVisualLogLine(const FName& InCategory, ELogVerbosity::Type InVerbosity, const FString& InLine, int64 InUserData);
+	ENGINE_API FVisualLogLine(const FName& InCategory, ELogVerbosity::Type InVerbosity, const FString& InLine, const FColor& InColor, bool bInMonospace);
 };
 
 struct FVisualLogStatusCategory
@@ -408,6 +410,9 @@ FVisualLogLine::FVisualLogLine(const FName& InCategory, ELogVerbosity::Type InVe
 , Category(InCategory)
 , Verbosity(InVerbosity)
 , UserData(0)
+, Color(FColor::White)
+, bMonospace(false)
+
 {
 
 }
@@ -418,6 +423,20 @@ FVisualLogLine::FVisualLogLine(const FName& InCategory, ELogVerbosity::Type InVe
 , Category(InCategory)
 , Verbosity(InVerbosity)
 , UserData(InUserData)
+, Color(FColor::White)
+, bMonospace(false)
+{
+
+}
+
+inline
+FVisualLogLine::FVisualLogLine(const FName& InCategory, ELogVerbosity::Type InVerbosity, const FString& InLine, const FColor& InColor, bool bInMonospace)
+	: Line(InLine)
+	, Category(InCategory)
+	, Verbosity(InVerbosity)
+	, UserData(0)
+	, Color(InColor)
+	, bMonospace(bInMonospace)
 {
 
 }
