[Engine]
* Changed so CacheShadowDepthsFromPrimitivesUsingWPO can be called on render threads

diff --git a/Engine/Source/Runtime/Engine/Private/PrimitiveSceneProxy.cpp b/Engine/Source/Runtime/Engine/Private/PrimitiveSceneProxy.cpp
--- a/Engine/Source/Runtime/Engine/Private/PrimitiveSceneProxy.cpp
+++ b/Engine/Source/Runtime/Engine/Private/PrimitiveSceneProxy.cpp
@@ -128,7 +128,7 @@ bool ShouldOptimizedWPOAffectNonNaniteShaderSelection()
 
 bool CacheShadowDepthsFromPrimitivesUsingWPO()
 {
-	return CVarCacheWPOPrimitives.GetValueOnAnyThread(true) != 0;
+	return CVarCacheWPOPrimitives.GetValueOnAnyThread() != 0;
 }
 
 bool SupportsCachingMeshDrawCommands(const FMeshBatch& MeshBatch)
