Expose texture source ID to blueprint for duplicate hunting.

#rnx

diff --git a/Engine/Source/Runtime/Engine/Classes/Engine/Texture.h b/Engine/Source/Runtime/Engine/Classes/Engine/Texture.h
--- a/Engine/Source/Runtime/Engine/Classes/Engine/Texture.h
+++ b/Engine/Source/Runtime/Engine/Classes/Engine/Texture.h
@@ -1596,6 +1596,13 @@ PRAGMA_ENABLE_DEPRECATION_WARNINGS
 	UFUNCTION(BlueprintCallable, BlueprintPure=false, Category = "Rendering|Texture")
 	bool ComputeTextureSourceChannelMinMax(FLinearColor & OutColorMin, FLinearColor & OutColorMax) const;
 
+	/**
+	 * Return the ID for the texture source.
+	 * If the source isn't valid or editor data isn't available, returns false.
+	 */
+	UFUNCTION(BlueprintCallable, meta = (DisplayName = "GetTextureSourceIdString"), Category = "Rendering|Texture")
+	bool Blueprint_GetTextureSourceIdString(FString& OutTextureSourceId);
+
 private:
 	/** Whether the async resource release process has already been kicked off or not */
 	UPROPERTY(transient)
diff --git a/Engine/Source/Runtime/Engine/Private/Texture.cpp b/Engine/Source/Runtime/Engine/Private/Texture.cpp
--- a/Engine/Source/Runtime/Engine/Private/Texture.cpp
+++ b/Engine/Source/Runtime/Engine/Private/Texture.cpp
@@ -4520,6 +4520,22 @@ void UTexture::Blueprint_GetTextureSourceDiskAndMemorySize(int64 & OutDiskSize,i
 #endif
 }
 
+bool UTexture::Blueprint_GetTextureSourceIdString(FString& OutTextureSourceId)
+{
+	OutTextureSourceId.Reset();
+
+#if WITH_EDITORONLY_DATA
+	if (!Source.IsValid())
+	{
+		return false;
+	}
+	OutTextureSourceId = Source.GetIdString();
+	return true;
+#else
+	return false;
+#endif
+}
+
 bool UTexture::ComputeTextureSourceChannelMinMax(FLinearColor & OutColorMin, FLinearColor & OutColorMax) const
 {
 	// make sure we fill the outputs if we return failure :
