UE-224742 - Maps that use redirectors for sublevels can sometimes hang the editor.

(Partial fix)

In rare cases, moving maps with streaming levels could end up serializing PackageNameToLoad. This could happen if a map was moved and the asset registry hadn't updated its soft path redirect table before ULevelStreaming::Serialize was called. As a result, PackageNameToLoad would be saved to disk with the updated path. However, if the map were moved again, there would be a discrepancy between PackageNameToLoad and WorldAsset. This would eventually create a corrupt UPackage where the name and file path don't match. Eventually, the editor would hang as it failed to load the streaming level and repeatedly re-attempt.

For now, we'll just clear out PackageNameToLoad during these scenarios. At some point, we'll need to phase out PackageNameToLoad completely in favor of WorldAsset. Having two fields attempt represent a package path is a bit clumsy and error prone.

diff --git a/Engine/Source/Runtime/Engine/Private/LevelStreaming.cpp b/Engine/Source/Runtime/Engine/Private/LevelStreaming.cpp
--- a/Engine/Source/Runtime/Engine/Private/LevelStreaming.cpp
+++ b/Engine/Source/Runtime/Engine/Private/LevelStreaming.cpp
@@ -513,10 +513,42 @@ void ULevelStreaming::Serialize( FArchive& Ar )
 	
 	if (Ar.IsLoading())
 	{
-		if (GetOutermost()->HasAnyPackageFlags(PKG_PlayInEditor) && GetOutermost()->GetPIEInstanceID() != INDEX_NONE)
+		const bool bIsRunningPIE =
+			GetOutermost()->HasAnyPackageFlags(PKG_PlayInEditor) && 
+			GetOutermost()->GetPIEInstanceID() != INDEX_NONE
+		;
+
+		if (bIsRunningPIE)
 		{
 			RenameForPIE(GetOutermost()->GetPIEInstanceID());
 		}
+
+#if WITH_EDITOR
+		// If PackageNameToLoad doesn't match WorldAsset, we can potentially create a corrupt UPackage where the name and file path don't match.
+		// In this case, we have no option but to clear out PackageNameToLoad, since WorldAsset is the source of truth.
+		// Note that RenameForPIE intentionally creates a mismatch, but that's acceptable since PIE effectively creates transient duplicates.
+		if (Ar.IsPersistent() && !bIsRunningPIE)
+		{
+			const bool bHasMismatchedPackageName =
+				(PackageNameToLoad != NAME_None) &&
+				(PackageNameToLoad != WorldAsset.GetLongPackageFName())
+			;
+
+			if (bHasMismatchedPackageName)
+			{
+				const UPackage* Package = GetPackage();
+				ensure(Package);
+				const FString& PackageName = Package ? Package->GetName() : TEXT("Unknown package");
+
+				UE_LOG(LogLevelStreaming, Warning, TEXT("WorldAsset (%s) and PackageNameToLoad (%s) point to different streaming levels. PackageNameToLoad will be set to none. Resave this map (%s) to remove warnings."),
+					*WorldAsset.GetLongPackageName(),
+					*PackageNameToLoad.ToString(),
+					*PackageName);
+
+				PackageNameToLoad = NAME_None;
+			}
+		}
+#endif
 	}
 }
 
