Account for possible invalid section when montage is force moved

diff --git a/Engine/Source/Runtime/Engine/Private/Animation/AnimMontage.cpp b/Engine/Source/Runtime/Engine/Private/Animation/AnimMontage.cpp
--- a/Engine/Source/Runtime/Engine/Private/Animation/AnimMontage.cpp
+++ b/Engine/Source/Runtime/Engine/Private/Animation/AnimMontage.cpp
@@ -2100,6 +2100,13 @@ EMontageSubStepResult FMontageSubStepper::Advance(float& InOut_P_Original, const
 			// We are modifying the current position so we also need to update the section and pos in section
 			InOut_P_Original = MontageInstance->ForcedNextFromPosition.GetValue();
 			CurrentSectionIndex = Montage->GetAnimCompositeSectionIndexFromPos(InOut_P_Original, PositionInSection);
+			
+			if (!Montage->IsValidSectionIndex(CurrentSectionIndex))
+			{
+				return EMontageSubStepResult::InvalidSection;
+			}
+			CurrentSectionStartTime = Montage->GetAnimCompositeSection(CurrentSectionIndex).GetTime();
+			CurrentSectionLength = Montage->GetSectionLength(CurrentSectionIndex);
 		}
 		DeltaMove = NewPosition - InOut_P_Original;
 		PlayRate = DeltaMove / TimeRemaining;
