Optimized actor range iterations in the editor : in the editor case, we are gathering all actors from the (many) worlds and dynamic cast them into the proper type during iteration. This can be very wasteful when the actor type's count is comparatively small to the total actor type (i.e. in most cases). We now pre-filter the list of actors based on their presence in the list of objects of their class (which was already retrieved), which is cache-friendlier than dynamic casting.

#tests Editor, cooked game

diff --git a/Engine/Source/Runtime/Engine/Public/EngineUtils.h b/Engine/Source/Runtime/Engine/Public/EngineUtils.h
--- a/Engine/Source/Runtime/Engine/Public/EngineUtils.h
+++ b/Engine/Source/Runtime/Engine/Public/EngineUtils.h
@@ -153,13 +153,13 @@ public:
 	/** Current world we are iterating upon						*/
 	const UWorld* CurrentWorld;
 	/** Results from the GetObjectsOfClass query				*/
-	TArray<UObject*> ObjectArray;
+	TArray<AActor*> ActorArray;
 	/** index of the current element in the object array		*/
 	int32 Index;
 	/** Whether we already reached the end						*/
-	bool	ReachedEnd;
+	bool ReachedEnd;
 	/** Number of actors that have been considered thus far		*/
-	int32		ConsideredCount;
+	int32 ConsideredCount;
 	/** Current actor pointed to by actor iterator				*/
 	AActor*	CurrentActor;
 	/** Contains any actors spawned during iteration			*/
@@ -184,12 +184,11 @@ public:
 		check(CurrentWorld);
 
 #if WITH_EDITOR
+		TArray<AActor*> LocalActorsOfClass;
 		if (InClass != AActor::StaticClass())
 		{
-			TArray<UObject*> ObjectsOfClass;
-			GetObjectsOfClass(InClass, ObjectsOfClass, true, RF_ClassDefaultObject, EInternalObjectFlags::Garbage);
-
-			if (ObjectsOfClass.IsEmpty())
+			GetObjectsOfClass(InClass, reinterpret_cast<TArray<UObject*>&>(LocalActorsOfClass), true, RF_ClassDefaultObject, EInternalObjectFlags::Garbage);
+			if (LocalActorsOfClass.IsEmpty())
 			{
 				return;
 			}
@@ -210,13 +209,14 @@ public:
 				MaxActors += Level->Actors.Num();
 			}
 		}
-		ObjectArray.Reserve(MaxActors);
+		TArray<AActor*> LocalWorldActors;
+		LocalWorldActors.Reserve(MaxActors);
 
 		// Add persistent level actors, we need to do this explicitly to handle worlds that are not initialized yet, as
 		// InWorld->GetLevels() would return an empty level list even if the world's persistent level contains actors.
 		if (InWorld->PersistentLevel)
 		{
-			ObjectArray.Append(InWorld->PersistentLevel->Actors);
+			LocalWorldActors.Append(InWorld->PersistentLevel->Actors);
 		}
 
 		// Add streaming levels actors
@@ -224,12 +224,37 @@ public:
 		{
 			if (Level && Level != InWorld->PersistentLevel)
 			{
-				ObjectArray.Append(Level->Actors);
+				LocalWorldActors.Append(Level->Actors);
+			}
+		}
+
+		if (InClass != AActor::StaticClass())
+		{
+			// Intersect the 2 arrays now to only retain the actors of the proper type. This is faster this way as we only rely on the pointer values to filter the actors (since LocalActorsOfClass 
+			//  is already a superset of all possible actors of class) and it avoids iterating on useless actors in the iteration loop (along with expensive RTTI calls to IsA()) 
+			
+			// LocalActorsOfClass is the superset of all actors of the proper type so there will never be more than this amount : 
+			ActorArray.Reserve(LocalActorsOfClass.Num());
+			// Pick the array that has the least element and turn it into a TSet for fast lookup
+			const bool bLookUpInActorsOfClass = LocalActorsOfClass.Num() < LocalWorldActors.Num();
+			TSet<AActor*> LookUpSet(bLookUpInActorsOfClass ? LocalActorsOfClass : LocalWorldActors);
+			TArray<AActor*>& ArrayToIterate = bLookUpInActorsOfClass ? LocalWorldActors : LocalActorsOfClass;
+			for (AActor* Actor : ArrayToIterate)
+			{
+				if (LookUpSet.Contains(Actor))
+				{
+					ActorArray.Add(Actor);
+				}
 			}
 		}
+		else
+		{
+			// If we're interested in all AActors, just use the World actors we have retrieved
+			Swap(LocalWorldActors, ActorArray);
+		}
 #else // WITH_EDITOR
 		constexpr EObjectFlags ExcludeFlags = RF_ClassDefaultObject;
-		GetObjectsOfClass(InClass, ObjectArray, true, ExcludeFlags, EInternalObjectFlags::Garbage);
+		GetObjectsOfClass(InClass, reinterpret_cast<TArray<UObject*>&>(ActorArray), true, ExcludeFlags, EInternalObjectFlags::Garbage);
 #endif
 
 		const auto ActorSpawnedDelegate = FOnActorSpawned::FDelegate::CreateRaw(this, &FActorIteratorState::OnActorSpawned);
@@ -258,7 +283,13 @@ private:
 	{
 		if (InActor->IsA(DesiredClass))
 		{
-			SpawnedActorArray.AddUnique(InActor);
+#if WITH_EDITOR
+			// In the editor, the actor list is pre-filtered by world so perform the test now instead of in the loop like at runtime :
+			if (ULevel* ActorLevel = InActor->GetLevel(); ActorLevel && (ActorLevel->GetWorld() == CurrentWorld))
+#endif // WITH_EDITOR
+			{
+				SpawnedActorArray.AddUnique(InActor);
+			}
 		}
 	}
 };
@@ -292,20 +323,22 @@ public:
 	void operator++()
 	{
 		// Use local version to avoid LHSs as compiler is not required to write out member variables to memory.
-		AActor*           LocalCurrentActor      = nullptr;
-		int32             LocalIndex             = State->Index;
-		TArray<UObject*>& LocalObjectArray       = State->ObjectArray;
-		TArray<AActor*>&  LocalSpawnedActorArray = State->SpawnedActorArray;
-		const UWorld*     LocalCurrentWorld      = State->CurrentWorld;
-		while(++LocalIndex < (LocalObjectArray.Num() + LocalSpawnedActorArray.Num()))
+		AActor* LocalCurrentActor = nullptr;
+		int32 LocalIndex = State->Index;
+		TArray<AActor*>& LocalActorArray = State->ActorArray;
+		TArray<AActor*>& LocalSpawnedActorArray = State->SpawnedActorArray;
+		const UWorld* LocalCurrentWorld = State->CurrentWorld;
+		// LocalActorArray is immutable within the loop so we can cache its number of elements :
+		const int32 LocalActorArrayNum = LocalActorArray.Num();
+		while(++LocalIndex < (LocalActorArrayNum + LocalSpawnedActorArray.Num()))
 		{
-			if (LocalIndex < LocalObjectArray.Num())
+			if (LocalIndex < LocalActorArrayNum)
 			{
-				LocalCurrentActor = static_cast<AActor*>(LocalObjectArray[LocalIndex]);
+				LocalCurrentActor = LocalActorArray[LocalIndex];
 			}
 			else
 			{
-				LocalCurrentActor = LocalSpawnedActorArray[LocalIndex - LocalObjectArray.Num()];
+				LocalCurrentActor = LocalSpawnedActorArray[LocalIndex - LocalActorArrayNum];
 			}
 			State->ConsideredCount++;
 			
@@ -313,7 +346,11 @@ public:
 			if ( ActorLevel
 				&& static_cast<const Derived*>(this)->IsActorSuitable(LocalCurrentActor)
 				&& static_cast<const Derived*>(this)->CanIterateLevel(ActorLevel)
-				&& ActorLevel->GetWorld() == LocalCurrentWorld)
+#if !WITH_EDITOR
+				// No need to perform the world test if ActorArray has already been pre-filtered, which is the case in editor :
+				&& (ActorLevel->GetWorld() == LocalCurrentWorld)
+#endif // !WITH_EDITOR
+				)
 			{
 				// ignore non-persistent world settings
 				if (ActorLevel == LocalCurrentWorld->PersistentLevel || !LocalCurrentActor->IsA(AWorldSettings::StaticClass()))
@@ -412,13 +449,6 @@ protected:
 			return false;
 		}
 
-#if WITH_EDITOR
-		// In editor, we don't filter actors by class when creating the iterator.
-		if (State->DesiredClass != AActor::StaticClass() && !Actor->IsA(State->DesiredClass))
-		{
-			return false;
-		}
-#endif
 		return true;
 	}
 
