NetUnitTest
* FTestWorlds is now responsible for restoring the PIE world and set necessary cvars for network tests to execute

diff --git a/Engine/Source/Runtime/Engine/Private/Tests/NetTestHelpers.cpp b/Engine/Source/Runtime/Engine/Private/Tests/NetTestHelpers.cpp
--- a/Engine/Source/Runtime/Engine/Private/Tests/NetTestHelpers.cpp
+++ b/Engine/Source/Runtime/Engine/Private/Tests/NetTestHelpers.cpp
@@ -3,12 +3,12 @@
 #include "Tests/NetTestHelpers.h"
 
 #include "CoreGlobals.h"
-#include "Engine/NetDriver.h"
 #include "Engine/NetworkObjectList.h"
 #include "Engine/Engine.h"
 #include "Engine/GameInstance.h"
 #include "Engine/World.h"
 #include "Engine/GameViewportClient.h"
+#include "Engine/PackageMapClient.h"
 #include "EngineUtils.h"
 #include "GameFramework/PlayerController.h"
 #include "GameFramework/PlayerState.h"
@@ -21,6 +21,11 @@
 #include "Settings/LevelEditorPlaySettings.h"
 #endif
 
+#if UE_WITH_IRIS
+#include "Iris/ReplicationSystem/ReplicationSystem.h"
+#include "Iris/ReplicationSystem/ObjectReplicationBridge.h"
+#endif
+
 #include "UObject/UObjectGlobals.h"
 #include "UObject/Package.h"
 
@@ -160,6 +165,22 @@ int32 FTestWorldInstance::FindUnusedPIEInstance()
 	return MaxUsedPIEInstance + 1;
 }
 
+FTestWorldInstance::FContext FTestWorldInstance::GetTestContext() const
+{
+#if UE_WITH_IRIS
+	UReplicationSystem* RepSystem = GetNetDriver() ? GetNetDriver()->GetReplicationSystem() : nullptr;
+	UObjectReplicationBridge* RepBridge  = RepSystem ? RepSystem->GetReplicationBridgeAs<UObjectReplicationBridge>() : nullptr;
+#endif
+
+	return FContext
+	{ 
+		.World = GetWorld(),
+		.NetDriver = GetNetDriver(),
+		.IrisRepSystem = RepSystem,
+		.IrisRepBridge = RepBridge,
+	};
+}
+
 UWorld* FTestWorldInstance::GetWorld() const
 {
 	return GameInstance ? GameInstance->GetWorld() : nullptr;
@@ -392,6 +413,53 @@ void FTestWorlds::TickClientsAndDelay(uint32 NumFramesToDelay)
 	return nullptr;
 }
 
+UObject* FTestWorlds::FindReplicatedObjectOnClient(UObject* ServerObject, uint32 ClientIndex) const
+{
+	using namespace UE::Net;
+
+	if (!Clients.IsValidIndex(ClientIndex))
+	{
+		ensureMsgf(false, TEXT("FTestWorlds::FindReplicatedObjectOnClient received invalid ClientIndex: %u"), ClientIndex);
+		return nullptr;
+	}
+
+	if (ServerObject->GetWorld() != Server.GetWorld())
+	{
+		ensureMsgf(false, TEXT("FTestWorlds::FindReplicatedObjectOnClient received object %s not part of the Server world"), *GetFullNameSafe(ServerObject));
+		return nullptr;
+	}
+
+	const FTestWorldInstance::FContext ServerContext = Server.GetTestContext();
+	const FTestWorldInstance::FContext ClientContext = Clients[ClientIndex].GetTestContext();
+
+	if (ServerContext.NetDriver->IsUsingIrisReplication())
+	{
+#if UE_WITH_IRIS
+	const FNetRefHandle NetHandle = ServerContext.IrisRepBridge->GetReplicatedRefHandle(ServerObject);
+
+		if (!NetHandle.IsValid())
+		{
+			ensureMsgf(false, TEXT("FTestWorlds::FindReplicatedObjectOnClient ServerObject: %s is not replicated."), *GetFullNameSafe(ServerObject));
+			return nullptr;
+		}
+
+		return ClientContext.IrisRepBridge->GetReplicatedObject(NetHandle);
+
+#endif // UE_WITH_IRIS
+	}
+	else
+	{
+		const FNetworkGUID NetGUID = ServerContext.NetDriver->GuidCache->GetNetGUID(ServerObject);
+
+		if (!NetGUID.IsValid())
+		{
+			ensureMsgf(false, TEXT("FTestWorlds::FindReplicatedObjectOnClient ServerObject: %s is not replicated."), *GetFullNameSafe(ServerObject));
+		}
+
+		return ClientContext.NetDriver->GuidCache->GetObjectFromNetGUID(NetGUID, false);
+	}
+}
+
 //------------------------------------------------------------------------
 
 FScopedCVarOverrideInt::FScopedCVarOverrideInt(const TCHAR* VariableName, int32 Value)
@@ -412,25 +480,38 @@ FScopedCVarOverrideInt::~FScopedCVarOverrideInt()
 	}
 }
 
+//------------------------------------------------------------------------
+
 FScopedTestSettings::FScopedTestSettings()
 	: AddressResolutionDisabled(TEXT("net.IpConnectionDisableResolution"), 1)
 	, BandwidthThrottlingDisabled(TEXT("net.DisableBandwithThrottling"), 1)
 	, RepGraphBandwidthThrottlingDisabled(TEXT("Net.RepGraph.DisableBandwithLimit"), 1)
 	, RandomNetUpdateDelayDisabled(TEXT("net.DisableRandomNetUpdateDelay"), 1)
 	, GameplayDebuggerDisabled(TEXT("GameplayDebugger.AutoCreateGameplayDebuggerManager"), 0)
-	, OldGWorld(GWorld)
+{
+}
+
+FScopedTestSettings::~FScopedTestSettings()
+{
+}
+
+//------------------------------------------------------------------------
+
+FScopedNetTestPIERestoration::FScopedNetTestPIERestoration()
+	: OldGWorld(GWorld)
 	, OldPIEID(UE::GetPlayInEditorID())
 	, OldGIsPlayInEditorWorld(GIsPlayInEditorWorld)
 {
+
 }
 
-FScopedTestSettings::~FScopedTestSettings()
+FScopedNetTestPIERestoration::~FScopedNetTestPIERestoration()
 {
 	GWorld = OldGWorld;
 	UE::SetPlayInEditorID(OldPIEID);
 	GIsPlayInEditorWorld = OldGIsPlayInEditorWorld;
 }
 
-}
+} // end namespace UE::Net
 
 #endif // WITH_EDITOR
\ No newline at end of file
diff --git a/Engine/Source/Runtime/Engine/Public/Tests/NetTestHelpers.h b/Engine/Source/Runtime/Engine/Public/Tests/NetTestHelpers.h
--- a/Engine/Source/Runtime/Engine/Public/Tests/NetTestHelpers.h
+++ b/Engine/Source/Runtime/Engine/Public/Tests/NetTestHelpers.h
@@ -3,15 +3,14 @@
 #pragma once
 
 #include "HAL/Platform.h"
-#include "Net/UnrealNetwork.h"
-#include "Engine/ActorChannel.h"
 #include "Engine/NetDriver.h"
-#include "Engine/NetworkObjectList.h"
 
 class APlayerController;
 class IConsoleVariable;
 class UGameInstance;
 class UWorld;
+class UReplicationSystem;
+class UObjectReplicationBridge;
 
 struct FGameInstancePIEParameters;
 
@@ -20,12 +19,80 @@
 namespace UE::Net
 {
 
+/*
+ * Used to set a cvar and restore it to it's original value when destroyed
+ */
+class FScopedCVarOverrideInt
+{
+public:
+	ENGINE_API FScopedCVarOverrideInt(const TCHAR* VariableName, int32 Value);
+	ENGINE_API ~FScopedCVarOverrideInt();
+
+	FScopedCVarOverrideInt(FScopedCVarOverrideInt&&) = delete;
+	FScopedCVarOverrideInt(const FScopedCVarOverrideInt&) = delete;
+	FScopedCVarOverrideInt& operator=(FScopedCVarOverrideInt&&) = delete;
+	FScopedCVarOverrideInt& operator=(const FScopedCVarOverrideInt&) = delete;
+
+private:
+	IConsoleVariable* Variable = nullptr;
+	int32 SavedValue = 0;
+};
+
 /**
+* Sets and restores  cvars needed to use FNetTestWorldInstances within a scope.
+* Meant to be used within a single function.
+*/
+class FScopedTestSettings
+{
+public:
+	ENGINE_API FScopedTestSettings();
+	ENGINE_API ~FScopedTestSettings();
+
+	FScopedTestSettings(FScopedTestSettings&&) = delete;
+	FScopedTestSettings(const FScopedTestSettings&) = delete;
+	FScopedTestSettings& operator=(FScopedTestSettings&&) = delete;
+	FScopedTestSettings& operator=(const FScopedTestSettings&) = delete;
+
+private:
+	FScopedCVarOverrideInt AddressResolutionDisabled;
+	FScopedCVarOverrideInt BandwidthThrottlingDisabled;
+	FScopedCVarOverrideInt RepGraphBandwidthThrottlingDisabled;
+	FScopedCVarOverrideInt RandomNetUpdateDelayDisabled;
+	FScopedCVarOverrideInt GameplayDebuggerDisabled;
+};
+
+/**
+* Stores and restore GWorld and PIE settings modified by the creation of the temporary test worlds
+*/
+class FScopedNetTestPIERestoration
+{
+public:
+
+	ENGINE_API FScopedNetTestPIERestoration();
+	ENGINE_API ~FScopedNetTestPIERestoration();
+
+	FScopedNetTestPIERestoration(FScopedNetTestPIERestoration&&) = delete;
+	FScopedNetTestPIERestoration(const FScopedNetTestPIERestoration&) = delete;
+	FScopedNetTestPIERestoration& operator=(FScopedNetTestPIERestoration&&) = delete;
+	FScopedNetTestPIERestoration& operator=(const FScopedNetTestPIERestoration&) = delete;
+
+private:
+
+	// Restores PIE world context
+	UWorld* OldGWorld = nullptr;
+	int32 OldPIEID = 0;
+	bool OldGIsPlayInEditorWorld = false;
+};
+
+/**
  * Properly scoped/RAII wrapper around a GameInstance/WorldContext/World that makes it easier to write tests
  * involving full UWorld functionality within the scope of one function.
  */
 struct FTestWorldInstance
 {
+	struct FContext;
+
+public:
 	ENGINE_API static FTestWorldInstance CreateServer(const TCHAR* InURL);
 	ENGINE_API static FTestWorldInstance CreateClient(int32 ServerPort);
 
@@ -38,9 +105,12 @@
 	
 	ENGINE_API FTestWorldInstance(FTestWorldInstance&& Other);
 	ENGINE_API FTestWorldInstance& operator=(FTestWorldInstance&& Other);
-	
+
+	ENGINE_API FWorldContext* GetWorldContext() const;	
+
+	ENGINE_API FContext GetTestContext() const;
+
 	ENGINE_API UWorld* GetWorld() const;
-	ENGINE_API FWorldContext* GetWorldContext() const;
 	ENGINE_API UNetDriver* GetNetDriver() const;
 
 	ENGINE_API int32 GetPort();
@@ -52,4 +122,16 @@
 
 public:
 
+	struct FContext
+	{
+		UWorld* World = nullptr;
+		UNetDriver* NetDriver = nullptr;
+#if UE_WITH_IRIS
+		UReplicationSystem* IrisRepSystem = nullptr;
+		UObjectReplicationBridge* IrisRepBridge = nullptr;
+#endif
+	};
+
+public:
+
 	UGameInstance* GameInstance = nullptr;
@@ -102,6 +184,14 @@
 	/** Return the Server's player state corresponding to a specific client */
 	ENGINE_API APlayerController* GetServerPlayerControllerOfClient(uint32 ClientIndex);
 
+	/**
+	 * Find the remote instance of a replicated object.
+	 * @param ServerObject The server object that you want the remote version of.
+	 * @param ClientIndex The client you want a remote instance of
+	 * @return Return the remote (client) instance of the same object if it exists.
+	 */
+	ENGINE_API UObject* FindReplicatedObjectOnClient(UObject* ServerObject, uint32 ClientIndex) const;
+
 public:
 
 	/** Server and Client Worlds */
@@ -113,7 +203,15 @@
 	void OnNetDriverCreated(UWorld* InWorld, UNetDriver* InNetDriver);
 	FDelegateHandle NetDriverCreatedHandle;
 
+private:
+
 	float TickDeltaSeconds = 0.0166f;
+
+	// Sets up important settings for the networking system to run optimally
+	FScopedTestSettings TestSettings;
+
+	// Restore GWorld and other PIE settings
+	FScopedNetTestPIERestoration PIERestoration;
 };
 
 //------------------------------------------------------------------------
@@ -141,52 +239,6 @@
 	return bPredicateResult;
 }
 
-//------------------------------------------------------------------------
-// FScopedCVarOverrideInt
-//------------------------------------------------------------------------
-class FScopedCVarOverrideInt
-{
-public:
-	ENGINE_API FScopedCVarOverrideInt(const TCHAR* VariableName, int32 Value);
-	ENGINE_API ~FScopedCVarOverrideInt();
-
-	FScopedCVarOverrideInt(FScopedCVarOverrideInt&&) = delete;
-	FScopedCVarOverrideInt(const FScopedCVarOverrideInt&) = delete;
-	FScopedCVarOverrideInt& operator=(FScopedCVarOverrideInt&&) = delete;
-	FScopedCVarOverrideInt& operator=(const FScopedCVarOverrideInt&) = delete;
-
-private:
-	IConsoleVariable* Variable = nullptr;
-	int32 SavedValue = 0;
-};
-
-/**
- * Sets and restores globals and cvars needed to use FNetTestWorldInstances within a scope.
- * Meant to be used within a single function.
- */
-class FScopedTestSettings
-{
-public:
-	ENGINE_API FScopedTestSettings();
-	ENGINE_API ~FScopedTestSettings();
-
-	FScopedTestSettings(FScopedTestSettings&&) = delete;
-	FScopedTestSettings(const FScopedTestSettings&) = delete;
-	FScopedTestSettings& operator=(FScopedTestSettings&&) = delete;
-	FScopedTestSettings& operator=(const FScopedTestSettings&) = delete;
-
-private:
-	FScopedCVarOverrideInt AddressResolutionDisabled;
-	FScopedCVarOverrideInt BandwidthThrottlingDisabled;
-	FScopedCVarOverrideInt RepGraphBandwidthThrottlingDisabled;
-	FScopedCVarOverrideInt RandomNetUpdateDelayDisabled;
-	FScopedCVarOverrideInt GameplayDebuggerDisabled;
-
-	UWorld* OldGWorld;
-	int32 OldPIEID;
-	bool OldGIsPlayInEditorWorld;
-};
-
 } // end namespace UE::Net
 
 #endif //WITH_EDITOR
