Fix for using custom texture mip streaming on platforms that support virtual mips
We now disable the use of virtual mips when custom mip streaming is present - as the virtual mip support is not yet implemented for that streaming path.

diff --git a/Engine/Source/Runtime/Engine/Private/Rendering/StreamableTextureResource.cpp b/Engine/Source/Runtime/Engine/Private/Rendering/StreamableTextureResource.cpp
--- a/Engine/Source/Runtime/Engine/Private/Rendering/StreamableTextureResource.cpp
+++ b/Engine/Source/Runtime/Engine/Private/Rendering/StreamableTextureResource.cpp
@@ -12,6 +12,7 @@
 #include "ProfilingDebugging/ScopedDebugInfo.h"
 #include "Stats/StatsTrace.h"
 #include "RHIUtilities.h"
+#include "Engine/TextureMipDataProviderFactory.h"
 
 #if STATS
 int64 GUITextureMemory = 0;
@@ -121,7 +122,9 @@ FStreamableTextureResource::FStreamableTextureResource(UTexture* InOwner, const
 
 	// Whether the virtual update path is enabled for this texture. This allows to map / unmap top mip memory in an out.
 	// Whether the texture will be created with TexCreate_Virtual depends on the requested mip count and "r.VirtualTextureReducedMemory"
-	bUsePartiallyResidentMips = bAllowPartiallyResidentMips && InPostInitState.bSupportsStreaming && CanCreateWithPartiallyResidentMips(CreationFlags);
+	// TODO: Remove the bHasNoMipProvider filter here, once virtual mip support is added to the Mip Provider streaming path
+	const bool bHasNoMipProvider = (InOwner->GetAssetUserData<UTextureMipDataProviderFactory>() == nullptr);
+	bUsePartiallyResidentMips = bHasNoMipProvider && bAllowPartiallyResidentMips && InPostInitState.bSupportsStreaming && CanCreateWithPartiallyResidentMips(CreationFlags);
 	 
 	STAT(LODGroupStatName = TextureGroupStatFNames[LODGroup]);
 	STAT(bIsNeverStream = InOwner->NeverStream);
