[HWRT] Remove FRDGAsyncTask from CommitRayTracingSkinnedGeometryUpdates because otherwise we might release the SkinCache entry before the lambda is executed and we'll end up with a dangling VB pointer.

#rnx

diff --git a/Engine/Source/Runtime/Engine/Private/RayTracingSkinnedGeometry.cpp b/Engine/Source/Runtime/Engine/Private/RayTracingSkinnedGeometry.cpp
--- a/Engine/Source/Runtime/Engine/Private/RayTracingSkinnedGeometry.cpp
+++ b/Engine/Source/Runtime/Engine/Private/RayTracingSkinnedGeometry.cpp
@@ -333,7 +333,7 @@ void FRayTracingSkinnedGeometryUpdateQueue::Commit(FRDGBuilder& GraphBuilder, ER
 			[
 				BuildRequests = MoveTemp(GeometryBuildRequests),
 				SharedScratchBuffer
-			](FRDGAsyncTask, FRHICommandList& RHICmdList)
+			](FRHICommandList& RHICmdList)
 			{
 				FRHIBufferRange ScratchBufferRange;
 				ScratchBufferRange.Buffer = SharedScratchBuffer->GetRHI();
