SplineComponent: Fix bClosedLoop not being used by construction scripts if set on an instance.

diff --git a/Engine/Source/Runtime/Engine/Classes/Components/SplineComponent.h b/Engine/Source/Runtime/Engine/Classes/Components/SplineComponent.h
--- a/Engine/Source/Runtime/Engine/Classes/Components/SplineComponent.h
+++ b/Engine/Source/Runtime/Engine/Classes/Components/SplineComponent.h
@@ -936,6 +936,9 @@ public:
 	UPROPERTY()
 	FSplineCurves SplineCurves;
 
+	UPROPERTY()
+	bool bClosedLoop = false;
+
 	UPROPERTY()
 	FSplineCurves SplineCurvesPreUCS;
 };
diff --git a/Engine/Source/Runtime/Engine/Private/Components/SplineComponent.cpp b/Engine/Source/Runtime/Engine/Private/Components/SplineComponent.cpp
--- a/Engine/Source/Runtime/Engine/Private/Components/SplineComponent.cpp
+++ b/Engine/Source/Runtime/Engine/Private/Components/SplineComponent.cpp
@@ -2207,6 +2207,7 @@ TStructOnScope<FActorComponentInstanceData> USplineComponent::GetComponentInstan
 	if (bSplineHasBeenEdited)
 	{
 		SplineInstanceData->SplineCurves = SplineCurves;
+		SplineInstanceData->bClosedLoop = bClosedLoop;
 	}
 	SplineInstanceData->bSplineHasBeenEdited = bSplineHasBeenEdited;
 
@@ -2257,6 +2258,7 @@ void USplineComponent::ApplyComponentInstanceData(FSplineInstanceData* SplineIns
 	if (SplineInstanceData->bSplineHasBeenEdited)
 	{
 		SplineCurves = SplineInstanceData->SplineCurves;
+		bClosedLoop = SplineInstanceData->bClosedLoop;
 		bModifiedByConstructionScript = false;
 	}
 
