Fix Substrate versioning issue.

* Register added nodes into EditorOnly expression collection to be saved/displayed correctly
* Fix conversion invalidation when a node is deleted
* Change Substrate version meaning:
  * -1 no conversion has run
  *   0 the conversion fonction has run and no conversion was done. Or the material has been changed/touched since its conversion
  * >=0 the conversion function has run

diff --git a/Engine/Source/Editor/MaterialEditor/Private/MaterialEditor.cpp b/Engine/Source/Editor/MaterialEditor/Private/MaterialEditor.cpp
--- a/Engine/Source/Editor/MaterialEditor/Private/MaterialEditor.cpp
+++ b/Engine/Source/Editor/MaterialEditor/Private/MaterialEditor.cpp
@@ -6223,6 +6223,7 @@ void FMaterialEditor::DeleteNodes(const TArray<UEdGraphNode*>& NodesToDelete, bo
 			}
 			RegenerateCodeView();
 		}
+		InvalidateSubstrateConversionVersion(Material);
 		UpdatePreviewMaterial();
 		Material->MarkPackageDirty();
 		SetMaterialDirty();
diff --git a/Engine/Source/Runtime/Engine/Private/Materials/Material.cpp b/Engine/Source/Runtime/Engine/Private/Materials/Material.cpp
--- a/Engine/Source/Runtime/Engine/Private/Materials/Material.cpp
+++ b/Engine/Source/Runtime/Engine/Private/Materials/Material.cpp
@@ -3315,12 +3315,18 @@ EBlendMode ConvertLegacyBlendMode(EBlendMode InBlendMode, FMaterialShadingModelF
 #define SUBSTRATE_MOVE_CONNECTION 0
 #define SUBSTRATE_COPY_CONNECTION 1
 
-#if WITH_EDITOR
+// Current conversion version - This means that the material has run ConvertMaterialToSubstrateMaterial() once, 
+// and its data has been converted.
+static int32 GetSubstrateConversionVersion()	{ return  1; }
+// Default version - This means that:
+//  * The material has run ConvertMaterialToSubstrateMaterial() once, but no conversion was needed 
+//  * Or that the material has been updated/touched since its conversion (e.g., manual edit)
+// In both case, there is no longer needs to rerun ConvertMaterialToSubstrateMaterial()
+static int32 GetSubstrateNoConversionVersion()	{ return  0; }
+// Invalid version - This means that the material has never run ConvertMaterialToSubstrateMaterial()
+static int32 GetSubstrateInvalidVersion()		{ return -1; }
 
-static int32 GetSubstrateConversionVersion()
-{ 
-	return 0;
-}
+#if WITH_EDITOR
 
 bool UMaterial::ConvertMaterialToSubstrateMaterial(bool bAllowEmptyMaterialUpdate)
 {
@@ -3347,10 +3353,11 @@ bool UMaterial::ConvertMaterialToSubstrateMaterial(bool bAllowEmptyMaterialUpdat
 
 	UMaterialEditorOnlyData* EditorOnly = GetEditorOnlyData();
 
-	// * If the current material has already been converted, skip the conversion
+	// * If the current material has already been converted, skip the conversion. 
 	// * Store the current version of Subtrate's auto-conversion
 	//   This allows to version the conversion, and safely update auto-converted materials if they have been saved.
-	if (EditorOnly->SubstrateConversionVersion == GetSubstrateConversionVersion())
+	// * If no conversion was needed, the conversion version is set to 0, to avoid rerunning this function
+	if (EditorOnly->SubstrateConversionVersion == GetSubstrateConversionVersion() || EditorOnly->SubstrateConversionVersion == GetSubstrateNoConversionVersion())
 	{
 		return false;
 	}
@@ -3569,6 +3576,7 @@ bool UMaterial::ConvertMaterialToSubstrateMaterial(bool bAllowEmptyMaterialUpdat
 		if (bUseMaterialAttributes && EditorOnly->MaterialAttributes.Expression && !EditorOnly->FrontMaterial.IsConnected() && !EditorOnly->MaterialAttributes.Expression->IsResultSubstrateMaterial(EditorOnly->MaterialAttributes.OutputIndex) && !bHasAnySubstrateNodes) // M_Rifle cause issues there
 		{
 			UMaterialExpressionSubstrateConvertMaterialAttributes* ConvertAttributeNode = NewObject<UMaterialExpressionSubstrateConvertMaterialAttributes>(this);
+			EditorOnly->ExpressionCollection.AddExpression(ConvertAttributeNode);
 			ConvertAttributeNode->Material = this;
 			SetPosXAndMoveReferenceToTheRight(ConvertAttributeNode);
 			ConvertAttributeNode->SubsurfaceProfile = bRequireNoSubsurfaceProfile ? nullptr : SubsurfaceProfile;
@@ -3601,11 +3609,11 @@ bool UMaterial::ConvertMaterialToSubstrateMaterial(bool bAllowEmptyMaterialUpdat
 				const FGuid CustomEyeTangentGuid = FMaterialAttributeDefinitionMap::GetCustomAttributeID(TEXT("CustomEyeTangent"));
 
 				UMaterialExpressionSetMaterialAttributes* SetAttributesNode = NewObject<UMaterialExpressionSetMaterialAttributes>(this);
+				EditorOnly->ExpressionCollection.AddExpression(SetAttributesNode); // Add the SetAttributesNode to ensure the MaterialAttribute input are correctly tracked & cached
 				SetAttributesNode->Material = this;
 				SetAttributesNode->Inputs[0].Connect(0, EditorOnly->MaterialAttributes.Expression);
 				SetAttributesNode->AttributeSetTypes.Add(CustomEyeTangentGuid);
 				SetAttributesNode->Inputs.Add(*TangentOutput->GetInput(0));
-				EditorOnly->ExpressionCollection.AddExpression(SetAttributesNode); // Add the SetAttributesNode to ensure the MaterialAttribute input are correctly tracked & cached
 
 				// Connect to conversion node
 				ConvertAttributeNode->MaterialAttributes.Connect(0, SetAttributesNode);
@@ -3638,6 +3646,7 @@ bool UMaterial::ConvertMaterialToSubstrateMaterial(bool bAllowEmptyMaterialUpdat
 
 				// Now pass through the convert to decal node, which flag the material as SSM_Decal, which will set the domain to Decal.
 				UMaterialExpressionSubstrateConvertToDecal* ConvertToDecalNode = NewObject<UMaterialExpressionSubstrateConvertToDecal>(this);
+				EditorOnly->ExpressionCollection.AddExpression(ConvertToDecalNode);
 				ConvertAttributeNode->Material = this;
 				ReplaceNodeAndMoveToTheRight(ConvertAttributeNode, ConvertToDecalNode);
 				ConvertToDecalNode->DecalMaterial.Connect(0, ConvertAttributeNode);
@@ -3663,7 +3672,7 @@ bool UMaterial::ConvertMaterialToSubstrateMaterial(bool bAllowEmptyMaterialUpdat
 				if (MFCallNode->SetMaterialFunction(DefaultMF))
 				{
 					// This is needed for input/output expressions to be set correctly, otherwise compilation will fail.
-					GetExpressionCollection().AddExpression(MFCallNode);
+					EditorOnly->ExpressionCollection.AddExpression(MFCallNode);
 
 					SetPosXAndMoveReferenceToTheRight(MFCallNode);
 					EditorOnly->FrontMaterial.Connect(0, MFCallNode);
@@ -3685,6 +3694,7 @@ bool UMaterial::ConvertMaterialToSubstrateMaterial(bool bAllowEmptyMaterialUpdat
 			{
 				// Or if it cannot be found, a slab node
 				UMaterialExpressionSubstrateSlabBSDF* SlabNode = NewObject<UMaterialExpressionSubstrateSlabBSDF>(this);
+				EditorOnly->ExpressionCollection.AddExpression(SlabNode);
 				SlabNode->Material = this;
 				SetPosXAndMoveReferenceToTheRight(SlabNode);
 				EditorOnly->FrontMaterial.Connect(0, SlabNode);
@@ -3707,6 +3717,7 @@ bool UMaterial::ConvertMaterialToSubstrateMaterial(bool bAllowEmptyMaterialUpdat
 
 						// Create metalness to Slab parameterisation conveersion node
 						UMaterialExpressionSubstrateMetalnessToDiffuseAlbedoF0* SubstrateMetalnessToDiffuseAlbedoF0 = NewObject<UMaterialExpressionSubstrateMetalnessToDiffuseAlbedoF0>(this);
+						EditorOnly->ExpressionCollection.AddExpression(SubstrateMetalnessToDiffuseAlbedoF0);
 						SetPosXAndMoveReferenceToTheRight(SubstrateMetalnessToDiffuseAlbedoF0);
 						ColorMatInputConnectionTo(EditorOnly->BaseColor, SubstrateMetalnessToDiffuseAlbedoF0, 0, MP_BaseColor);
 						ScalarMatInputConnectionTo(EditorOnly->Metallic, SubstrateMetalnessToDiffuseAlbedoF0, 1, MP_Metallic);
@@ -3714,6 +3725,7 @@ bool UMaterial::ConvertMaterialToSubstrateMaterial(bool bAllowEmptyMaterialUpdat
 					
 						// Top slab BSDF as a simple Disney material
 						UMaterialExpressionSubstrateSlabBSDF* BottomSlabBSDF = NewObject<UMaterialExpressionSubstrateSlabBSDF>(this);
+						EditorOnly->ExpressionCollection.AddExpression(BottomSlabBSDF);
 						BottomSlabBSDF->Material = this;
 						SetPosXAndMoveReferenceToTheRight(BottomSlabBSDF);
 						BottomSlabBSDF->GetInput(0)->Connect(0, SubstrateMetalnessToDiffuseAlbedoF0);
@@ -3728,6 +3740,7 @@ bool UMaterial::ConvertMaterialToSubstrateMaterial(bool bAllowEmptyMaterialUpdat
 
 						// Now weight the top base material by opacity.
 						UMaterialExpressionSubstrateSlabBSDF* TopSlabBSDF = NewObject<UMaterialExpressionSubstrateSlabBSDF>(this);
+						EditorOnly->ExpressionCollection.AddExpression(TopSlabBSDF);
 						TopSlabBSDF->Material = this;
 						TopSlabBSDF->MaterialExpressionEditorX = BottomSlabBSDF->MaterialExpressionEditorX;
 						TopSlabBSDF->MaterialExpressionEditorY = BottomSlabBSDF->MaterialExpressionEditorY + 650;
@@ -3737,6 +3750,7 @@ bool UMaterial::ConvertMaterialToSubstrateMaterial(bool bAllowEmptyMaterialUpdat
 
 						//  The top layer has a hard coded specular value of 0.5 (F0 = 0.04)
 						UMaterialExpressionConstant* ConstantHalf = NewObject<UMaterialExpressionConstant>(this);
+						EditorOnly->ExpressionCollection.AddExpression(ConstantHalf);
 						ReplaceNodeAndMoveToTheRight(TopSlabBSDF, ConstantHalf);
 						ConstantHalf->R = 0.5f * 0.08f;
 						TopSlabBSDF->GetInput(1)->Connect(0, ConstantHalf);
@@ -3745,15 +3759,18 @@ bool UMaterial::ConvertMaterialToSubstrateMaterial(bool bAllowEmptyMaterialUpdat
 						// To simplify things, we set the top slab BSDF as having a constant Grey scale transmittance.
 						// As for the original, this is achieved with coverage so both transmittance and specular contribution vanishes
 						UMaterialExpressionConstant* ConstantZero = NewObject<UMaterialExpressionConstant>(this);
+						EditorOnly->ExpressionCollection.AddExpression(ConstantZero);
 						ReplaceNodeAndMoveToTheRight(TopSlabBSDF, ConstantZero);
 						ConstantZero->R = 0.0f;
 						TopSlabBSDF->GetInput(0)->Connect(0, ConstantZero);							// BaseColor = 0 to only feature absorption, no scattering
 
 						// Now setup the mean free path with a hard coded transmittance of 0.75 when viewing the surface perpendicularly
 						UMaterialExpressionConstant* Constant075 = NewObject<UMaterialExpressionConstant>(this);
+						EditorOnly->ExpressionCollection.AddExpression(Constant075);
 						ReplaceNodeAndMoveToTheRight(TopSlabBSDF, Constant075);
 						Constant075->R = 0.75f;
 						UMaterialExpressionSubstrateTransmittanceToMFP* TransToMDFP = NewObject<UMaterialExpressionSubstrateTransmittanceToMFP>(this);
+						EditorOnly->ExpressionCollection.AddExpression(TransToMDFP);
 						ReplaceNodeAndMoveToTheRight(TopSlabBSDF, TransToMDFP);
 						TransToMDFP->GetInput(0)->Connect(0, Constant075);
 						TopSlabBSDF->GetInput(7)->Connect(0, TransToMDFP);							// MFP -> MFP
@@ -3761,12 +3778,14 @@ bool UMaterial::ConvertMaterialToSubstrateMaterial(bool bAllowEmptyMaterialUpdat
 
 						// Now weight the top base material by ClearCoat
 						UMaterialExpressionSubstrateWeight* TopSlabBSDFWithCoverage = NewObject<UMaterialExpressionSubstrateWeight>(this);
+						EditorOnly->ExpressionCollection.AddExpression(SubstrateMetalnessToDiffuseAlbedoF0); // Add the SetAttributesNode to ensure the MaterialAttribute input are correctly tracked & cached
 						SetPosXAndMoveReferenceToTheRight(TopSlabBSDFWithCoverage);
 						TopSlabBSDFWithCoverage->GetInput(0)->Connect(0, TopSlabBSDF);												// TopSlabBSDF -> A
 						ScalarMatInputConnectionTo(EditorOnly->ClearCoat, TopSlabBSDFWithCoverage, 1, MP_CustomData0);				// ClearCoat -> Weight
 						ScalarMatInputConnectionTo(EditorOnly->ClearCoatRoughness, TopSlabBSDFWithCoverage, 1, MP_CustomData1);		// ClearCoat -> Weight
 
 						UMaterialExpressionSubstrateVerticalLayering* VerticalLayering = NewObject<UMaterialExpressionSubstrateVerticalLayering>(this);
+						EditorOnly->ExpressionCollection.AddExpression(VerticalLayering);
 						SetPosXAndMoveReferenceToTheRight(VerticalLayering);
 						VerticalLayering->GetInput(0)->Connect(0, TopSlabBSDFWithCoverage);			// Top -> Top
 						VerticalLayering->GetInput(1)->Connect(0, BottomSlabBSDF);					// Bottom -> Base
@@ -3780,6 +3799,7 @@ bool UMaterial::ConvertMaterialToSubstrateMaterial(bool bAllowEmptyMaterialUpdat
 				if (!bClearCoatConversionDone)
 				{
 					ConvertNode = NewObject<UMaterialExpressionSubstrateShadingModels>(this);
+					EditorOnly->ExpressionCollection.AddExpression(ConvertNode); // Add the SetAttributesNode to ensure the MaterialAttribute input are correctly tracked & cached
 					ConvertNode->Material = this;
 					SetPosXAndMoveReferenceToTheRight(ConvertNode);
 					ConvertNode->SubsurfaceProfile = bRequireNoSubsurfaceProfile ? nullptr : SubsurfaceProfile;
@@ -3836,6 +3856,7 @@ bool UMaterial::ConvertMaterialToSubstrateMaterial(bool bAllowEmptyMaterialUpdat
 			else if (MaterialDomain == MD_Volume)
 			{
 				UMaterialExpressionSubstrateVolumetricFogCloudBSDF* VolBSDF = NewObject<UMaterialExpressionSubstrateVolumetricFogCloudBSDF>(this);
+				EditorOnly->ExpressionCollection.AddExpression(VolBSDF);
 				VolBSDF->Material = this;
 				SetPosXAndMoveReferenceToTheRight(VolBSDF);
 				ColorMatInputConnectionTo(EditorOnly->BaseColor, VolBSDF, 0, MP_BaseColor);	
@@ -3858,6 +3879,7 @@ bool UMaterial::ConvertMaterialToSubstrateMaterial(bool bAllowEmptyMaterialUpdat
 
 				// Only Emissive & Opacity are valid input for PostProcess material
 				UMaterialExpressionSubstrateLightFunction* LightFunctionNode = NewObject<UMaterialExpressionSubstrateLightFunction>(this);
+				EditorOnly->ExpressionCollection.AddExpression(LightFunctionNode);
 				LightFunctionNode->Material = this;
 				SetPosXAndMoveReferenceToTheRight(LightFunctionNode);
 				ColorMatInputConnectionTo(EditorOnly->EmissiveColor, LightFunctionNode, 0, MP_EmissiveColor);
@@ -3878,6 +3900,7 @@ bool UMaterial::ConvertMaterialToSubstrateMaterial(bool bAllowEmptyMaterialUpdat
 				}
 
 				UMaterialExpressionSubstratePostProcess* PostProcNode = NewObject<UMaterialExpressionSubstratePostProcess>(this);
+				EditorOnly->ExpressionCollection.AddExpression(PostProcNode);
 				PostProcNode->Material = this;
 				SetPosXAndMoveReferenceToTheRight(PostProcNode);
 
@@ -3895,6 +3918,7 @@ bool UMaterial::ConvertMaterialToSubstrateMaterial(bool bAllowEmptyMaterialUpdat
 				ShadingModels.AddShadingModel(MSM_DefaultLit);
 
 				ConvertNode = NewObject<UMaterialExpressionSubstrateShadingModels>(this);
+				EditorOnly->ExpressionCollection.AddExpression(ConvertNode);
 				ConvertNode->Material = this;
 				SetPosXAndMoveReferenceToTheRight(ConvertNode);
 				ColorMatInputConnectionTo(EditorOnly->BaseColor, ConvertNode, 0, MP_BaseColor);
@@ -3916,6 +3940,7 @@ bool UMaterial::ConvertMaterialToSubstrateMaterial(bool bAllowEmptyMaterialUpdat
 
 				// Now pass through the convert to decal node, which flag the material as SSM_Decal, which will set the domain to Decal.
 				UMaterialExpressionSubstrateConvertToDecal* ConvertToDecalNode= NewObject<UMaterialExpressionSubstrateConvertToDecal>(this);
+				EditorOnly->ExpressionCollection.AddExpression(ConvertToDecalNode);
 				ConvertToDecalNode->Material = this;
 				ReplaceNodeAndMoveToTheRight(ConvertNode, ConvertToDecalNode);
 				ConvertToDecalNode->DecalMaterial.Connect(0, ConvertNode);
@@ -3931,6 +3956,7 @@ bool UMaterial::ConvertMaterialToSubstrateMaterial(bool bAllowEmptyMaterialUpdat
 				ShadingModels.AddShadingModel(MSM_Unlit);
 
 				UMaterialExpressionSubstrateUI* UINode = NewObject<UMaterialExpressionSubstrateUI>(this);
+				EditorOnly->ExpressionCollection.AddExpression(UINode);
 				UINode->Material = this;
 				SetPosXAndMoveReferenceToTheRight(UINode);
 				ColorMatInputConnectionTo(EditorOnly->EmissiveColor, UINode, 0, MP_EmissiveColor);
@@ -7718,11 +7744,11 @@ UMaterialEditorOnlyData::UMaterialEditorOnlyData()
 	AmbientOcclusion.Constant = FMaterialAttributeDefinitionMap::GetDefaultValue(MP_AmbientOcclusion).X;
 	Refraction.Constant = FMaterialAttributeDefinitionMap::GetDefaultValue(MP_Refraction).X;
 	SurfaceThickness.Constant = FMaterialAttributeDefinitionMap::GetDefaultValue(MP_SurfaceThickness).X;
-	ResetSubstrateConversionVersion();
+	SubstrateConversionVersion = GetSubstrateInvalidVersion();
 }
 
 void UMaterialEditorOnlyData::ResetSubstrateConversionVersion()
 {
-	SubstrateConversionVersion = -1; // Invalid version
+	SubstrateConversionVersion = GetSubstrateNoConversionVersion();
 }
 #undef LOCTEXT_NAMESPACE
