Tooltip update on some Substrate nodes.

* Add default thickness value for Substrate's node in tooltips
* Clarify meaning of thickness in the tooltip of the transmittance to MPF node.

#pregflight skip

diff --git a/Engine/Source/Runtime/Engine/Private/Materials/MaterialExpressionSubstrate.cpp b/Engine/Source/Runtime/Engine/Private/Materials/MaterialExpressionSubstrate.cpp
--- a/Engine/Source/Runtime/Engine/Private/Materials/MaterialExpressionSubstrate.cpp
+++ b/Engine/Source/Runtime/Engine/Private/Materials/MaterialExpressionSubstrate.cpp
@@ -3094,7 +3094,7 @@ void UMaterialExpressionSubstrateTransmittanceToMFP::GetConnectorToolTip(int32 I
 			ConvertToMultilineToolTip(TEXT("The colored transmittance for a view perpendicular to the surface. The transmittance for other view orientations will automatically be deduced according to surface thickness."), 80, OutToolTip);
 			break;
 		case 1:
-			ConvertToMultilineToolTip(TEXT("The desired thickness in centimeter. This can be set lower than 0.1mm(= 0.01cm) to enable the Thin lighting model on the slab node for instance. Another use case example: this node output called thickness can be modulated before it is plugged in a slab node.this can be used to achieve simple scattering/transmittance variation of the same material."), 80, OutToolTip);
+			ConvertToMultilineToolTip(TEXT("The thickness (in centimeter) at which the desired colored transmittance is reached. Default thickness: 0.01cm. Another use case example: this node output called thickness can be modulated before it is plugged in a slab node.this can be used to achieve simple scattering/transmittance variation of the same material."), 80, OutToolTip);
 			break;
 		}
 	}
diff --git a/Engine/Source/Runtime/Engine/Public/Materials/MaterialExpressionSubstrate.h b/Engine/Source/Runtime/Engine/Public/Materials/MaterialExpressionSubstrate.h
--- a/Engine/Source/Runtime/Engine/Public/Materials/MaterialExpressionSubstrate.h
+++ b/Engine/Source/Runtime/Engine/Public/Materials/MaterialExpressionSubstrate.h
@@ -271,7 +271,7 @@ class UMaterialExpressionSubstrateSlabBSDF : public UMaterialExpressionSubstrate
 	FExpressionInput SSSMFP;
 
 	/**
-	 * Scale the mean free path radius of the SSS profile according to a value between 0 and 1. Always used, when a subsurface profile is provided or not. (type = float, unitless, defaults to 1)
+	 * Scale the mean free path length. Value between 0 and 1. Always used, without or without a subsurface profile. (type = float, unitless, defaults to 1)
 	 */
 	UPROPERTY()
 	FExpressionInput SSSMFPScale;
@@ -971,7 +971,7 @@ class UMaterialExpressionSubstrateVerticalLayering : public UMaterialExpressionS
 	FExpressionInput Base;
 
 	/**
-	 * Thickness of the Top material layer in centimeter.
+	 * Thickness of the Top material layer in centimeter. Default value: 0.01cm.
 	 * It can be modulated to achieve simple scattering/transmittance variation of the same material.
 	 */
 	UPROPERTY()
@@ -1086,7 +1086,7 @@ class UMaterialExpressionSubstrateTransmittanceToMFP : public UMaterialExpressio
 	FExpressionInput TransmittanceColor;
 
 	/**
-	* Thickness of the layer in centimeter.
+	* Thickness of the layer in centimeter. Default value: 0.01cm.
 	* Example of use case: this node output called thickness can be modulated before it is plugged in a Vertical Layering node Thickness input. This can be used to achieve simple scattering/transmittance variation of the same material.
 	*/
 	UPROPERTY()
