Fixed code in UMaterial::SetMaterialUsage() to check editor is running.

This fixes crash with Slate Notification call in UMaterial::SetMaterialUsage when coming from UWorld::SendAllEndOfFrameUpdates.

diff --git a/Engine/Source/Runtime/Engine/Private/Materials/Material.cpp b/Engine/Source/Runtime/Engine/Private/Materials/Material.cpp
--- a/Engine/Source/Runtime/Engine/Private/Materials/Material.cpp
+++ b/Engine/Source/Runtime/Engine/Private/Materials/Material.cpp
@@ -106,6 +106,7 @@
 #include "MaterialGraph/MaterialGraph.h"
 #include "Widgets/Notifications/SNotificationList.h"
 #include "ThumbnailRendering/SceneThumbnailInfoWithPrimitive.h"
+#include "Editor.h"
 #endif
 #include "ShaderCodeLibrary.h"
 #include "Curves/CurveLinearColorAtlas.h"
@@ -1914,9 +1915,9 @@ bool UMaterial::SetMaterialUsage(bool &bNeedsRecompile, EMaterialUsage Usage, UM
 				{
 					UE_LOG(LogMaterial, Warning, TEXT("     The material will recompile every editor launch until resaved."));
 				}
-				else if (GIsEditor && !FApp::IsGame())
+				#if WITH_EDITOR
+				else if (GEditor && !GEditor->IsPlaySessionInProgress())
 				{
-#if WITH_EDITOR
 					FFormatNamedArguments Args;
 					Args.Add(TEXT("UsageName"), FText::FromString(GetUsageName(Usage)));
 					FNotificationInfo Info(FText::Format(LOCTEXT("CouldntSetMaterialUsage","Material didn't allow automatic setting of usage flag {UsageName} needed to render on this component, using Default Material instead."), Args));
@@ -1925,8 +1926,8 @@ bool UMaterial::SetMaterialUsage(bool &bNeedsRecompile, EMaterialUsage Usage, UM
 
 					// Give the user feedback as to why they are seeing the default material
 					FSlateNotificationManager::Get().AddNotification(Info);
-#endif
 				}
+				#endif
 
 				UsageFlagWarnings |= UsageFlagBit;
 			}
