Skeleton Modifier: make sure that the new parent is re-added first when reparenting to avoid parent index > child index

#rnx

diff --git a/Engine/Source/Runtime/Engine/Private/Animation/ReferenceSkeleton.cpp b/Engine/Source/Runtime/Engine/Private/Animation/ReferenceSkeleton.cpp
--- a/Engine/Source/Runtime/Engine/Private/Animation/ReferenceSkeleton.cpp
+++ b/Engine/Source/Runtime/Engine/Private/Animation/ReferenceSkeleton.cpp
@@ -305,6 +305,12 @@ namespace FReferenceSkeletonLocals
 
 			auto ElementToIndex = [this, &Visited, &OutMapping](const FElement& Element, auto&& ElementToIndexArg) -> void
 			{
+				// make sure the parent is added first
+				if (Element.Parent && !Visited[Element.Parent->RawIndex])
+				{
+					ElementToIndexArg(*Element.Parent, ElementToIndexArg);
+				}
+				
 				// add the element
 				if (!Visited[Element.RawIndex])
 				{
