Fix depth prepass disabling MSAA

diff --git a/Engine/Source/Runtime/Engine/Private/SceneTexturesConfig.cpp b/Engine/Source/Runtime/Engine/Private/SceneTexturesConfig.cpp
--- a/Engine/Source/Runtime/Engine/Private/SceneTexturesConfig.cpp
+++ b/Engine/Source/Runtime/Engine/Private/SceneTexturesConfig.cpp
@@ -173,10 +173,11 @@ void GetSceneColorFormatAndCreateFlags(ERHIFeatureLevel::Type FeatureLevel, bool
 	SceneColorCreateFlags |= sRGBFlag;
 }
 
-static ETextureCreateFlags GetSceneDepthStencilCreateFlags(uint32 NumSamples, bool bKeepDepthContent, bool bMemorylessMSAA, ETextureCreateFlags ExtraSceneDepthCreateFlags)
+static ETextureCreateFlags GetSceneDepthStencilCreateFlags(uint32 NumSamples, bool bKeepDepthContent, bool bMemorylessMSAA, EShaderPlatform ShaderPlatform, ETextureCreateFlags ExtraSceneDepthCreateFlags) //ericado
 {
 	ETextureCreateFlags DepthCreateFlags = TexCreate_DepthStencilTargetable | TexCreate_ShaderResource | TexCreate_InputAttachmentRead | ExtraSceneDepthCreateFlags;
-	if (!bKeepDepthContent || (NumSamples > 1 && bMemorylessMSAA))
+	// We can't discard the SceneDepth MSAA target if full depth prepass is enabled because it will be used in the base pass later on.
+	if (!bKeepDepthContent || (NumSamples > 1 && bMemorylessMSAA && (!IsMobilePlatform(ShaderPlatform) || !MobileUsesFullDepthPrepass(ShaderPlatform))))
 	{
 		DepthCreateFlags |= TexCreate_Memoryless;
 	}
@@ -357,7 +358,7 @@ void FSceneTexturesConfig::SetupMobileGBufferFlags(bool bRequiresMultiPass)
 void FSceneTexturesConfig::BuildSceneColorAndDepthFlags()
 {
     GetSceneColorFormatAndCreateFlags(FeatureLevel, bRequiresAlphaChannel, ExtraSceneColorCreateFlags, NumSamples, bMemorylessMSAA, ColorFormat, ColorCreateFlags);
-    DepthCreateFlags = GetSceneDepthStencilCreateFlags(NumSamples, bKeepDepthContent, bMemorylessMSAA, ExtraSceneDepthCreateFlags);
+    DepthCreateFlags = GetSceneDepthStencilCreateFlags(NumSamples, bKeepDepthContent, bMemorylessMSAA, ShaderPlatform, ExtraSceneDepthCreateFlags);
 }
 
 uint32 FSceneTexturesConfig::GetGBufferRenderTargetsInfo(FGraphicsPipelineRenderTargetsInfo& RenderTargetsInfo, EGBufferLayout Layout) const 
diff --git a/Engine/Source/Runtime/Engine/Private/SceneUtils.cpp b/Engine/Source/Runtime/Engine/Private/SceneUtils.cpp
--- a/Engine/Source/Runtime/Engine/Private/SceneUtils.cpp
+++ b/Engine/Source/Runtime/Engine/Private/SceneUtils.cpp
@@ -208,14 +208,12 @@ ENGINE_API uint32 GetDefaultMSAACount(const FStaticFeatureLevel InFeatureLevel,
 			if (InFeatureLevel == ERHIFeatureLevel::ES3_1)
 			{
 				bool bMobilePixelProjectedReflection = IsUsingMobilePixelProjectedReflection(ShaderPlatform);
-				
-				bool bIsFullDepthPrepassEnabled = MobileUsesFullDepthPrepass(ShaderPlatform);
 
-				bRendererSupportMSAA = bRHISupportsMSAA && !bMobilePixelProjectedReflection && !bIsFullDepthPrepassEnabled;
+				bRendererSupportMSAA = bRHISupportsMSAA && !bMobilePixelProjectedReflection;
 
 				if (!bRendererSupportMSAA)
 				{
-					FailedReason = FString::Printf(TEXT("RHISupportsMSAA %d, MobilePixelProjectedReflection %d, MobileFullDepthPrepass %d"), bRHISupportsMSAA ? 1 : 0, bMobilePixelProjectedReflection ? 1 : 0, bIsFullDepthPrepassEnabled ? 1 : 0);
+					FailedReason = FString::Printf(TEXT("RHISupportsMSAA %d, MobilePixelProjectedReflection %d"), bRHISupportsMSAA ? 1 : 0, bMobilePixelProjectedReflection ? 1 : 0);
 				}
 			}
 			else
