Added functionality on anim sequence and pose asset tp avoid setting transient retarget source assets
Also, created a console command to scan all the assets in the project and optiomally clean the retarget data if the source is missing
Added a console var to temporarily disable warnings at postload check

diff --git a/Engine/Config/ConsoleVariables.ini b/Engine/Config/ConsoleVariables.ini
--- a/Engine/Config/ConsoleVariables.ini
+++ b/Engine/Config/ConsoleVariables.ini
@@ -211,3 +211,5 @@ fx.SkipVectorVMBackendOptimizations=1
 ds.CADTranslator.Meshing.ActivateThinZoneMeshing=false
 ds.CADTranslator.Stitching.RemoveThinFaces=false
 
+; Temporarily disable anim sequence retarget data source error reporting (until assets are fixed)
+a.EnablePostLoadRetargetSourceErrorReporting=0
\ No newline at end of file
diff --git a/Engine/Plugins/Animation/IKRig/Source/IKRigEditor/Private/RetargetEditor/IKRetargetBatchOperation.cpp b/Engine/Plugins/Animation/IKRig/Source/IKRigEditor/Private/RetargetEditor/IKRetargetBatchOperation.cpp
--- a/Engine/Plugins/Animation/IKRig/Source/IKRigEditor/Private/RetargetEditor/IKRetargetBatchOperation.cpp
+++ b/Engine/Plugins/Animation/IKRig/Source/IKRigEditor/Private/RetargetEditor/IKRetargetBatchOperation.cpp
@@ -242,7 +242,7 @@ void UIKRetargetBatchOperation::RetargetAssets(
 			
 			// set the retarget source to the target skeletal mesh
 			AnimSequenceToRetarget->RetargetSource = NAME_None;
-			AnimSequenceToRetarget->RetargetSourceAsset = Context.TargetMesh;
+			AnimSequenceToRetarget->SetRetargetSourceAsset(Context.TargetMesh);
 			Controller.UpdateWithSkeleton(NewSkeleton, bShouldTransact);
 
 			// done editing sequence data, close bracket
@@ -258,8 +258,6 @@ void UIKRetargetBatchOperation::RetargetAssets(
 	// Call PostEditChange after the references of all assets were replaced, to prevent order dependence of post edit
 	// change hooks. If PostEditChange is called right after ReplaceReferredAnimations it can access references that are
 	// still queued for retarget and follow the current asset in the array.
-	static const FName RetargetSourceAssetPropertyName =  GET_MEMBER_NAME_STRING_CHECKED(UAnimSequence, RetargetSourceAsset);
-	static FProperty* RetargetAssetProperty = UAnimSequence::StaticClass()->FindPropertyByName(RetargetSourceAssetPropertyName);
 	for (UAnimationAsset* AssetToRetarget : AnimationAssetsToRetarget)
 	{
 		if (Progress.ShouldCancel())
@@ -270,8 +268,7 @@ void UIKRetargetBatchOperation::RetargetAssets(
 		// force updating of the retarget pose, this is normally done on PreSave() but is guarded against procedural saves
 		if (UAnimSequence* AnimSequenceToRetarget = Cast<UAnimSequence>(AssetToRetarget))
 		{
-			FPropertyChangedEvent RetargetAssetPropertyChangedEvent(RetargetAssetProperty);
-			AnimSequenceToRetarget->PostEditChangeProperty(RetargetAssetPropertyChangedEvent);
+			AnimSequenceToRetarget->UpdateRetargetSourceAssetData();
 		}
 		
 		AssetToRetarget->PostEditChange();
diff --git a/Engine/Source/Editor/DetailCustomizations/Private/AnimSequenceDetails.cpp b/Engine/Source/Editor/DetailCustomizations/Private/AnimSequenceDetails.cpp
--- a/Engine/Source/Editor/DetailCustomizations/Private/AnimSequenceDetails.cpp
+++ b/Engine/Source/Editor/DetailCustomizations/Private/AnimSequenceDetails.cpp
@@ -82,7 +82,9 @@ void FAnimSequenceDetails::CustomizeDetails(IDetailLayoutBuilder& DetailBuilder)
 
 	// *** Retarget source handler ***
 	RetargetSourceNameHandler = DetailBuilder.GetProperty(GET_MEMBER_NAME_CHECKED(UAnimSequence, RetargetSource));
+	PRAGMA_DISABLE_DEPRECATION_WARNINGS
 	RetargetSourceAssetHandle = DetailBuilder.GetProperty(GET_MEMBER_NAME_CHECKED(UAnimSequence, RetargetSourceAsset));
+	PRAGMA_ENABLE_DEPRECATION_WARNINGS
 
 	// first create profile combo list
 	RetargetSourceComboList.Empty();
@@ -468,10 +470,12 @@ EVisibility FAnimSequenceDetails::UpdateRetargetSourceAssetDataVisibility() cons
 	{
 		if (UAnimSequence* AnimSequence = WeakAnimSequence.Get())
 		{
+			PRAGMA_DISABLE_DEPRECATION_WARNINGS
 			if (!AnimSequence->RetargetSourceAsset.IsNull())
 			{
 				return EVisibility::Visible;
 			}
+			PRAGMA_ENABLE_DEPRECATION_WARNINGS
 		}
 	}
 
diff --git a/Engine/Source/Editor/DetailCustomizations/Private/PoseAssetDetails.cpp b/Engine/Source/Editor/DetailCustomizations/Private/PoseAssetDetails.cpp
--- a/Engine/Source/Editor/DetailCustomizations/Private/PoseAssetDetails.cpp
+++ b/Engine/Source/Editor/DetailCustomizations/Private/PoseAssetDetails.cpp
@@ -80,7 +80,9 @@ void FPoseAssetDetails::CustomizeDetails(IDetailLayoutBuilder& DetailBuilder)
 	/////////////////////////////////////////////////////////////////////////////////
 	IDetailCategoryBuilder& AnimationCategory = DetailBuilder.EditCategory("Animation");
 	RetargetSourceNameHandler = DetailBuilder.GetProperty(GET_MEMBER_NAME_CHECKED(UPoseAsset, RetargetSource));
+	PRAGMA_DISABLE_DEPRECATION_WARNINGS
 	RetargetSourceAssetHandle = DetailBuilder.GetProperty(GET_MEMBER_NAME_CHECKED(UAnimSequence, RetargetSourceAsset));
+	PRAGMA_ENABLE_DEPRECATION_WARNINGS
 
 	// first create profile combo list
 	RetargetSourceComboList.Empty();
@@ -576,10 +578,12 @@ void FPoseAssetDetails::CachePoseAssetData()
  {
 	 if (UPoseAsset* Pose = PoseAsset.Get())
 	 {
+		 PRAGMA_DISABLE_DEPRECATION_WARNINGS
 		 if (!Pose->RetargetSourceAsset.IsNull())
 		 {
 			 return EVisibility::Visible;
 		 }
+		 PRAGMA_ENABLE_DEPRECATION_WARNINGS
 	 }
 
 	 return EVisibility::Collapsed;
diff --git a/Engine/Source/Editor/SequenceRecorder/Private/AnimationRecorder.cpp b/Engine/Source/Editor/SequenceRecorder/Private/AnimationRecorder.cpp
--- a/Engine/Source/Editor/SequenceRecorder/Private/AnimationRecorder.cpp
+++ b/Engine/Source/Editor/SequenceRecorder/Private/AnimationRecorder.cpp
@@ -299,13 +299,8 @@ void FAnimationRecorder::StartRecord(USkeletalMeshComponent* Component, UAnimSeq
 	AnimationObject->RetargetSource = Component->GetSkeletalMeshAsset() ? AnimSkeleton->GetRetargetSourceForMesh(Component->GetSkeletalMeshAsset()) : NAME_None;
 	if (AnimationObject->RetargetSource == NAME_None)
 	{
-		AnimationObject->RetargetSourceAsset = Component->GetSkeletalMeshAsset();
-		//UpdateRetargetSourceAssetData() is protected so need to do a posteditchagned
-#if WITH_EDITOR
-		FProperty* PropertyChanged = UAnimSequence::StaticClass()->FindPropertyByName(GET_MEMBER_NAME_CHECKED(UAnimSequence, RetargetSourceAsset));
-		FPropertyChangedEvent PropertyUpdateStruct(PropertyChanged);
-		AnimationObject->PostEditChangeProperty(PropertyUpdateStruct);
-#endif
+		AnimationObject->SetRetargetSourceAsset(Component->GetSkeletalMeshAsset());
+		AnimationObject->UpdateRetargetSourceAssetData();
 	}
 
 	// record first-frame notifies
diff --git a/Engine/Source/Runtime/Engine/Classes/Animation/AnimSequence.h b/Engine/Source/Runtime/Engine/Classes/Animation/AnimSequence.h
--- a/Engine/Source/Runtime/Engine/Classes/Animation/AnimSequence.h
+++ b/Engine/Source/Runtime/Engine/Classes/Animation/AnimSequence.h
@@ -325,6 +325,7 @@ public:
 
 #if WITH_EDITORONLY_DATA
 	/** If RetargetSource is set to Default (None), this is asset for the base pose to use when retargeting. Transform data will be saved in RetargetSourceAssetReferencePose. */
+	UE_DEPRECATED(5.5, "Direct access to RetargetSourceAsset has been deprecated. Please use member SetRetargetSourceAssetData instead.")
 	UPROPERTY(EditAnywhere, AssetRegistrySearchable, Category=Animation, meta = (DisallowedClasses = "/Script/ApexDestruction.DestructibleMesh"))
 	TSoftObjectPtr<USkeletalMesh> RetargetSourceAsset;
 #endif
@@ -491,9 +492,13 @@ protected:
 	
 	UE_DEPRECATED(5.3, "Please use UpdateCompressedCurveName that takes FNames.")
 	void UpdateCompressedCurveName(SmartName::UID_Type CurveUID, const struct FSmartName& NewCurveName) {}
-private:
-	ENGINE_API void UpdateRetargetSourceAsset();
 
+public:
+	ENGINE_API void SetRetargetSourceAsset(USkeletalMesh* InRetargetSourceAsset);
+	ENGINE_API USkeletalMesh* GetRetargetSourceAsset() const;
+	ENGINE_API void UpdateRetargetSourceAssetData();
+
+private:
 	/** Updates the stored sampling frame-rate using the sequence length and number of sampling keys */
 	UE_DEPRECATED(5.0, "UpdateFrameRate has been deprecated see UAnimDataController::SetFrameRate")
 	ENGINE_API void UpdateFrameRate();
diff --git a/Engine/Source/Runtime/Engine/Classes/Animation/PoseAsset.h b/Engine/Source/Runtime/Engine/Classes/Animation/PoseAsset.h
--- a/Engine/Source/Runtime/Engine/Classes/Animation/PoseAsset.h
+++ b/Engine/Source/Runtime/Engine/Classes/Animation/PoseAsset.h
@@ -205,6 +205,7 @@ public:
 
 #if WITH_EDITORONLY_DATA
 	/** If RetargetSource is set to Default (None), this is asset for the base pose to use when retargeting. Transform data will be saved in RetargetSourceAssetReferencePose. */
+	UE_DEPRECATED(5.5, "Direct access to RetargetSourceAsset has been deprecated. Please use member SetRetargetSourceAsset instead.")
 	UPROPERTY(EditAnywhere, AssetRegistrySearchable, Category=Animation, meta = (DisallowedClasses = "/Script/ApexDestruction.DestructibleMesh"))
 	TSoftObjectPtr<USkeletalMesh> RetargetSourceAsset;
 #endif
@@ -421,9 +422,14 @@ private:
 	void UpdateTrackBoneIndices();
 	bool RemoveInvalidTracks();
 
+public:
 #if WITH_EDITORONLY_DATA
-	void UpdateRetargetSourceAsset();
+	ENGINE_API void SetRetargetSourceAsset(USkeletalMesh* InRetargetSourceAsset);
+	ENGINE_API USkeletalMesh* GetRetargetSourceAsset() const;
+	ENGINE_API void UpdateRetargetSourceAssetData();
 #endif
+
+private:
 	const TArray<FTransform>& GetRetargetTransforms() const;
 	FName GetRetargetTransformsSourceName() const;
 
diff --git a/Engine/Source/Runtime/Engine/Private/Animation/AnimRetargetHelpers.cpp b/Engine/Source/Runtime/Engine/Private/Animation/AnimRetargetHelpers.cpp
new file mode 100644
--- /dev/null
+++ b/Engine/Source/Runtime/Engine/Private/Animation/AnimRetargetHelpers.cpp
@@ -0,0 +1,183 @@
+// Copyright Epic Games, Inc. All Rights Reserved.
+
+#include "Animation/AnimRetargetHelpers.h"
+
+
+#if WITH_EDITOR
+#include "Engine/SkeletalMesh.h"
+#include "Animation/Skeleton.h"
+#include "AssetRegistry/IAssetRegistry.h"
+#include "Animation/AnimSequence.h"
+#include "Animation/PoseAsset.h"
+#include "AssetRegistry/AssetData.h"
+#include "AssetRegistry/AssetRegistryModule.h"
+#include "Templates/UnrealTemplate.h"
+#endif // WITH_EDITOR
+
+#define LOCTEXT_NAMESPACE "AnimRetargetHelpers"
+
+namespace UE::Anim::RetargetHelpers
+{
+
+#if WITH_EDITOR
+
+namespace Private
+{
+
+static bool GRunningCheckRetargetSourceAssetDataCmd = false;
+
+int32 GEnablePostLoadRetargetSourceErrorReporting = 1;
+static FAutoConsoleVariableRef CVarStripAdditiveRefPose(
+	TEXT("a.EnablePostLoadRetargetSourceErrorReporting"),
+	GEnablePostLoadRetargetSourceErrorReporting,
+	TEXT("1 = Enables validation of retarget source asset data. 0 = off"));
+
+template<typename AssetType>
+ERetargetSourceAssetStatus CheckRetargetSourceAssetDataImpl(const AssetType* InAsset, bool bLogIssues)
+{
+	PRAGMA_DISABLE_DEPRECATION_WARNINGS
+	const TSoftObjectPtr<USkeletalMesh>& RetargetSourceAsset = InAsset->RetargetSourceAsset;
+	PRAGMA_ENABLE_DEPRECATION_WARNINGS
+
+	FAssetData AssetData;
+	UE::AssetRegistry::EExists AssetDataStatus = RetargetSourceAsset.IsNull() 
+		? UE::AssetRegistry::EExists::DoesNotExist
+		: IAssetRegistry::GetChecked().TryGetAssetByObjectPath(RetargetSourceAsset.ToSoftObjectPath(), AssetData);
+
+	if (AssetDataStatus != AssetRegistry::EExists::Unknown)
+	{
+		if (InAsset->RetargetSourceAssetReferencePose.IsEmpty())
+		{
+			return ERetargetSourceAssetStatus::NoRetargetDataSet;
+		}
+		else // InAsset->RetargetSourceAssetReferencePose.Num() > 0
+		{
+			if (AssetDataStatus == UE::AssetRegistry::EExists::Exists)
+			{
+				return ERetargetSourceAssetStatus::RetargetDataOk;
+			}
+			else // if (AssetDataStatus == UE::AssetRegistry::EExists::DoesNotExist)
+			{
+				if (bLogIssues)
+				{
+					UE_LOG(LogAnimation, Warning, TEXT("Asset [%s] references a missing Retarget Source Asset [%s]. Retarget Reference Pose has [%d] elements. Please, add a correct retarget source asset and resave.")
+						, *InAsset->GetFullName()
+						, *(RetargetSourceAsset.GetLongPackageName() + FString(TEXT("/")) + RetargetSourceAsset.GetAssetName())
+						, InAsset->RetargetSourceAssetReferencePose.Num());
+				}
+				return ERetargetSourceAssetStatus::RetargetSourceMissing;
+			}
+		}
+	}
+
+	// if Asset registry is not ready, we go with a slow path
+	if (InAsset->RetargetSourceAssetReferencePose.Num() > 0)
+	{
+		const USkeletalMesh* SourceReferenceMesh = InAsset->GetRetargetSourceAsset();
+		if (SourceReferenceMesh == nullptr)
+		{
+			if (bLogIssues)
+			{
+				UE_LOG(LogAnimation, Warning, TEXT("Asset [%s] references a missing Retarget Source Asset [%s]. Retarget Reference Pose has [%d] elements. Please, add a correct retarget source asset and resave.")
+					, *InAsset->GetFullName()
+					, *(RetargetSourceAsset.GetLongPackageName() + FString(TEXT("/")) + RetargetSourceAsset.GetAssetName())
+					, InAsset->RetargetSourceAssetReferencePose.Num());
+			}
+			return ERetargetSourceAssetStatus::RetargetSourceMissing;
+		}
+		else
+		{
+			return ERetargetSourceAssetStatus::RetargetDataOk;
+		}
+	}
+
+	return ERetargetSourceAssetStatus::NoRetargetDataSet;
+}
+
+void CheckRetargetSourceAssetData(bool bFixAssets)
+{
+	TGuardValue<bool> CompilationGuard(GRunningCheckRetargetSourceAssetDataCmd, true);
+
+	TArray<FAssetData> Assets;
+	IAssetRegistry& AssetRegistry = FModuleManager::LoadModuleChecked<FAssetRegistryModule>(TEXT("AssetRegistry")).Get();
+	AssetRegistry.GetAssetsByClass(UAnimSequence::StaticClass()->GetClassPathName(), Assets);
+	AssetRegistry.GetAssetsByClass(UPoseAsset::StaticClass()->GetClassPathName(), Assets);
+
+	const int32 NumAssets = Assets.Num();
+	for (int32 Idx = 0; Idx < NumAssets; Idx++)
+	{
+		if (UAnimSequence* AnimSequence = Cast<UAnimSequence>(Assets[Idx].GetAsset()))
+		{
+			const ERetargetSourceAssetStatus Status = Private::CheckRetargetSourceAssetDataImpl(AnimSequence, true);
+			if (bFixAssets)
+			{
+				if (Status == ERetargetSourceAssetStatus::RetargetSourceMissing)
+				{
+					AnimSequence->Modify();
+					PRAGMA_DISABLE_DEPRECATION_WARNINGS
+					AnimSequence->RetargetSourceAsset.Reset();
+					PRAGMA_ENABLE_DEPRECATION_WARNINGS
+					AnimSequence->RetargetSourceAssetReferencePose.Empty();
+					AnimSequence->MarkPackageDirty();
+				}
+			}
+		}
+		else if (UPoseAsset* PoseAsset = Cast<UPoseAsset>(Assets[Idx].GetAsset()))
+		{
+			const ERetargetSourceAssetStatus Status = Private::CheckRetargetSourceAssetDataImpl(PoseAsset, true);
+			if (bFixAssets)
+			{
+				if (Status == ERetargetSourceAssetStatus::RetargetSourceMissing)
+				{
+					PoseAsset->Modify();
+					PRAGMA_DISABLE_DEPRECATION_WARNINGS
+					PoseAsset->RetargetSourceAsset.Reset();
+					PRAGMA_ENABLE_DEPRECATION_WARNINGS
+					PoseAsset->RetargetSourceAssetReferencePose.Empty();
+					PoseAsset->MarkPackageDirty();
+				}
+			}
+		}
+	}
+}
+
+static FAutoConsoleCommand CheckRetargetSourceAssetDataCmd(
+	TEXT("a.CheckRetargetSourceAssetData"),
+	TEXT("Checks if Anim Sequences and Pose Assets RetargetSourceAsset is valid. Use: 'a.CheckRetargetSourceAssetData' to check or 'a.CheckRetargetSourceAssetData true' to check and fix the assets."),
+	FConsoleCommandWithArgsDelegate::CreateLambda([](const TArray<FString>& Args)
+		{
+			if (Args.Num() > 0)
+			{
+				CheckRetargetSourceAssetData(Args[0].ToBool());
+			}
+			else
+			{
+				CheckRetargetSourceAssetData(false);
+			}
+		}
+	));
+} // end namespace Private
+
+bool ShouldCheckRetargetSourceAssetData()
+{
+#if WITH_EDITOR
+	return !Private::GRunningCheckRetargetSourceAssetDataCmd;
+#else
+	return false;
+#endif
+}
+
+ERetargetSourceAssetStatus CheckRetargetSourceAssetData(const UAnimSequence* InAsset)
+{
+	return Private::CheckRetargetSourceAssetDataImpl(InAsset, Private::GEnablePostLoadRetargetSourceErrorReporting != 0);
+}
+
+ERetargetSourceAssetStatus CheckRetargetSourceAssetData(const UPoseAsset* InAsset)
+{
+	return Private::CheckRetargetSourceAssetDataImpl(InAsset, Private::GEnablePostLoadRetargetSourceErrorReporting != 0);
+}
+#endif // WITH_EDITOR
+
+} // end namespace UE::Anim::RetargetHelpers
+
+#undef LOCTEXT_NAMESPACE 
diff --git a/Engine/Source/Runtime/Engine/Private/Animation/AnimSequence.cpp b/Engine/Source/Runtime/Engine/Private/Animation/AnimSequence.cpp
--- a/Engine/Source/Runtime/Engine/Private/Animation/AnimSequence.cpp
+++ b/Engine/Source/Runtime/Engine/Private/Animation/AnimSequence.cpp
@@ -58,16 +58,21 @@ LLM_DEFINE_TAG(SequenceData);
 #include "Animation/AnimationCompressionDerivedData.h"
 #include "Animation/AnimData/IAnimationDataModel.h"
 #include "Animation/BuiltInAttributeTypes.h"
+#include "AssetRegistry/AssetData.h"
 #include "AssetRegistry/AssetRegistryModule.h"
 #include "Cooker/CookDeterminismHelper.h"
 #include "ProfilingDebugging/CookStats.h"
 #include "Serialization/CompactBinaryWriter.h"
 #include "Serialization/MemoryHasher.h"
 #include "Misc/DataValidation.h"
+#include "Animation/PoseAsset.h"
 #endif // WITH_EDITOR
 
 #include "Animation/AnimSequenceHelpers.h"
 
+#if WITH_EDITORONLY_DATA
+#include "Animation/AnimRetargetHelpers.h"
+#endif
 
 #include UE_INLINE_GENERATED_CPP_BY_NAME(AnimSequence)
 
@@ -984,7 +989,7 @@ void UAnimSequence::PreSave(FObjectPreSaveContext ObjectSaveContext)
 
 	if (!ObjectSaveContext.IsProceduralSave())
 	{
-		UpdateRetargetSourceAsset();
+		UpdateRetargetSourceAssetData();
 	}
 
 	if (ObjectSaveContext.IsDeterminismDebug())
@@ -1071,6 +1076,16 @@ void UAnimSequence::PostLoad()
 		}
 #endif
 	}
+
+#if WITH_EDITOR
+	if (UE::Anim::RetargetHelpers::ShouldCheckRetargetSourceAssetData())
+	{
+		if (UE::Anim::RetargetHelpers::CheckRetargetSourceAssetData(this) == UE::Anim::RetargetHelpers::ERetargetSourceAssetStatus::RetargetSourceMissing)
+		{
+			UpdateRetargetSourceAssetData();
+		}
+	}
+#endif // WITH_EDITOR
 }
 
 #if WITH_EDITOR
@@ -1169,10 +1184,12 @@ void UAnimSequence::PostEditChangeProperty(FPropertyChangedEvent& PropertyChange
 
 	if(PropertyChangedEvent.Property)
 	{
+		PRAGMA_DISABLE_DEPRECATION_WARNINGS
 		if (PropertyChangedEvent.Property->GetFName() == GET_MEMBER_NAME_CHECKED(UAnimSequence, RetargetSourceAsset))
 		{
-			UpdateRetargetSourceAsset();
+			UpdateRetargetSourceAssetData();
 		}
+		PRAGMA_ENABLE_DEPRECATION_WARNINGS
 
 		const bool bChangedRefFrameIndex = PropertyChangedEvent.Property->GetFName() == GET_MEMBER_NAME_CHECKED(UAnimSequence, RefFrameIndex);
 
@@ -1955,9 +1972,33 @@ void UAnimSequence::UpdateCompressedCurveName(const FName& OldCurveName, const F
 #endif // WITH_EDITOR
 
 #if WITH_EDITORONLY_DATA
-void UAnimSequence::UpdateRetargetSourceAsset()
+void UAnimSequence::SetRetargetSourceAsset(USkeletalMesh* InRetargetSourceAsset)
+{
+	if (InRetargetSourceAsset != nullptr && InRetargetSourceAsset->HasAnyFlags(RF_Transient))
+	{
+		UE_LOG(LogAnimation, Error, TEXT("Error, Transient asset [%s] can not be assigned as Retarget Source for AnimSequence [%s]. Please, use a non transient asset as retarget surce.")
+			, *(InRetargetSourceAsset->GetFullName())
+			, *GetFullName());
+		ensure(false);
+		return;
+	}
+
+	PRAGMA_DISABLE_DEPRECATION_WARNINGS
+	RetargetSourceAsset = InRetargetSourceAsset;
+	PRAGMA_ENABLE_DEPRECATION_WARNINGS
+}
+
+USkeletalMesh* UAnimSequence::GetRetargetSourceAsset() const
 {
+	PRAGMA_DISABLE_DEPRECATION_WARNINGS
 	USkeletalMesh* SourceReferenceMesh = RetargetSourceAsset.LoadSynchronous();
+	PRAGMA_ENABLE_DEPRECATION_WARNINGS
+	return SourceReferenceMesh;
+}
+
+void UAnimSequence::UpdateRetargetSourceAssetData()
+{
+	USkeletalMesh* SourceReferenceMesh = GetRetargetSourceAsset();
 	const USkeleton* MySkeleton = GetSkeleton();
 	if (SourceReferenceMesh && MySkeleton)
 	{
@@ -4973,6 +5014,7 @@ void UAnimSequence::RequestSyncAnimRecompression(bool bOutput)
 }
 #endif // WITH_EDITOR
 
+
 /*-----------------------------------------------------------------------------
 	AnimNotify& subclasses
 -----------------------------------------------------------------------------*/
diff --git a/Engine/Source/Runtime/Engine/Private/Animation/PoseAsset.cpp b/Engine/Source/Runtime/Engine/Private/Animation/PoseAsset.cpp
--- a/Engine/Source/Runtime/Engine/Private/Animation/PoseAsset.cpp
+++ b/Engine/Source/Runtime/Engine/Private/Animation/PoseAsset.cpp
@@ -25,6 +25,10 @@
 #include "UObject/UnrealType.h"
 #include "Widgets/Notifications/SNotificationList.h"
 
+#if WITH_EDITORONLY_DATA
+#include "Animation/AnimRetargetHelpers.h"
+#endif
+
 #include UE_INLINE_GENERATED_CPP_BY_NAME(PoseAsset)
 
 #define LOCTEXT_NAMESPACE "PoseAsset"
@@ -993,6 +997,16 @@ void UPoseAsset::PostLoad()
 	{
 		UpdateTrackBoneIndices();
 	}
+
+#if WITH_EDITOR
+	if (UE::Anim::RetargetHelpers::ShouldCheckRetargetSourceAssetData())
+	{
+		if (UE::Anim::RetargetHelpers::CheckRetargetSourceAssetData(this) == UE::Anim::RetargetHelpers::ERetargetSourceAssetStatus::RetargetSourceMissing)
+		{
+			UpdateRetargetSourceAssetData();
+		}
+	}
+#endif // WITH_EDITOR
 }
 
 void UPoseAsset::Serialize(FArchive& Ar)
@@ -1017,7 +1031,7 @@ void UPoseAsset::PreSave(FObjectPreSaveContext ObjectSaveContext)
 #if WITH_EDITOR
 	if (!ObjectSaveContext.IsProceduralSave())
 	{
-		UpdateRetargetSourceAsset();
+		UpdateRetargetSourceAssetData();
 	}
 #endif // WITH_EDITOR
 	Super::PreSave(ObjectSaveContext);
@@ -1364,11 +1378,13 @@ void UPoseAsset::PostEditChangeProperty(FPropertyChangedEvent& PropertyChangedEv
 	if (PropertyChangedEvent.Property)
 	{
 		bool bConvertToAdditivePose = false;
+		PRAGMA_DISABLE_DEPRECATION_WARNINGS
 		if (PropertyChangedEvent.Property->GetFName() == GET_MEMBER_NAME_CHECKED(UPoseAsset, RetargetSourceAsset))
 		{
 			bConvertToAdditivePose = true;
-			UpdateRetargetSourceAsset();
+			UpdateRetargetSourceAssetData();
 		}
+		PRAGMA_ENABLE_DEPRECATION_WARNINGS
 
 		if (PropertyChangedEvent.Property->GetFName() == GET_MEMBER_NAME_CHECKED(UPoseAsset, RetargetSource))
 		{
@@ -1805,9 +1821,33 @@ bool UPoseAsset::RemoveInvalidTracks()
 }
 
 #if WITH_EDITORONLY_DATA
-void UPoseAsset::UpdateRetargetSourceAsset()
+void UPoseAsset::SetRetargetSourceAsset(USkeletalMesh* InRetargetSourceAsset)
 {
+	if (InRetargetSourceAsset != nullptr && InRetargetSourceAsset->HasAnyFlags(RF_Transient))
+	{
+		UE_LOG(LogAnimation, Error, TEXT("Error, Transient asset [%s] can not be assigned as Retarget Source for Pose Asset [%s]. Please, use a non transient asset as retarget surce.")
+			, *(InRetargetSourceAsset->GetFullName())
+			, *GetFullName());
+		ensure(false);
+		return;
+	}
+
+	PRAGMA_DISABLE_DEPRECATION_WARNINGS
+	RetargetSourceAsset = InRetargetSourceAsset;
+	PRAGMA_ENABLE_DEPRECATION_WARNINGS
+}
+
+USkeletalMesh* UPoseAsset::GetRetargetSourceAsset() const
+{
+	PRAGMA_DISABLE_DEPRECATION_WARNINGS
 	USkeletalMesh* SourceReferenceMesh = RetargetSourceAsset.LoadSynchronous();
+	PRAGMA_ENABLE_DEPRECATION_WARNINGS
+	return SourceReferenceMesh;
+}
+
+void UPoseAsset::UpdateRetargetSourceAssetData()
+{
+	USkeletalMesh* SourceReferenceMesh = GetRetargetSourceAsset();
 	const USkeleton* MySkeleton = GetSkeleton();
 	if (SourceReferenceMesh && MySkeleton)
 	{
@@ -1818,6 +1858,7 @@ void UPoseAsset::UpdateRetargetSourceAsset()
 		RetargetSourceAssetReferencePose.Empty();
 	}
 }
+
 #endif // WITH_EDITORONLY_DATA
 
 const TArray<FTransform>& UPoseAsset::GetRetargetTransforms() const
diff --git a/Engine/Source/Runtime/Engine/Public/Animation/AnimRetargetHelpers.h b/Engine/Source/Runtime/Engine/Public/Animation/AnimRetargetHelpers.h
new file mode 100644
--- /dev/null
+++ b/Engine/Source/Runtime/Engine/Public/Animation/AnimRetargetHelpers.h
@@ -0,0 +1,30 @@
+// Copyright Epic Games, Inc. All Rights Reserved.
+
+#pragma once
+
+#include "CoreMinimal.h"
+
+class UAnimSequence;
+class UPoseAsset;
+
+namespace UE::Anim::RetargetHelpers
+{
+
+#if WITH_EDITOR
+
+enum class ERetargetSourceAssetStatus
+{
+	NoRetargetDataSet = 0,
+	RetargetSourceMissing,
+	RetargetDataOk,
+};
+
+ERetargetSourceAssetStatus CheckRetargetSourceAssetData(const UAnimSequence* InAsset);
+
+ERetargetSourceAssetStatus CheckRetargetSourceAssetData(const UPoseAsset* InAsset);
+
+bool ShouldCheckRetargetSourceAssetData();
+
+#endif // WITH_EDITOR
+
+} // end namespace UE::Anim::RetargetHelpers
