Fix bone names being drawn for bones when they're not visible

diff --git a/Engine/Source/Editor/Persona/Private/AnimationEditorViewportClient.cpp b/Engine/Source/Editor/Persona/Private/AnimationEditorViewportClient.cpp
--- a/Engine/Source/Editor/Persona/Private/AnimationEditorViewportClient.cpp
+++ b/Engine/Source/Editor/Persona/Private/AnimationEditorViewportClient.cpp
@@ -282,6 +282,8 @@ FAnimationViewportClient::~FAnimationViewportClient()
 				}
 			}
 		}
+
+		AnimationEditorPreviewScene->UnregisterOnSelectedBoneChanged(OnSelectedBoneChangedHandle);
 	}
 	OnPhysicsCreatedDelegateHandle.Reset();
 	OnMeshChangedDelegateHandle.Reset();
@@ -303,12 +305,20 @@ void FAnimationViewportClient::Initialize()
 		AnimationEditorPreviewScene->RegisterOnFocusViews(FSimpleDelegate::CreateSP(this, &FAnimationViewportClient::HandleFocusViews));
 		AnimationEditorPreviewScene->RegisterOnPreTick(FSimpleDelegate::CreateSP(this, &FAnimationViewportClient::HandlePreviewScenePreTick));
 		AnimationEditorPreviewScene->RegisterOnPostTick(FSimpleDelegate::CreateSP(this, &FAnimationViewportClient::HandlePreviewScenePostTick));
+
+		OnSelectedBoneChangedHandle = AnimationEditorPreviewScene->RegisterOnSelectedBoneChanged(FOnSelectedBoneChanged::CreateLambda([this](const FName&, ESelectInfo::Type)
+		{
+			UpdateBonesToDraw();
+		}));
 	}
 
 	// Register delegate to update the show flags when the post processing is turned on or off
 	UAssetViewerSettings::Get()->OnAssetViewerSettingsChanged().AddSP(this, &FAnimationViewportClient::OnAssetViewerSettingsChanged);
 	// Set correct flags according to current profile settings
 	UAssetViewerSettings::GetCurrentUserProjectProfile().SetShowFlags(EngineShowFlags);
+
+	// Setup bones to draw on initialise
+	UpdateBonesToDraw();
 }
 
 void FAnimationViewportClient::OnToggleAutoAlignFloor()
@@ -596,6 +606,8 @@ void FAnimationViewportClient::HandleSkeletalMeshChanged(USkeletalMesh* OldSkele
 		}
 	}
 
+	UpdateBonesToDraw();
+
 	Invalidate();
 }
 
@@ -941,6 +953,11 @@ void FAnimationViewportClient::ShowBoneNames( FCanvas* Canvas, FSceneView* View,
 	{
 		const int32 BoneIndex = LODData.RequiredBones[i];
 
+		if (!BonesToDraw[BoneIndex])
+		{
+			continue;
+		}
+
 		// If previewing a specific section, only show the bone names that belong to it
 		if ((PreviewMeshComponent->GetSectionPreview() >= 0) && !LODData.RenderSections[PreviewMeshComponent->GetSectionPreview()].BoneMap.Contains(BoneIndex))
 		{
@@ -1629,6 +1646,8 @@ void FAnimationViewportClient::SetBoneDrawMode(EBoneDrawMode::Type AxesMode)
 {
 	ConfigOption->SetDefaultBoneDrawSelection(AxesMode);
 	RedrawRequested(Viewport);
+
+	UpdateBonesToDraw();
 }
 
 bool FAnimationViewportClient::IsBoneDrawModeSet(EBoneDrawMode::Type AxesMode) const
@@ -1864,6 +1883,27 @@ void FAnimationViewportClient::DrawBonesFromSkeleton(UDebugSkelMeshComponent * M
 		bUseMultiColor);
 }
 
+void FAnimationViewportClient::UpdateBonesToDraw()
+{
+	if (UDebugSkelMeshComponent* MeshComponent = GetAnimPreviewScene()->GetPreviewMeshComponent())
+	{
+		const FReferenceSkeleton& RefSkeleton = MeshComponent->GetReferenceSkeleton();
+
+		TArray<int32> ParentIndices;
+		ParentIndices.AddUninitialized(RefSkeleton.GetNum());
+		for (int32 BoneIndex = 0; BoneIndex < RefSkeleton.GetNum(); ++BoneIndex)
+		{
+			ParentIndices[BoneIndex] = RefSkeleton.GetParentIndex(BoneIndex);
+		}
+
+		SkeletalDebugRendering::CalculateBonesToDraw(
+			ParentIndices,
+			MeshComponent->BonesOfInterest,
+			GetBoneDrawMode(),
+			BonesToDraw);
+	}
+}
+
 void FAnimationViewportClient::DrawMeshBones(UDebugSkelMeshComponent* MeshComponent, FPrimitiveDrawInterface* PDI) const
 {
 	if (!MeshComponent ||
@@ -1960,7 +2000,8 @@ void FAnimationViewportClient::DrawBones(
 		InSelectedBones,
 		BoneColors,
 		HitProxies,
-		DrawConfig
+		DrawConfig,
+		BonesToDraw
 	);
 }
 
diff --git a/Engine/Source/Editor/Persona/Public/AnimationEditorViewportClient.h b/Engine/Source/Editor/Persona/Public/AnimationEditorViewportClient.h
--- a/Engine/Source/Editor/Persona/Public/AnimationEditorViewportClient.h
+++ b/Engine/Source/Editor/Persona/Public/AnimationEditorViewportClient.h
@@ -106,6 +106,9 @@ public:
 	UE_DEPRECATED(5.1, "This version of InputAxis is deprecated. Please use the version that takes DeviceId instead.")
 	virtual bool InputAxis(FViewport* InViewport, int32 ControllerId, FKey Key, float Delta, float DeltaTime, int32 NumSamples = 1, bool bGamepad = false) override;
 	virtual bool InputAxis(FViewport* InViewport, FInputDeviceId DeviceId, FKey Key, float Delta, float DeltaTime, int32 NumSamples = 1, bool bGamepad = false) override;
+
+	// Sets what bones are drawn by DrawMeshBones and ShowBoneNames
+	virtual void UpdateBonesToDraw();
 	
 //	virtual void ProcessClick(class FSceneView& View, class HHitProxy* HitProxy, FKey Key, EInputEvent Event, uint32 HitX, uint32 HitY) override;
 //	virtual bool InputWidgetDelta( FViewport* Viewport, EAxisList::Type CurrentAxis, FVector& Drag, FRotator& Rot, FVector& Scale ) override;
@@ -404,6 +407,11 @@ private:
 	/* Member use to unregister OnMeshChanged for our preview skeletal mesh */
 	FDelegateHandle OnMeshChangedDelegateHandle;
 
+	/* Bit field indexed on bone index that stores what bones are visible in the viewport, updated with UpdateBonesToDraw */
+	TBitArray<> BonesToDraw;
+
+	FDelegateHandle OnSelectedBoneChangedHandle;
+
 private:
 
 	void SetCameraTargetLocation(const FSphere &BoundSphere, float DeltaSeconds);
diff --git a/Engine/Source/Runtime/Engine/Private/SkeletalDebugRendering.cpp b/Engine/Source/Runtime/Engine/Private/SkeletalDebugRendering.cpp
--- a/Engine/Source/Runtime/Engine/Private/SkeletalDebugRendering.cpp
+++ b/Engine/Source/Runtime/Engine/Private/SkeletalDebugRendering.cpp
@@ -255,7 +255,8 @@ void DrawBonesFromPoseWatch(
 		/*SelectedBones*/TArray<int32>(),
 		/*BoneColors*/TArray<FLinearColor>(),
 		/*HitProxies*/TArray<TRefCountPtr<HHitProxy>>(),
-		DrawConfig);
+		DrawConfig,
+		TBitArray<>{});
 }
 #endif
 
@@ -269,6 +270,32 @@ void DrawBones(
 	const TArray<FLinearColor>& BoneColors,
 	const TArray<TRefCountPtr<HHitProxy>>& HitProxies,
 	const FSkelDebugDrawConfig& DrawConfig)
+{
+	DrawBones(
+		PDI,
+		ComponentOrigin,
+		RequiredBones,
+		RefSkeleton,
+		WorldTransforms,
+		InSelectedBones,
+		BoneColors,
+		HitProxies,
+		DrawConfig,
+		TBitArray<>{});
+}
+
+void DrawBones(
+	FPrimitiveDrawInterface* PDI,
+	const FVector& ComponentOrigin,
+	const TArray<FBoneIndexType>& RequiredBones,
+	const FReferenceSkeleton& RefSkeleton,
+	const TArray<FTransform>& WorldTransforms,
+	const TArray<int32>& InSelectedBones,
+	const TArray<FLinearColor>& BoneColors,
+	const TArray<TRefCountPtr<HHitProxy>>& HitProxies,
+	const FSkelDebugDrawConfig& DrawConfig,
+	// Overrides the bones that're drawn
+	const TBitArray<>& BonesToDrawOverride)
 {
 	// get parent indices of bones
 	TArray<int32> ParentIndices;
@@ -287,9 +314,11 @@ void DrawBones(
 		InSelectedBones,
 		BoneColors,
 		HitProxies,
-		DrawConfig);
+		DrawConfig,
+		BonesToDrawOverride);
 }
 
+
 void DrawBonesInternal(
 	FPrimitiveDrawInterface* PDI,
 	const FVector& ComponentOrigin,
@@ -299,7 +328,8 @@ void DrawBonesInternal(
 	const TArray<int32>& InSelectedBones,
 	const TArray<FLinearColor>& BoneColors,
 	const TArray<TRefCountPtr<HHitProxy>>& HitProxies,
-	const FSkelDebugDrawConfig& DrawConfig)
+	const FSkelDebugDrawConfig& DrawConfig,
+	const TBitArray<>& BonesToDrawOverride)
 {
 	const auto GetParentIndex = [ParentIndices](const int32 InBoneIndex) -> int32
 	{
@@ -312,57 +342,16 @@ void DrawBonesInternal(
 
 	// first determine which bones to draw, and which to filter out
 	const int32 NumBones = ParentIndices.Num();
-	TBitArray<> BonesToDraw(false, NumBones);
 	const bool bDrawAll = DrawConfig.BoneDrawMode == EBoneDrawMode::All;
 	const bool bDrawSelected = DrawConfig.BoneDrawMode == EBoneDrawMode::Selected;
 	const bool bDrawSelectedAndParents = DrawConfig.BoneDrawMode == EBoneDrawMode::SelectedAndParents;
 	const bool bDrawSelectedAndChildren = DrawConfig.BoneDrawMode == EBoneDrawMode::SelectedAndChildren;
 	const bool bDrawSelectedAndParentsAndChildren = DrawConfig.BoneDrawMode == EBoneDrawMode::SelectedAndParentsAndChildren;
 
-	// draw all bones
-	if (bDrawAll)
-	{
-		BonesToDraw.Init(true, NumBones);
-	}
-
-	// add selected bones
-	if (bDrawSelected || bDrawSelectedAndParents || bDrawSelectedAndChildren || bDrawSelectedAndParentsAndChildren)
-	{
-		for (int32 BoneIndex : InSelectedBones)
-		{
-			if (BoneIndex != INDEX_NONE && BonesToDraw.IsValidIndex(BoneIndex))
-			{
-				BonesToDraw[BoneIndex] = true;
-			}
-		}
-	}
-
-	// add children of selected
-	if (bDrawSelectedAndChildren || bDrawSelectedAndParentsAndChildren)
+	TBitArray<> BonesToDraw = BonesToDrawOverride;
+	if (BonesToDraw.IsEmpty())
 	{
-		for (int32 BoneIndex = 0; BoneIndex < NumBones; ++BoneIndex)
-		{
-			const int32 ParentIndex = GetParentIndex(BoneIndex);
-			if (ParentIndex != INDEX_NONE && BonesToDraw[ParentIndex])
-			{
-				BonesToDraw[BoneIndex] = true;
-			}
-		}
-	}
-
-	// add parents of selected
-	if (bDrawSelectedAndParents || bDrawSelectedAndParentsAndChildren)
-	{
-		for (const int32 BoneIndex : InSelectedBones)
-		{
-			if (BoneIndex != INDEX_NONE)
-			{
-				for (int32 ParentIndex = GetParentIndex(BoneIndex); ParentIndex != INDEX_NONE; ParentIndex = GetParentIndex(ParentIndex))
-				{
-					BonesToDraw[ParentIndex] = true;
-				}
-			}
-		}
+		CalculateBonesToDraw(ParentIndices, InSelectedBones, DrawConfig.BoneDrawMode, BonesToDraw);
 	}
 
 	// determine which bones are "affected" (these are ALL children of selected bones)
@@ -500,4 +489,76 @@ void FillWithMultiColors(TArray<FLinearColor>& BoneColors, const int32 NumBones)
 	}
 }
 
+
+void CalculateBonesToDraw(
+	const TArray<int32>& ParentIndices,
+	const TArray<int32>& InSelectedBones,
+	const EBoneDrawMode::Type BoneDrawMode,
+	TBitArray<>& OutBonesToDraw)
+{
+	const auto GetParentIndex = [ParentIndices](const int32 InBoneIndex) -> int32
+	{
+		if (ParentIndices.IsValidIndex(InBoneIndex))
+		{
+			return ParentIndices[InBoneIndex];
+		}
+		return INDEX_NONE;
+	};
+
+	const int32 NumBones = ParentIndices.Num();
+	OutBonesToDraw.Init(false, NumBones);
+
+	const bool bDrawAll = BoneDrawMode == EBoneDrawMode::All;
+	const bool bDrawSelected = BoneDrawMode == EBoneDrawMode::Selected;
+	const bool bDrawSelectedAndParents = BoneDrawMode == EBoneDrawMode::SelectedAndParents;
+	const bool bDrawSelectedAndChildren = BoneDrawMode == EBoneDrawMode::SelectedAndChildren;
+	const bool bDrawSelectedAndParentsAndChildren = BoneDrawMode == EBoneDrawMode::SelectedAndParentsAndChildren;
+
+	// draw all bones
+	if (bDrawAll)
+	{
+		OutBonesToDraw.Init(true, NumBones);
+	}
+
+	// add selected bones
+	if (bDrawSelected || bDrawSelectedAndParents || bDrawSelectedAndChildren || bDrawSelectedAndParentsAndChildren)
+	{
+		for (int32 BoneIndex : InSelectedBones)
+		{
+			if (BoneIndex != INDEX_NONE && OutBonesToDraw.IsValidIndex(BoneIndex))
+			{
+				OutBonesToDraw[BoneIndex] = true;
+			}
+		}
+	}
+
+	// add children of selected
+	if (bDrawSelectedAndChildren || bDrawSelectedAndParentsAndChildren)
+	{
+		for (int32 BoneIndex = 0; BoneIndex < NumBones; ++BoneIndex)
+		{
+			const int32 ParentIndex = GetParentIndex(BoneIndex);
+			if (ParentIndex != INDEX_NONE && OutBonesToDraw[ParentIndex])
+			{
+				OutBonesToDraw[BoneIndex] = true;
+			}
+		}
+	}
+
+	// add parents of selected
+	if (bDrawSelectedAndParents || bDrawSelectedAndParentsAndChildren)
+	{
+		for (const int32 BoneIndex : InSelectedBones)
+		{
+			if (BoneIndex != INDEX_NONE)
+			{
+				for (int32 ParentIndex = GetParentIndex(BoneIndex); ParentIndex != INDEX_NONE; ParentIndex = GetParentIndex(ParentIndex))
+				{
+					OutBonesToDraw[ParentIndex] = true;
+				}
+			}
+		}
+	}
+}
+
 }
diff --git a/Engine/Source/Runtime/Engine/Public/SkeletalDebugRendering.h b/Engine/Source/Runtime/Engine/Public/SkeletalDebugRendering.h
--- a/Engine/Source/Runtime/Engine/Public/SkeletalDebugRendering.h
+++ b/Engine/Source/Runtime/Engine/Public/SkeletalDebugRendering.h
@@ -141,6 +141,19 @@ ENGINE_API void DrawBones(
 	const TArray<TRefCountPtr<HHitProxy>>& HitProxies,
 	const FSkelDebugDrawConfig& DrawConfig);
 
+/* Alternative signature that allows overriding the drawn bones that're computed from EBoneDrawMode */
+ENGINE_API void DrawBones(
+	FPrimitiveDrawInterface* PDI,
+	const FVector& ComponentOrigin,
+	const TArray<FBoneIndexType>& RequiredBones,
+	const FReferenceSkeleton& RefSkeleton,
+	const TArray<FTransform>& WorldTransforms,
+	const TArray<int32>& InSelectedBones,
+	const TArray<FLinearColor>& BoneColors,
+	const TArray<TRefCountPtr<HHitProxy>>& HitProxies,
+	const FSkelDebugDrawConfig& DrawConfig,
+	const TBitArray<>& BonesToDrawOverride);
+
 
 void DrawBonesInternal(
 	FPrimitiveDrawInterface* PDI,
@@ -151,8 +164,15 @@ void DrawBonesInternal(
 	const TArray<int32>& InSelectedBones,
 	const TArray<FLinearColor>& BoneColors,
 	const TArray<TRefCountPtr<HHitProxy>>& HitProxies,
-	const FSkelDebugDrawConfig& DrawConfig);
+	const FSkelDebugDrawConfig& DrawConfig,
+	const TBitArray<>& BonesToDrawOverride);
 
 ENGINE_API FLinearColor GetSemiRandomColorForBone(const int32 BoneIndex, float Value=1.0f, float Saturation=0.75f);
 ENGINE_API void FillWithMultiColors(TArray<FLinearColor>& BoneColors, const int32 NumBones);
+
+ENGINE_API void CalculateBonesToDraw(
+	const TArray<int32>& ParentIndices,
+	const TArray<int32>& InSelectedBones,
+	const EBoneDrawMode::Type BoneDrawMode,
+	TBitArray<>& OutBonesToDraw);
 }
