Fix for a mistake in implementation of DevOptions.Shaders.bUseShaderCompilerFromLaunchDir, set OutFilename when returning true

#rnx

diff --git a/Engine/Source/Runtime/Engine/Private/ShaderCompiler/ShaderCompiler.cpp b/Engine/Source/Runtime/Engine/Private/ShaderCompiler/ShaderCompiler.cpp
--- a/Engine/Source/Runtime/Engine/Private/ShaderCompiler/ShaderCompiler.cpp
+++ b/Engine/Source/Runtime/Engine/Private/ShaderCompiler/ShaderCompiler.cpp
@@ -716,6 +716,7 @@ static bool FindShaderCompileWorkerExecutableInLaunchDir(const FString& Executab
 		}
 	}
 
+	OutFilename = LocalShaderCompileWorkerName;
 	return true;
 }
 
