Make sure required bones are up to date when ticking animation

diff --git a/Engine/Source/Runtime/Engine/Private/Components/SkeletalMeshComponent.cpp b/Engine/Source/Runtime/Engine/Private/Components/SkeletalMeshComponent.cpp
--- a/Engine/Source/Runtime/Engine/Private/Components/SkeletalMeshComponent.cpp
+++ b/Engine/Source/Runtime/Engine/Private/Components/SkeletalMeshComponent.cpp
@@ -1437,8 +1437,14 @@ void USkeletalMeshComponent::TickAnimation(float DeltaTime, bool bNeedsValidRoot
 	SCOPE_CYCLE_COUNTER(STAT_AnimGameThreadTime);
 	SCOPE_CYCLE_COUNTER(STAT_AnimTickTime);
 
-	// if curves have to be refreshed before updating animation
-	if (!AreRequiredCurvesUpToDate())
+	// Recalculate the RequiredBones array, if necessary
+	if (!bRequiredBonesUpToDate)
+	{
+		QUICK_SCOPE_CYCLE_COUNTER(STAT_USkeletalMeshComponent_RefreshBoneTransforms_RecalcRequiredBones);
+		RecalcRequiredBones(GetPredictedLODLevel());
+	}
+	// if curves have to be refreshed
+	else if (!AreRequiredCurvesUpToDate())
 	{
 		QUICK_SCOPE_CYCLE_COUNTER(STAT_USkeletalMeshComponent_RefreshBoneTransforms_RecalcRequiredCurves);
 		RecalcRequiredCurves();
