[Shaders] Fix compile issue in non-editor build for new GetOverrideJobPriority() function. Follow up on CL 36644883.

#rnx

diff --git a/Engine/Source/Runtime/Engine/Private/ShaderCompiler/ShaderCompiler.cpp b/Engine/Source/Runtime/Engine/Private/ShaderCompiler/ShaderCompiler.cpp
--- a/Engine/Source/Runtime/Engine/Private/ShaderCompiler/ShaderCompiler.cpp
+++ b/Engine/Source/Runtime/Engine/Private/ShaderCompiler/ShaderCompiler.cpp
@@ -1437,7 +1437,7 @@ bool FShaderCompilingManager::IsCompilingShaderMap(uint32 Id)
 	return false;
 }
 
-FShaderCompileJob* FShaderCompilingManager::PrepareShaderCompileJob(uint32 Id, const FShaderCompileJobKey& Key, EShaderCompileJobPriority InPriority)
+FShaderCompileJob* FShaderCompilingManager::PrepareShaderCompileJob(uint32 Id, const FShaderCompileJobKey& Key, EShaderCompileJobPriority Priority)
 {
 	// don't allow any jobs if not allowed
 	if (!AllowShaderCompiling())
@@ -1445,15 +1445,17 @@ FShaderCompileJob* FShaderCompilingManager::PrepareShaderCompileJob(uint32 Id, c
 		return nullptr;
 	}
 
+#if WITH_EDITOR
 	// Check if shader type overrides job priority
 	check(Key.ShaderType != nullptr);
 	const EShaderCompileJobPriority OverrideJobPriority = Key.ShaderType->GetOverrideJobPriority();
-	const EShaderCompileJobPriority JobPriority = (OverrideJobPriority < EShaderCompileJobPriority::Num ? FMath::Max(OverrideJobPriority, InPriority) : InPriority);
+	Priority = (OverrideJobPriority < EShaderCompileJobPriority::Num ? FMath::Max(OverrideJobPriority, Priority) : Priority);
+#endif
 
-	return AllJobs.PrepareJob(Id, Key, JobPriority);
+	return AllJobs.PrepareJob(Id, Key, Priority);
 }
 
-FShaderPipelineCompileJob* FShaderCompilingManager::PreparePipelineCompileJob(uint32 Id, const FShaderPipelineCompileJobKey& Key, EShaderCompileJobPriority InPriority)
+FShaderPipelineCompileJob* FShaderCompilingManager::PreparePipelineCompileJob(uint32 Id, const FShaderPipelineCompileJobKey& Key, EShaderCompileJobPriority Priority)
 {
 	// don't allow any jobs if not allowed
 	if (!AllowShaderCompiling())
@@ -1461,16 +1463,17 @@ FShaderPipelineCompileJob* FShaderCompilingManager::PreparePipelineCompileJob(ui
 		return nullptr;
 	}
 
+#if WITH_EDITOR
 	// Check if shader types in pipeline override job priority and pick highest one
 	check(Key.ShaderPipeline != nullptr);
-	EShaderCompileJobPriority JobPriority = InPriority;
 	for (const FShaderType* ShaderType : Key.ShaderPipeline->GetStages())
 	{
 		const EShaderCompileJobPriority OverrideJobPriority = ShaderType->GetOverrideJobPriority();
-		JobPriority = (OverrideJobPriority < EShaderCompileJobPriority::Num ? FMath::Max(OverrideJobPriority, JobPriority) : JobPriority);
+		Priority = (OverrideJobPriority < EShaderCompileJobPriority::Num ? FMath::Max(OverrideJobPriority, Priority) : Priority);
 	}
+#endif
 
-	return AllJobs.PrepareJob(Id, Key, JobPriority);
+	return AllJobs.PrepareJob(Id, Key, Priority);
 }
 
 void FShaderCompilingManager::ProcessFinishedJob(FShaderCommonCompileJob* FinishedJob)
