Restore fix for frame sync issues that was backed out in 36633942, but without the changes to FPendingCleanupObjects
 - Only the additional RHI thread fences are inserted by FFrameEndSync::Sync.
 - Since we are still fencing the game and render threads to be no more than 1 frame apart, the original FPendingCleanupObjects mechanism should still work.

Original CL desc:

Fix frame sync issues
 - The engine still uses legacy syncing behaviour where the game thread syncs with the render thread. Renderer refactors have removed most of the RHI thread flushes that happened per frame, which were the only thing synchronizing the game/render threads with the RHI thread. Without these flushes, and when occlusion queries are disabled, the game thread can run ahead of the RHI thread by several hundred frames, since it is now entirely unsynchronized.
 - This fix changes mode 0 of "r.GTSyncType" to sync with the N-2 RHI thread frame when the RHI thread is active, rather than the N-1 render thread frame. The game thread is now always synchronized with the RHI thread to prevent it running ahead.
 - Mode 1 of "r.GTSyncType" now works even when vsync is disabled, and syncs the game thread with the N-1 RHI thread frame (same behaviour as before).

diff --git a/Engine/Plugins/Experimental/TextureGraph/Source/TextureGraphEngine/Model/Mix/MixManager.cpp b/Engine/Plugins/Experimental/TextureGraph/Source/TextureGraphEngine/Model/Mix/MixManager.cpp
--- a/Engine/Plugins/Experimental/TextureGraph/Source/TextureGraphEngine/Model/Mix/MixManager.cpp
+++ b/Engine/Plugins/Experimental/TextureGraph/Source/TextureGraphEngine/Model/Mix/MixManager.cpp
@@ -141,7 +141,7 @@ void MixManager::Exit()
 	
 	FRenderCommandFence SuspendFence;
 	
-	SuspendFence.BeginFence(true);
+	SuspendFence.BeginFence();
 	SuspendFence.Wait();
 	
 	Suspend();
diff --git a/Engine/Source/Runtime/Engine/Private/UnrealEngine.cpp b/Engine/Source/Runtime/Engine/Private/UnrealEngine.cpp
--- a/Engine/Source/Runtime/Engine/Private/UnrealEngine.cpp
+++ b/Engine/Source/Runtime/Engine/Private/UnrealEngine.cpp
@@ -455,6 +455,17 @@ TAutoConsoleVariable<int32> CVarAllowOneFrameThreadLag(
 	TEXT("Whether to allow the rendering thread to lag one frame behind the game thread (0: disabled, otherwise enabled)")
 );
 
+TAutoConsoleVariable<int32> CVarGTSyncType(
+	TEXT("r.GTSyncType"),
+	0,
+	TEXT("Determines how the game thread syncs with the render thread, RHI thread and GPU.\n")
+	TEXT("Syncing to the GPU swap chain flip allows for lower frame latency.\n")
+	TEXT(" <= 0 - Sync the game thread with the N-1 render thread frame. Then sync with the N-m RHI thread frame where m is (2 + (-r.GTSyncType)) (i.e. negative values increase the amount of RHI thread overlap) (default = 0).\n")
+	TEXT("    1 - Sync the game thread with the N-1 RHI thread frame.\n")
+	TEXT("    2 - Sync the game thread with the GPU swap chain flip (only on supported platforms).\n"),
+	ECVF_Default
+);
+
 static FAutoConsoleVariable CVarSystemResolution(
 	TEXT("r.SetRes"),
 	TEXT("1280x720w"),
@@ -736,9 +747,8 @@ void CalculateFPSTimings()
 	double CurrentTime = FPlatformTime::Seconds();
 	float FrameTimeMS = (float)((CurrentTime - LastTime) * 1000.0);
 
-	static auto CVarGTSyncType = IConsoleManager::Get().FindConsoleVariable(TEXT("r.GTSyncType"));
 	static auto CVarVsync = IConsoleManager::Get().FindConsoleVariable(TEXT("r.VSync"));
-	if (CVarGTSyncType->GetInt() == 2 && CVarVsync->GetInt() != 0)
+	if (CVarGTSyncType.GetValueOnAnyThread() == 2 && CVarVsync->GetInt() != 0)
 	{
 		float RHIFrameTime = RHIGetFrameTime();
 		if (RHIFrameTime != 0)
@@ -13296,70 +13306,122 @@ UFont* GetStatsFont()
 	return GEngine->GetSmallFont();
 }
 
-
-FFrameEndSync::FFrameEndSync()
+namespace FFrameEndSync
 {
-	// FFrameEndSync instances are often used as static local vars. we need to cleanup them on engine exit to avoid static destruciton order problem
-	CleanupDelegate = FCoreDelegates::OnEnginePreExit.AddRaw(this, &FFrameEndSync::Cleanup);
-}
+	using ESyncDepth = FRenderCommandFence::ESyncDepth;
 
-FFrameEndSync::~FFrameEndSync()
-{
-	// if it's destroyed before the engine exit, remove the delegate so it doesn't use the instance after destruction
-	if (CleanupDelegate.IsValid())
+	struct FRenderThreadFence
 	{
-		FCoreDelegates::OnEnginePreExit.Remove(CleanupDelegate);
-	}
-}
+		// Legacy game code assumes the game thread will never get further than 1 frame ahead of the render thread.
+		// This fence is used to sync the game thread with the N-1 render thread frame.
+		FRenderCommandFence Fence;
 
-void FFrameEndSync::Cleanup()
-{
-	// forcing static deallocation order:
-	// fences can hold task completion handles that need to be freed before their allocator is destroyed. wating for fences does the job
-	Fence[0].Wait();
-	Fence[1].Wait();
+		FRenderThreadFence()
+		{
+			Fence.BeginFence(ESyncDepth::RenderThread);
+		}
 
-	FCoreDelegates::OnEnginePreExit.Remove(CleanupDelegate);
-	// notify the destructor
-	CleanupDelegate = {};
-}
+		~FRenderThreadFence()
+		{
+			Fence.Wait(true);
+		}
+	};
+	TArray<FRenderThreadFence, TInlineAllocator<2>> RenderThreadFences;
 
-void FFrameEndSync::Sync( bool bAllowOneFrameThreadLag )
-{
-	check(IsInGameThread());			
+	// Additional fences to await. These sync with either the RHI thread or swapchain,
+	// and are used to prevent the game thread running too far ahead of presented frames.
+	TArray<FRenderCommandFence, TInlineAllocator<3>> PipelineFences;
 
-#if !UE_BUILD_SHIPPING && PLATFORM_SUPPORTS_FLIP_TRACKING
-	// Set the FrameDebugInfo on platforms that have accurate frame tracking.
-	ENQUEUE_RENDER_COMMAND(FrameDebugInfo)(
-		[CurrentFrameCounter = GFrameCounter, CurrentInputTime = GInputTime](FRHICommandListImmediate& RHICmdList)
+	void Sync(bool bFullSync)
 	{
-		RHICmdList.EnqueueLambda(
-			[CurrentFrameCounter, CurrentInputTime](FRHICommandListImmediate&)
+		// The "r.OneFrameThreadLag" cvar forces a full sync, meaning the game thread will
+		// not start work until all the rendering work for the previous frame has completed.
+		bFullSync |= CVarAllowOneFrameThreadLag.GetValueOnAnyThread() <= 0;
+
+		SCOPE_CYCLE_COUNTER(STAT_FrameSyncTime);
+
+		check(IsInGameThread());
+
+	#if !UE_BUILD_SHIPPING && PLATFORM_SUPPORTS_FLIP_TRACKING
+		// Set the FrameDebugInfo on platforms that have accurate frame tracking.
+		ENQUEUE_RENDER_COMMAND(FrameDebugInfo)(
+			[CurrentFrameCounter = GFrameCounter, CurrentInputTime = GInputTime](FRHICommandListImmediate& RHICmdList)
 		{
-			// Set the FrameCount and InputTime for input latency stats and flip debugging.
-			RHISetFrameDebugInfo(GRHIPresentCounter - 1, CurrentFrameCounter, CurrentInputTime);
+			RHICmdList.EnqueueLambda(
+				[CurrentFrameCounter, CurrentInputTime](FRHICommandListImmediate&)
+			{
+				// Set the FrameCount and InputTime for input latency stats and flip debugging.
+				RHISetFrameDebugInfo(GRHIPresentCounter - 1, CurrentFrameCounter, CurrentInputTime);
+			});
 		});
-	});
-#endif
+	#endif
+
+		// Always sync with the render thread (either current frame, or N-1 frame)
+		RenderThreadFences.Emplace();
+		while (RenderThreadFences.Num() > (bFullSync ? 0 : 1))
+		{
+			RenderThreadFences.RemoveAt(0);
+		}
 
-	// Since this is the frame end sync, allow sync with the RHI and GPU (true).
-	Fence[EventIndex].BeginFence(true);
+		// Insert an additional fence based on how we want to sync with the RHI thread / swapchain
+		ESyncDepth SyncDepth;
+		int32 NumFramesOverlap;
 
-	bool bEmptyGameThreadTasks = !FTaskGraphInterface::Get().IsThreadProcessingTasks(ENamedThreads::GameThread);
+		int32 const GTSyncType = CVarGTSyncType.GetValueOnAnyThread();
 
-	if (bEmptyGameThreadTasks)
-	{
-		// need to process gamethread tasks at least once a frame no matter what
-		FTaskGraphInterface::Get().ProcessThreadUntilIdle(ENamedThreads::GameThread);
-	}
+		if (bFullSync)
+		{
+			SyncDepth = GTSyncType >= 2
+				? ESyncDepth::Swapchain
+				: ESyncDepth::RHIThread;
 
-	// Use two events if we allow a one frame lag.
-	if( bAllowOneFrameThreadLag )
-	{
-		EventIndex = (EventIndex + 1) % 2;
-	}
+			NumFramesOverlap = 0;
+		}
+		else if (GTSyncType >= 2)
+		{
+			SyncDepth = ESyncDepth::Swapchain;
+			NumFramesOverlap = 1;
+		}
+		else if (GTSyncType == 1)
+		{
+			SyncDepth = ESyncDepth::RHIThread;
+			NumFramesOverlap = 1;
+		}
+		else
+		{
+			check(GTSyncType <= 0);
+
+			// Modes <= 0 allows N frames of overlap with the RHI thread.
+			SyncDepth = ESyncDepth::RHIThread;
+			NumFramesOverlap = 2 + (-GTSyncType);
+		}
+
+		if (SyncDepth == ESyncDepth::Swapchain)
+		{
+			// Swapchain sync mode does not work when vsync is disabled. Fallback to RHI thread sync in that case.
+			static auto CVarVsync = IConsoleManager::Get().FindConsoleVariable(TEXT("r.VSync"));
+			check(CVarVsync != nullptr);
+
+			if (CVarVsync->GetInt() == 0)
+			{
+				SyncDepth = ESyncDepth::RHIThread;
+			}
+		}
 
-	Fence[EventIndex].Wait(bEmptyGameThreadTasks);  // here we also opportunistically execute game thread tasks while we wait
+		PipelineFences.Emplace_GetRef().BeginFence(SyncDepth);
+
+		if (!FTaskGraphInterface::Get().IsThreadProcessingTasks(ENamedThreads::GameThread))
+		{
+			// need to process gamethread tasks at least once a frame no matter what
+			FTaskGraphInterface::Get().ProcessThreadUntilIdle(ENamedThreads::GameThread);
+		}
+
+		while (PipelineFences.Num() > NumFramesOverlap)
+		{
+			PipelineFences[0].Wait(true);
+			PipelineFences.RemoveAt(0);
+		}
+	}
 }
 
 FString appGetStartupMap(const TCHAR* CommandLine)
diff --git a/Engine/Source/Runtime/Engine/Public/UnrealEngine.h b/Engine/Source/Runtime/Engine/Public/UnrealEngine.h
--- a/Engine/Source/Runtime/Engine/Public/UnrealEngine.h
+++ b/Engine/Source/Runtime/Engine/Public/UnrealEngine.h
@@ -368,31 +368,14 @@ extern ENGINE_API UFont* GetStatsFont();
 	Frame end sync object implementation.
 -----------------------------------------------------------------------------*/
 
-/**
- * Special helper class for frame end sync. It respects a passed in option to allow one frame
- * of lag between the game and the render thread by using two events in round robin fashion.
- */
-class FFrameEndSync
+namespace FFrameEndSync
 {
-	/** Pair of fences. */
-	FRenderCommandFence Fence[2];
-	/** Current index into events array. */
-	int32 EventIndex;
-	/** cleanup delegate for engine pre-exit */
-	FDelegateHandle CleanupDelegate;
-
-public:
-	ENGINE_API FFrameEndSync();
-	ENGINE_API ~FFrameEndSync();
-
 	/**
-	 * Syncs the game thread with the render thread. Depending on passed in bool this will be a total
-	 * sync or a one frame lag.
+	 * Syncs the game thread based on progress throughout the rendering pipeline
+	 * @param bFullSync - when true, blocks the caller until all rendering work is completed,
+	 * otherwise only blocks until the N - m frame has completed, where m is driven by various config.
 	 */
-	ENGINE_API void Sync( bool bAllowOneFrameThreadLag );
-
-private:
-	void Cleanup();
+	ENGINE_API void Sync(bool bFullSync = false);
 };
 
 
diff --git a/Engine/Source/Runtime/Launch/Private/LaunchEngineLoop.cpp b/Engine/Source/Runtime/Launch/Private/LaunchEngineLoop.cpp
--- a/Engine/Source/Runtime/Launch/Private/LaunchEngineLoop.cpp
+++ b/Engine/Source/Runtime/Launch/Private/LaunchEngineLoop.cpp
@@ -6071,14 +6071,9 @@ void FEngineLoop::Tick()
 		FPendingCleanupObjects* PreviousPendingCleanupObjects = PendingCleanupObjects;
 		PendingCleanupObjects = GetPendingCleanupObjects();
 
-		{
-			SCOPE_CYCLE_COUNTER(STAT_FrameSyncTime);
-			// this could be perhaps moved down to get greater parallelism
-			// Sync game and render thread. Either total sync or allowing one frame lag.
-			static FFrameEndSync FrameEndSync;
-			static auto CVarAllowOneFrameThreadLag = IConsoleManager::Get().FindTConsoleVariableDataInt(TEXT("r.OneFrameThreadLag"));
-			FrameEndSync.Sync( CVarAllowOneFrameThreadLag->GetValueOnGameThread() != 0 );
-		}
+		// This could be perhaps moved down to get greater parallelism
+		// Sync game and render/RHI threads.
+		FFrameEndSync::Sync();
 
 		// tick core ticker, threads & deferred commands
 		{
diff --git a/Engine/Source/Runtime/RHI/Private/RHIUtilities.cpp b/Engine/Source/Runtime/RHI/Private/RHIUtilities.cpp
--- a/Engine/Source/Runtime/RHI/Private/RHIUtilities.cpp
+++ b/Engine/Source/Runtime/RHI/Private/RHIUtilities.cpp
@@ -178,6 +178,7 @@ class FRHIFrameFlipTrackingRunnable : public FRunnable
 
 		~FFramePair()
 		{
+			TRACE_CPUPROFILER_EVENT_SCOPE(SyncTrigger_Swapchain);
 			Event.Trigger();
 		}
 
@@ -524,6 +525,7 @@ void FRHIFrameFlipTrackingRunnable::TriggerTaskEventOnFlip(uint64 PresentIndex,
 	{
 		// Platform does not support flip tracking.
 		// Signal the event now...
+		TRACE_CPUPROFILER_EVENT_SCOPE(SyncTrigger_Swapchain);
 		Event.Trigger();
 	}
 }
diff --git a/Engine/Source/Runtime/RenderCore/Private/RenderingThread.cpp b/Engine/Source/Runtime/RenderCore/Private/RenderingThread.cpp
--- a/Engine/Source/Runtime/RenderCore/Private/RenderingThread.cpp
+++ b/Engine/Source/Runtime/RenderCore/Private/RenderingThread.cpp
@@ -1008,178 +1008,87 @@ bool IsRenderThreadTimeoutSuspended()
 	return GTimeoutSuspendCount > 0;
 }
 
-TAutoConsoleVariable<int32> CVarGTSyncType(
-	TEXT("r.GTSyncType"),
-	0,
-	TEXT("Determines how the game thread syncs with the render thread, RHI thread and GPU.\n")
-	TEXT("Syncing to the GPU swap chain flip allows for lower frame latency.\n")
-	TEXT(" 0 - Sync the game thread with the render thread (default).\n")
-	TEXT(" 1 - Sync the game thread with the RHI thread.\n")
-	TEXT(" 2 - Sync the game thread with the GPU swap chain flip (only on supported platforms).\n"),
-	ECVF_Default);
-
-FRHICOMMAND_MACRO(FRHISyncFrameCommand)
-{
-	UE::Tasks::FTaskEvent TaskEvent;
-	int32 GTSyncType;
-
-	FORCEINLINE_DEBUGGABLE FRHISyncFrameCommand(UE::Tasks::FTaskEvent InTaskEvent, int32 InGTSyncType)
-		: TaskEvent(MoveTemp(InTaskEvent))
-		, GTSyncType(InGTSyncType)
-	{}
-
-	void Execute(FRHICommandListBase& CmdList)
-	{
-		if (GTSyncType == 1)
-		{
-			// Sync the Game Thread with the RHI Thread
-
-			// "Complete" the graph event
-			TaskEvent.Trigger();
-		}
-		else
-		{
-			// This command runs *after* a present has happened, so the counter has already been incremented.
-			// Subtracting 1 gives us the index of the frame that has *just* been presented.
-			RHITriggerTaskEventOnFlip(GRHIPresentCounter - 1, TaskEvent);
-		}
-	}
-};
-
 FRenderCommandFence::FRenderCommandFence() = default;
 FRenderCommandFence::~FRenderCommandFence() = default;
 
-void FRenderCommandFence::BeginFence(bool bSyncToRHIAndGPU)
+void FRenderCommandFence::BeginFence(ESyncDepth SyncDepth)
 {
 	if (!GIsThreadedRendering)
 	{
 		return;
 	}
+
+	check(IsInGameThread());
 	
-	if (GRenderCommandFenceBundlerState.Event && IsInGameThread())
+	if (GRenderCommandFenceBundlerState.Event && SyncDepth == ESyncDepth::RenderThread)
 	{
+		// Case for game->render thread syncs when fence bundling is enabled. These are used
+		// throughout the engine when resources are destroyed. The fence bundling is an optimization
+		// to avoid the overhead of hundreds of individual fences.
+		// We aren't syncing any deeper than the render thread, so just use the bundled fence event.
 		CompletionTask = *GRenderCommandFenceBundlerState.Event;
 		return;
 	}
 
 	TRACE_CPUPROFILER_EVENT_SCOPE(FRenderCommandFence::BeginFence);
+	UE::Tasks::FTaskEvent Event{ UE_SOURCE_LOCATION };
 
-	struct FRenderCommandPipeFence : public TConcurrentLinearObject<FRenderCommandPipeFence>
-	{
-		FRenderCommandPipeFence(int32 InNumRefs)
-			: NumRefs(InNumRefs)
-		{}
-
-		void Trigger(int32 NumTriggerRefs = 1)
-		{
-			if (NumRefs.fetch_sub(NumTriggerRefs, std::memory_order_release) == 1)
-			{
-				std::atomic_thread_fence(std::memory_order_acquire);
-				CompletionTaskEvent.Trigger();
-				delete this;
-			}
-		}
-
-		UE::Tasks::FTaskEvent CompletionTaskEvent{ UE_SOURCE_LOCATION };
-		std::atomic_int32_t NumRefs;
-	};
-
-	TConstArrayView<FRenderCommandPipe*> Pipes = GRenderCommandPipeMode == ERenderCommandPipeMode::All
-		? UE::RenderCommandPipe::GetPipes()
-		: TConstArrayView<FRenderCommandPipe*>{};
-
-	FRenderCommandPipeBitArray ActivePipeBits;
-	int32 NumActivePipes = 0;
-
-	for (FRenderCommandPipe* Pipe : Pipes)
+	if (GRenderCommandFenceBundlerState.Event)
 	{
-		// Skip pipes that aren't recording or replaying any work.
-		const bool bIsActive = Pipe->IsRecording() && !Pipe->IsEmpty();
-		ActivePipeBits.Add(bIsActive);
-		NumActivePipes += bIsActive ? 1 : 0;
+		// Render command fences are bundled, but we're syncing deeper than the render thread.
+		// Flush the fence bundler so we can insert an RHIThread (or deeper) fence in the right location.
+		Event.AddPrerequisites(*GRenderCommandFenceBundlerState.Event);
+		FlushRenderCommandFenceBundler();
 	}
 
-	FRenderCommandPipeFence* Fence = nullptr;
-
-	if (NumActivePipes > 0)
+	if (GRenderCommandPipeMode == ERenderCommandPipeMode::All)
 	{
-		Fence = new FRenderCommandPipeFence(NumActivePipes + 1);
-
-		for (FRenderCommandPipeSetBitIterator BitIt(ActivePipeBits); BitIt; ++BitIt)
+		for (FRenderCommandPipe* Pipe : UE::RenderCommandPipe::GetPipes())
 		{
-			FRenderCommandPipe* Pipe = Pipes[BitIt.GetIndex()];
-
-			ENQUEUE_RENDER_COMMAND(BeginFence)(Pipe, [Fence]
+			// Skip pipes that aren't recording or replaying any work.
+			if (Pipe->IsRecording() && !Pipe->IsEmpty())
 			{
-				Fence->Trigger();
-			});
-		}
-	}
-
-	const int32 GTSyncType = CVarGTSyncType.GetValueOnAnyThread();
+				UE::Tasks::FTaskEvent PipeEvent { UE_SOURCE_LOCATION };
+				Event.AddPrerequisites(PipeEvent);
 
-	if (bSyncToRHIAndGPU)
-	{
-		// Don't sync to the RHI and GPU if GtSyncType is disabled, or we're not vsyncing
-		//@TODO: do this logic in the caller?
-		static auto CVarVsync = IConsoleManager::Get().FindConsoleVariable(TEXT("r.VSync"));
-		check(CVarVsync != nullptr);
-
-		if (GTSyncType == 0 || CVarVsync->GetInt() == 0)
-		{
-			bSyncToRHIAndGPU = false;
+				ENQUEUE_RENDER_COMMAND(BeginFence)([PipeEvent = MoveTemp(PipeEvent)](FRHICommandList&) mutable
+				{
+					PipeEvent.Trigger();
+				});
+			}
 		}
 	}
 
-	if (bSyncToRHIAndGPU)
+	ENQUEUE_RENDER_COMMAND(BeginFence)([Event, SyncDepth](FRHICommandListImmediate& RHICmdList) mutable
 	{
-		UE::Tasks::FTaskEvent CompletionTaskEvent{ UE_SOURCE_LOCATION };
-
-		if (Fence)
+		if (SyncDepth == ESyncDepth::RenderThread)
 		{
-			// RHI frame sync command requires a task event, so connect it to the ref-counted fence event.
-			CompletionTaskEvent.AddPrerequisites(Fence->CompletionTaskEvent);
-			Fence->Trigger();
+			TRACE_CPUPROFILER_EVENT_SCOPE(SyncTrigger_RenderThread);
+			Event.Trigger();
 		}
-
-		ENQUEUE_RENDER_COMMAND(FSyncFrameCommand)(
-			[CompletionTaskEvent, GTSyncType, bSyncToRHIAndGPU](FRHICommandListImmediate& RHICmdList) mutable
+		else
 		{
-			if (IsRHIThreadRunning())
-			{
-				ALLOC_COMMAND_CL(RHICmdList, FRHISyncFrameCommand)(MoveTemp(CompletionTaskEvent), GTSyncType);
-				RHICmdList.ImmediateFlush(EImmediateFlushType::DispatchToRHIThread);
-			}
-			else
+			RHICmdList.EnqueueLambda([SyncDepth, Event](FRHICommandListImmediate&) mutable
 			{
-				FRHISyncFrameCommand Command(MoveTemp(CompletionTaskEvent), GTSyncType);
-				Command.Execute(RHICmdList);
-			}
-		});
-
-		CompletionTask = MoveTemp(CompletionTaskEvent);
-	}
-	else if (Fence)
-	{
-		CompletionTask = Fence->CompletionTaskEvent;
-
-		ENQUEUE_RENDER_COMMAND(BeginFence)([Fence](FRHICommandListBase& RHICmdList)
-		{
-			Fence->Trigger();
-		});
-	}
-	else
-	{
-		UE::Tasks::FTaskEvent CompletionTaskEvent{ UE_SOURCE_LOCATION };
-
-		ENQUEUE_RENDER_COMMAND(BeginFence)([CompletionTaskEvent](FRHICommandListBase& RHICmdList) mutable
-		{
-			CompletionTaskEvent.Trigger();
-		});
+				if (SyncDepth == ESyncDepth::RHIThread)
+				{
+					// Sync the Game Thread with the RHI Thread
+					TRACE_CPUPROFILER_EVENT_SCOPE(SyncTrigger_RHIThread);
+					Event.Trigger();
+				}
+				else
+				{
+					// This command runs *after* a present has happened, so the counter has already been incremented.
+					// Subtracting 1 gives us the index of the frame that has *just* been presented.
+					RHITriggerTaskEventOnFlip(GRHIPresentCounter - 1, Event);
+				}
+			});
+			RHICmdList.ImmediateFlush(EImmediateFlushType::DispatchToRHIThread);
+		}
+	});
 
-		CompletionTask = MoveTemp(CompletionTaskEvent);
-	}
-} //-V773
+	CompletionTask = MoveTemp(Event);
+}
 
 bool FRenderCommandFence::IsFenceComplete() const
 {
diff --git a/Engine/Source/Runtime/RenderCore/Public/RenderCommandFence.h b/Engine/Source/Runtime/RenderCore/Public/RenderCommandFence.h
--- a/Engine/Source/Runtime/RenderCore/Public/RenderCommandFence.h
+++ b/Engine/Source/Runtime/RenderCore/Public/RenderCommandFence.h
@@ -14,14 +14,25 @@
 class FRenderCommandFence
 {
 public:
+	enum class ESyncDepth
+	{
+		// The fence will be signalled by the render thread.
+		RenderThread,
+
+		// The fence will be enqueued to the RHI thread via a command on the immediate command list
+		// and signalled once all prior parallel translation and submission is complete.
+		RHIThread,
+
+		// The fence will be signalled according to the rate of flips in the swapchain.
+		// This is only supported on some platforms. On unsupported platforms, this behaves like RHIThread mode.
+		Swapchain
+	};
 
 	/**
-	 * Adds a fence command to the rendering command queue.
-	 * Conceptually, the pending fence count is incremented to reflect the pending fence command.
-	 * Once the rendering thread has executed the fence command, it decrements the pending fence count.
-	 * @param bSyncToRHIAndGPU, true if we should wait for the RHI thread or GPU, otherwise we only wait for the render thread.
+	 * Inserts this fence in the rendering pipeline.
+	 * @param SyncDepth, determines which stage of the pipeline will signal the fence.
 	 */
-	RENDERCORE_API void BeginFence(bool bSyncToRHIAndGPU = false);
+	RENDERCORE_API void BeginFence(ESyncDepth SyncDepth = ESyncDepth::RenderThread);
 
 	/**
 	 * Waits for pending fence commands to retire.
diff --git a/Engine/Source/Runtime/SlateNullRenderer/Private/SlateNullRenderer.cpp b/Engine/Source/Runtime/SlateNullRenderer/Private/SlateNullRenderer.cpp
--- a/Engine/Source/Runtime/SlateNullRenderer/Private/SlateNullRenderer.cpp
+++ b/Engine/Source/Runtime/SlateNullRenderer/Private/SlateNullRenderer.cpp
@@ -144,9 +144,6 @@ void FSlateNullRenderer::ClearScenes()
 void FSlateNullRenderer::Sync() const
 {
 #if UE_SLATE_NULL_RENDERER_WITH_ENGINE
-	// Sync game and render thread. Either total sync or allowing one frame lag.
-	static FFrameEndSync FrameEndSync;
-	static auto CVarAllowOneFrameThreadLag = IConsoleManager::Get().FindTConsoleVariableDataInt(TEXT("r.OneFrameThreadLag"));
-	FrameEndSync.Sync(CVarAllowOneFrameThreadLag->GetValueOnAnyThread() != 0);
+	FFrameEndSync::Sync();
 #endif
 }
diff --git a/Engine/Source/Runtime/SlateRHIRenderer/Private/SlateRHIRenderer.cpp b/Engine/Source/Runtime/SlateRHIRenderer/Private/SlateRHIRenderer.cpp
--- a/Engine/Source/Runtime/SlateRHIRenderer/Private/SlateRHIRenderer.cpp
+++ b/Engine/Source/Runtime/SlateRHIRenderer/Private/SlateRHIRenderer.cpp
@@ -1460,10 +1460,7 @@ void FSlateRHIRenderer::FlushCommands() const
 
 void FSlateRHIRenderer::Sync() const
 {
-	// Sync game and render thread. Either total sync or allowing one frame lag.
-	static FFrameEndSync FrameEndSync;
-	static auto CVarAllowOneFrameThreadLag = IConsoleManager::Get().FindTConsoleVariableDataInt(TEXT("r.OneFrameThreadLag"));
-	FrameEndSync.Sync(CVarAllowOneFrameThreadLag->GetValueOnAnyThread() != 0);
+	FFrameEndSync::Sync();
 }
 
 void FSlateRHIRenderer::BeginFrame() const
