Heterogeneous Volume Component API update

Added MID initialization to PostInitProperties and PostReinitProperties. This is necessary for HV components which are parented to other objects.

Added support for slot-name material bindings.

diff --git a/Engine/Source/Runtime/Engine/Classes/Components/HeterogeneousVolumeComponent.h b/Engine/Source/Runtime/Engine/Classes/Components/HeterogeneousVolumeComponent.h
--- a/Engine/Source/Runtime/Engine/Classes/Components/HeterogeneousVolumeComponent.h
+++ b/Engine/Source/Runtime/Engine/Classes/Components/HeterogeneousVolumeComponent.h
@@ -106,20 +106,26 @@ public:
 	virtual void TickComponent(float DeltaTime, ELevelTick TickType, FActorComponentTickFunction* ThisTickFunction) override;
 
 	virtual void PostLoad() override;
+	virtual void PostInitProperties() override;
+	virtual void PostReinitProperties() override;
 #if WITH_EDITOR
 	virtual void PostEditChangeProperty(FPropertyChangedEvent& PropertyChangedEvent) override;
 #endif // WITH_EDITOR
-
 	//~ End USceneComponent Interface.
 
 	//~ Begin UPrimitiveComponent Interface.
 	virtual FPrimitiveSceneProxy* CreateSceneProxy() override;
-	virtual int32 GetNumMaterials() const override { return 1; }
 	virtual bool ShouldRenderSelected() const override { return true; }
-	virtual void GetUsedMaterials(TArray<UMaterialInterface*>& OutMaterials, bool bGetDebugMaterials = false) const override;
-	virtual void SetMaterial(int32 ElementIndex, class UMaterialInterface* Material) override;
+	virtual bool IsMaterialSlotNameValid(FName MaterialSlotName) const override;
 	//~ End UPrimitiveComponent Interface.
 
+	//~ Begin UMeshComponent Interface.
+	virtual int32 GetNumMaterials() const override { return 1; }
+	virtual int32 GetMaterialIndex(FName MaterialSlotName) const override;
+	virtual void SetMaterial(int32 ElementIndex, class UMaterialInterface* Material) override;
+	virtual void GetUsedMaterials(TArray<UMaterialInterface*>& OutMaterials, bool bGetDebugMaterials = false) const override;
+	//~ End UMeshComponent Interface.
+
 private:
 	const USparseVolumeTexture* PreviousSVT;
 
diff --git a/Engine/Source/Runtime/Engine/Private/Components/HeterogeneousVolumeComponent.cpp b/Engine/Source/Runtime/Engine/Private/Components/HeterogeneousVolumeComponent.cpp
--- a/Engine/Source/Runtime/Engine/Private/Components/HeterogeneousVolumeComponent.cpp
+++ b/Engine/Source/Runtime/Engine/Private/Components/HeterogeneousVolumeComponent.cpp
@@ -388,6 +388,28 @@ void UHeterogeneousVolumeComponent::PostLoad()
 	}
 }
 
+void UHeterogeneousVolumeComponent::PostInitProperties()
+{
+	Super::PostInitProperties();
+
+	MaterialInstanceDynamic = nullptr;
+	if (UMaterialInterface* MaterialInterface = GetHeterogeneousVolumeMaterial())
+	{
+		MaterialInstanceDynamic = CreateOrCastToMID(MaterialInterface);
+	}
+}
+
+void UHeterogeneousVolumeComponent::PostReinitProperties()
+{
+	Super::PostReinitProperties();
+
+	MaterialInstanceDynamic = nullptr;
+	if (UMaterialInterface* MaterialInterface = GetHeterogeneousVolumeMaterial())
+	{
+		MaterialInstanceDynamic = CreateOrCastToMID(MaterialInterface);
+	}
+}
+
 USparseVolumeTexture* UHeterogeneousVolumeComponent::GetSparseVolumeTexture(UMaterialInterface* MaterialInterface, int32 ParameterIndex, FName* OutParamName)
 {
 	USparseVolumeTexture* SparseVolumeTexture = nullptr;
@@ -555,6 +577,26 @@ void UHeterogeneousVolumeComponent::GetUsedMaterials(TArray<UMaterialInterface*>
 	}
 }
 
+bool UHeterogeneousVolumeComponent::IsMaterialSlotNameValid(FName MaterialSlotName) const
+{
+	return GetMaterialIndex(MaterialSlotName) >= 0;
+}
+
+int32 UHeterogeneousVolumeComponent::GetMaterialIndex(FName MaterialSlotName) const
+{
+	TArray<UMaterialInterface*> Materials;
+	GetUsedMaterials(Materials);
+	for (int32 MaterialIndex = 0; MaterialIndex < Materials.Num(); ++MaterialIndex)
+	{
+		if (Materials[MaterialIndex]->GetFName() == MaterialSlotName)
+		{
+			return MaterialIndex;
+		}
+	}
+
+	return INDEX_NONE;
+}
+
 void UHeterogeneousVolumeComponent::SetMaterial(int32 ElementIndex, UMaterialInterface* Material)
 {
 	Super::SetMaterial(ElementIndex, Material);
@@ -574,6 +616,16 @@ void UHeterogeneousVolumeComponent::TickComponent(float DeltaTime, ELevelTick Ti
 {
 	Super::TickComponent(DeltaTime, TickType, ThisTickFunction);
 
+	// Forcibly update the MID if it's not initialized
+	if (!MaterialInstanceDynamic)
+	{
+		UMaterialInterface* MaterialInterface = GetHeterogeneousVolumeMaterial();
+		if (MaterialInterface)
+		{
+			MaterialInstanceDynamic = CreateOrCastToMID(MaterialInterface);
+		}
+	}
+
 	if (ShouldRender() && MaterialInstanceDynamic)
 	{
 		const int32 SVTParameterIndex = 0;
