The -SaveWinPos command-line option now saves window positions to specific slots based on the number used.
This allows multiple UE game processes to track their window positions independently, useful when running multiple clients for multiplayer testing, for example.

Deprecated UGameUserSettings::WindowPosX and UGameUserSettings::WindowPosY and replaced with a WindowPositions array.

diff --git a/Engine/Source/Runtime/Engine/Classes/Engine/GameEngine.h b/Engine/Source/Runtime/Engine/Classes/Engine/GameEngine.h
--- a/Engine/Source/Runtime/Engine/Classes/Engine/GameEngine.h
+++ b/Engine/Source/Runtime/Engine/Classes/Engine/GameEngine.h
@@ -110,6 +110,7 @@ public:
 	/**
 	 * Called when the game window is moved
 	 */
+	UE_DEPRECATED(5.6, "Replaced by private OnGameWindowMoved_SavePosition function.")
 	ENGINE_API void OnGameWindowMoved( const TSharedRef<SWindow>& WindowBeingMoved );
 
 	/**
@@ -187,6 +188,11 @@ private:
 	ENGINE_API virtual void HandleNetworkFailure_NotifyGameInstance(UWorld* World, UNetDriver* NetDriver, ENetworkFailure::Type FailureType) override;
 	ENGINE_API virtual void HandleTravelFailure_NotifyGameInstance(UWorld* World, ETravelFailure::Type FailureType) override;
 
+	/**
+	 * Called when the game window is moved
+	 */
+	void OnGameWindowMoved_SavePosition(const TSharedRef<SWindow>& WindowBeingMoved, int32 SaveWinPosIndex);
+
 	/** Last time the logs have been flushed. */
 	double LastTimeLogsFlushed;
 
diff --git a/Engine/Source/Runtime/Engine/Classes/GameFramework/GameUserSettings.h b/Engine/Source/Runtime/Engine/Classes/GameFramework/GameUserSettings.h
--- a/Engine/Source/Runtime/Engine/Classes/GameFramework/GameUserSettings.h
+++ b/Engine/Source/Runtime/Engine/Classes/GameFramework/GameUserSettings.h
@@ -281,10 +281,21 @@ public:
 	UFUNCTION(BlueprintCallable, Category=Settings)
 	ENGINE_API virtual void ResetToCurrentSettings();
 
+	/** Sets the window position in save slot index 0. */
 	ENGINE_API virtual void SetWindowPosition(int32 WindowPosX, int32 WindowPosY);
 
+	/** Sets the window position in the given index, which must be less than GetMaxNumWindowPositionSaveSlots(). */
+	ENGINE_API virtual void SetWindowPosition(int32 WindowPosX, int32 WindowPosY, int32 SaveSlotIndex);
+
+	/** Returns the window position in save slot index 0. */
 	ENGINE_API virtual FIntPoint GetWindowPosition();
+
+	/** Returns the window position in the given index, or the default if the index is invalid. */
+	ENGINE_API virtual FIntPoint GetWindowPosition(int32 SaveSlotIndex);
 	
+	/** Reloads the WindowPositions array from config files. Use before SaveConfig to avoid clobbering other processes' window positions. */
+	ENGINE_API void ReloadWindowPositions();
+
 	UFUNCTION(BlueprintCallable, Category=Settings)
 	ENGINE_API virtual void SetToDefaults();
 
@@ -382,12 +393,18 @@ protected:
 	uint32 LastUserConfirmedResolutionSizeY;
 
 	/** Window PosX */
+	UE_DEPRECATED(5.6, "Use the WindowPositions array instead.")
 	UPROPERTY(config)
-	int32 WindowPosX;
+	int32 WindowPosX_DEPRECATED;
 
 	/** Window PosY */
+	UE_DEPRECATED(5.6, "Use the WindowPositions array instead.")
 	UPROPERTY(config)
-	int32 WindowPosY;
+	int32 WindowPosY_DEPRECATED;
+
+	/** Saved window positions. Multiple instances running simultaneously can use different ones. */
+	UPROPERTY(config)
+	TArray<FIntPoint> WindowPositions;
 
 	/**
 	 * Game window fullscreen mode
@@ -511,7 +528,10 @@ public:
 	{
 		OnGameUserSettingsUINeedsUpdate.Broadcast();
 	}
-    
+
+	/** Returns the maximum supported number save slots for window positions. */
+	int32 GetMaxNumWindowPositionSaveSlots() const;
+
     /**
      * We call this to refresh the config file from a listening Cloud subsystem, before we open it.
      */
diff --git a/Engine/Source/Runtime/Engine/Private/GameEngine.cpp b/Engine/Source/Runtime/Engine/Private/GameEngine.cpp
--- a/Engine/Source/Runtime/Engine/Private/GameEngine.cpp
+++ b/Engine/Source/Runtime/Engine/Private/GameEngine.cpp
@@ -265,7 +265,21 @@ void UGameEngine::CreateGameViewport( UGameViewportClient* GameViewportClient )
 	int32 SaveWinPos;
 	if (FParse::Value(FCommandLine::Get(), TEXT("SAVEWINPOS="), SaveWinPos) && SaveWinPos > 0 )
 	{
-		Window->SetOnWindowMoved( FOnWindowMoved::CreateUObject( this, &UGameEngine::OnGameWindowMoved ) );
+		// Convert to 0-based slot index
+		const int32 SaveWinPosIndex = SaveWinPos - 1;
+		const int32 MaxSlots = GetGameUserSettings()->GetMaxNumWindowPositionSaveSlots();
+
+		if (SaveWinPosIndex < MaxSlots)
+		{
+			Window->SetOnWindowMoved(FOnWindowMoved::CreateWeakLambda(this, [SaveWinPosIndex, this](const TSharedRef<SWindow>& WindowBeingMoved)
+				{
+					this->OnGameWindowMoved_SavePosition(WindowBeingMoved, SaveWinPosIndex);
+				}));
+		}
+		else
+		{
+			UE_LOG(LogEngine, Warning, TEXT("SaveWinPos index %d exceeds the maximum of %d. See cvar GameUserSettings.WindowPositionsMaxSize. Window position won't be saved."), SaveWinPosIndex, MaxSlots - 1);
+		}
 	}
 
 	SceneViewport = MakeShareable( GameViewportClient->CreateGameViewport(GameViewportWidgetRef) );
@@ -562,11 +576,39 @@ TSharedRef<SWindow> UGameEngine::CreateGameWindow()
 		FString ScriptEngineCategory = TEXT("/Script/Engine.Engine");
 		FString GameUserSettingsCategory = TEXT("/Script/Engine.GameUserSettings");
 		GConfig->GetString(*ScriptEngineCategory, TEXT("GameUserSettingsClassName"), GameUserSettingsCategory, GEngineIni);
+
+		// Backwards compatibility for the move from a single saved window position to an array to handle multiple instances.
+		// Old behavior: If SaveWinPos > 0 and WindowPosX and WindowPosY config values exist, use those values for the window position.
+		// New behavior: If SaveWinPos > 0, use this value as a slot number to save multiple positions in the WindowPositions array.
+		// For compatibility, if WindowPosX and WindowPosY exist in the config, use the old behavior.
+		// Upon re-saving the config, those values will be moved into the new WindowPositions array.
 		if (GConfig->GetInt(*GameUserSettingsCategory, TEXT("WindowPosX"), WinX, GGameUserSettingsIni) &&
 			GConfig->GetInt(*GameUserSettingsCategory, TEXT("WindowPosY"), WinY, GGameUserSettingsIni))
 		{
 			AutoCenterType = EAutoCenter::None;
 		}
+		else
+		{
+			TArray<FString> PositionStrings;
+			GConfig->GetArray(*GameUserSettingsCategory, TEXT("WindowPositions"), PositionStrings, GGameUserSettingsIni);
+
+			// SaveWinPos is 1-based, the WindowPositions array is 0-based.
+			int32 Index = SaveWinPos - 1;
+			if (PositionStrings.IsValidIndex(Index))
+			{
+				FIntPoint Position;
+				if (Position.InitFromString(PositionStrings[Index]))
+				{
+					WinX = Position.X;
+					WinY = Position.Y;
+					AutoCenterType = EAutoCenter::None;
+				}
+				else
+				{
+					UE_LOG(LogEngine, Warning, TEXT("Saved position for SaveWinPos=%d invalid, ignoring."), SaveWinPos);
+				}
+			}
+		}
 	}
 
 	// Give the window the max width/height of either the requested resolution, or your available desktop resolution
@@ -717,13 +759,24 @@ void UGameEngine::OnGameWindowClosed( const TSharedRef<SWindow>& WindowBeingClos
 	SceneViewport.Reset();
 }
 
-void UGameEngine::OnGameWindowMoved( const TSharedRef<SWindow>& WindowBeingMoved )
+void UGameEngine::OnGameWindowMoved_SavePosition(const TSharedRef<SWindow>& WindowBeingMoved, int32 SaveWinPosIndex)
 {
 	const FSlateRect WindowRect = WindowBeingMoved->GetRectInScreen();
-	GetGameUserSettings()->SetWindowPosition(WindowRect.Left, WindowRect.Top);
+
+	// Reload config in case another running instance has modified it since the last load.
+	// Prevents this instance from clobbering other SaveWinPos slots.
+	GetGameUserSettings()->ReloadWindowPositions();
+
+	// The -SaveWinPos=X command-line option is 1-based, but the GameUserSettings array is 0-based.
+	GetGameUserSettings()->SetWindowPosition(WindowRect.Left, WindowRect.Top, SaveWinPosIndex);
 	GetGameUserSettings()->SaveConfig();
 }
 
+void UGameEngine::OnGameWindowMoved(const TSharedRef<SWindow>& WindowBeingMoved)
+{
+	OnGameWindowMoved_SavePosition(WindowBeingMoved, 0);
+}
+
 void UGameEngine::RedrawViewports( bool bShouldPresent /*= true*/ )
 {
 	SCOPE_CYCLE_COUNTER(STAT_RedrawViewports);
diff --git a/Engine/Source/Runtime/Engine/Private/GameUserSettings.cpp b/Engine/Source/Runtime/Engine/Private/GameUserSettings.cpp
--- a/Engine/Source/Runtime/Engine/Private/GameUserSettings.cpp
+++ b/Engine/Source/Runtime/Engine/Private/GameUserSettings.cpp
@@ -28,6 +28,11 @@
 
 extern EWindowMode::Type GetWindowModeType(EWindowMode::Type WindowMode);
 
+static TAutoConsoleVariable<int32> CVar_GameUserSettings_WindowPositionsMaxSize(
+	TEXT("GameUserSettings.WindowPositionsMaxSize"),
+	32,
+	TEXT("Maximum number of saved window positions for the -SaveWinPos command-line option."));
+
 enum EGameUserSettingsVersion
 {
 	/** Version for user game settings. All settings will be wiped if the serialized version differs. */
@@ -220,8 +225,11 @@ void UGameUserSettings::SetToDefaults()
 {
 	ResolutionSizeX = LastUserConfirmedResolutionSizeX = GetDefaultResolution().X;
 	ResolutionSizeY = LastUserConfirmedResolutionSizeY = GetDefaultResolution().Y;
-	WindowPosX = GetDefaultWindowPosition().X;
-	WindowPosY = GetDefaultWindowPosition().Y;
+	PRAGMA_DISABLE_DEPRECATION_WARNINGS
+	WindowPosX_DEPRECATED = GetDefaultWindowPosition().X;
+	WindowPosY_DEPRECATED = GetDefaultWindowPosition().Y;
+	PRAGMA_ENABLE_DEPRECATION_WARNINGS
+	WindowPositions.Reset();
 	FullscreenMode = GetDefaultWindowMode();
 	FrameRateLimit = 0.0f;
 	PRAGMA_DISABLE_DEPRECATION_WARNINGS
@@ -551,6 +559,23 @@ void UGameUserSettings::LoadSettings(bool bForceReload/*=false*/)
 
 	// Update r.FullScreenMode CVar
 	SetPreferredFullscreenMode(PreferredFullscreenMode);
+
+	// Window position backwards compatability
+	PRAGMA_DISABLE_DEPRECATION_WARNINGS
+	const FIntPoint OldWindowPosition(WindowPosX_DEPRECATED, WindowPosY_DEPRECATED);
+	PRAGMA_ENABLE_DEPRECATION_WARNINGS
+
+	if (OldWindowPosition != GetDefaultWindowPosition())
+	{
+		if (WindowPositions.IsEmpty())
+		{
+			WindowPositions.SetNum(1);
+		}
+		WindowPositions[0] = OldWindowPosition;
+
+		// Write out the non-deprecated value
+		SaveConfig(CPF_Config, *GGameUserSettingsIni);
+	}
 }
 
 void UGameUserSettings::RequestResolutionChange(int32 InResolutionX, int32 InResolutionY, EWindowMode::Type InWindowMode, bool bInDoOverrides /* = true */)
@@ -574,6 +599,8 @@ void UGameUserSettings::SaveSettings()
 
 	// Save the Scalability state to the same ini file as it was loaded from in FEngineLoop::Preinit
 	Scalability::SaveState(GIsEditor ? GEditorSettingsIni : GGameUserSettingsIni);
+
+	ReloadWindowPositions();
 	SaveConfig(CPF_Config, *GGameUserSettingsIni);
     
     if (OnUpdateCloudDataFromGameUserSettings.IsBound())
@@ -764,15 +791,58 @@ void UGameUserSettings::ResetToCurrentSettings()
 	}
 }
 
+void UGameUserSettings::SetWindowPosition(int32 WinX, int32 WinY, int32 SaveSlotIndex)
+{
+	if (SaveSlotIndex >= 0 && SaveSlotIndex < CVar_GameUserSettings_WindowPositionsMaxSize.GetValueOnAnyThread())
+	{
+		const int32 NumToAdd = FMath::Max(0, SaveSlotIndex - WindowPositions.Num() + 1);
+		for (int32 i = 0; i < NumToAdd; ++i)
+		{
+			WindowPositions.Add(GetDefaultWindowPosition());
+		}
+		WindowPositions[SaveSlotIndex].X = WinX;
+		WindowPositions[SaveSlotIndex].Y = WinY;
+	}
+	else
+	{
+		UE_LOG(LogEngine, Warning, TEXT("UGameUserSettings::SetWindowPosition attempted to save window position in invalid slot %d. Maximum slot index is %d."),
+			SaveSlotIndex, CVar_GameUserSettings_WindowPositionsMaxSize.GetValueOnAnyThread() - 1);
+	}
+}
+
 void UGameUserSettings::SetWindowPosition(int32 WinX, int32 WinY)
 {
-	WindowPosX = WinX;
-	WindowPosY = WinY;
+	SetWindowPosition(WinX, WinY, 0);
+}
+
+FIntPoint UGameUserSettings::GetWindowPosition(int32 SaveSlotIndex)
+{
+	if (WindowPositions.IsValidIndex(SaveSlotIndex))
+	{
+		return FIntPoint(WindowPositions[SaveSlotIndex].X, WindowPositions[SaveSlotIndex].Y);
+	}
+
+	return GetDefaultWindowPosition();
 }
 
 FIntPoint UGameUserSettings::GetWindowPosition()
 {
-	return FIntPoint(WindowPosX, WindowPosY);
+	return GetWindowPosition(0);
+}
+
+void UGameUserSettings::ReloadWindowPositions()
+{
+	FProperty* WindowPositionsProperty = GetClass()->FindPropertyByName(GET_MEMBER_NAME_CHECKED(UGameUserSettings, WindowPositions));
+	if (WindowPositionsProperty != nullptr)
+	{
+		LoadConfigIni(true);
+		ReloadConfig(GetClass(), *GGameUserSettingsIni, UE::LCPF_None, WindowPositionsProperty);
+	}
+}
+
+int32 UGameUserSettings::GetMaxNumWindowPositionSaveSlots() const
+{
+	return CVar_GameUserSettings_WindowPositionsMaxSize.GetValueOnAnyThread();
 }
 
 void UGameUserSettings::SetBenchmarkFallbackValues()
