Fix ensure occurring inside FSkeletalMeshSceneProxy when calling InverseFast on LocalToWorld matrices where one component of the scale is zero.

There is no need to call the inverse at all since these methods don't care about WorldToLocal.

diff --git a/Engine/Source/Runtime/Engine/Private/SkeletalMeshSceneProxy.cpp b/Engine/Source/Runtime/Engine/Private/SkeletalMeshSceneProxy.cpp
--- a/Engine/Source/Runtime/Engine/Private/SkeletalMeshSceneProxy.cpp
+++ b/Engine/Source/Runtime/Engine/Private/SkeletalMeshSceneProxy.cpp
@@ -1091,11 +1091,7 @@ int32 FSkeletalMeshSceneProxy::GetCurrentLODIndex()
  */
 void FSkeletalMeshSceneProxy::DebugDrawPhysicsAsset(int32 ViewIndex, FMeshElementCollector& Collector, const FEngineShowFlags& EngineShowFlags) const
 {
-	FMatrix ProxyLocalToWorld, WorldToLocal;
-	if (!GetWorldMatrices(ProxyLocalToWorld, WorldToLocal))
-	{
-		return; // Cannot draw this, world matrix not valid
-	}
+	const FMatrix& ProxyLocalToWorld = GetLocalToWorld();
 
 	FMatrix ScalingMatrix = ProxyLocalToWorld;
 	FVector TotalScale = ScalingMatrix.ExtractScaling();
@@ -1140,11 +1136,7 @@ void FSkeletalMeshSceneProxy::DebugDrawPoseWatchSkeletons(int32 ViewIndex, FMesh
 void FSkeletalMeshSceneProxy::DebugDrawSkeleton(int32 ViewIndex, FMeshElementCollector& Collector, const FEngineShowFlags& EngineShowFlags) const
 {
 #if !(UE_BUILD_SHIPPING || UE_BUILD_TEST)
-	FMatrix ProxyLocalToWorld, WorldToLocal;
-	if (!GetWorldMatrices(ProxyLocalToWorld, WorldToLocal))
-	{
-		return; // Cannot draw this, world matrix not valid
-	}
+	const FMatrix& ProxyLocalToWorld = GetLocalToWorld();
 
 	// Can't draw this, don't have ComponentSpaceTransforms. This happens with sk meshes rendered with FSkeletalMeshObjectStatic.
 	if (!MeshObject->GetComponentSpaceTransforms())
