Reducing copy-pasted code by adding BeginInitResource function that takes an OwnerName.
Removing deprecated FRenderResource methods.

diff --git a/Engine/Source/Runtime/Engine/Private/SkeletalMeshLODRenderData.cpp b/Engine/Source/Runtime/Engine/Private/SkeletalMeshLODRenderData.cpp
--- a/Engine/Source/Runtime/Engine/Private/SkeletalMeshLODRenderData.cpp
+++ b/Engine/Source/Runtime/Engine/Private/SkeletalMeshLODRenderData.cpp
@@ -242,10 +242,8 @@ void FSkeletalMeshLODRenderData::InitResources(bool bNeedsVertexColors, int32 LO
 	MultiSizeIndexContainer.SetOwnerName(OwnerName);
 	MultiSizeIndexContainer.InitResources();
 
-	StaticVertexBuffers.PositionVertexBuffer.SetOwnerName(OwnerName);
-	BeginInitResource(&StaticVertexBuffers.PositionVertexBuffer, &UE::RenderCommandPipe::SkeletalMesh);
-	StaticVertexBuffers.StaticMeshVertexBuffer.SetOwnerName(OwnerName);
-	BeginInitResource(&StaticVertexBuffers.StaticMeshVertexBuffer, &UE::RenderCommandPipe::SkeletalMesh);
+	BeginInitResource(OwnerName, &StaticVertexBuffers.PositionVertexBuffer, &UE::RenderCommandPipe::SkeletalMesh);
+	BeginInitResource(OwnerName, &StaticVertexBuffers.StaticMeshVertexBuffer, &UE::RenderCommandPipe::SkeletalMesh);
 
 	SkinWeightVertexBuffer.SetOwnerName(OwnerName);
 	SkinWeightVertexBuffer.BeginInitResources();
@@ -253,15 +251,13 @@ void FSkeletalMeshLODRenderData::InitResources(bool bNeedsVertexColors, int32 LO
 	if (bNeedsVertexColors)
 	{
 		// Only init the color buffer if the mesh has vertex colors
-		StaticVertexBuffers.ColorVertexBuffer.SetOwnerName(OwnerName);
-		BeginInitResource(&StaticVertexBuffers.ColorVertexBuffer, &UE::RenderCommandPipe::SkeletalMesh);
+		BeginInitResource(OwnerName, &StaticVertexBuffers.ColorVertexBuffer, &UE::RenderCommandPipe::SkeletalMesh);
 	}
 
 	if (ClothVertexBuffer.GetNumVertices() > 0)
 	{
 		// Only init the clothing buffer if the mesh has clothing data
-		ClothVertexBuffer.SetOwnerName(OwnerName);
-		BeginInitResource(&ClothVertexBuffer, &UE::RenderCommandPipe::SkeletalMesh);
+		BeginInitResource(OwnerName, &ClothVertexBuffer, &UE::RenderCommandPipe::SkeletalMesh);
 	}
 
 	// We can discard any the cooked DuplicatedVertices here based on runtime settings.
@@ -281,8 +277,7 @@ void FSkeletalMeshLODRenderData::InitResources(bool bNeedsVertexColors, int32 LO
 			{
 				// No need to discard CPU data in cooked builds as bNeedsCPUAccess is false (see FDuplicatedVerticesBuffer constructor), 
 				// so it'd be auto-discarded after the RHI has copied the resource data. Keep CPU data when in the editor for geometry operations.
-				RenderSection.DuplicatedVerticesBuffer.SetOwnerName(OwnerName);
-				BeginInitResource(&RenderSection.DuplicatedVerticesBuffer, &UE::RenderCommandPipe::SkeletalMesh);
+				BeginInitResource(OwnerName, &RenderSection.DuplicatedVerticesBuffer, &UE::RenderCommandPipe::SkeletalMesh);
 			}
 		}
 	}
@@ -297,22 +292,19 @@ void FSkeletalMeshLODRenderData::InitResources(bool bNeedsVertexColors, int32 LO
 	
 	if (!MorphTargetVertexInfoBuffers.IsRHIInitialized() && MorphTargetVertexInfoBuffers.IsMorphCPUDataValid() && MorphTargetVertexInfoBuffers.NumTotalBatches > 0)
 	{
-		MorphTargetVertexInfoBuffers.SetOwnerName(OwnerName);
-		BeginInitResource(&MorphTargetVertexInfoBuffers, &UE::RenderCommandPipe::SkeletalMesh);
+		BeginInitResource(OwnerName, &MorphTargetVertexInfoBuffers, &UE::RenderCommandPipe::SkeletalMesh);
 	}
 
 	VertexAttributeBuffers.InitResources();
 	
-	HalfEdgeBuffer.SetOwnerName(OwnerName);
-	BeginInitResource(&HalfEdgeBuffer, &UE::RenderCommandPipe::SkeletalMesh);	
+	BeginInitResource(OwnerName, &HalfEdgeBuffer, &UE::RenderCommandPipe::SkeletalMesh);	
 
 #if RHI_RAYTRACING
 	if (IsRayTracingAllowed())
 	{
 		if (SourceRayTracingGeometry.RawData.Num() > 0)
 		{
-			SourceRayTracingGeometry.SetOwnerName(OwnerName);
-			BeginInitResource(&SourceRayTracingGeometry, &UE::RenderCommandPipe::SkeletalMesh);
+			BeginInitResource(OwnerName, &SourceRayTracingGeometry, &UE::RenderCommandPipe::SkeletalMesh);
 		}
 	}
 #endif
diff --git a/Engine/Source/Runtime/Engine/Private/SkeletalRenderGPUSkin.cpp b/Engine/Source/Runtime/Engine/Private/SkeletalRenderGPUSkin.cpp
--- a/Engine/Source/Runtime/Engine/Private/SkeletalRenderGPUSkin.cpp
+++ b/Engine/Source/Runtime/Engine/Private/SkeletalRenderGPUSkin.cpp
@@ -191,20 +191,17 @@ void FMorphVertexBuffer::ReleaseRHI()
 /*-----------------------------------------------------------------------------
 FMorphVertexBufferPool
 -----------------------------------------------------------------------------*/
-void FMorphVertexBufferPool::InitResources(const FName& OwnerName)
+void FMorphVertexBufferPool::InitResources(FName OwnerName)
 {
 	// InitResources may be called again when morph vertex data is persisted during render state re-creation.
 	if (!bInitializedResources)
 	{
-		MorphVertexBuffers[0].SetOwnerName(OwnerName);
-		MorphVertexBuffers[1].SetOwnerName(OwnerName);
-
 		check(!MorphVertexBuffers[0].VertexBufferRHI.IsValid());
 		check(!MorphVertexBuffers[1].VertexBufferRHI.IsValid());
-		BeginInitResource(&MorphVertexBuffers[0], &UE::RenderCommandPipe::SkeletalMesh);
+		BeginInitResource(OwnerName, &MorphVertexBuffers[0], &UE::RenderCommandPipe::SkeletalMesh);
 		if (bDoubleBuffer)
 		{
-			BeginInitResource(&MorphVertexBuffers[1], &UE::RenderCommandPipe::SkeletalMesh);
+			BeginInitResource(OwnerName, &MorphVertexBuffers[1], &UE::RenderCommandPipe::SkeletalMesh);
 		}
 
 		bInitializedResources = true;
diff --git a/Engine/Source/Runtime/Engine/Private/SkeletalRenderGPUSkin.h b/Engine/Source/Runtime/Engine/Private/SkeletalRenderGPUSkin.h
--- a/Engine/Source/Runtime/Engine/Private/SkeletalRenderGPUSkin.h
+++ b/Engine/Source/Runtime/Engine/Private/SkeletalRenderGPUSkin.h
@@ -329,7 +329,7 @@ public:
 		}
 	}
 
-	void InitResources(const FName& OwnerName);
+	void InitResources(FName OwnerName);
 	void ReleaseResources();
 	SIZE_T GetResourceSize() const;
 	void EnableDoubleBuffer(FRHICommandListBase& RHICmdList);
diff --git a/Engine/Source/Runtime/Engine/Private/StaticMesh.cpp b/Engine/Source/Runtime/Engine/Private/StaticMesh.cpp
--- a/Engine/Source/Runtime/Engine/Private/StaticMesh.cpp
+++ b/Engine/Source/Runtime/Engine/Private/StaticMesh.cpp
@@ -1268,13 +1268,6 @@ void FStaticMeshVertexBuffers::InitFromDynamicVertex(FRHICommandListBase* RHICmd
 	}
 };
 
-void FStaticMeshVertexBuffers::SetOwnerName(const FName& OwnerName)
-{
-	PositionVertexBuffer.SetOwnerName(OwnerName);
-	StaticMeshVertexBuffer.SetOwnerName(OwnerName);
-	ColorVertexBuffer.SetOwnerName(OwnerName);
-}
-
 FStaticMeshLODResources::FStaticMeshLODResources(bool bAddRef)
 	: CardRepresentationData(nullptr)
 	, MaxDeviation(0.0f)
@@ -1362,46 +1355,40 @@ void FStaticMeshLODResources::InitResources(UStaticMesh* Parent, int32 LODIndex)
 		UpdateIndexMemoryStats<true>();
 	}
 
-	IndexBuffer.SetOwnerName(OwnerName);
-	BeginInitResource(&IndexBuffer);
+	BeginInitResource(OwnerName, &IndexBuffer);
+
 	if(bHasWireframeIndices)
 	{
-		AdditionalIndexBuffers->WireframeIndexBuffer.SetOwnerName(OwnerName);
-		BeginInitResource(&AdditionalIndexBuffers->WireframeIndexBuffer);
+		BeginInitResource(OwnerName, &AdditionalIndexBuffers->WireframeIndexBuffer);
 	}
-	VertexBuffers.StaticMeshVertexBuffer.SetOwnerName(OwnerName);
-	BeginInitResource(&VertexBuffers.StaticMeshVertexBuffer);
-	VertexBuffers.PositionVertexBuffer.SetOwnerName(OwnerName);
-	BeginInitResource(&VertexBuffers.PositionVertexBuffer);
+
+	BeginInitResource(OwnerName, &VertexBuffers.StaticMeshVertexBuffer);
+	BeginInitResource(OwnerName, &VertexBuffers.PositionVertexBuffer);
+
 	if(bHasColorVertexData)
 	{
-		VertexBuffers.ColorVertexBuffer.SetOwnerName(OwnerName);
-		BeginInitResource(&VertexBuffers.ColorVertexBuffer);
+		BeginInitResource(OwnerName, &VertexBuffers.ColorVertexBuffer);
 	}
 
 	if (bHasReversedIndices)
 	{
-		AdditionalIndexBuffers->ReversedIndexBuffer.SetOwnerName(OwnerName);
-		BeginInitResource(&AdditionalIndexBuffers->ReversedIndexBuffer);
+		BeginInitResource(OwnerName, &AdditionalIndexBuffers->ReversedIndexBuffer);
 	}
 
 	if (bHasDepthOnlyIndices)
 	{
-		DepthOnlyIndexBuffer.SetOwnerName(OwnerName);
-		BeginInitResource(&DepthOnlyIndexBuffer);
+		BeginInitResource(OwnerName, &DepthOnlyIndexBuffer);
 	}
 
 	if (bHasReversedDepthOnlyIndices)
 	{
-		AdditionalIndexBuffers->ReversedDepthOnlyIndexBuffer.SetOwnerName(OwnerName);
-		BeginInitResource(&AdditionalIndexBuffers->ReversedDepthOnlyIndexBuffer);
+		BeginInitResource(OwnerName, &AdditionalIndexBuffers->ReversedDepthOnlyIndexBuffer);
 	}
 
 	if (Parent && Parent->bSupportGpuUniformlyDistributedSampling && Parent->bSupportUniformlyDistributedSampling && (AreaWeightedSampler.GetNumEntries() > 0))
 	{
 		AreaWeightedSectionSamplersBuffer.Init(&AreaWeightedSectionSamplers);
-		AreaWeightedSectionSamplersBuffer.SetOwnerName(OwnerName);
-		BeginInitResource(&AreaWeightedSectionSamplersBuffer);
+		BeginInitResource(OwnerName, &AreaWeightedSectionSamplersBuffer);
 	}
 
 	if (Parent && !Parent->bSupportRayTracing)
@@ -1626,12 +1613,9 @@ void FStaticMeshRayTracingProxyLOD::InitResources(UStaticMesh* Parent, int32 LOD
 
 	if (bOwnsBuffers)
 	{
-		IndexBuffer->SetOwnerName(OwnerName);
-		BeginInitResource(IndexBuffer);
-		VertexBuffers->StaticMeshVertexBuffer.SetOwnerName(OwnerName);
-		BeginInitResource(&VertexBuffers->StaticMeshVertexBuffer);
-		VertexBuffers->PositionVertexBuffer.SetOwnerName(OwnerName);
-		BeginInitResource(&VertexBuffers->PositionVertexBuffer);
+		BeginInitResource(OwnerName, IndexBuffer);
+		BeginInitResource(OwnerName, &VertexBuffers->StaticMeshVertexBuffer);
+		BeginInitResource(OwnerName, &VertexBuffers->PositionVertexBuffer);
 
 		// todo: support ColorVertexBuffer
 	}
diff --git a/Engine/Source/Runtime/Engine/Public/StaticMeshResources.h b/Engine/Source/Runtime/Engine/Public/StaticMeshResources.h
--- a/Engine/Source/Runtime/Engine/Public/StaticMeshResources.h
+++ b/Engine/Source/Runtime/Engine/Public/StaticMeshResources.h
@@ -317,8 +317,10 @@ struct FStaticMeshVertexBuffers
 {
 	/** The buffer containing vertex data. */
 	FStaticMeshVertexBuffer StaticMeshVertexBuffer;
+
 	/** The buffer containing the position vertex data. */
 	FPositionVertexBuffer PositionVertexBuffer;
+
 	/** The buffer containing the vertex color data. */
 	FColorVertexBuffer ColorVertexBuffer;
 
@@ -369,7 +371,14 @@ struct FStaticMeshVertexBuffers
 		InitModelVF(nullptr, nullptr, VertexFactory);
 	}
 
-	void ENGINE_API SetOwnerName(const FName& OwnerName);
+	void SetOwnerName(FName OwnerName)
+	{
+#if RHI_ENABLE_RESOURCE_INFO
+		PositionVertexBuffer.SetOwnerName(OwnerName);
+		StaticMeshVertexBuffer.SetOwnerName(OwnerName);
+		ColorVertexBuffer.SetOwnerName(OwnerName);
+#endif
+	}
 
 private:
 	void ENGINE_API InitWithDummyData(FRHICommandListBase* RHICmdList, FRenderCommandPipe* RenderCommandPipe, FLocalVertexFactory* VertexFactory, uint32 NumVertices, uint32 NumTexCoords, uint32 LightMapIndex);
diff --git a/Engine/Source/Runtime/Landscape/Private/LandscapeCollision.cpp b/Engine/Source/Runtime/Landscape/Private/LandscapeCollision.cpp
--- a/Engine/Source/Runtime/Landscape/Private/LandscapeCollision.cpp
+++ b/Engine/Source/Runtime/Landscape/Private/LandscapeCollision.cpp
@@ -658,18 +658,17 @@ FPrimitiveSceneProxy* ULandscapeHeightfieldCollisionComponent::CreateSceneProxy(
 			if (Vertices.Num() > 0)
 			{
 #if RHI_ENABLE_RESOURCE_INFO
-				FName Name = FName(TEXT("FLandscapeHeightfieldCollisionComponentSceneProxy ") + GetOwnerName().ToString());
-				VertexBuffers.SetOwnerName(Name);
-				IndexBuffer.SetOwnerName(Name);
-				VertexFactory.SetOwnerName(Name);
+				const FName OwnerName(TEXT("FLandscapeHeightfieldCollisionComponentSceneProxy ") + GetOwnerName().ToString());
+#else
+				const FName OwnerName = NAME_None;
 #endif
 
 				VertexBuffers.InitFromDynamicVertex(&VertexFactory, Vertices);
-				BeginInitResource(&VertexBuffers.PositionVertexBuffer);
-				BeginInitResource(&VertexBuffers.StaticMeshVertexBuffer);
-				BeginInitResource(&VertexBuffers.ColorVertexBuffer);
-				BeginInitResource(&IndexBuffer);
-				BeginInitResource(&VertexFactory);
+				BeginInitResource(OwnerName, &VertexBuffers.PositionVertexBuffer);
+				BeginInitResource(OwnerName, &VertexBuffers.StaticMeshVertexBuffer);
+				BeginInitResource(OwnerName, &VertexBuffers.ColorVertexBuffer);
+				BeginInitResource(OwnerName, &IndexBuffer);
+				BeginInitResource(OwnerName, &VertexFactory);
 
 				WireframeMaterialInstance.Reset(new FColoredMaterialRenderProxy(
 					GEngine->WireframeMaterial ? GEngine->WireframeMaterial->GetRenderProxy() : nullptr,
diff --git a/Engine/Source/Runtime/RenderCore/Private/RenderResource.cpp b/Engine/Source/Runtime/RenderCore/Private/RenderResource.cpp
--- a/Engine/Source/Runtime/RenderCore/Private/RenderResource.cpp
+++ b/Engine/Source/Runtime/RenderCore/Private/RenderResource.cpp
@@ -212,11 +212,6 @@ void FRenderResource::ReleaseResource()
 	}
 }
 
-void FRenderResource::UpdateRHI()
-{
-	UpdateRHI(FRHICommandListImmediate::Get());
-}
-
 void FRenderResource::UpdateRHI(FRHICommandListBase& RHICmdList)
 {
 	if (IsInitialized() && GIsRHIInitialized)
@@ -279,22 +274,6 @@ FBufferRHIRef FRenderResource::CreateRHIBufferInternal(
 	return Buffer;
 }
 
-void FRenderResource::SetOwnerName(const FName& InOwnerName)
-{
-#if RHI_ENABLE_RESOURCE_INFO
-	OwnerName = InOwnerName;
-#endif
-}
-
-FName FRenderResource::GetOwnerName() const
-{
-#if RHI_ENABLE_RESOURCE_INFO
-	return OwnerName;
-#else
-	return NAME_None;
-#endif
-}
-
 void BeginInitResource(FRenderResource* Resource, FRenderCommandPipe* RenderCommandPipe)
 {
 	ENQUEUE_RENDER_COMMAND(InitCommand)(RenderCommandPipe,
diff --git a/Engine/Source/Runtime/RenderCore/Public/RenderResource.h b/Engine/Source/Runtime/RenderCore/Public/RenderResource.h
--- a/Engine/Source/Runtime/RenderCore/Public/RenderResource.h
+++ b/Engine/Source/Runtime/RenderCore/Public/RenderResource.h
@@ -96,24 +96,6 @@ public:
 	/** Destructor used to catch unreleased resources. */
 	RENDERCORE_API virtual ~FRenderResource();
 
-	/**
-	 * Initializes the dynamic RHI resource and/or RHI render target used by this resource.
-	 * Called when the resource is initialized, or when reseting all RHI resources.
-	 * Resources that need to initialize after a D3D device reset must implement this function.
-	 * This is only called by the rendering thread.
-	 */
-	UE_DEPRECATED(5.3, "InitDynamicRHI is now unified with InitRHI.")
-	virtual void InitDynamicRHI() final {}
-
-	/**
-	 * Releases the dynamic RHI resource and/or RHI render target resources used by this resource.
-	 * Called when the resource is released, or when reseting all RHI resources.
-	 * Resources that need to release before a D3D device reset must implement this function.
-	 * This is only called by the rendering thread.
-	 */
-	UE_DEPRECATED(5.3, "ReleaseDynamicRHI is now unified with ReleaseRHI.")
-	virtual void ReleaseDynamicRHI() final {}
-
 	/**
 	 * Initializes the RHI resources used by this resource.
 	 * Called when entering the state where both the resource and the RHI have been initialized.
@@ -121,9 +103,6 @@ public:
 	 */
 	virtual void InitRHI(FRHICommandListBase& RHICmdList) {}
 
-	UE_DEPRECATED(5.3, "InitRHI now requires a command list.")
-	virtual void InitRHI() final {}
-
 	/**
 	 * Releases the RHI resources used by this resource.
 	 * Called when leaving the state where both the resource and the RHI have been initialized.
@@ -137,9 +116,6 @@ public:
 	 */
 	RENDERCORE_API virtual void InitResource(FRHICommandListBase& RHICmdList);
 
-	UE_DEPRECATED(5.3, "InitResource now requires a command list.")
-	virtual void InitResource() final { InitResource(GetImmediateCommandList()); }
-
 	/**
 	 * Prepares the resource for deletion.
 	 * This is only called by the rendering thread.
@@ -150,8 +126,6 @@ public:
 	 * If the resource's RHI resources have been initialized, then release and reinitialize it.  Otherwise, do nothing.
 	 * This is only called by the rendering thread.
 	 */
-	UE_DEPRECATED(5.3, "UpdateRHI now requires a command list.")
-	RENDERCORE_API void UpdateRHI();
 	RENDERCORE_API void UpdateRHI(FRHICommandListBase& RHICmdList);
 
 	/** @return The resource's friendly name.  Typically a UObject name. */
@@ -164,8 +138,21 @@ public:
 	EInitPhase GetInitPhase() const { return InitPhase; }
 
 	/** SetOwnerName should be called before BeginInitResource for the owner name to be successfully tracked. */
-	RENDERCORE_API void SetOwnerName(const FName& InOwnerName);
-	RENDERCORE_API FName GetOwnerName() const;
+	void SetOwnerName(FName InOwnerName)
+	{
+#if RHI_ENABLE_RESOURCE_INFO
+		OwnerName = InOwnerName;
+#endif
+	}
+
+	FName GetOwnerName() const
+	{
+#if RHI_ENABLE_RESOURCE_INFO
+		return OwnerName;
+#else
+		return NAME_None;
+#endif
+	}
 
 protected:
 	// This is used during mobile editor preview refactor, this will eventually be replaced with a parameter to InitRHI() etc..
@@ -234,6 +221,15 @@ public:
  */
 extern RENDERCORE_API void BeginInitResource(FRenderResource* Resource, FRenderCommandPipe* RenderCommandPipe = nullptr);
 
+inline void BeginInitResource(FName OwnerName, FRenderResource* Resource, FRenderCommandPipe* RenderCommandPipe = nullptr)
+{
+#if RHI_ENABLE_RESOURCE_INFO
+	Resource->SetOwnerName(OwnerName);
+#endif
+
+	BeginInitResource(Resource, RenderCommandPipe);
+}
+
 /**
  * Sends a message to the rendering thread to update a resource.
  * This is called in the game thread.
