Entity
* Made IrisCore a public dependancy for Entity and EngineRestricted modules
* This makes it so user modules don't need to add IrisCore if they use SceneGraph or use [REDACTED_HANDLE] tags in their verse code

#rnx

diff --git a/Engine/Source/Runtime/Engine/Engine.Build.cs b/Engine/Source/Runtime/Engine/Engine.Build.cs
--- a/Engine/Source/Runtime/Engine/Engine.Build.cs
+++ b/Engine/Source/Runtime/Engine/Engine.Build.cs
@@ -488,8 +488,8 @@ public class Engine : ModuleRules
 			PublicDefinitions.Add("WITH_ODSC=0");
 		}
 
-		const bool bIrisAddAsPublicDepedency = true;
-		SetupIrisSupport(Target, bIrisAddAsPublicDepedency);
+		const bool bIrisAddAsPublicDependency = true;
+		SetupIrisSupport(Target, bIrisAddAsPublicDependency);
 
 		PrivateDefinitions.Add("UE_DEPRECATE_LEGACY_MATH_CONSTANT_MACRO_NAMES=1");
 	}
