Skinned mesh relevance flag fix for feature levels other than SM6.  Need to pull the feature level from the scene, rather than hard coding to SM6.  I noticed this issue when motion vectors weren't working on an object in my scene, due to the relevance flags all being zero.

#rnx

diff --git a/Engine/Source/Runtime/Engine/Private/SkinnedMeshSceneProxyDesc.cpp b/Engine/Source/Runtime/Engine/Private/SkinnedMeshSceneProxyDesc.cpp
--- a/Engine/Source/Runtime/Engine/Private/SkinnedMeshSceneProxyDesc.cpp
+++ b/Engine/Source/Runtime/Engine/Private/SkinnedMeshSceneProxyDesc.cpp
@@ -46,7 +46,7 @@ void FSkinnedMeshSceneProxyDesc::InitializeFrom(const USkinnedMeshComponent* InC
 #if WITH_EDITORONLY_DATA
 	bClothPainting = InComponent->IsA<UDebugSkelMeshComponent>();
 #endif
-	MaterialRelevance = InComponent->GetMaterialRelevance(ERHIFeatureLevel::SM6);
+	MaterialRelevance = InComponent->GetMaterialRelevance(InComponent->GetScene()->GetFeatureLevel());
 	StreamingDistanceMultiplier = InComponent->StreamingDistanceMultiplier;
 	CapsuleIndirectShadowMinVisibility = InComponent->CapsuleIndirectShadowMinVisibility;
 	OverlayMaterialMaxDrawDistance = InComponent->OverlayMaterialMaxDrawDistance;
