Fixed UInstancedStaticMeshComponent navigation bounds cache not always updated.

diff --git a/Engine/Source/Runtime/Engine/Private/InstancedStaticMesh.cpp b/Engine/Source/Runtime/Engine/Private/InstancedStaticMesh.cpp
--- a/Engine/Source/Runtime/Engine/Private/InstancedStaticMesh.cpp
+++ b/Engine/Source/Runtime/Engine/Private/InstancedStaticMesh.cpp
@@ -2765,7 +2765,13 @@ FBoxSphereBounds UInstancedStaticMeshComponent::GetOrCalcBounds(const FTransform
 			return TargetBounds.GetValue();
 		}
 	}
-	return FBoxSphereBounds(InBoundTransform.GetLocation(), FVector::ZeroVector, 0.f);
+	FBoxSphereBounds Value(InBoundTransform.GetLocation(), FVector::ZeroVector, 0.f);
+	if (InType == EBoundsCacheType::Navigation)
+	{
+		// Store result for Navigation Cache as GetNavigationBounds() relies on it.
+		CachedBounds[(int32)EBoundsCacheType::Navigation].CacheValue(Value);
+	}
+	return Value;
 }
 
 // Deprecated
