FRawStaticIndexBuffer::bShouldExpandTo32Bit doesn't need to be kept around on FRawStaticIndexBuffer. The value is always initialized at the start of Serialize and only ever used inside of the same function.

diff --git a/Engine/Source/Runtime/Engine/Private/RawIndexBuffer.cpp b/Engine/Source/Runtime/Engine/Private/RawIndexBuffer.cpp
--- a/Engine/Source/Runtime/Engine/Private/RawIndexBuffer.cpp
+++ b/Engine/Source/Runtime/Engine/Private/RawIndexBuffer.cpp
@@ -149,7 +149,6 @@ FRawStaticIndexBuffer::FRawStaticIndexBuffer(bool InNeedsCPUAccess)
 	: IndexStorage(InNeedsCPUAccess)
 	, CachedNumIndices(-1)
 	, b32Bit(false)
-	, bShouldExpandTo32Bit(false)
 {
 }
 
@@ -400,18 +399,17 @@ void FRawStaticIndexBuffer::Serialize(FArchive& Ar, bool bNeedsCPUAccess)
 		IndexStorage.BulkSerialize(Ar);
 		CachedNumIndices = IndexStorage.Num() / (b32Bit ? 4 : 2);
 
-		if (Ar.IsCooking() && (CachedNumIndices > 0) && !b32Bit)
+		// Set when cooking for Android if the 16-bit index data potentially needs to be converted to 32-bit on load to work around bugs on certain devices.
+		bool bShouldExpandTo32Bit = false;
+
+		if (Ar.IsCooking() && CachedNumIndices > 0 && !b32Bit)
 		{
 			bShouldExpandTo32Bit = Ar.CookingTarget()->ShouldExpandTo32Bit((const uint16*)IndexStorage.GetData(), CachedNumIndices);
 		}
-		else
-		{
-			bShouldExpandTo32Bit = false;
-		}
 
 		Ar << bShouldExpandTo32Bit;
 
-		if (Ar.IsLoading() && (CachedNumIndices > 0) && bShouldExpandTo32Bit && FPlatformMisc::Expand16BitIndicesTo32BitOnLoad())
+		if (Ar.IsLoading() && bShouldExpandTo32Bit && FPlatformMisc::Expand16BitIndicesTo32BitOnLoad())
 		{
 			ExpandTo32Bit();
 		}
diff --git a/Engine/Source/Runtime/Engine/Public/RawIndexBuffer.h b/Engine/Source/Runtime/Engine/Public/RawIndexBuffer.h
--- a/Engine/Source/Runtime/Engine/Public/RawIndexBuffer.h
+++ b/Engine/Source/Runtime/Engine/Public/RawIndexBuffer.h
@@ -317,9 +317,6 @@ private:
 
 	/** 32bit or 16bit? */
 	bool b32Bit;
-
-	/** Set when cooking for Android if the 16-bit index data potentially needs to be converted to 32-bit on load to work around bugs on certain devices. Only when FPlatformMisc::Expand16BitIndicesTo32BitOnLoad equals true*/
-	bool bShouldExpandTo32Bit;
 };
 
 /**
