Add a "ThrowScriptError" function to the KismetSystemLibrary. This used in combination

PR #3075

diff --git a/Engine/Source/Editor/UnrealEd/Private/Kismet2/KismetDebugUtilities.cpp b/Engine/Source/Editor/UnrealEd/Private/Kismet2/KismetDebugUtilities.cpp
--- a/Engine/Source/Editor/UnrealEd/Private/Kismet2/KismetDebugUtilities.cpp
+++ b/Engine/Source/Editor/UnrealEd/Private/Kismet2/KismetDebugUtilities.cpp
@@ -306,6 +306,7 @@ void FKismetDebugUtilities::OnScriptException(const UObject* ActiveObject, const
 		case EBlueprintExceptionType::WireTracepoint:
 			break;
 		case EBlueprintExceptionType::AccessViolation:
+		case EBlueprintExceptionType::UserRaisedError:
 			if ( GIsEditor && GIsPlayInEditorWorld )
 			{
 				// declared as its own variable since it's flushed (logs pushed to std output) on destruction
diff --git a/Engine/Source/Runtime/CoreUObject/Public/Blueprint/BlueprintExceptionInfo.h b/Engine/Source/Runtime/CoreUObject/Public/Blueprint/BlueprintExceptionInfo.h
--- a/Engine/Source/Runtime/CoreUObject/Public/Blueprint/BlueprintExceptionInfo.h
+++ b/Engine/Source/Runtime/CoreUObject/Public/Blueprint/BlueprintExceptionInfo.h
@@ -17,6 +17,7 @@ namespace EBlueprintExceptionType
 		NonFatalError,
 		FatalError,
 		AbortExecution,
+		UserRaisedError,
 	};
 }
 
diff --git a/Engine/Source/Runtime/Engine/Classes/Kismet/KismetSystemLibrary.h b/Engine/Source/Runtime/Engine/Classes/Kismet/KismetSystemLibrary.h
--- a/Engine/Source/Runtime/Engine/Classes/Kismet/KismetSystemLibrary.h
+++ b/Engine/Source/Runtime/Engine/Classes/Kismet/KismetSystemLibrary.h
@@ -110,6 +110,15 @@ class UKismetSystemLibrary : public UBlueprintFunctionLibrary
 		StackTraceImpl(Stack);
 	}
 
+	/**
+	 * Raises an error in the form of a Blueprint exception.
+	 * Logs the error to the Output and Message Log with information pertaining to where the error was raised.
+	 *
+	 * @param ErrorMessage		The error message to display
+	 */
+	UFUNCTION(BlueprintCallable, Category = "Utilities|Debugging", meta=(Keywords = "raise error", DevelopmentOnly))
+	static void RaiseScriptError(const FString& ErrorMessage = FString(TEXT("An error occurred")));
+
 	// Return true if the object is usable : non-null and not pending kill
 	UFUNCTION(BlueprintPure, Category = "Utilities")
 	static ENGINE_API bool IsValid(const UObject* Object);
diff --git a/Engine/Source/Runtime/Engine/Private/KismetSystemLibrary.cpp b/Engine/Source/Runtime/Engine/Private/KismetSystemLibrary.cpp
--- a/Engine/Source/Runtime/Engine/Private/KismetSystemLibrary.cpp
+++ b/Engine/Source/Runtime/Engine/Private/KismetSystemLibrary.cpp
@@ -123,6 +123,22 @@ void UKismetSystemLibrary::StackTraceImpl(const FFrame& StackFrame)
 	UE_LOG(LogBlueprintUserMessages, Log, TEXT("\n%s"), *Trace);
 }
 
+void UKismetSystemLibrary::RaiseScriptError(const FString& ErrorMessage)
+{
+#if !(UE_BUILD_TEST || UE_BUILD_SHIPPING)
+	FFrame* TopFrame = FFrame::GetThreadLocalTopStackFrame();
+	if (TopFrame)
+	{
+	#if WITH_EDITOR
+		const FBlueprintExceptionInfo ExceptionInfo(EBlueprintExceptionType::UserRaisedError, FText::FromString(ErrorMessage));
+		FBlueprintCoreDelegates::ThrowScriptException(TopFrame->Object, *TopFrame, ExceptionInfo);
+	#else
+		UE_LOG(LogBlueprintUserMessages, Error, TEXT("%s:\n%s"), *ErrorMessage, *TopFrame->GetStackTrace());
+	#endif	// WITH_EDITOR
+	}
+#endif	// !(UE_BUILD_TEST || UE_BUILD_SHIPPING)
+}
+
 FString UKismetSystemLibrary::GetObjectName(const UObject* Object)
 {
 	return GetNameSafe(Object);
