[Property Sheet Soundbase] Bind to new property sheet changes delegate to get notified whenever a relevant property is changed
#rnx

diff --git a/Engine/Source/Runtime/AudioExtensions/Public/AudioPropertiesSheetAssetBase.h b/Engine/Source/Runtime/AudioExtensions/Public/AudioPropertiesSheetAssetBase.h
--- a/Engine/Source/Runtime/AudioExtensions/Public/AudioPropertiesSheetAssetBase.h
+++ b/Engine/Source/Runtime/AudioExtensions/Public/AudioPropertiesSheetAssetBase.h
@@ -3,11 +3,39 @@
 #pragma once
 
 #include "Delegates/IDelegateInstance.h"
+#include "UObject/Interface.h"
 #include "UObject/Object.h"
 #include "UObject/ObjectPtr.h"
 
 #include "AudioPropertiesSheetAssetBase.generated.h"
 
+UINTERFACE(MinimalApi, meta=(CannotImplementInterfaceInBlueprint))
+class UAudioPropertiesSheetAssetUserInterface : public UInterface
+{
+	GENERATED_BODY()
+};
+
+/**
+*	IAudioPropertiesSheetAssetUserInterface
+* 
+*	Target objects can implement this interface to bypass parsing coming from a property sheet. 
+*
+*	Allow/Ignore Property Parsing can be used to set a local property bypass - e.g. from a details view.
+*
+*	ShouldParseProperty can be used by a parser to determine if a property should be parsed or not.
+*/
+class IAudioPropertiesSheetAssetUserInterface
+{
+	GENERATED_BODY()
+
+public:
+#if WITH_EDITORONLY_DATA
+	virtual void AllowPropertyParsing(const FProperty& TargetProperty) = 0;
+	virtual void IgnorePropertyParsing(const FProperty& PropertyToIgnore) = 0;
+	virtual bool ShouldParseProperty(const FProperty& TargetProperty) const = 0;
+#endif
+};
+
 UCLASS(Abstract)
 class AUDIOEXTENSIONS_API UAudioPropertiesSheetAssetBase : public UObject
 {
diff --git a/Engine/Source/Runtime/Engine/Classes/Sound/SoundBase.h b/Engine/Source/Runtime/Engine/Classes/Sound/SoundBase.h
--- a/Engine/Source/Runtime/Engine/Classes/Sound/SoundBase.h
+++ b/Engine/Source/Runtime/Engine/Classes/Sound/SoundBase.h
@@ -3,6 +3,7 @@
 
 
 #include "AudioDefines.h"
+#include "AudioPropertiesSheetAssetBase.h"
 #include "CoreMinimal.h"
 #include "Sound/SoundTimecodeOffset.h"
 #include "SoundConcurrency.h"
@@ -74,7 +75,7 @@ enum class EVirtualizationMode : uint8
  * The base class for a playable sound object
  */
 UCLASS(config=Engine, hidecategories=Object, abstract, editinlinenew, BlueprintType, MinimalAPI)
-class USoundBase : public UObject, public IInterface_AssetUserData
+class USoundBase : public UObject, public IAudioPropertiesSheetAssetUserInterface, public IInterface_AssetUserData
 {
 	GENERATED_UCLASS_BODY()
 
@@ -206,14 +207,31 @@ public:
 	/** Array of user data stored with the asset */
 	UPROPERTY(EditAnywhere, AdvancedDisplay, Instanced, Category = Advanced)
 	TArray<TObjectPtr<UAssetUserData>> AssetUserData;
+
 #if WITH_EDITORONLY_DATA	
 	UPROPERTY(EditAnywhere, Category = AudioProperties)
 	TObjectPtr<UAudioPropertiesSheetAssetBase> AudioPropertiesSheet;
 
+	// //~ Begin IAudioPropertiesSheetAssetUserInterface
+    ENGINE_API virtual void AllowPropertyParsing(const FProperty& TargetProperty) override;
+    ENGINE_API virtual void IgnorePropertyParsing(const FProperty& PropertyToIgnore) override;
+    ENGINE_API virtual bool ShouldParseProperty(const FProperty& TargetProperty) const override;
+	//~ End IAudioPropertiesSheetAssetUserInterface
+#endif
+
 private:
 	UPROPERTY()
 	FSoundTimecodeOffset TimecodeOffset;
-#endif //WITH_EDITORONLY_DATA
+
+	//Property Sheets support
+	void InjectPropertySheet();
+	void BindToPropertySheetChanges();
+	void UnbindFromPropertySheetChanges();
+
+#if WITH_EDITORONLY_DATA
+    UPROPERTY()
+    TArray<FName> LocalAudioProperties;
+#endif
 
 public:
 	//~ Begin UObject Interface.
@@ -226,6 +244,7 @@ public:
 
 #if WITH_EDITOR
 	ENGINE_API virtual void PostEditChangeProperty(struct FPropertyChangedEvent& PropertyChangedEvent) override;
+	ENGINE_API virtual void PreEditChange(FProperty* PropertyAboutToChange) override;
 #endif
 
 	//~ End UObject interface.
@@ -350,7 +369,6 @@ public:
 	ENGINE_API void SetTimecodeOffset(const FSoundTimecodeOffset& InTimecodeOffset);
 	ENGINE_API TOptional<FSoundTimecodeOffset> GetTimecodeOffset() const;
 
-	void InjectPropertySheet();
 #endif //WITH_EDITORONLY_DATA
 
 protected:
diff --git a/Engine/Source/Runtime/Engine/Private/SoundBase.cpp b/Engine/Source/Runtime/Engine/Private/SoundBase.cpp
--- a/Engine/Source/Runtime/Engine/Private/SoundBase.cpp
+++ b/Engine/Source/Runtime/Engine/Private/SoundBase.cpp
@@ -2,9 +2,9 @@
 #include "Sound/SoundBase.h"
 
 #include "AudioDevice.h"
-#include "AudioPropertiesSheetAssetBase.h"
 #include "Engine/AssetUserData.h"
 #include "IAudioParameterTransmitter.h"
+#include "UObject/UObjectGlobals.h"
 #include "Sound/SoundSubmix.h"
 
 #include UE_INLINE_GENERATED_CPP_BY_NAME(SoundBase)
@@ -201,15 +201,31 @@ void USoundBase::PostEditChangeProperty(struct FPropertyChangedEvent& PropertyCh
 		const FName& Name = PropertyThatChanged->GetFName();
 		if (Name == AudioPropertiesSheetFName)
 		{
+			BindToPropertySheetChanges();
 			InjectPropertySheet();
 		}
 	}
 }
+
+void USoundBase::PreEditChange(FProperty* PropertyAboutToChange)
+{
+	Super::PreEditChange(PropertyAboutToChange);
+
+	if (!PropertyAboutToChange)
+	{
+		return;
+	}
+
+	if (PropertyAboutToChange->GetName() == GET_MEMBER_NAME_CHECKED(USoundBase, AudioPropertiesSheet) && AudioPropertiesSheet)
+	{
+		UnbindFromPropertySheetChanges();
+	}
+}
+
 #endif
 
 #if WITH_EDITORONLY_DATA
 
-
 void USoundBase::PostLoad()
 {
 	Super::PostLoad();
@@ -231,6 +247,22 @@ void USoundBase::PostLoad()
 		ConcurrencyOverrides.MaxCount = FMath::Max(MaxConcurrentPlayCount_DEPRECATED, 1);
 		ConcurrencyOverrides.ResolutionRule = MaxConcurrentResolutionRule_DEPRECATED;
 	}
+
+	if (AudioPropertiesSheet && AudioPropertiesSheet->HasAnyFlags(RF_WasLoaded))
+	{
+		BindToPropertySheetChanges();
+	}
+	else
+	{
+		FCoreUObjectDelegates::OnAssetLoaded.AddWeakLambda(this, [this](UObject* LoadedObject)
+			{
+				if (LoadedObject == AudioPropertiesSheet)
+				{
+					BindToPropertySheetChanges();
+					FCoreUObjectDelegates::OnAssetLoaded.RemoveAll(this);
+				}
+			});
+	}
 }
 #endif // WITH_EDITORONLY_DATA
 
@@ -349,6 +381,21 @@ TOptional<FSoundTimecodeOffset> USoundBase::GetTimecodeOffset() const
 	return TimecodeOffset;
 }
 
+void USoundBase::AllowPropertyParsing(const FProperty& TargetProperty)
+{
+	LocalAudioProperties.RemoveSwap(TargetProperty.GetFName());
+}
+
+void USoundBase::IgnorePropertyParsing(const FProperty& PropertyToIgnore)
+{
+	LocalAudioProperties.AddUnique(PropertyToIgnore.GetFName());
+}
+
+bool USoundBase::ShouldParseProperty(const FProperty& TargetProperty) const
+{
+	return !LocalAudioProperties.Contains(TargetProperty.GetFName());
+}
+
 void USoundBase::InjectPropertySheet()
 {
 	if (AudioPropertiesSheet)
@@ -357,6 +404,22 @@ void USoundBase::InjectPropertySheet()
 	}
 }
 
+void USoundBase::BindToPropertySheetChanges()
+{
+	if (AudioPropertiesSheet)
+	{
+		AudioPropertiesSheet->BindPropertiesCopyToSheetChanges(this);
+	}
+}
+
+void USoundBase::UnbindFromPropertySheetChanges()
+{
+	if (AudioPropertiesSheet)
+	{
+		AudioPropertiesSheet->UnbindCopyFromPropertySheetChanges(this);
+	}
+}
+
 #endif //WITH_EDITORONLY_DATA
 
 float USoundBase::ComputeMaxDistance() const
