Fixed issues with AnimBank DDC builds when fetching anim sequence data for targets other than the running platform.

diff --git a/Engine/Source/Runtime/Engine/Private/Animation/AnimBank.cpp b/Engine/Source/Runtime/Engine/Private/Animation/AnimBank.cpp
--- a/Engine/Source/Runtime/Engine/Private/Animation/AnimBank.cpp
+++ b/Engine/Source/Runtime/Engine/Private/Animation/AnimBank.cpp
@@ -210,6 +210,8 @@ private:
 
 	TArray<FAnimBankSequence> BankItems;
 
+	const ITargetPlatform* TargetPlatform = nullptr;
+
 	UE::DerivedData::FRequestOwner Owner;
 	TRefCountPtr<IExecutionResource> ExecutionResource;
 
@@ -240,11 +242,12 @@ FAnimBankBuildAsyncCacheTask::FAnimBankBuildAsyncCacheTask(
 	const FIoHash& InKeyHash,
 	FAnimBankData* InData,
 	UAnimBank& InBank,
-	const ITargetPlatform* TargetPlatform
+	const ITargetPlatform* InTargetPlatform
 )
 : Data(InData)
 , WeakAnimBank(&InBank)
 , BankItems(InBank.Sequences)
+, TargetPlatform(InTargetPlatform)
 // Once we pass the BeginCache throttling gate, we want to finish as fast as possible
 // to avoid holding on to memory for a long time. We use the high priority since it will go fast,
 // but also it will avoid starving the critical threads in the subsequent task.
@@ -364,7 +367,7 @@ void FAnimBankBuildAsyncCacheTask::BeginCache(const FIoHash& InKeyHash)
 		FQueuedThreadPool* ThreadPool = FAnimBankCompilingManager::Get().GetThreadPool();
 		EQueuedWorkPriority BasePriority = FAnimBankCompilingManager::Get().GetBasePriority(Bank);
 
-		int64 RequiredMemory = -1; // @todo RequiredMemory
+		int64 RequiredMemory = -1; // @todo RequiredMemory - possibly follow FAnimationSequenceAsyncCacheTask::GetRequiredMemoryEstimate()
 
 		check(BuildTask == nullptr);
 		BuildTask = MakeUnique<FAnimBankAsyncBuildTask>(this, InKeyHash);
@@ -448,7 +451,8 @@ static void GetRefBoneGlobalSpace(const FReferenceSkeleton& RefSkeleton, TArray<
 			}
 
 			IterBoneIndex = ParentIndex;
-		} while (RefSkeleton.IsValidIndex(IterBoneIndex));
+		}
+		while (RefSkeleton.IsValidIndex(IterBoneIndex));
 
 		Transforms[BoneIndex] = ComponentSpaceTransform;
 	}
@@ -500,6 +504,11 @@ bool FAnimBankBuildAsyncCacheTask::BuildData(const UE::FSharedString& Name, cons
 	Data->Entries.SetNum(AnimBank->Sequences.Num());
 	for (int32 ItemIndex = 0; ItemIndex < AnimBank->Sequences.Num(); ++ItemIndex)
 	{
+		if (Owner.IsCanceled())
+		{
+			return false;
+		}
+
 		FAnimBankEntry& BankEntry = Data->Entries[ItemIndex];
 
 		FAnimBankSequence& BankSequence = AnimBank->Sequences[ItemIndex];
@@ -511,18 +520,34 @@ bool FAnimBankBuildAsyncCacheTask::BuildData(const UE::FSharedString& Name, cons
 
 		Sequence->FinishAsyncTasks();
 
-		check(!Sequence->DoesNeedRecompress());
-		check(Sequence->IsCompressedDataValid() && Sequence->CompressedData.CompressedTrackToSkeletonMapTable.Num() > 0);
+		const FCompressedAnimSequence* CompressedData = nullptr;
+
+		const FIoHash SequenceKeyHash = Sequence->CreateDerivedDataKeyHash(TargetPlatform);
+		if (SequenceKeyHash == Sequence->DataKeyHash)
+		{
+			// Key hash is from running platform, use compressed data on the sequence object
+			CompressedData = &Sequence->CompressedData;
+		}
+		else
+		{
+			// Fetch compressed data for the appropriate target platform
+			CompressedData = Sequence->DataByPlatformKeyHash[SequenceKeyHash].Get();
+		}
 
-		const FCompressedAnimSequence& SampledSequence = Sequence->CompressedData;
+		// We should always have compressed data at this point
+		if (CompressedData == nullptr || CompressedData->CompressedDataStructure == nullptr || CompressedData->CompressedTrackToSkeletonMapTable.Num() == 0)
+		{
+			UE_LOG(LogAnimBank, Error, TEXT("Animation bank referenced sequence is missing compressed data!"));
+			return false;
+		}
 
 		// Set up mapping tables for the decompressor to map internal tracks to the pose array (which is in bone order).
-		const int32 NumTracks = Sequence->CompressedData.CompressedTrackToSkeletonMapTable.Num();
+		const int32 NumTracks = CompressedData->CompressedTrackToSkeletonMapTable.Num();
 
 		TrackToBoneIndexMap.Reset(NumTracks);
 		for (int32 TrackIndex = 0; TrackIndex < NumTracks; ++TrackIndex)
 		{
-			const int32 BoneIndex = SampledSequence.GetSkeletonIndexFromTrackIndex(TrackIndex);
+			const int32 BoneIndex = CompressedData->GetSkeletonIndexFromTrackIndex(TrackIndex);
 
 			// We only care about raw bones.
 			if (BoneIndex < NumAnimBones)
@@ -537,9 +562,9 @@ bool FAnimBankBuildAsyncCacheTask::BuildData(const UE::FSharedString& Name, cons
 			Sequence->GetSamplingFrameRate().AsFrameTime(Sequence->GetPlayLength()).RoundToFrame().Value,
 			EAnimInterpolationType::Linear,
 			Sequence->GetRetargetTransformsSourceName(),
-			*SampledSequence.CompressedDataStructure,
+			*CompressedData->CompressedDataStructure,
 			AnimRefSkeleton.GetRefBonePose(),
-			SampledSequence.CompressedTrackToSkeletonMapTable,
+			CompressedData->CompressedTrackToSkeletonMapTable,
 			nullptr,
 			false,
 			AAT_None
@@ -585,7 +610,7 @@ bool FAnimBankBuildAsyncCacheTask::BuildData(const UE::FSharedString& Name, cons
 			DecompressionContext.Seek(SeekTime);
 
 			TArrayView<FTransform> SampledLocalPoseView(SampledLocalAnimPose);
-			SampledSequence.BoneCompressionCodec->DecompressPose(DecompressionContext, TrackToBoneIndexMap, TrackToBoneIndexMap, TrackToBoneIndexMap, SampledLocalPoseView);
+			CompressedData->BoneCompressionCodec->DecompressPose(DecompressionContext, TrackToBoneIndexMap, TrackToBoneIndexMap, TrackToBoneIndexMap, SampledLocalPoseView);
 
 			// Retarget from the anim skeleton to the mesh skeleton.
 			for (int32 MeshBoneIndex = 0; MeshBoneIndex < NumMeshBones; ++MeshBoneIndex)
@@ -645,6 +670,12 @@ bool FAnimBankBuildAsyncCacheTask::BuildData(const UE::FSharedString& Name, cons
 		BankEntry.SampledBounds.BoxExtent *= BankSequence.BoundsScale;
 		BankEntry.SampledBounds.SphereRadius *= BankSequence.BoundsScale;
 
+		if (BankEntry.SampledBounds.ContainsNaN())
+		{
+			UE_LOG(LogAnimBank, Error, TEXT("BankEntry contains NaN in sampled bounds!"));
+			return false;
+		}
+
 		check(BankEntry.PositionKeys.Num() > 0);
 		check(BankEntry.RotationKeys.Num() > 0);
 		check(BankEntry.FrameCount > 0);
@@ -996,7 +1027,7 @@ FIoHash UAnimBank::CreateDerivedDataKeyHash(const ITargetPlatform* TargetPlatfor
 	}
 
 	FMemoryHasherBlake3 Writer;
-	FGuid AnimBankVersionGuid(0x50961173, 0xDBAF4D9D, 0xE03ABA3C, 0x40EC8690);
+	FGuid AnimBankVersionGuid(0x51842261, 0xDBAF4D8C, 0xB13ABA3C, 0x50EC8691);
 	Writer << AnimBankVersionGuid;
 
 	if (IsValid(Asset))
@@ -1012,6 +1043,9 @@ FIoHash UAnimBank::CreateDerivedDataKeyHash(const ITargetPlatform* TargetPlatfor
 		FAnimBankSequence& BankSequence = Sequences[SequenceIndex];
 		if (IsValid(BankSequence.Sequence))
 		{
+			// Make sure all our required dependencies are loaded, we need them to compute the KeyHash
+			FAnimationUtils::EnsureAnimSequenceLoaded(*BankSequence.Sequence);
+
 			FIoHash SequenceHash = BankSequence.Sequence->CreateDerivedDataKeyHash(TargetPlatform);
 			Writer << SequenceHash;
 		}
@@ -1073,10 +1107,7 @@ FIoHash UAnimBank::BeginCacheDerivedData(const ITargetPlatform* TargetPlatform)
 			continue;
 		}
 
-		if (Sequence->DoesNeedRecompress() || !Sequence->IsCompressedDataValid())
-		{
-			Sequence->BeginCacheDerivedData(TargetPlatform);
-		}
+		Sequence->BeginCacheDerivedData(TargetPlatform);
 	}
 
 	CacheTasksByKeyHash.Emplace(KeyHash, MakePimpl<FAnimBankBuildAsyncCacheTask>(KeyHash, TargetData, *this, TargetPlatform));
