Material Layer Prototype Project:
Introduce the tintermediate thumbnails representing the blended result at intermediate steps of the Layers.
In order to have access to the MaterialExpressions correspoding to the assembled graph for a given Material Instance, we are saving the assembled expressions in the FMaterialResource of the plaform.

diff --git a/Engine/Source/Editor/MaterialEditor/Private/MaterialEditorUtilities.cpp b/Engine/Source/Editor/MaterialEditor/Private/MaterialEditorUtilities.cpp
--- a/Engine/Source/Editor/MaterialEditor/Private/MaterialEditorUtilities.cpp
+++ b/Engine/Source/Editor/MaterialEditor/Private/MaterialEditorUtilities.cpp
@@ -584,7 +584,15 @@ void FMaterialEditorUtilities::GetVisibleMaterialParametersFromExpression(
 			LayersExpression->OverrideLayerGraph(&LayersValue);
 			if (LayersExpression->bIsLayerGraphBuilt)
 			{
-				for (auto& Layer : LayersExpression->LayerCallers)
+				TArray<TObjectPtr<UMaterialExpressionMaterialFunctionCall>>& LayerCallers = LayersExpression->LayerCallers;
+				TArray<TObjectPtr<UMaterialExpressionMaterialFunctionCall>>& BlendCallers = LayersExpression->BlendCallers;
+#if ENABLE_MATERIAL_LAYER_PROTOTYPE
+				TSharedPtr< FMaterialLayersFunctionsRuntimeGraphCache > LayerRuntimeTree = LayersValue.RuntimeGraphCache;
+				LayerCallers = LayerRuntimeTree->LayerCallers;
+				BlendCallers = LayerRuntimeTree->BlendCallers;
+#endif
+
+				for (auto& Layer : LayerCallers)
 				{
 					// Possible that Layer->FunctionOutputs will be empty if this is a newly create layer
 					if (Layer && Layer->MaterialFunction && Layer->FunctionOutputs.IsValidIndex(MaterialExpressionKey.OutputIndex))
@@ -598,7 +606,7 @@ void FMaterialEditorUtilities::GetVisibleMaterialParametersFromExpression(
 					}
 				}
 
-				for (auto& Blend : LayersExpression->BlendCallers)
+				for (auto& Blend : BlendCallers)
 				{
 					if (Blend && Blend->MaterialFunction && Blend->FunctionOutputs.IsValidIndex(MaterialExpressionKey.OutputIndex))
 					{
diff --git a/Engine/Source/Editor/MaterialEditor/Private/SMaterialSubstrateTree.cpp b/Engine/Source/Editor/MaterialEditor/Private/SMaterialSubstrateTree.cpp
--- a/Engine/Source/Editor/MaterialEditor/Private/SMaterialSubstrateTree.cpp
+++ b/Engine/Source/Editor/MaterialEditor/Private/SMaterialSubstrateTree.cpp
@@ -38,6 +38,7 @@
 #include "Widgets/Layout/SSeparator.h"
 #include "Widgets/Text/SInlineEditableTextBlock.h"
 #include "Materials/MaterialFunctionInstance.h"
+#include "Materials/MaterialFunction.h"
 #include "Materials/MaterialExpressionTextureSampleParameter.h"
 #include "Framework/Application/SlateApplication.h"
 #include "Framework/MultiBox/MultiBoxBuilder.h"
@@ -1246,6 +1247,33 @@ void SMaterialSubstrateTree::RecursiveCreateWidgets(FRecursiveCreateWidgetsConte
 					MaterialEditorInstance->StoredLayerPreviews[Payload.Layer]->SetParentEditorOnly(EditedMaterial);
 				}
 			}
+#if ENABLE_MATERIAL_LAYER_PROTOTYPE
+			else 
+			{
+				UMaterialInterface* MaterialInterface = MaterialEditorInstance->GetMaterialInterface();
+				const FMaterialResource* MaterialResource = MaterialInterface->GetMaterialResource(ERHIFeatureLevel::SM6);
+				const FMaterialLayersFunctions* LayersFunctions = MaterialResource->GetMaterialLayers();
+				TSharedPtr<FMaterialLayersFunctionsRuntimeGraphCache> LayerTreeCache = LayersFunctions->RuntimeGraphCache;
+				if (LayerTreeCache)
+				{
+					UMaterialExpressionMaterialFunctionCall* Call = LayerTreeCache->NodeMaterialGraphExpressions[Payload.Layer];
+					if (Call->MaterialFunction)
+					{
+						if (MaterialEditorInstance->StoredLayerPreviews[Payload.Layer] == nullptr)
+						{
+							MaterialEditorInstance->StoredLayerPreviews[Payload.Layer] = (NewObject<UMaterialInstanceConstant>(MaterialEditorInstance, NAME_None));
+						}
+
+						UMaterial* PreviewMaterial = LayersFunctions->GetRuntimeNodePreviewMaterial(Payload.Layer);
+
+						if (MaterialEditorInstance->StoredLayerPreviews[Payload.Layer] && MaterialEditorInstance->StoredLayerPreviews[Payload.Layer]->Parent != PreviewMaterial)
+						{
+							MaterialEditorInstance->StoredLayerPreviews[Payload.Layer]->SetParentEditorOnly(PreviewMaterial);
+						}
+					}
+				}
+			}
+#endif
 		}
 
 		StackProperty->Children.Add(ChildProperty);
diff --git a/Engine/Source/Runtime/Engine/Private/Materials/HLSLMaterialTranslator.cpp b/Engine/Source/Runtime/Engine/Private/Materials/HLSLMaterialTranslator.cpp
--- a/Engine/Source/Runtime/Engine/Private/Materials/HLSLMaterialTranslator.cpp
+++ b/Engine/Source/Runtime/Engine/Private/Materials/HLSLMaterialTranslator.cpp
@@ -8778,6 +8778,17 @@ const FMaterialLayersFunctions* FHLSLMaterialTranslator::GetMaterialLayers()
 	return StaticParameters.bHasMaterialLayers ? &CachedMaterialLayers : nullptr;
 }
 
+void FHLSLMaterialTranslator::FeedbackMaterialLayersInstancedGraphFromCompilation(const FMaterialLayersFunctions* InLayers)
+{
+	if (InLayers && StaticParameters.bHasMaterialLayers)
+	{
+		if (Material)
+		{
+			Material->FeedbackMaterialLayersInstancedGraphFromCompilation(InLayers);
+		}
+	}
+}
+
 bool FHLSLMaterialTranslator::GetStaticBoolValue(int32 BoolIndex, bool& bSucceeded)
 {
 	bSucceeded = true;
diff --git a/Engine/Source/Runtime/Engine/Private/Materials/HLSLMaterialTranslator.h b/Engine/Source/Runtime/Engine/Private/Materials/HLSLMaterialTranslator.h
--- a/Engine/Source/Runtime/Engine/Private/Materials/HLSLMaterialTranslator.h
+++ b/Engine/Source/Runtime/Engine/Private/Materials/HLSLMaterialTranslator.h
@@ -1115,6 +1115,8 @@ protected:
 	virtual int32 StaticComponentMask(int32 Vector, FName ParameterName, bool bDefaultR, bool bDefaultG, bool bDefaultB, bool bDefaultA) override;
 	virtual const FMaterialLayersFunctions* GetMaterialLayers() override;
 
+	virtual void FeedbackMaterialLayersInstancedGraphFromCompilation(const FMaterialLayersFunctions* InLayers) override;
+
 	virtual bool GetStaticBoolValue(int32 BoolIndex, bool& bSucceeded) override;
 
 	virtual int32 StaticTerrainLayerWeight(FName ParameterName, int32 Default, bool bTextureArray = false) override;
diff --git a/Engine/Source/Runtime/Engine/Private/Materials/MaterialExpressions.cpp b/Engine/Source/Runtime/Engine/Private/Materials/MaterialExpressions.cpp
--- a/Engine/Source/Runtime/Engine/Private/Materials/MaterialExpressions.cpp
+++ b/Engine/Source/Runtime/Engine/Private/Materials/MaterialExpressions.cpp
@@ -8495,15 +8495,21 @@ void UMaterialExpressionMaterialAttributeLayers::RebuildLayerGraph(bool bReportE
 	const TArray<UMaterialFunctionInterface*>& Blends = GetBlends();
 	const TArray<bool>& LayerStates = GetLayerStates();
 
+	if (!LayersTree->RuntimeGraphCache)
+	{
+		LayersTree->RuntimeGraphCache = MakeShareable(new FMaterialLayersFunctionsRuntimeGraphCache());
+	}
+	TSharedPtr< FMaterialLayersFunctionsRuntimeGraphCache> LayersTreeCache = LayersTree->RuntimeGraphCache;
+
 	// Pre-populate callers, we maintain these transient objects to avoid
 	// heavy UObject recreation as the graphs are frequently rebuilt
-	while (LayerCallers.Num() < Layers.Num())
+	while (LayersTreeCache->LayerCallers.Num() < Layers.Num())
 	{
-		LayerCallers.Add(NewObject<UMaterialExpressionMaterialFunctionCall>(GetTransientPackage()));
+		LayersTreeCache->LayerCallers.Add(NewObject<UMaterialExpressionMaterialFunctionCall>(GetTransientPackage()));
 	}
-	while (BlendCallers.Num() < Blends.Num())
+	while (LayersTreeCache->BlendCallers.Num() < Blends.Num())
 	{
-		BlendCallers.Add(NewObject<UMaterialExpressionMaterialFunctionCall>(GetTransientPackage()));
+		LayersTreeCache->BlendCallers.Add(NewObject<UMaterialExpressionMaterialFunctionCall>(GetTransientPackage()));
 	}
 
 	// Reset graph connectivity
@@ -8518,25 +8524,26 @@ void UMaterialExpressionMaterialAttributeLayers::RebuildLayerGraph(bool bReportE
 		{
 			if (Layers[LayerIndex] && LayerStates[LayerIndex])
 			{
-				LayerCallers[LayerIndex]->MaterialFunction = Layers[LayerIndex];
-				LayerCallers[LayerIndex]->FunctionParameterInfo.Association = EMaterialParameterAssociation::LayerParameter;
-				LayerCallers[LayerIndex]->FunctionParameterInfo.Index = LayerIndex;
+				LayersTreeCache->LayerCallers[LayerIndex]->MaterialFunction = Layers[LayerIndex];
+				LayersTreeCache->LayerCallers[LayerIndex]->FunctionParameterInfo.Association = EMaterialParameterAssociation::LayerParameter;
+				LayersTreeCache->LayerCallers[LayerIndex]->FunctionParameterInfo.Index = LayerIndex;
 
-				Layers[LayerIndex]->GetInputsAndOutputs(LayerCallers[LayerIndex]->FunctionInputs, LayerCallers[LayerIndex]->FunctionOutputs);
-				for (FFunctionExpressionOutput& FunctionOutput : LayerCallers[LayerIndex]->FunctionOutputs)
+				Layers[LayerIndex]->GetInputsAndOutputs(LayersTreeCache->LayerCallers[LayerIndex]->FunctionInputs, LayersTreeCache->LayerCallers[LayerIndex]->FunctionOutputs);
+				
+				for (FFunctionExpressionOutput& FunctionOutput : LayersTreeCache->LayerCallers[LayerIndex]->FunctionOutputs)
 				{
-					LayerCallers[LayerIndex]->Outputs.Add(FunctionOutput.Output);
+					LayersTreeCache->LayerCallers[LayerIndex]->Outputs.Add(FunctionOutput.Output);
 				}
 
 				// Recursively run through internal functions to allow connection of inputs/outputs
-				LayerCallers[LayerIndex]->UpdateFromFunctionResource();
+				LayersTreeCache->LayerCallers[LayerIndex]->UpdateFromFunctionResource();
 
 				++NumActiveLayerCallers;
 			}
 			else
 			{
 				// Empty entries for opaque layers
-				LayerCallers[LayerIndex]->MaterialFunction = nullptr;
+				LayersTreeCache->LayerCallers[LayerIndex]->MaterialFunction = nullptr;
 			}
 		}
 
@@ -8544,43 +8551,50 @@ void UMaterialExpressionMaterialAttributeLayers::RebuildLayerGraph(bool bReportE
 		{
 			if (Blends[BlendIndex])
 			{
-				BlendCallers[BlendIndex]->MaterialFunction = Blends[BlendIndex];
-				BlendCallers[BlendIndex]->FunctionParameterInfo.Association = EMaterialParameterAssociation::BlendParameter;
-				BlendCallers[BlendIndex]->FunctionParameterInfo.Index = BlendIndex;
+				LayersTreeCache->BlendCallers[BlendIndex]->MaterialFunction = Blends[BlendIndex];
+				LayersTreeCache->BlendCallers[BlendIndex]->FunctionParameterInfo.Association = EMaterialParameterAssociation::BlendParameter;
+				LayersTreeCache->BlendCallers[BlendIndex]->FunctionParameterInfo.Index = BlendIndex;
 
-				Blends[BlendIndex]->GetInputsAndOutputs(BlendCallers[BlendIndex]->FunctionInputs, BlendCallers[BlendIndex]->FunctionOutputs);
-				for (FFunctionExpressionOutput& FunctionOutput : BlendCallers[BlendIndex]->FunctionOutputs)
+				Blends[BlendIndex]->GetInputsAndOutputs(LayersTreeCache->BlendCallers[BlendIndex]->FunctionInputs, LayersTreeCache->BlendCallers[BlendIndex]->FunctionOutputs);
+				for (FFunctionExpressionOutput& FunctionOutput : LayersTreeCache->BlendCallers[BlendIndex]->FunctionOutputs)
 				{
-					BlendCallers[BlendIndex]->Outputs.Add(FunctionOutput.Output);
+					LayersTreeCache->BlendCallers[BlendIndex]->Outputs.Add(FunctionOutput.Output);
 				}
 
 				// Recursively run through internal functions to allow connection of inputs/ouputs
-				BlendCallers[BlendIndex]->UpdateFromFunctionResource();
+				LayersTreeCache->BlendCallers[BlendIndex]->UpdateFromFunctionResource();
 
 				++NumActiveBlendCallers;
 			}
 			else
 			{
 				// Empty entries for opaque layers
-				BlendCallers[BlendIndex]->MaterialFunction = nullptr;
+				LayersTreeCache->BlendCallers[BlendIndex]->MaterialFunction = nullptr;
 			}
 		}
 
 		// Empty out unused callers
-		for (int32 CallerIndex = Layers.Num(); CallerIndex < LayerCallers.Num(); ++CallerIndex)
+		for (int32 CallerIndex = Layers.Num(); CallerIndex < LayersTreeCache->LayerCallers.Num(); ++CallerIndex)
 		{
-			LayerCallers[CallerIndex]->MaterialFunction = nullptr;
+			LayersTreeCache->LayerCallers[CallerIndex]->MaterialFunction = nullptr;
 		}
 
-		for (int32 CallerIndex = Blends.Num(); CallerIndex < BlendCallers.Num(); ++CallerIndex)
+		for (int32 CallerIndex = Blends.Num(); CallerIndex < LayersTreeCache->BlendCallers.Num(); ++CallerIndex)
 		{
-			BlendCallers[CallerIndex]->MaterialFunction = nullptr;
+			LayersTreeCache->BlendCallers[CallerIndex]->MaterialFunction = nullptr;
 		}
 
+
+		LayersTreeCache->NodeMaterialGraphExpressions.Empty(LayersTree->Tree.Nodes.Num());
+		LayersTreeCache->NodeMaterialGraphExpressions.InsertZeroed(0, LayersTree->Tree.Nodes.Num());
+		LayersTreeCache->NodePreviewMaterials.Empty(LayersTree->Tree.Nodes.Num());
+		LayersTreeCache->NodePreviewMaterials.InsertZeroed(0, LayersTree->Tree.Nodes.Num());
+
 		struct TreeGraphBuilder
 		{
 			UMaterialExpressionMaterialAttributeLayers* This = nullptr;
 			const FMaterialLayersFunctions* Tree = nullptr;
+			TSharedPtr< FMaterialLayersFunctionsRuntimeGraphCache> TreeCache;
 
 			UMaterialExpressionMaterialFunctionCall* AssembleBranchGraph(int32 InSourceId)
 			{
@@ -8590,8 +8604,10 @@ void UMaterialExpressionMaterialAttributeLayers::RebuildLayerGraph(bool bReportE
 				if (InSourceId != -1)
 				{
 					int32 BackgroundLayerFuncIdx = Tree->GetLayerFuncIndex(InSourceId);
-					PreviousLayerInput = (This->LayerCallers.IsValidIndex(BackgroundLayerFuncIdx) ? This->LayerCallers[BackgroundLayerFuncIdx] : nullptr);
+					PreviousLayerInput = (TreeCache->LayerCallers.IsValidIndex(BackgroundLayerFuncIdx) ? TreeCache->LayerCallers[BackgroundLayerFuncIdx] : nullptr);
 					PreviousLayerInput = (PreviousLayerInput && PreviousLayerInput->MaterialFunction ? PreviousLayerInput : nullptr);
+					if (PreviousLayerInput)
+						TreeCache->ExpressionCollection.AddExpression(PreviousLayerInput);
 				}
 
 				// For child of this node, try to assemble the blend/layer tree 
@@ -8601,13 +8617,17 @@ void UMaterialExpressionMaterialAttributeLayers::RebuildLayerGraph(bool bReportE
 					// Assemble the branch of layer / blend under this sub node by calling recursively  into that sub node
 					UMaterialExpressionMaterialFunctionCall* CurrentLayerInput = AssembleBranchGraph(NodeId);
 
+
 					// Fetch the blend function of the node
 					int32 BlendFuncIdx = Tree->GetBlendFuncIndex(NodeId);
-					UMaterialExpressionMaterialFunctionCall* CurrentBlendCaller = (This->BlendCallers.IsValidIndex(BlendFuncIdx) ? This->BlendCallers[BlendFuncIdx] : nullptr);
+					UMaterialExpressionMaterialFunctionCall* CurrentBlendCaller = (TreeCache->BlendCallers.IsValidIndex(BlendFuncIdx) ? TreeCache->BlendCallers[BlendFuncIdx] : nullptr);
 					CurrentBlendCaller = (CurrentBlendCaller && CurrentBlendCaller->MaterialFunction ? CurrentBlendCaller : nullptr);
 
 					if (CurrentBlendCaller) // Blend is the new current layer if valid and connect the previous layer and current layer is valid
 					{
+						if (PreviousLayerInput)
+							TreeCache->ExpressionCollection.AddExpression(CurrentBlendCaller);
+
 						if (PreviousLayerInput)
 							CurrentBlendCaller->FunctionInputs[0].Input.Connect(0, PreviousLayerInput);
 						if (CurrentLayerInput)
@@ -8622,9 +8642,14 @@ void UMaterialExpressionMaterialAttributeLayers::RebuildLayerGraph(bool bReportE
 					//else this node's layer and and blend functions aren't valid, ignore and move on to the next sub node
 				}
 
+				if (InSourceId != -1)
+				{
+					TreeCache->NodeMaterialGraphExpressions[InSourceId] = (PreviousLayerInput);
+				}
+
 				return PreviousLayerInput;
 			}
-		} BuildTreeGraph = { this, LayersTree };
+		} BuildTreeGraph = { this, LayersTree, LayersTreeCache };
 
 		OutputCaller = BuildTreeGraph.AssembleBranchGraph(-1);
 
@@ -9058,6 +9083,7 @@ FSubstrateOperator* UMaterialExpressionMaterialAttributeLayers::SubstrateGenerat
 		if (OutputCaller)
 		{
 			OutOperator = OutputCaller->SubstrateGenerateMaterialTopologyTree(Compiler, Parent, 0);
+			Compiler->FeedbackMaterialLayersInstancedGraphFromCompilation(OverrideLayers);
 		}
 		
 		if (!OutOperator)
@@ -19791,6 +19817,43 @@ void FMaterialLayersFunctions::MoveLayerNode(FLayerNodeId InNodeId, FLayerNodeId
 	}
 }
 
+UMaterial* FMaterialLayersFunctions::GetRuntimeNodePreviewMaterial(FLayerNodeId InNodeId) const
+{
+#if ENABLE_MATERIAL_LAYER_PROTOTYPE
+	if (RuntimeGraphCache)
+	{
+		if (Tree.IsValidId(InNodeId))
+		{
+			if (!RuntimeGraphCache->NodePreviewMaterials[InNodeId])
+			{
+				if (RuntimeGraphCache->NodeMaterialGraphExpressions[InNodeId])
+				{
+					UMaterialExpressionMaterialFunctionCall* NodeCall = RuntimeGraphCache->NodeMaterialGraphExpressions[InNodeId];
+
+					UMaterial* PreviewMaterial = NewObject<UMaterial>(NodeCall, NAME_None, RF_Transient | RF_Public);
+					PreviewMaterial->bIsPreviewMaterial = true;
+
+					//Update cached expression data to ensure function calls are populated for resolving the preview
+					PreviewMaterial->AssignExpressionCollection(RuntimeGraphCache->ExpressionCollection);
+					PreviewMaterial->UpdateCachedExpressionData();
+
+					NodeCall->ConnectToPreviewMaterial(PreviewMaterial, 0);
+
+					//Compile the material.
+					PreviewMaterial->PreEditChange(nullptr);
+					PreviewMaterial->PostEditChange();
+
+					RuntimeGraphCache->NodePreviewMaterials[InNodeId] = PreviewMaterial;
+				}
+			}
+
+			return RuntimeGraphCache->NodePreviewMaterials[InNodeId];	
+		}
+	}
+#endif
+	return nullptr;
+}
+
 #endif // WITH_EDITOR
 
 ///////////////////////////////////////////////////////////////////////////////
diff --git a/Engine/Source/Runtime/Engine/Private/Materials/MaterialShared.cpp b/Engine/Source/Runtime/Engine/Private/Materials/MaterialShared.cpp
--- a/Engine/Source/Runtime/Engine/Private/Materials/MaterialShared.cpp
+++ b/Engine/Source/Runtime/Engine/Private/Materials/MaterialShared.cpp
@@ -1678,6 +1678,19 @@ TConstArrayView<TObjectPtr<UTextureCollection>> FMaterialResource::GetReferenced
 	return UMaterial::GetDefaultMaterial(MD_Surface)->GetReferencedTextureCollections();
 }
 
+void FMaterialResource::FeedbackMaterialLayersInstancedGraphFromCompilation(const FMaterialLayersFunctions* InLayers)
+{
+	if (InLayers)
+	{
+		MaterialLayersFunctions = *(InLayers);
+	}
+}
+
+const FMaterialLayersFunctions* FMaterialResource::GetMaterialLayers() const
+{
+	return &MaterialLayersFunctions;
+}
+
 void FMaterialResource::AddReferencedObjects(FReferenceCollector& Collector)
 {
 	FMaterial::AddReferencedObjects(Collector);
diff --git a/Engine/Source/Runtime/Engine/Public/MaterialCompiler.h b/Engine/Source/Runtime/Engine/Public/MaterialCompiler.h
--- a/Engine/Source/Runtime/Engine/Public/MaterialCompiler.h
+++ b/Engine/Source/Runtime/Engine/Public/MaterialCompiler.h
@@ -447,10 +447,14 @@ public:
 	virtual int32 DynamicBoolParameter(FName ParameterName,bool bDefaultValue) = 0;
 	virtual int32 StaticComponentMask(int32 Vector,FName ParameterName,bool bDefaultR,bool bDefaultG,bool bDefaultB,bool bDefaultA) = 0;
 	virtual const FMaterialLayersFunctions* GetMaterialLayers() = 0;
+
 	virtual bool GetStaticBoolValue(int32 BoolIndex, bool& bSucceeded) = 0;
 	virtual int32 StaticTerrainLayerWeight(FName ParameterName,int32 Default, bool bTextureArray = false) = 0;
 
 	virtual int32 FontSignedDistanceData() = 0;
+	// Feedback the instanced material layers data after the UMaterialExpressionMaterialAttributeLayers is assembled
+	virtual void FeedbackMaterialLayersInstancedGraphFromCompilation(const FMaterialLayersFunctions* InLayers) {} // TODO: SHould be pure virtual
+
 
 	virtual int32 VertexColor() = 0;
 	virtual int32 MeshPaintTextureCoordinateIndex() = 0;
diff --git a/Engine/Source/Runtime/Engine/Public/MaterialShared.h b/Engine/Source/Runtime/Engine/Public/MaterialShared.h
--- a/Engine/Source/Runtime/Engine/Public/MaterialShared.h
+++ b/Engine/Source/Runtime/Engine/Public/MaterialShared.h
@@ -2612,6 +2612,9 @@ public:
 	void SetODSCMetaData(uint8 InODSCMetaData) const { ODSCMetaData.store(InODSCMetaData, std::memory_order_relaxed); };
 #endif
 
+	ENGINE_API virtual void FeedbackMaterialLayersInstancedGraphFromCompilation(const FMaterialLayersFunctions* InLayers) {}
+	ENGINE_API virtual const FMaterialLayersFunctions* GetMaterialLayers() const { return nullptr; }
+
 protected:
 	// shared code needed for GetUniformScalarParameterExpressions, GetUniformVectorParameterExpressions, GetUniformCubeTextureExpressions..
 	// @return can be 0
@@ -3070,10 +3073,14 @@ public:
 	inline void SetMaterial(UMaterial* InMaterial) { Material = InMaterial; }
 	inline void SetMaterialInstance(UMaterialInstance* InMaterialInstance) { MaterialInstance = InMaterialInstance; }
 
+	ENGINE_API virtual void FeedbackMaterialLayersInstancedGraphFromCompilation(const FMaterialLayersFunctions* InLayers) override;
+	ENGINE_API virtual const FMaterialLayersFunctions* GetMaterialLayers() const override;
 protected:
 	TObjectPtr<UMaterial> Material;
 	TObjectPtr<UMaterialInstance> MaterialInstance;
 
+	FMaterialLayersFunctions MaterialLayersFunctions;
+
 	/** Entry point for compiling a specific material property.  This must call SetMaterialProperty. */
 	ENGINE_API virtual int32 CompilePropertyAndSetMaterialProperty(EMaterialProperty Property, class FMaterialCompiler* Compiler, EShaderFrequency OverrideShaderFrequency, bool bUsePreviousFrameTime) const override;
 #if HANDLE_CUSTOM_OUTPUTS_AS_MATERIAL_ATTRIBUTES
diff --git a/Engine/Source/Runtime/Engine/Public/Materials/MaterialLayersFunctions.h b/Engine/Source/Runtime/Engine/Public/Materials/MaterialLayersFunctions.h
--- a/Engine/Source/Runtime/Engine/Public/Materials/MaterialLayersFunctions.h
+++ b/Engine/Source/Runtime/Engine/Public/Materials/MaterialLayersFunctions.h
@@ -13,6 +13,7 @@
 class FArchive;
 class FShaderKeyGenerator;
 struct FMaterialLayersFunctions;
+class UMaterialExpressionMaterialFunctionCall;
 
 UENUM()
 enum class EMaterialLayerLinkState : uint8
@@ -343,6 +344,24 @@ struct TStructOpsTypeTraits<FMaterialLayersFunctionsTree>
 	};
 };
 
+USTRUCT() 
+struct FMaterialLayersFunctionsRuntimeGraphCache
+{
+	GENERATED_BODY()
+
+
+	TArray<TObjectPtr<UMaterialExpressionMaterialFunctionCall>> LayerCallers;
+	TArray<TObjectPtr<UMaterialExpressionMaterialFunctionCall>> BlendCallers;
+
+#if WITH_EDITOR
+	FMaterialExpressionCollection								ExpressionCollection;
+
+	TArray<TObjectPtr<UMaterialExpressionMaterialFunctionCall>> NodeMaterialGraphExpressions;
+	TArray<TObjectPtr<UMaterial>>								NodePreviewMaterials;
+#endif
+};
+
+
 USTRUCT()
 struct FMaterialLayersFunctionsRuntimeData
 {
@@ -358,12 +377,17 @@ struct FMaterialLayersFunctionsRuntimeData
 	UPROPERTY()
 	FMaterialLayersFunctionsTree Tree;
 
+#if ENABLE_MATERIAL_LAYER_PROTOTYPE
+	mutable TSharedPtr<FMaterialLayersFunctionsRuntimeGraphCache> RuntimeGraphCache;
+#endif
+
 	FMaterialLayersFunctionsRuntimeData() = default;
 	FMaterialLayersFunctionsRuntimeData(const FMaterialLayersFunctionsRuntimeData& Rhs)
 		: Layers(Rhs.Layers)
 		, Blends(Rhs.Blends)
 #if ENABLE_MATERIAL_LAYER_PROTOTYPE
 		, Tree(Rhs.Tree)
+		, RuntimeGraphCache(Rhs.RuntimeGraphCache)
 #endif
 	{}
 
@@ -375,6 +399,7 @@ struct FMaterialLayersFunctionsRuntimeData
 		Blends = Rhs.Blends;
 #if ENABLE_MATERIAL_LAYER_PROTOTYPE
 		Tree = Rhs.Tree;
+		RuntimeGraphCache = Rhs.RuntimeGraphCache;
 #endif
 		return *this;
 	}
@@ -389,6 +414,7 @@ struct FMaterialLayersFunctionsRuntimeData
 		Blends.Empty();
 #if ENABLE_MATERIAL_LAYER_PROTOTYPE
 		Tree.Empty();
+		RuntimeGraphCache.Reset();
 #endif
 	}
 
@@ -632,6 +658,10 @@ struct FMaterialLayersFunctions : public FMaterialLayersFunctionsRuntimeData
 
 	ENGINE_API void MoveLayerNode(FLayerNodeId InNodeId, FLayerNodeId DstParentId, int32 InSiblingIndex, bool bShouldDuplicate);
 
+	// Cached preview material can be generated directly from the embedded runtime graph cache struct
+	// Return such a material for the specific node if available, nullptr otherwise
+	ENGINE_API UMaterial* GetRuntimeNodePreviewMaterial(FLayerNodeId InNodeId) const;
+
 #endif // WITH_EDITOR
 
 private:
