Added UMaterialInterface::IsUsingWorldPositionOffset_Concurrent to optimize and replace calls doing Material->GetRelevance_Concurrent(FeatureLevel).bUsesWorldPositionOffset.
Saves time spent in UPrimitiveComponent::SetupPrecachePSOParams, ProcessAsyncLoading and AddToWorld.

diff --git a/Engine/Source/Runtime/Engine/Private/Components/PrimitiveComponent.cpp b/Engine/Source/Runtime/Engine/Private/Components/PrimitiveComponent.cpp
--- a/Engine/Source/Runtime/Engine/Private/Components/PrimitiveComponent.cpp
+++ b/Engine/Source/Runtime/Engine/Private/Components/PrimitiveComponent.cpp
@@ -4760,13 +4760,10 @@ void UPrimitiveComponent::SetupPrecachePSOParams(FPSOPrecacheParams& Params)
 	GetUsedMaterials(UsedMaterials);
 	for (const UMaterialInterface* MaterialInterface : UsedMaterials)
 	{
-		if (MaterialInterface)
+		if (MaterialInterface && MaterialInterface->IsUsingWorldPositionOffset_Concurrent(GMaxRHIFeatureLevel))
 		{
-			if (MaterialInterface->GetRelevance_Concurrent(GMaxRHIFeatureLevel).bUsesWorldPositionOffset)
-			{
-				Params.bAnyMaterialHasWorldPositionOffset = true;
-				break;
-			}
+			Params.bAnyMaterialHasWorldPositionOffset = true;
+			break;
 		}
 	}
 }
diff --git a/Engine/Source/Runtime/Engine/Private/InstancedStaticMesh.cpp b/Engine/Source/Runtime/Engine/Private/InstancedStaticMesh.cpp
--- a/Engine/Source/Runtime/Engine/Private/InstancedStaticMesh.cpp
+++ b/Engine/Source/Runtime/Engine/Private/InstancedStaticMesh.cpp
@@ -1371,10 +1371,7 @@ void FInstancedStaticMeshSceneProxy::SetupProxy(const FInstancedStaticMeshSceneP
 			{
 				Section.Material = UMaterial::GetDefaultMaterial(MD_Surface);
 			}
-
-			bAnySegmentUsesWorldPositionOffset |= Section.Material->GetRelevance_Concurrent(GMaxRHIFeatureLevel).bUsesWorldPositionOffset;
-
-			const FMaterialCachedExpressionData& CachedMaterialData = Section.Material->GetCachedExpressionData();
+			bAnySegmentUsesWorldPositionOffset |= Section.Material->IsUsingWorldPositionOffset_Concurrent(GMaxRHIFeatureLevel);
 		}
 	}
 
diff --git a/Engine/Source/Runtime/Engine/Private/Materials/MaterialInterface.cpp b/Engine/Source/Runtime/Engine/Private/Materials/MaterialInterface.cpp
--- a/Engine/Source/Runtime/Engine/Private/Materials/MaterialInterface.cpp
+++ b/Engine/Source/Runtime/Engine/Private/Materials/MaterialInterface.cpp
@@ -429,6 +429,23 @@ bool UMaterialInterface::GetStaticComponentMaskParameterValue(const FHashedMater
 }
 #endif // WITH_EDITORONLY_DATA
 
+bool UMaterialInterface::IsUsingWorldPositionOffset_Concurrent(ERHIFeatureLevel::Type InFeatureLevel) const
+{
+	if (const UMaterial* Material = GetMaterial_Concurrent())
+	{
+		const FMaterialResource* MaterialResource = GetMaterialResource(InFeatureLevel);
+		if (!MaterialResource || !MaterialResource->GetMaterial())
+		{
+			MaterialResource = Material->GetMaterialResource(InFeatureLevel);
+		}
+		if (MaterialResource)
+		{
+			return MaterialResource->MaterialUsesWorldPositionOffset_GameThread();
+		}
+	}
+	return false;
+}
+
 FMaterialRelevance UMaterialInterface::GetRelevance_Internal(const UMaterial* Material, ERHIFeatureLevel::Type InFeatureLevel) const
 {
 	if (!Material)
diff --git a/Engine/Source/Runtime/Engine/Private/PrimitiveSceneProxy.cpp b/Engine/Source/Runtime/Engine/Private/PrimitiveSceneProxy.cpp
--- a/Engine/Source/Runtime/Engine/Private/PrimitiveSceneProxy.cpp
+++ b/Engine/Source/Runtime/Engine/Private/PrimitiveSceneProxy.cpp
@@ -651,13 +651,10 @@ FPrimitiveSceneProxy::FPrimitiveSceneProxy(const FPrimitiveSceneProxyDesc& InPro
 		InProxyDesc.GetUsedMaterials(UsedMaterials);
 		for (const UMaterialInterface* MaterialInterface : UsedMaterials)
 		{
-			if (MaterialInterface)
+			if (MaterialInterface && MaterialInterface->IsUsingWorldPositionOffset_Concurrent(FeatureLevel))
 			{
-				if (MaterialInterface->GetRelevance_Concurrent(FeatureLevel).bUsesWorldPositionOffset)
-				{
-					bAnyMaterialHasWorldPositionOffset = true;
-					break;
-				}
+				bAnyMaterialHasWorldPositionOffset = true;
+				break;
 			}
 		}
 
diff --git a/Engine/Source/Runtime/Engine/Public/Materials/MaterialInterface.h b/Engine/Source/Runtime/Engine/Public/Materials/MaterialInterface.h
--- a/Engine/Source/Runtime/Engine/Public/Materials/MaterialInterface.h
+++ b/Engine/Source/Runtime/Engine/Public/Materials/MaterialInterface.h
@@ -716,6 +716,9 @@ public:
 	/** @return The material's relevance, from concurrent render thread updates. */
 	ENGINE_API FMaterialRelevance GetRelevance_Concurrent(ERHIFeatureLevel::Type InFeatureLevel) const;
 
+	/** @return The material's relevance bUsesWorldPositionOffset, from concurrent render thread updates. */
+	ENGINE_API bool IsUsingWorldPositionOffset_Concurrent(ERHIFeatureLevel::Type InFeatureLevel) const;
+
 #if !(UE_BUILD_SHIPPING || UE_BUILD_TEST)
 	/**
 	 * Output to the log which materials and textures are used by this material.
