PR #12371: FStreamingSourceShape::ToString(): Fix Scale: %3.2f output

#rnx

diff --git a/Engine/Source/Runtime/Engine/Private/WorldPartition/WorldPartitionStreamingSource.cpp b/Engine/Source/Runtime/Engine/Private/WorldPartition/WorldPartitionStreamingSource.cpp
--- a/Engine/Source/Runtime/Engine/Private/WorldPartition/WorldPartitionStreamingSource.cpp
+++ b/Engine/Source/Runtime/Engine/Private/WorldPartition/WorldPartitionStreamingSource.cpp
@@ -32,7 +32,7 @@ FString FStreamingSourceShape::ToString() const
 		StringBuilder += TEXT("UsesGridLoadingRange ");
 		if (!FMath::IsNearlyEqual(LoadingRangeScale, 1.f))
 		{
-			StringBuilder.Appendf(TEXT("Scale: 3.2f"), LoadingRangeScale);
+			StringBuilder.Appendf(TEXT("Scale: %3.2f"), LoadingRangeScale);
 		}
 	}
 	else
