[Shaders] Add descriptive string to distributed shader job filename.

This improves the visualization of UBA traces for batches that include only a single shader job and can also improve the readability for other distributed shader controllers.
For example, a distributed shader job batch called "145.uba.in" can now be called "145.uba.Global-FTSRRejectShadingCS-455.in" derived from the debug group name "Global/FTSRRejectShadingCS/455".

#rnx

diff --git a/Engine/Source/Runtime/Engine/Private/ShaderCompiler/ShaderCompilerDistributed.cpp b/Engine/Source/Runtime/Engine/Private/ShaderCompiler/ShaderCompilerDistributed.cpp
--- a/Engine/Source/Runtime/Engine/Private/ShaderCompiler/ShaderCompilerDistributed.cpp
+++ b/Engine/Source/Runtime/Engine/Private/ShaderCompiler/ShaderCompilerDistributed.cpp
@@ -132,12 +132,51 @@ static FString BuildCompactTaskDescription(const TArray<FShaderCommonCompileJobP
 	return Description;
 }
 
+static bool IsCharValidForFilename(TCHAR Ch)
+{
+	const TCHAR* ValidSpecialChars = TEXT("_-+()[]");
+	return TChar<TCHAR>::IsAlnum(Ch) || FCString::Strchr(ValidSpecialChars, Ch) != nullptr;
+}
+
+static void ConvertDebugNameToFilename(TStringBuilder<1024>& StringBuilder, const FString& Name)
+{
+	for (TCHAR Ch : Name)
+	{
+		StringBuilder << (IsCharValidForFilename(Ch) ? Ch : TEXT('-'));
+	}
+}
+
+static void BuildDescriptiveTaskFilename(TStringBuilder<1024>& StringBuilder, const FShaderCompileJob& SingleJob)
+{
+	ConvertDebugNameToFilename(StringBuilder, SingleJob.Input.DebugGroupName);
+}
+
 void FShaderCompileDistributedThreadRunnable_Interface::DispatchShaderCompileJobsBatch(TArray<FShaderCommonCompileJobPtr>& JobsToSerialize)
 {
-	const FString BaseFilePath = CachedController.CreateUniqueFilePath();
+	// Generate unique filename for shader compiler I/O files
+	FString BaseFilePath = CachedController.CreateUniqueFilePath();
+
+	TStringBuilder<1024> BaseFilePathBuilder;
+	if (JobsToSerialize.Num() == 1)
+	{
+		// Decorate filename with single job description
+		if (const FShaderCompileJob* FirstSingleJob = JobsToSerialize[0]->GetSingleShaderJob())
+		{
+			BaseFilePathBuilder << TEXT(".");
+			BuildDescriptiveTaskFilename(BaseFilePathBuilder, *FirstSingleJob);
+		}
+	}
+	else
+	{
+		// Decorate filename with number of jobs
+		BaseFilePathBuilder << TEXT(".j-") << JobsToSerialize.Num();
+	}
+	BaseFilePath += BaseFilePathBuilder.ToView();
+
 	FString InputFilePath = BaseFilePath + TEXT(".in");
 	FString OutputFilePath = BaseFilePath + TEXT(".out");
 
+	// Set up remote task
 	const FString WorkingDirectory = FPaths::GetPath(InputFilePath);
 
 	// Serialize the jobs to the input file
