GetLocalBounds now considers bUseAttachParentBound so it is consistent with USceneComponent::UpdateBounds.

Set p.LocalBoundsUseAttachParentBound to 0 to restore the old behaviour: it will ignore the parent and compute the local bounds through CalcBounds

[RN] minor framework

diff --git a/Engine/Source/Runtime/Engine/Private/Components/SceneComponent.cpp b/Engine/Source/Runtime/Engine/Private/Components/SceneComponent.cpp
--- a/Engine/Source/Runtime/Engine/Private/Components/SceneComponent.cpp
+++ b/Engine/Source/Runtime/Engine/Private/Components/SceneComponent.cpp
@@ -3771,12 +3771,26 @@ FScopedPreventAttachedComponentMove::~FScopedPreventAttachedComponentMove()
 	}
 }
 
+bool GLocalBoundsUseAttachParentBound = true;
+static FAutoConsoleVariableRef CVarLocalBoundsUseAttachParentBound (
+	TEXT("p.LocalBoundsUseAttachParentBound"),
+	GLocalBoundsUseAttachParentBound,
+	TEXT("If enabled, Components with 'Use Attach Parent Bound' will compute their local bounds based their parent's bounds"),
+	ECVF_Default
+);
+
 FBoxSphereBounds USceneComponent::GetLocalBounds() const
 {
 	if (bComputeFastLocalBounds)
 	{
 		return Bounds.TransformBy(ComponentToWorld.Inverse());
 	}
+
+	if (GLocalBoundsUseAttachParentBound && bUseAttachParentBound && GetAttachParent() != nullptr)
+	{
+		return GetAttachParent()->Bounds.TransformBy(ComponentToWorld.Inverse());
+	}
+
 	return CalcBounds(FTransform::Identity);
 }
 
