[World Partition]
- Added RF_NonPIEDuplicateTransient support for transient actors:
  - They get duplicated with the world in PIE for non-partitioned worlds.
  - They get part of the streaming process for partitioned worlds, similar to dirty actors.
- Added transient flag AActor::bIgnoreInPIE to ignore actors from the PIE session.

diff --git a/Engine/Source/Runtime/Engine/Classes/GameFramework/Actor.h b/Engine/Source/Runtime/Engine/Classes/GameFramework/Actor.h
--- a/Engine/Source/Runtime/Engine/Classes/GameFramework/Actor.h
+++ b/Engine/Source/Runtime/Engine/Classes/GameFramework/Actor.h
@@ -329,6 +329,11 @@ private:
 public:
 	UPROPERTY(EditAnywhere, AdvancedDisplay, Category = LevelInstance, meta = (Tooltip = "If checked, this Actor will only get loaded in a main world (persistent level), it will not be loaded through Level Instances."))
 	uint8 bIsMainWorldOnly : 1;
+
+	/** Experimental: this will force this actor to be ignored in PIE, will not work if the actor has references to other actors without the flag set, use with care. */
+	UPROPERTY(Transient)
+	uint8 bIgnoreInPIE : 1;
+
 private:
 #endif
 
diff --git a/Engine/Source/Runtime/Engine/Private/Level.cpp b/Engine/Source/Runtime/Engine/Private/Level.cpp
--- a/Engine/Source/Runtime/Engine/Private/Level.cpp
+++ b/Engine/Source/Runtime/Engine/Private/Level.cpp
@@ -648,6 +648,13 @@ void ULevel::Serialize( FArchive& Ar )
 
 			if (Actor->HasAnyFlags(RF_Transient))
 			{
+#if WITH_EDITOR
+				// When duplicating transient actors for PIE, allow 
+				if (ensure(!Actor->IsPackageExternal()) && Actor->HasAnyFlags(RF_NonPIEDuplicateTransient) && (Ar.GetPortFlags() & PPF_DuplicateForPIE))
+				{
+					return !bIsPartitioned;
+				}
+#endif
 				return false;
 			}
 
@@ -671,6 +678,11 @@ void ULevel::Serialize( FArchive& Ar )
 					return Actor->IsForceExternalActorLevelReferenceForPIE();
 				}
 			}
+
+			if ((Ar.GetPortFlags() & PPF_DuplicateForPIE) && Actor->bIgnoreInPIE)
+			{
+				return false;
+			}
 #endif
 			return true;
 		});
diff --git a/Engine/Source/Runtime/Engine/Private/WorldPartition/ActorDescContainer.cpp b/Engine/Source/Runtime/Engine/Private/WorldPartition/ActorDescContainer.cpp
--- a/Engine/Source/Runtime/Engine/Private/WorldPartition/ActorDescContainer.cpp
+++ b/Engine/Source/Runtime/Engine/Private/WorldPartition/ActorDescContainer.cpp
@@ -352,8 +352,16 @@ bool UActorDescContainer::IsActorDescHandled(const AActor* InActor, bool bInUseL
 	{
 		const FName LoadedPackageName = InActor->GetPackage()->GetLoadedPath().GetPackageFName();
 		const FString ActorPackageName = bInUseLoadedPath && !LoadedPackageName.IsNone() ? LoadedPackageName.ToString() : InActor->GetPackage()->GetName();
-		const FString ExternalActorPath = GetExternalActorPath() / TEXT("");
-		return ActorPackageName.StartsWith(ExternalActorPath);
+
+		if (InActor->GetExternalPackage())
+		{
+			const FString ExternalActorPath = GetExternalActorPath() / TEXT("");
+			return ActorPackageName.StartsWith(ExternalActorPath);
+		}
+		else if (InActor->HasAllFlags(RF_Transient))
+		{
+			return ActorPackageName == LoadedPackageName.ToString();
+		}
 	}
 	return false;
 }
diff --git a/Engine/Source/Runtime/Engine/Private/WorldPartition/WorldPartitionStreamingGeneration.cpp b/Engine/Source/Runtime/Engine/Private/WorldPartition/WorldPartitionStreamingGeneration.cpp
--- a/Engine/Source/Runtime/Engine/Private/WorldPartition/WorldPartitionStreamingGeneration.cpp
+++ b/Engine/Source/Runtime/Engine/Private/WorldPartition/WorldPartitionStreamingGeneration.cpp
@@ -778,6 +778,11 @@ class FWorldPartitionStreamingGenerator
 			{
 				if (OutActor)
 				{
+					if (bIsPIE && Actor->bIgnoreInPIE)
+					{
+						return false;
+					}
+
 					*OutActor = Actor;
 				}
 			}
@@ -834,7 +839,7 @@ class FWorldPartitionStreamingGenerator
 					}
 
 					// Dirty actors
-					if (bShouldHandleUnsavedActors && (bIsTempContainerPackage || Actor->GetPackage()->IsDirty()))
+					if (bShouldHandleUnsavedActors && (bIsTempContainerPackage || (Actor->GetPackage()->IsDirty() || Actor->HasAllFlags(RF_Transient | RF_NonPIEDuplicateTransient))))
 					{
 						// Dirty, unsaved actor for PIE
 						TUniquePtr<FStreamingGenerationUnsavedDirtyActorDescInstance>& UnsavedDirtyRef = OutUnsavedDirtyInstances.Add_GetRef(FStreamingGenerationUnsavedDirtyActorDescInstance::Create(Actor, InActorDescCollection));
@@ -880,7 +885,7 @@ class FWorldPartitionStreamingGenerator
 				// Here, FindHandlingContainer is used to make sure that the actor is handled by the collection
 				// The main reason is that UWorldPartition::CheckForErrors currently builds a collection per ActorDescContainer of the WorldPartition.
 				// This is probably a limitation introduced by ContentBundles. 
-				if (IsValid(Actor) && Actor->IsPackageExternal() && Actor->IsMainPackageActor()
+				if (IsValid(Actor) && ((Actor->IsPackageExternal() && Actor->IsMainPackageActor()) || Actor->HasAllFlags(RF_Transient | RF_NonPIEDuplicateTransient))
 					&& InActorDescCollection.FindHandlingContainerInstance(Actor)
 					&& !InActorDescCollection.GetActorDescInstance(Actor->GetActorGuid()))
 				{
@@ -888,7 +893,7 @@ class FWorldPartitionStreamingGenerator
 					if (ShouldRegisterActorDesc(UnsavedViewPtr.Get()))
 					{
 						FStreamingGenerationActorDescView ModifiedActorDescView(OutActorDescViewMap, UnsavedViewPtr.Get(), true);
-						OutUnsavedDirtyInstances.Add(MoveTemp(UnsavedViewPtr));						
+						OutUnsavedDirtyInstances.Add(MoveTemp(UnsavedViewPtr));
 						RegisterActorDescView(MoveTemp(ModifiedActorDescView));
 					}
 				}
