TextureRenderTarget : select CompressionSetting more suited to render target format

diff --git a/Engine/Source/Runtime/Engine/Private/ImageUtils.cpp b/Engine/Source/Runtime/Engine/Private/ImageUtils.cpp
--- a/Engine/Source/Runtime/Engine/Private/ImageUtils.cpp
+++ b/Engine/Source/Runtime/Engine/Private/ImageUtils.cpp
@@ -923,6 +923,10 @@ UTexture * FImageUtils::CreateTexture(ETextureClass TextureClass, const FImageVi
 
 	if ( DoPostEditChange )
 	{
+		// note: Tex->CompressionSettings , SRGB, etc. have not been set yet
+		//	calling PostEditChange usually doesn't make sense
+		//	you should first set those things up, then call PostEditChange yourself outside
+ 
 		Tex->PostEditChange();
 	}
 
diff --git a/Engine/Source/Runtime/Engine/Private/TextureRenderTarget.cpp b/Engine/Source/Runtime/Engine/Private/TextureRenderTarget.cpp
--- a/Engine/Source/Runtime/Engine/Private/TextureRenderTarget.cpp
+++ b/Engine/Source/Runtime/Engine/Private/TextureRenderTarget.cpp
@@ -193,6 +193,18 @@ bool UTextureRenderTarget::UpdateTexture(UTexture* InTexture, EConstructTextureF
 	}
 
 	TextureCompressionSettings CompressionSettingsForTexture = IsHDR(PixelFormat) ? TC_HDR : TC_Default;
+	if ( PixelFormat == PF_R16F || PixelFormat == PF_R32_FLOAT )
+	{
+		// single 16F instead of TC_HDR which is RGBA16F
+		CompressionSettingsForTexture = TC_HalfFloat;
+	}
+	#if 0
+	// @todo Oodle : this seems to work fine but not sure it's desirable because of the way format affects usage in Unreal
+	else if ( PixelFormat == PF_G8 || PixelFormat == PF_G16 )
+	{
+		CompressionSettingsForTexture = TC_Grayscale;
+	}
+	#endif
 
 	const int32 NumMips = 1; // There's only support for mip 0 ATM
 
