Anim Montage: Check for valid MontageInstanceToUpdate

diff --git a/Engine/Source/Runtime/Engine/Private/Animation/AnimMontage.cpp b/Engine/Source/Runtime/Engine/Private/Animation/AnimMontage.cpp
--- a/Engine/Source/Runtime/Engine/Private/Animation/AnimMontage.cpp
+++ b/Engine/Source/Runtime/Engine/Private/Animation/AnimMontage.cpp
@@ -2921,7 +2921,10 @@ UAnimMontage* FAnimMontageInstance::SetSequencerMontagePosition(FName SlotName,
 				AnimInst->Montage_Play(PlayingMontage, 1.f, EMontagePlayReturnType::MontageLength, 0.f, false);
 				MontageInstanceToUpdate = AnimInst->GetActiveInstanceForMontage(PlayingMontage);
 				// this is sequencer set up, we disable auto blend out
-				MontageInstanceToUpdate->bEnableAutoBlendOut = false;
+				if (MontageInstanceToUpdate)
+				{
+					MontageInstanceToUpdate->bEnableAutoBlendOut = false;
+				}
 			}
 		}
 
