Added logging to debug a rare bug where DataLayerInstances in external packages are not properly de-externalized at cook.

diff --git a/Engine/Source/Runtime/Engine/Private/WorldPartition/DataLayer/WorldDataLayers.cpp b/Engine/Source/Runtime/Engine/Private/WorldPartition/DataLayer/WorldDataLayers.cpp
--- a/Engine/Source/Runtime/Engine/Private/WorldPartition/DataLayer/WorldDataLayers.cpp
+++ b/Engine/Source/Runtime/Engine/Private/WorldPartition/DataLayer/WorldDataLayers.cpp
@@ -1190,7 +1190,9 @@ void AWorldDataLayers::OnDataLayerManagerInitialized()
 	if (IsRunningCookCommandlet())
 	{
 		// Embed external DataLayerInstances when cooking
+		UE_CLOG(IsUsingExternalPackageDataLayerInstances(), LogWorldPartition, Display, TEXT("Internalizing DataLayerInstances in %s in package (%s)"), *GetPathName(), *GetPackage()->GetName());
 		SetUseExternalPackageDataLayerInstances(false);
+		UE_CLOG(IsUsingExternalPackageDataLayerInstances(), LogWorldPartition, Error, TEXT("Error while internalizing DataLayerInstances."));
 	}
 
 	if (RootExternalDataLayerInstance)
@@ -1488,7 +1490,12 @@ bool AWorldDataLayers::SetUseExternalPackageDataLayerInstances(bool bInNewValue,
 	for (UDataLayerInstance* DataLayerInstance : GetDataLayerInstances())
 	{
 		check(DataLayerInstance->GetDirectOuterWorldDataLayers() == this);
+		UPackage* PreviousPackage = DataLayerInstance->GetPackage();
 		FExternalPackageHelper::SetPackagingMode(DataLayerInstance, this, bInNewValue);
+		UE_LOG(LogWorldPartition, Display, TEXT("DataLayerInstance %s changed package from %s to %s"),
+			*DataLayerInstance->GetPathName(),
+			*PreviousPackage->GetName(),
+			*DataLayerInstance->GetPackage()->GetName());
 	}
 	bUseExternalPackageDataLayerInstances = bInNewValue;
 	Swap(DataLayerInstances, ExternalPackageDataLayerInstances);
