Added specific names instead of "FPhysicsFieldResource" for buffers in order to help debugging
Transition Clipmapbuffer to SRVMask in after the Init due to a rhi validation errorfor usage of the buffers in a dispatch call

diff --git a/Engine/Source/Runtime/Engine/Private/PhysicsField/PhysicsFieldComponent.cpp b/Engine/Source/Runtime/Engine/Private/PhysicsField/PhysicsFieldComponent.cpp
--- a/Engine/Source/Runtime/Engine/Private/PhysicsField/PhysicsFieldComponent.cpp
+++ b/Engine/Source/Runtime/Engine/Private/PhysicsField/PhysicsFieldComponent.cpp
@@ -113,7 +113,7 @@ FAutoConsoleVariableRef CVarPhysicsFieldEnableCulling(
 */
 
 template<typename BufferType, int ElementSize, EPixelFormat PixelFormat>
-void InitInternalBuffer(FRHICommandListBase& RHICmdList, const uint32 ElementCount, FRWBuffer& OutputBuffer)
+void InitInternalBuffer(FRHICommandListBase& RHICmdList, const TCHAR* DebugName, const uint32 ElementCount, FRWBuffer& OutputBuffer)
 {
 	if (ElementCount > 0)
 	{
@@ -122,7 +122,7 @@ void InitInternalBuffer(FRHICommandListBase& RHICmdList, const uint32 ElementCou
 		
 		{
 			LLM_SCOPE(ELLMTag::Physics);
-			OutputBuffer.Initialize(RHICmdList, TEXT("FPhysicsFieldResource"), sizeof(BufferType), BufferCount, PixelFormat, BUF_Static);
+			OutputBuffer.Initialize(RHICmdList, DebugName, sizeof(BufferType), BufferCount, PixelFormat, BUF_Static);
 		}
 
 		if (OutputBuffer.UAV)
@@ -144,7 +144,7 @@ void InitInternalBuffer(FRHICommandListBase& RHICmdList, const uint32 ElementCou
 }
 
 template<typename BufferType, int ElementSize, EPixelFormat PixelFormat>
-void UpdateInternalBuffer(FRHICommandListBase& RHICmdList, const uint32 ElementCount, const BufferType* InputData, FRWBuffer& OutputBuffer, const bool bInitField = false)
+void UpdateInternalBuffer(FRHICommandListBase& RHICmdList, const TCHAR* DebugName, const uint32 ElementCount, const BufferType* InputData, FRWBuffer& OutputBuffer, const bool bInitField = false)
 {
 	if (ElementCount > 0 && InputData)
 	{
@@ -154,7 +154,7 @@ void UpdateInternalBuffer(FRHICommandListBase& RHICmdList, const uint32 ElementC
 		if (bInitField)
 		{
 			LLM_SCOPE(ELLMTag::Physics);
-			OutputBuffer.Initialize(RHICmdList, TEXT("FPhysicsFieldResource"), sizeof(BufferType), BufferCount, PixelFormat, BUF_Static);
+			OutputBuffer.Initialize(RHICmdList, DebugName, sizeof(BufferType), BufferCount, PixelFormat, BUF_Static);
 		}
 
 		void* OutputData = RHICmdList.LockBuffer(OutputBuffer.Buffer, 0, BufferBytes, RLM_WriteOnly);
@@ -423,21 +423,24 @@ void FPhysicsFieldResource::InitRHI(FRHICommandListBase& RHICmdList)
 	SCOPE_CYCLE_COUNTER(STAT_PhysicsFields_UpdateResource_RT);
 
 	const int32 DatasCount = FieldInfos.ClipmapCount * FieldInfos.TargetCount;
-	InitInternalBuffer<int32, 1, EPixelFormat::PF_R32_SINT>(RHICmdList, EFieldPhysicsType::Field_PhysicsType_Max + 1, TargetsOffsets);
+	InitInternalBuffer<int32, 1, EPixelFormat::PF_R32_SINT>(RHICmdList, TEXT("PhysicsField_TargetOffsets"), EFieldPhysicsType::Field_PhysicsType_Max + 1, TargetsOffsets);
 
-	InitInternalBuffer<FVector4f, 1, EPixelFormat::PF_A32B32G32R32F>(RHICmdList, 1, BoundsMin);
-	InitInternalBuffer<FVector4f, 1, EPixelFormat::PF_A32B32G32R32F>(RHICmdList, 1, BoundsMax);
+	InitInternalBuffer<FVector4f, 1, EPixelFormat::PF_A32B32G32R32F>(RHICmdList, TEXT("PhysicsField_BoundsMin"), 1, BoundsMin);
+	InitInternalBuffer<FVector4f, 1, EPixelFormat::PF_A32B32G32R32F>(RHICmdList, TEXT("PhysicsField_BoundsMax"), 1, BoundsMax);
 
-	InitInternalBuffer<float, 1, EPixelFormat::PF_R32_FLOAT>(RHICmdList, 1, NodesParams);
-	InitInternalBuffer<int32, 1, EPixelFormat::PF_R32_SINT>(RHICmdList, 1, NodesOffsets);
+	InitInternalBuffer<float, 1, EPixelFormat::PF_R32_FLOAT>(RHICmdList, TEXT("PhysicsField_NodesParams"), 1, NodesParams);
+	InitInternalBuffer<int32, 1, EPixelFormat::PF_R32_SINT>(RHICmdList, TEXT("PhysicsField_NodesOffsets"), 1, NodesOffsets);
 
 	if (FieldInfos.bBuildClipmap)
 	{
 		const int32 CellsCount = FieldInfos.ClipmapCount * EFieldPhysicsType::Field_PhysicsType_Max;
-		InitInternalBuffer<float, 1, EPixelFormat::PF_R32_FLOAT>(RHICmdList, FieldInfos.ClipmapResolution * FieldInfos.ClipmapResolution * FieldInfos.ClipmapResolution * DatasCount, ClipmapBuffer);
-		InitInternalBuffer<int32, 1, EPixelFormat::PF_R32_SINT>(RHICmdList, CellsCount + 1, CellsOffsets);
-		InitInternalBuffer<FIntVector4, 1, EPixelFormat::PF_R32G32B32A32_UINT>(RHICmdList, CellsCount, CellsMin);
-		InitInternalBuffer<FIntVector4, 1, EPixelFormat::PF_R32G32B32A32_UINT>(RHICmdList, CellsCount, CellsMax);
+		InitInternalBuffer<float, 1, EPixelFormat::PF_R32_FLOAT>(RHICmdList, TEXT("PhysicsField_ClipmapBuffer"), FieldInfos.ClipmapResolution * FieldInfos.ClipmapResolution * FieldInfos.ClipmapResolution * DatasCount, ClipmapBuffer);
+		InitInternalBuffer<int32, 1, EPixelFormat::PF_R32_SINT>(RHICmdList, TEXT("PhysicsField_CellsOffsets"), CellsCount + 1, CellsOffsets);
+		InitInternalBuffer<FIntVector4, 1, EPixelFormat::PF_R32G32B32A32_UINT>(RHICmdList, TEXT("PhysicsField_CellsMin"), CellsCount, CellsMin);
+		InitInternalBuffer<FIntVector4, 1, EPixelFormat::PF_R32G32B32A32_UINT>(RHICmdList, TEXT("PhysicsField_CellsMax"), CellsCount, CellsMax );
+
+		FRHIComputeCommandList& RHICmdListCompute = FRHIComputeCommandList::Get(RHICmdList);
+		RHICmdListCompute.Transition(FRHITransitionInfo(ClipmapBuffer.UAV, ERHIAccess::SRVCompute, ERHIAccess::SRVMask));
 	}
 }
 
@@ -538,9 +541,9 @@ void FPhysicsFieldResource::UpdateResource(FRHICommandList& RHICmdList,
 		RHI_BREADCRUMB_EVENT_STAT(RHICmdList, PhysicsFields_UpdateBuffers, "PhysicsFields_UpdateBuffers");
 		SCOPED_GPU_STAT(RHICmdList, PhysicsFields_UpdateBuffers);
 
-		UpdateInternalBuffer<float, 1, EPixelFormat::PF_R32_FLOAT>(RHICmdList, NodesParamsDatas.Num(), NodesParamsDatas.GetData(), NodesParams, true);
-		UpdateInternalBuffer<int32, 1, EPixelFormat::PF_R32_SINT>(RHICmdList, NodesOffsetsDatas.Num(), NodesOffsetsDatas.GetData(), NodesOffsets, true);
-		UpdateInternalBuffer<int32, 1, EPixelFormat::PF_R32_SINT>(RHICmdList, TargetsOffsetsDatas.Num(), TargetsOffsetsDatas.GetData(), TargetsOffsets);
+		UpdateInternalBuffer<float, 1, EPixelFormat::PF_R32_FLOAT>(RHICmdList, TEXT("PhysicsField_NodeParams"), NodesParamsDatas.Num(), NodesParamsDatas.GetData(), NodesParams, true);
+		UpdateInternalBuffer<int32, 1, EPixelFormat::PF_R32_SINT>(RHICmdList, TEXT("PhysicsField_NodeOffset"), NodesOffsetsDatas.Num(), NodesOffsetsDatas.GetData(), NodesOffsets, true);
+		UpdateInternalBuffer<int32, 1, EPixelFormat::PF_R32_SINT>(RHICmdList, TEXT("PhysicsField_TargetOffset"), TargetsOffsetsDatas.Num(), TargetsOffsetsDatas.GetData(), TargetsOffsets);
 
 		FieldInfos.TimeSeconds = TimeSeconds;
 		for (uint32 BoundIndex = 0; BoundIndex < EFieldPhysicsType::Field_PhysicsType_Max + 1; ++BoundIndex)
@@ -549,8 +552,8 @@ void FPhysicsFieldResource::UpdateResource(FRHICommandList& RHICmdList,
 		}
 		// LWC_TODO: Perf pessimization
 		// LWC_TODO: Precision loss
-		UpdateInternalBuffer<FVector4f, 1, EPixelFormat::PF_A32B32G32R32F>(RHICmdList, BoundsMinDatas.Num(), UE::LWC::ConvertArrayType<FVector4f>(BoundsMinDatas).GetData(), BoundsMin, true);	
-		UpdateInternalBuffer<FVector4f, 1, EPixelFormat::PF_A32B32G32R32F>(RHICmdList, BoundsMaxDatas.Num(),  UE::LWC::ConvertArrayType<FVector4f>(BoundsMaxDatas).GetData(), BoundsMax, true);
+		UpdateInternalBuffer<FVector4f, 1, EPixelFormat::PF_A32B32G32R32F>(RHICmdList, TEXT("PhysicsField_BoundsMin"), BoundsMinDatas.Num(), UE::LWC::ConvertArrayType<FVector4f>(BoundsMinDatas).GetData(), BoundsMin, true);
+		UpdateInternalBuffer<FVector4f, 1, EPixelFormat::PF_A32B32G32R32F>(RHICmdList, TEXT("PhysicsField_BoundsMax"), BoundsMaxDatas.Num(),  UE::LWC::ConvertArrayType<FVector4f>(BoundsMaxDatas).GetData(), BoundsMax, true);
 	}
 
 	if (FieldInfos.bBuildClipmap)
@@ -581,9 +584,9 @@ void FPhysicsFieldResource::UpdateResource(FRHICommandList& RHICmdList,
 			FieldInfos.ClipmapCenter = FieldInfos.ViewOrigin;
 			UpdateBounds(TargetsMinDatas, TargetsMaxDatas, TargetsOffsetsDatas, BoundsOffsetsDatas);
 
-			UpdateInternalBuffer<int32, 1, EPixelFormat::PF_R32_SINT>(RHICmdList, FieldInfos.CellsOffsets.Num(), FieldInfos.CellsOffsets.GetData(), CellsOffsets);
-			UpdateInternalBuffer<FIntVector4, 1, EPixelFormat::PF_R32G32B32A32_UINT>(RHICmdList, FieldInfos.CellsMin.Num(), FieldInfos.CellsMin.GetData(), CellsMin);
-			UpdateInternalBuffer<FIntVector4, 1, EPixelFormat::PF_R32G32B32A32_UINT>(RHICmdList, FieldInfos.CellsMax.Num(), FieldInfos.CellsMax.GetData(), CellsMax);
+			UpdateInternalBuffer<int32, 1, EPixelFormat::PF_R32_SINT>(RHICmdList, TEXT("PhysicsField_CellsOffsets"), FieldInfos.CellsOffsets.Num(), FieldInfos.CellsOffsets.GetData(), CellsOffsets);
+			UpdateInternalBuffer<FIntVector4, 1, EPixelFormat::PF_R32G32B32A32_UINT>(RHICmdList, TEXT("PhysicsField_CellsMin"), FieldInfos.CellsMin.Num(), FieldInfos.CellsMin.GetData(), CellsMin);
+			UpdateInternalBuffer<FIntVector4, 1, EPixelFormat::PF_R32G32B32A32_UINT>(RHICmdList, TEXT("PhysicsField_CellsMax"), FieldInfos.CellsMax.Num(), FieldInfos.CellsMax.GetData(), CellsMax);
 		}
 		{
 			RHI_BREADCRUMB_EVENT_STAT(RHICmdList, PhysicsFields_BuildClipmap, "PhysicsFields_BuildClipmap");
