Removing PreCulled Indices
- UpdatePreCulledData was the only function that set pre culled indices and that was deprecated in 4.27
- Blanket removing all data that was set from that function.

diff --git a/Engine/Source/Runtime/Engine/Classes/Components/StaticMeshComponent.h b/Engine/Source/Runtime/Engine/Classes/Components/StaticMeshComponent.h
--- a/Engine/Source/Runtime/Engine/Classes/Components/StaticMeshComponent.h
+++ b/Engine/Source/Runtime/Engine/Classes/Components/StaticMeshComponent.h
@@ -739,9 +739,6 @@ public:
 	ENGINE_API void RemoveInstanceVertexColors();
 #endif
 
-	UE_DEPRECATED(4.27, "This function is no longer used")
-	ENGINE_API void UpdatePreCulledData(int32 LODIndex, const TArray<uint32>& PreCulledData, const TArray<int32>& NumTrianglesPerSection);
-
 #if WITH_EDITORONLY_DATA
 	/**
 	*	Sets the value of the SectionIndexPreview flag and reattaches the component as necessary.
diff --git a/Engine/Source/Runtime/Engine/Private/Components/StaticMeshComponent.cpp b/Engine/Source/Runtime/Engine/Private/Components/StaticMeshComponent.cpp
--- a/Engine/Source/Runtime/Engine/Private/Components/StaticMeshComponent.cpp
+++ b/Engine/Source/Runtime/Engine/Private/Components/StaticMeshComponent.cpp
@@ -1817,70 +1817,11 @@ void UStaticMeshComponent::PrivateFixupOverrideColors()
 }
 #endif // WITH_EDITOR
 
-float GKeepPreCulledIndicesThreshold = .95f;
-
-FAutoConsoleVariableRef CKeepPreCulledIndicesThreshold(
-	TEXT("r.KeepPreCulledIndicesThreshold"),
-	GKeepPreCulledIndicesThreshold,
-	TEXT(""),
-	ECVF_Scalability | ECVF_RenderThreadSafe
-	);
-
-void UStaticMeshComponent::UpdatePreCulledData(int32 LODIndex, const TArray<uint32>& PreCulledData, const TArray<int32>& NumTrianglesPerSection)
-{
-	const FStaticMeshLODResources& StaticMeshLODResources = GetStaticMesh()->GetRenderData()->LODResources[LODIndex];
-
-	int32 NumOriginalTriangles = 0;
-	int32 NumVisibleTriangles = 0;
-
-	for (int32 SectionIndex = 0; SectionIndex < StaticMeshLODResources.Sections.Num(); SectionIndex++)
-	{
-		const FStaticMeshSection& Section = StaticMeshLODResources.Sections[SectionIndex];
-		NumOriginalTriangles += Section.NumTriangles;
-		NumVisibleTriangles += NumTrianglesPerSection[SectionIndex];
-	}
-
-	if (NumVisibleTriangles / (float)NumOriginalTriangles < GKeepPreCulledIndicesThreshold)
-	{
-		SetLODDataCount(LODIndex + 1, LODData.Num());
-
-		DEC_DWORD_STAT_BY(STAT_StaticMeshPreCulledIndexMemory, LODData[LODIndex].PreCulledIndexBuffer.GetAllocatedSize());
-		//@todo - game thread
-		check(IsInRenderingThread());
-		LODData[LODIndex].PreCulledIndexBuffer.ReleaseResource();
-		LODData[LODIndex].PreCulledIndexBuffer.SetIndices(PreCulledData, EIndexBufferStride::AutoDetect);
-		LODData[LODIndex].PreCulledIndexBuffer.InitResource(FRHICommandListImmediate::Get());
-
-		INC_DWORD_STAT_BY(STAT_StaticMeshPreCulledIndexMemory, LODData[LODIndex].PreCulledIndexBuffer.GetAllocatedSize());
-		LODData[LODIndex].PreCulledSections.Empty(StaticMeshLODResources.Sections.Num());
-
-		int32 FirstIndex = 0;
-
-		for (int32 SectionIndex = 0; SectionIndex < StaticMeshLODResources.Sections.Num(); SectionIndex++)
-		{
-			FPreCulledStaticMeshSection PreCulledSection;
-			PreCulledSection.FirstIndex = FirstIndex;
-			PreCulledSection.NumTriangles = NumTrianglesPerSection[SectionIndex];
-			FirstIndex += PreCulledSection.NumTriangles * 3;
-			LODData[LODIndex].PreCulledSections.Add(PreCulledSection);
-		}
-	}
-	else if (LODIndex < LODData.Num())
-	{
-		LODData[LODIndex].PreCulledIndexBuffer.ReleaseResource();
-		TArray<uint32> EmptyIndices;
-		LODData[LODIndex].PreCulledIndexBuffer.SetIndices(EmptyIndices, EIndexBufferStride::AutoDetect);
-		LODData[LODIndex].PreCulledSections.Empty(StaticMeshLODResources.Sections.Num());
-	}
-}
-
 void UStaticMeshComponent::ReleaseResources()
 {
 	for(int32 LODIndex = 0;LODIndex < LODData.Num();LODIndex++)
 	{
 		LODData[LODIndex].BeginReleaseOverrideVertexColors();
-		DEC_DWORD_STAT_BY(STAT_StaticMeshPreCulledIndexMemory, LODData[LODIndex].PreCulledIndexBuffer.GetAllocatedSize());
-		BeginReleaseResource(&LODData[LODIndex].PreCulledIndexBuffer);
 	}
 
 	DetachFence.BeginFence();
diff --git a/Engine/Source/Runtime/Engine/Private/StaticMeshSceneProxy.cpp b/Engine/Source/Runtime/Engine/Private/StaticMeshSceneProxy.cpp
--- a/Engine/Source/Runtime/Engine/Private/StaticMeshSceneProxy.cpp
+++ b/Engine/Source/Runtime/Engine/Private/StaticMeshSceneProxy.cpp
@@ -76,15 +76,6 @@ static FAutoConsoleCommand GToggleShadowIndexBuffersCmd(
 	FConsoleCommandDelegate::CreateStatic(ToggleShadowIndexBuffers)
 	);
 
-static bool GUsePreCulledIndexBuffer = true;
-
-void TogglePreCulledIndexBuffers( UWorld* InWorld )
-{
-	FGlobalComponentRecreateRenderStateContext Context;
-	FlushRenderingCommands();
-	GUsePreCulledIndexBuffer = !GUsePreCulledIndexBuffer;
-}
-
 static bool GStaticMeshComponentBoostPSOPrecachePri = false;
 static FAutoConsoleVariableRef CVarStaticMeshComponentBoostPSOPrecachePri(
 	TEXT("r.PSOPrecache.StaticMeshComponentPSOPrecachePriority"),
@@ -95,10 +86,6 @@ static FAutoConsoleVariableRef CVarStaticMeshComponentBoostPSOPrecachePri(
 	ECVF_Default
 );
 
-FAutoConsoleCommandWithWorld GToggleUsePreCulledIndexBuffersCmd(
-	TEXT("r.TogglePreCulledIndexBuffers"),
-	TEXT("Toggles use of preculled index buffers from the command 'PreCullIndexBuffers'"),
-	FConsoleCommandWithWorldDelegate::CreateStatic(TogglePreCulledIndexBuffers));
 
 static FAutoConsoleCommand GToggleReversedIndexBuffersCmd(
 	TEXT("ToggleReversedIndexBuffers"),
@@ -1009,7 +996,7 @@ uint32 FStaticMeshSceneProxy::SetMeshElementGeometrySource(
 	int32 SectionIndex,
 	bool bWireframe,
 	bool bUseReversedIndices,
-	bool bAllowPreCulledIndices,
+	bool bAllowPreCulledIndices, // unused
 	const FVertexFactory* VertexFactory,
 	FMeshBatch& OutMeshBatch) const
 {
@@ -1027,9 +1014,6 @@ uint32 FStaticMeshSceneProxy::SetMeshElementGeometrySource(
 	FMeshBatchElement& OutMeshBatchElement = OutMeshBatch.Elements[0];
 	uint32 NumPrimitives = 0;
 
-	const bool bHasPreculledTriangles = LODInfo.Sections[SectionIndex].NumPreCulledTriangles >= 0;
-	const bool bUsePreculledIndices = bAllowPreCulledIndices && GUsePreCulledIndexBuffer && bHasPreculledTriangles;
-
 	if (bWireframe)
 	{
 		if (LODModel.AdditionalIndexBuffers && LODModel.AdditionalIndexBuffers->WireframeIndexBuffer.IsInitialized())
@@ -1043,18 +1027,9 @@ uint32 FStaticMeshSceneProxy::SetMeshElementGeometrySource(
 		{
 			OutMeshBatch.Type = PT_TriangleList;
 
-			if (bUsePreculledIndices)
-			{
-				OutMeshBatchElement.IndexBuffer = LODInfo.PreCulledIndexBuffer;
-				OutMeshBatchElement.FirstIndex = 0;
-				NumPrimitives = LODInfo.PreCulledIndexBuffer->GetNumIndices() / 3;
-			}
-			else
-			{
-				OutMeshBatchElement.FirstIndex = 0;
-				OutMeshBatchElement.IndexBuffer = &LODModel.IndexBuffer;
-				NumPrimitives = LODModel.IndexBuffer.GetNumIndices() / 3;
-			}
+			OutMeshBatchElement.FirstIndex = 0;
+			OutMeshBatchElement.IndexBuffer = &LODModel.IndexBuffer;
+			NumPrimitives = LODModel.IndexBuffer.GetNumIndices() / 3;
 
 			OutMeshBatch.bWireframe = true;
 			OutMeshBatch.bDisableBackfaceCulling = true;
@@ -1064,18 +1039,9 @@ uint32 FStaticMeshSceneProxy::SetMeshElementGeometrySource(
 	{
 		OutMeshBatch.Type = PT_TriangleList;
 
-		if (bUsePreculledIndices)
-		{
-			OutMeshBatchElement.IndexBuffer = LODInfo.PreCulledIndexBuffer;
-			OutMeshBatchElement.FirstIndex = SectionInfo.FirstPreCulledIndex;
-			NumPrimitives = SectionInfo.NumPreCulledTriangles;
-		}
-		else
-		{
-			OutMeshBatchElement.IndexBuffer = bUseReversedIndices ? &LODModel.AdditionalIndexBuffers->ReversedIndexBuffer : &LODModel.IndexBuffer;
-			OutMeshBatchElement.FirstIndex = Section.FirstIndex;
-			NumPrimitives = Section.NumTriangles;
-		}
+		OutMeshBatchElement.IndexBuffer = bUseReversedIndices ? &LODModel.AdditionalIndexBuffers->ReversedIndexBuffer : &LODModel.IndexBuffer;
+		OutMeshBatchElement.FirstIndex = Section.FirstIndex;
+		NumPrimitives = Section.NumTriangles;
 	}
 
 	OutMeshBatchElement.NumPrimitives = NumPrimitives;
@@ -2175,7 +2141,6 @@ bool FStaticMeshSceneProxy::HasDynamicIndirectShadowCasterRepresentation() const
 FStaticMeshSceneProxy::FLODInfo::FLODInfo(const FStaticMeshSceneProxyDesc& InProxyDesc, const FStaticMeshVertexFactoriesArray& InLODVertexFactories, int32 LODIndex, int32 InClampedMinLOD, bool bLODsShareStaticLighting)
 	: FLightCacheInterface()
 	, OverrideColorVertexBuffer(nullptr)
-	, PreCulledIndexBuffer(nullptr)
 	, bUsesMeshModifyingMaterials(false)
 {
 	const auto FeatureLevel =  InProxyDesc.GetWorld()->GetFeatureLevel();
@@ -2233,8 +2198,6 @@ FStaticMeshSceneProxy::FLODInfo::FLODInfo(const FStaticMeshSceneProxyDesc& InPro
 			}
 		}
 		
-		PreCulledIndexBuffer = &ComponentLODInfo.PreCulledIndexBuffer;
-
 		// Initialize this LOD's overridden vertex colors, if it has any
 		if( ComponentLODInfo.OverrideVertexColors )
 		{
@@ -2339,17 +2302,6 @@ FStaticMeshSceneProxy::FLODInfo::FLODInfo(const FStaticMeshSceneProxyDesc& InPro
 		}
 #endif
 
-		if (LODIndex < InProxyDesc.LODData.Num())
-		{
-			const FStaticMeshComponentLODInfo& ComponentLODInfo = InProxyDesc.LODData[LODIndex];
-
-			if (SectionIndex < ComponentLODInfo.PreCulledSections.Num())
-			{
-				SectionInfo.FirstPreCulledIndex = ComponentLODInfo.PreCulledSections[SectionIndex].FirstIndex;
-				SectionInfo.NumPreCulledTriangles = ComponentLODInfo.PreCulledSections[SectionIndex].NumTriangles;
-			}
-		}
-
 		// Store the element info.
 		Sections.Add(SectionInfo);
 
diff --git a/Engine/Source/Runtime/Engine/Public/StaticMeshComponentLODInfo.h b/Engine/Source/Runtime/Engine/Public/StaticMeshComponentLODInfo.h
--- a/Engine/Source/Runtime/Engine/Public/StaticMeshComponentLODInfo.h
+++ b/Engine/Source/Runtime/Engine/Public/StaticMeshComponentLODInfo.h
@@ -9,13 +9,6 @@
 class FColorVertexBuffer;
 class FMeshMapBuildData;
 
-struct FPreCulledStaticMeshSection
-{
-	/** Range of vertices and indices used when rendering this section. */
-	uint32 FirstIndex;
-	uint32 NumTriangles;
-};
-
 USTRUCT()
 struct FStaticMeshComponentLODInfo
 {
@@ -43,11 +36,6 @@ struct FStaticMeshComponentLODInfo
 	/** Vertex colors to use for this mesh LOD */
 	FColorVertexBuffer* OverrideVertexColors;
 
-	/** Information for each section about what range of PreCulledIndexBuffer to use.  If no preculled index data is available, PreCulledSections will be empty. */
-	TArray<FPreCulledStaticMeshSection> PreCulledSections;
-
-	FRawStaticIndexBuffer PreCulledIndexBuffer;
-
 	/** 
 	 * Owner of this FStaticMeshComponentLODInfo 
 	 * Warning, can be NULL for a component created via SpawnActor off of a blueprint default (LODData will be created without a call to SetLODDataCount).
diff --git a/Engine/Source/Runtime/Engine/Public/StaticMeshSceneProxy.h b/Engine/Source/Runtime/Engine/Public/StaticMeshSceneProxy.h
--- a/Engine/Source/Runtime/Engine/Public/StaticMeshSceneProxy.h
+++ b/Engine/Source/Runtime/Engine/Public/StaticMeshSceneProxy.h
@@ -179,35 +179,21 @@ protected:
 		/** Information about an element of a LOD. */
 		struct FSectionInfo
 		{
-			/** Default constructor. */
-			FSectionInfo()
-				: Material(nullptr)
-			#if WITH_EDITOR
-				, bSelected(false)
-				, HitProxy(nullptr)
-			#endif
-				, FirstPreCulledIndex(0)
-				, NumPreCulledTriangles(-1)
-			{}
-
 			/** The material with which to render this section. */
-			UMaterialInterface* Material;
+			UMaterialInterface* Material = nullptr;
 
 		#if WITH_EDITOR
 			/** True if this section should be rendered as selected (editor only). */
-			bool bSelected;
+			bool bSelected = false;
 
 			/** The editor needs to be able to individual sub-mesh hit detection, so we store a hit proxy on each mesh. */
-			HHitProxy* HitProxy;
+			HHitProxy* HitProxy = nullptr;
 		#endif
 
 		#if WITH_EDITORONLY_DATA
 			// The material index from the component. Used by the texture streaming accuracy viewmodes.
-			int32 MaterialIndex;
+			int32 MaterialIndex = 0;
 		#endif
-
-			int32 FirstPreCulledIndex;
-			int32 NumPreCulledTriangles;
 		};
 
 		/** Per-section information. */
@@ -218,8 +204,6 @@ protected:
 
 		TUniformBufferRef<FLocalVertexFactoryUniformShaderParameters> OverrideColorVFUniformBuffer;
 
-		const FRawStaticIndexBuffer* PreCulledIndexBuffer;
-
 		/** Initialization constructor. */
 		FLODInfo(const FStaticMeshSceneProxyDesc& InProxyDesc, const FStaticMeshVertexFactoriesArray& InLODVertexFactories, int32 InLODIndex, int32 InClampedMinLOD, bool bLODsShareStaticLighting);
 
