[misc] fix FString::Appendf format errors

diff --git a/Engine/Source/Runtime/D3D12RHI/Private/Windows/WindowsD3D12PipelineState.cpp b/Engine/Source/Runtime/D3D12RHI/Private/Windows/WindowsD3D12PipelineState.cpp
--- a/Engine/Source/Runtime/D3D12RHI/Private/Windows/WindowsD3D12PipelineState.cpp
+++ b/Engine/Source/Runtime/D3D12RHI/Private/Windows/WindowsD3D12PipelineState.cpp
@@ -614,7 +614,7 @@ static void DumpShaderAsm(FString& String, const D3D12_SHADER_BYTECODE& Shader)
 			ID3DBlob* blob = nullptr;
 			if (SUCCEEDED(D3DDisasmFunc(Shader.pShaderBytecode, Shader.BytecodeLength, 0, "", &blob)))
 			{
-				String.Appendf(TEXT("%S\n"), blob->GetBufferPointer());
+				String.Appendf(TEXT("%.*hs\n"), blob->GetBufferSize(), static_cast<const char*>(blob->GetBufferPointer()));
 				blob->Release();
 			}
 		}
diff --git a/Engine/Source/Runtime/Engine/Private/Materials/MaterialIRToHLSLTranslator.cpp b/Engine/Source/Runtime/Engine/Private/Materials/MaterialIRToHLSLTranslator.cpp
--- a/Engine/Source/Runtime/Engine/Private/Materials/MaterialIRToHLSLTranslator.cpp
+++ b/Engine/Source/Runtime/Engine/Private/Materials/MaterialIRToHLSLTranslator.cpp
@@ -40,8 +40,8 @@ struct FHLSLPrinter
 	bool bFirstListItem = false;
 	int32 Tabs = 0;
 
-	template <int N, typename... Types>
-	void Appendf(const TCHAR (&Format)[N], Types... Args)
+	template <typename... Types>
+	void Appendf(UE::Core::TCheckedFormatString<FString::FmtCharType, Types...> Format, Types... Args)
 	{
 		Buffer.Appendf(Format, Args...);
 	}
@@ -803,4 +803,4 @@ void FMaterialIRToHLSLTranslation::Run(TMap<FString, FString>& OutParameters, FS
 }
 
 #undef TAB
-#endif // #if WITH_EDITOR
\ No newline at end of file
+#endif // #if WITH_EDITOR
