Fix writer thread access to GetPSOMetricsDelegate

#rnx

diff --git a/Engine/Source/Runtime/Apple/MetalRHI/Private/MetalPipeline.cpp b/Engine/Source/Runtime/Apple/MetalRHI/Private/MetalPipeline.cpp
--- a/Engine/Source/Runtime/Apple/MetalRHI/Private/MetalPipeline.cpp
+++ b/Engine/Source/Runtime/Apple/MetalRHI/Private/MetalPipeline.cpp
@@ -285,9 +285,9 @@ public:
 				{
 					const double CompilationStartTime = FPlatformTime::Seconds();
 					Desc = CreateMTLRenderPipeline(Device, bSync, Key, Init, State);
-					const double CompilationDuration = FPlatformTime::Seconds() - CompilationStartTime;
+					const float CompilationDuration = static_cast<float>(FPlatformTime::Seconds() - CompilationStartTime);
 
-					GetPSOMetricsDelegate().ExecuteIfBound((float)CompilationDuration);
+					AccumulatePSOMetrics(CompilationDuration);
 
 
 					if (Desc != nullptr)
diff --git a/Engine/Source/Runtime/Engine/Private/PSOMetrics.cpp b/Engine/Source/Runtime/Engine/Private/PSOMetrics.cpp
--- a/Engine/Source/Runtime/Engine/Private/PSOMetrics.cpp
+++ b/Engine/Source/Runtime/Engine/Private/PSOMetrics.cpp
@@ -14,22 +14,12 @@ namespace PSOMetrics
 	int Count = 0;
 }
 
-static FPSOMetricsEvent PSOMetricsEvent;
-FPSOMetricsEvent& GetPSOMetricsDelegate()
+void AccumulatePSOMetrics(float CompilationDuration)
 {
-	if (!PSOMetricsEvent.IsBound())
-	{
-		PSOMetricsEvent.BindLambda([](float CompilationDuration)
-		{
-			// we need this scope because this will be called from RHIThread
-			FScopeLock ScopeLock(&PSOMetrics::PSOCriticalSection);
-
-			PSOMetrics::DurationSum += CompilationDuration;
-			++PSOMetrics::Count;
-		});
-	}
-
-	return PSOMetricsEvent;
+	FScopeLock ScopeLock(&PSOMetrics::PSOCriticalSection);
+	
+	PSOMetrics::DurationSum += CompilationDuration;
+	++PSOMetrics::Count;
 }
 
 void GetPSOCompilationMetrics(float& DurationSum, int& Count)
@@ -42,4 +32,4 @@ void GetPSOCompilationMetrics(float& DurationSum, int& Count)
 
 	DurationSum = 0;
 	Count = 0;
-}
\ No newline at end of file
+}
diff --git a/Engine/Source/Runtime/Engine/Public/PSOMetrics.h b/Engine/Source/Runtime/Engine/Public/PSOMetrics.h
--- a/Engine/Source/Runtime/Engine/Public/PSOMetrics.h
+++ b/Engine/Source/Runtime/Engine/Public/PSOMetrics.h
@@ -9,11 +9,10 @@
 #include "Delegates/Delegate.h"
 
 /**
- * PSO metric delegate is called for each individual PSO compilation.
+ * Accumulate PSO metric is called for each individual PSO compilation.
  * At the moment only IOS and Android platforms will call this.
  */ 
-DECLARE_DELEGATE_OneParam(FPSOMetricsEvent, float /*CompilationDuration*/);
-extern ENGINE_API FPSOMetricsEvent& GetPSOMetricsDelegate();
+extern ENGINE_API void AccumulatePSOMetrics(float CompilationDuration);
 
 // retrieves the current metrics and set it to zero
-extern ENGINE_API void GetPSOCompilationMetrics(float& DurationSum, int& Count);
\ No newline at end of file
+extern ENGINE_API void GetPSOCompilationMetrics(float& DurationSum, int& Count);
diff --git a/Engine/Source/Runtime/OpenGLDrv/Private/Android/AndroidOpenGL.cpp b/Engine/Source/Runtime/OpenGLDrv/Private/Android/AndroidOpenGL.cpp
--- a/Engine/Source/Runtime/OpenGLDrv/Private/Android/AndroidOpenGL.cpp
+++ b/Engine/Source/Runtime/OpenGLDrv/Private/Android/AndroidOpenGL.cpp
@@ -948,7 +948,7 @@ TArray<uint8> FAndroidOpenGL::DispatchAndWaitForRemoteGLProgramCompile(FGraphics
 				CompiledProgramBinary.SetNumUninitialized(len);
 				Env->GetByteArrayRegion(*ProgramResult, 0, len, reinterpret_cast<jbyte*>(CompiledProgramBinary.GetData()));
 				float CompilationDuration = (float)Env->GetFloatField(*ProgramResponseObj, OpenGLRemoteGLProgramCompileJNI.ProgramResponse_CompilationDurationField);
-				GetPSOMetricsDelegate().ExecuteIfBound(CompilationDuration);
+				AccumulatePSOMetrics(CompilationDuration);
 			}
 			else
 			{
diff --git a/Engine/Source/Runtime/VulkanRHI/Private/Android/VulkanAndroidPlatform.cpp b/Engine/Source/Runtime/VulkanRHI/Private/Android/VulkanAndroidPlatform.cpp
--- a/Engine/Source/Runtime/VulkanRHI/Private/Android/VulkanAndroidPlatform.cpp
+++ b/Engine/Source/Runtime/VulkanRHI/Private/Android/VulkanAndroidPlatform.cpp
@@ -1638,7 +1638,7 @@ VkPipelineCache FVulkanAndroidPlatform::PrecompilePSO(
 			{
 				const int ProgramResultSharedHandle = Env->GetIntField(*ProgramResponseObj, VKRemoteProgramCompileJNI.ProgramResponse_SHMOutputHandleField);
 				const float ProgramResultCompilationDuration = Env->GetFloatField(*ProgramResponseObj, VKRemoteProgramCompileJNI.ProgramResponse_CompilationDurationField);
-				GetPSOMetricsDelegate().ExecuteIfBound(ProgramResultCompilationDuration);
+				AccumulatePSOMetrics(ProgramResultCompilationDuration);
 				if(ensure(ProgramResultSharedHandle > -1))
 				{
 					const uint32 ResultMemSize = (uint32)ASharedMemory_getSize(ProgramResultSharedHandle);
