[Interchange] Fix ensure if there is some bones in the animation viewport and we re-import a skeletalmesh.

#rnx

diff --git a/Engine/Source/Editor/Persona/Private/SAnimationEditorViewport.cpp b/Engine/Source/Editor/Persona/Private/SAnimationEditorViewport.cpp
--- a/Engine/Source/Editor/Persona/Private/SAnimationEditorViewport.cpp
+++ b/Engine/Source/Editor/Persona/Private/SAnimationEditorViewport.cpp
@@ -2584,7 +2584,7 @@ void SAnimationEditorViewportTabBody::AddPostProcessNotification()
 	{
 		if (UDebugSkelMeshComponent* PreviewComponent = GetPreviewScene()->GetPreviewMeshComponent())
 		{
-			if(PreviewComponent->GetSkeletalMeshAsset() && PreviewComponent->GetSkeletalMeshAsset()->GetPostProcessAnimBlueprint())
+			if(PreviewComponent->GetSkeletalMeshAsset() && !PreviewComponent->GetSkeletalMeshAsset()->IsCompiling() && PreviewComponent->GetSkeletalMeshAsset()->GetPostProcessAnimBlueprint())
 			{
 				GEditor->GetEditorSubsystem<UAssetEditorSubsystem>()->OpenEditorForAssets(TArray<UObject*>({ PreviewComponent->GetSkeletalMeshAsset()->GetPostProcessAnimBlueprint()->ClassGeneratedBy }));
 			}
diff --git a/Engine/Source/Runtime/Engine/Private/Components/SkinnedMeshComponent.cpp b/Engine/Source/Runtime/Engine/Private/Components/SkinnedMeshComponent.cpp
--- a/Engine/Source/Runtime/Engine/Private/Components/SkinnedMeshComponent.cpp
+++ b/Engine/Source/Runtime/Engine/Private/Components/SkinnedMeshComponent.cpp
@@ -4966,6 +4966,11 @@ bool USkinnedMeshComponent::IsUsingSkinWeightProfile() const
 
 	if (USkinnedAsset* Asset = GetSkinnedAsset())
 	{
+		if (Asset->IsCompiling())
+		{
+			return false;
+		}
+
 		if (const FSkinWeightProfilesData* ProfileData = Asset->GetSkinWeightProfilesData(GetPredictedLODLevel()))
 		{
 			return ProfileData->IsDefaultOverridden() || ProfileData->IsStaticOverridden();
@@ -4993,6 +4998,12 @@ TArray<FName> USkinnedMeshComponent::GetCurrentSkinWeightProfileLayerNames() con
 
 	if (USkinnedAsset* Asset = GetSkinnedAsset())
 	{
+		if (Asset->IsCompiling())
+		{
+			//Return an array with name none entry
+			return { NAME_None };
+		}
+
 		if (const FSkinWeightProfilesData* ProfileData = Asset->GetSkinWeightProfilesData(GetPredictedLODLevel()))
 		{
 			if( ProfileData->IsDefaultOverridden() || ProfileData->IsStaticOverridden())
