World Partition - LocationVolume: Hide most properties categories from the UI as they make no sense

diff --git a/Engine/Source/Runtime/Engine/Public/LocationVolume.h b/Engine/Source/Runtime/Engine/Public/LocationVolume.h
--- a/Engine/Source/Runtime/Engine/Public/LocationVolume.h
+++ b/Engine/Source/Runtime/Engine/Public/LocationVolume.h
@@ -13,7 +13,7 @@ class FLoaderAdapterActor;
 /**
  * A volume representing a location in the world. Used for World Partition loading regions.
  */
-UCLASS(BlueprintType, Blueprintable, MinimalAPI)
+UCLASS(BlueprintType, Blueprintable, HideCategories=(Mobility, HLOD, Collision, Navigation, BrushSettings, Tags, Cooking, Replication, Networking, Actor, Mobile, RayTracing, AssetUserData, LevelInstance, DataLayers), MinimalAPI)
 class ALocationVolume : public AVolume, public IWorldPartitionActorLoaderInterface
 {
 	GENERATED_UCLASS_BODY()
