Data Layers - Store user settings as references to the DataLayer assets rather than FNames
-> Will invalidate data layer user settings

diff --git a/Engine/Source/Runtime/Engine/Private/WorldPartition/DataLayer/DataLayerManager.cpp b/Engine/Source/Runtime/Engine/Private/WorldPartition/DataLayer/DataLayerManager.cpp
--- a/Engine/Source/Runtime/Engine/Private/WorldPartition/DataLayer/DataLayerManager.cpp
+++ b/Engine/Source/Runtime/Engine/Private/WorldPartition/DataLayer/DataLayerManager.cpp
@@ -699,21 +699,21 @@ TSubclassOf<UDataLayerInstanceWithAsset> UDataLayerManager::GetDataLayerInstance
 
 void UDataLayerManager::UpdateDataLayerEditorPerProjectUserSettings() const
 {
-	TArray<FName> DataLayersNotLoadedInEditor;
-	TArray<FName> DataLayersLoadedInEditor;
+	TArray<const UDataLayerInstance*> DataLayersNotLoadedInEditor;
+	TArray<const UDataLayerInstance*> DataLayersLoadedInEditor;
 
 	GetUserLoadedInEditorStates(DataLayersLoadedInEditor, DataLayersNotLoadedInEditor);
 
 	GetMutableDefault<UWorldPartitionEditorPerProjectUserSettings>()->SetWorldDataLayersNonDefaultEditorLoadStates(GetWorld(), DataLayersLoadedInEditor, DataLayersNotLoadedInEditor);
 }
 
-void UDataLayerManager::GetUserLoadedInEditorStates(TArray<FName>& OutDataLayersLoadedInEditor, TArray<FName>& OutDataLayersNotLoadedInEditor) const
+void UDataLayerManager::GetUserLoadedInEditorStates(TArray<const UDataLayerInstance*>& OutDataLayersLoadedInEditor, TArray<const UDataLayerInstance*>& OutDataLayersNotLoadedInEditor) const
 {
 	OutDataLayersLoadedInEditor.Empty();
 	OutDataLayersNotLoadedInEditor.Empty();
 
-	const TArray<FName>& SettingsDataLayersNotLoadedInEditor = GetMutableDefault<UWorldPartitionEditorPerProjectUserSettings>()->GetWorldDataLayersNotLoadedInEditor(GetWorld());
-	const TArray<FName>& SettingsDataLayersLoadedInEditor = GetMutableDefault<UWorldPartitionEditorPerProjectUserSettings>()->GetWorldDataLayersLoadedInEditor(GetWorld());
+	const TArray<const UDataLayerInstance*>& SettingsDataLayersNotLoadedInEditor = GetDefault<UWorldPartitionEditorPerProjectUserSettings>()->GetWorldDataLayersNotLoadedInEditor(GetWorld());
+	const TArray<const UDataLayerInstance*>& SettingsDataLayersLoadedInEditor = GetDefault<UWorldPartitionEditorPerProjectUserSettings>()->GetWorldDataLayersLoadedInEditor(GetWorld());
 
 	ForEachDataLayerInstance([&OutDataLayersLoadedInEditor, &OutDataLayersNotLoadedInEditor, &SettingsDataLayersNotLoadedInEditor, &SettingsDataLayersLoadedInEditor](UDataLayerInstance* DataLayerInstance)
 	{
@@ -721,24 +721,24 @@ void UDataLayerManager::GetUserLoadedInEditorStates(TArray<FName>& OutDataLayers
 		{
 			if (!DataLayerInstance->IsLoadedInEditor() && DataLayerInstance->IsInitiallyLoadedInEditor())
 			{
-				OutDataLayersNotLoadedInEditor.Add(DataLayerInstance->GetDataLayerFName());
+				OutDataLayersNotLoadedInEditor.Add(DataLayerInstance);
 			}
 			else if (DataLayerInstance->IsLoadedInEditor() && !DataLayerInstance->IsInitiallyLoadedInEditor())
 			{
-				OutDataLayersLoadedInEditor.Add(DataLayerInstance->GetDataLayerFName());
+				OutDataLayersLoadedInEditor.Add(DataLayerInstance);
 			}
 
 			DataLayerInstance->ClearLoadedInEditorChangedByUserOperation();
 		}
 		else
 		{
-			if (SettingsDataLayersNotLoadedInEditor.Contains(DataLayerInstance->GetDataLayerFName()))
+			if (SettingsDataLayersNotLoadedInEditor.Contains(DataLayerInstance))
 			{
-				OutDataLayersNotLoadedInEditor.Add(DataLayerInstance->GetDataLayerFName());
+				OutDataLayersNotLoadedInEditor.Add(DataLayerInstance);
 			}
-			else if (SettingsDataLayersLoadedInEditor.Contains(DataLayerInstance->GetDataLayerFName()))
+			else if (SettingsDataLayersLoadedInEditor.Contains(DataLayerInstance))
 			{
-				OutDataLayersLoadedInEditor.Add(DataLayerInstance->GetDataLayerFName());
+				OutDataLayersLoadedInEditor.Add(DataLayerInstance);
 			}
 		}
 
diff --git a/Engine/Source/Runtime/Engine/Private/WorldPartition/DataLayer/WorldDataLayers.cpp b/Engine/Source/Runtime/Engine/Private/WorldPartition/DataLayer/WorldDataLayers.cpp
--- a/Engine/Source/Runtime/Engine/Private/WorldPartition/DataLayer/WorldDataLayers.cpp
+++ b/Engine/Source/Runtime/Engine/Private/WorldPartition/DataLayer/WorldDataLayers.cpp
@@ -1225,19 +1225,19 @@ void AWorldDataLayers::OnDataLayerManagerInitialized()
 	});
 
 	// Initialize DataLayer's IsLoadedInEditor based on DataLayerEditorPerProjectUserSettings
-	TArray<FName> SettingsDataLayersNotLoadedInEditor = GetMutableDefault<UWorldPartitionEditorPerProjectUserSettings>()->GetWorldDataLayersNotLoadedInEditor(GetWorld());
-	for (const FName& DataLayerName : SettingsDataLayersNotLoadedInEditor)
+	TArray<UDataLayerInstance*> SettingsDataLayersNotLoadedInEditor = GetMutableDefault<UWorldPartitionEditorPerProjectUserSettings>()->GetWorldDataLayersNotLoadedInEditor(GetWorld());
+	for (UDataLayerInstance* DataLayerInstance : SettingsDataLayersNotLoadedInEditor)
 	{
-		if (UDataLayerInstance* DataLayerInstance = const_cast<UDataLayerInstance*>(GetDataLayerInstance(DataLayerName)))
+		if (DataLayerInstance)
 		{
 			DataLayerInstance->SetIsLoadedInEditor(false, /*bFromUserChange*/false);
 		}
 	}
 
-	TArray<FName> SettingsDataLayersLoadedInEditor = GetMutableDefault<UWorldPartitionEditorPerProjectUserSettings>()->GetWorldDataLayersLoadedInEditor(GetWorld());
-	for (const FName& DataLayerName : SettingsDataLayersLoadedInEditor)
+	TArray<UDataLayerInstance*> SettingsDataLayersLoadedInEditor = GetMutableDefault<UWorldPartitionEditorPerProjectUserSettings>()->GetWorldDataLayersLoadedInEditor(GetWorld());
+	for (UDataLayerInstance* DataLayerInstance : SettingsDataLayersLoadedInEditor)
 	{
-		if (UDataLayerInstance* DataLayerInstance = const_cast<UDataLayerInstance*>(GetDataLayerInstance(DataLayerName)))
+		if (DataLayerInstance)
 		{
 			DataLayerInstance->SetIsLoadedInEditor(true, /*bFromUserChange*/false);
 		}
diff --git a/Engine/Source/Runtime/Engine/Private/WorldPartition/WorldPartitionEditorPerProjectUserSettings.cpp b/Engine/Source/Runtime/Engine/Private/WorldPartition/WorldPartitionEditorPerProjectUserSettings.cpp
--- a/Engine/Source/Runtime/Engine/Private/WorldPartition/WorldPartitionEditorPerProjectUserSettings.cpp
+++ b/Engine/Source/Runtime/Engine/Private/WorldPartition/WorldPartitionEditorPerProjectUserSettings.cpp
@@ -1,6 +1,7 @@
 // Copyright Epic Games, Inc. All Rights Reserved.
 
 #include "WorldPartition/WorldPartitionEditorPerProjectUserSettings.h"
+#include "WorldPartition/DataLayer/DataLayerManager.h"
 #include "Engine/World.h"
 #include "GameFramework/WorldSettings.h"
 #include "Misc/PackageName.h"
@@ -11,14 +12,17 @@
 
 #if WITH_EDITOR
 
-void UWorldPartitionEditorPerProjectUserSettings::SetWorldDataLayersNonDefaultEditorLoadStates(UWorld* InWorld, const TArray<FName>& InDataLayersLoadedInEditor, const TArray<FName>& InDataLayersNotLoadedInEditor)
+void UWorldPartitionEditorPerProjectUserSettings::SetWorldDataLayersNonDefaultEditorLoadStates(UWorld* InWorld, const TArray<const UDataLayerInstance*>& InDataLayersLoadedInEditor, const TArray<const UDataLayerInstance*>& InDataLayersNotLoadedInEditor)
 {
 	if (ShouldSaveSettings(InWorld))
 	{
 		FWorldPartitionPerWorldSettings& PerWorldSettings = PerWorldEditorSettings.FindOrAdd(TSoftObjectPtr<UWorld>(InWorld));
-		PerWorldSettings.NotLoadedDataLayers = InDataLayersNotLoadedInEditor;
-		PerWorldSettings.LoadedDataLayers = InDataLayersLoadedInEditor;
-		
+		PerWorldSettings.LoadedDataLayers.Reset(InDataLayersLoadedInEditor.Num());
+		PerWorldSettings.NotLoadedDataLayers.Reset(InDataLayersNotLoadedInEditor.Num());
+		auto IsValidDataLayerInstance = [](const UDataLayerInstance* DataLayerInstance) { return DataLayerInstance && DataLayerInstance->GetAsset(); };
+		auto GetDataLayerInstanceAsset = [](const UDataLayerInstance* DataLayerInstance) { return DataLayerInstance->GetAsset(); };
+		Algo::TransformIf(InDataLayersLoadedInEditor, PerWorldSettings.LoadedDataLayers, IsValidDataLayerInstance, GetDataLayerInstanceAsset);
+		Algo::TransformIf(InDataLayersNotLoadedInEditor, PerWorldSettings.NotLoadedDataLayers, IsValidDataLayerInstance, GetDataLayerInstanceAsset);
 		SaveConfig();
 	}
 }
@@ -65,24 +69,70 @@ TArray<FName> UWorldPartitionEditorPerProjectUserSettings::GetEditorLoadedLocati
 	return TArray<FName>();
 }
 
-TArray<FName> UWorldPartitionEditorPerProjectUserSettings::GetWorldDataLayersNotLoadedInEditor(UWorld* InWorld) const
+TArray<const UDataLayerInstance*> UWorldPartitionEditorPerProjectUserSettings::GetWorldDataLayersNotLoadedInEditor(UWorld* InWorld) const
 {
-	if (const FWorldPartitionPerWorldSettings* PerWorldSettings = GetWorldPartitionPerWorldSettings(InWorld))
+	TArray<const UDataLayerInstance*> NotLoadedDataLayers;
+
+	if (UDataLayerManager* DataLayerManager = UDataLayerManager::GetDataLayerManager(InWorld))
 	{
-		return PerWorldSettings->NotLoadedDataLayers;
+		if (const FWorldPartitionPerWorldSettings* PerWorldSettings = GetWorldPartitionPerWorldSettings(InWorld))
+		{
+			for (const TSoftObjectPtr<const UDataLayerAsset>& DataLayerAsset : PerWorldSettings->NotLoadedDataLayers)
+			{
+				if (DataLayerAsset != nullptr)
+				{
+					const UDataLayerInstance* DataLayerInstance = DataLayerManager->GetDataLayerInstance(DataLayerAsset.Get());
+					if (DataLayerInstance != nullptr)
+					{
+						NotLoadedDataLayers.Emplace(DataLayerInstance);
+					}
+				}
+			}
+		}
 	}
 
-	return TArray<FName>();
+	return NotLoadedDataLayers;
 }
 
-TArray<FName> UWorldPartitionEditorPerProjectUserSettings::GetWorldDataLayersLoadedInEditor(UWorld* InWorld) const
+TArray<UDataLayerInstance*> UWorldPartitionEditorPerProjectUserSettings::GetWorldDataLayersNotLoadedInEditor(UWorld* InWorld)
 {
-	if (const FWorldPartitionPerWorldSettings* PerWorldSettings = GetWorldPartitionPerWorldSettings(InWorld))
+	TArray<const UDataLayerInstance*> NotLoadedDataLayers = AsConst(*this).GetWorldDataLayersNotLoadedInEditor(InWorld);
+	TArray<UDataLayerInstance*> NotLoadedDataLayersMutable;
+	Algo::Transform(NotLoadedDataLayers, NotLoadedDataLayersMutable, [](const UDataLayerInstance* DataLayerInstance) { return const_cast<UDataLayerInstance*>(DataLayerInstance); });
+	return NotLoadedDataLayersMutable;
+}
+
+TArray<const UDataLayerInstance*> UWorldPartitionEditorPerProjectUserSettings::GetWorldDataLayersLoadedInEditor(UWorld* InWorld) const
+{
+	TArray<const UDataLayerInstance*> LoadedDataLayers;
+
+	if (UDataLayerManager* DataLayerManager = UDataLayerManager::GetDataLayerManager(InWorld))
 	{
-		return PerWorldSettings->LoadedDataLayers;
+		if (const FWorldPartitionPerWorldSettings* PerWorldSettings = GetWorldPartitionPerWorldSettings(InWorld))
+		{
+			for (const TSoftObjectPtr<const UDataLayerAsset>& DataLayerAsset : PerWorldSettings->LoadedDataLayers)
+			{
+				if (DataLayerAsset != nullptr)
+				{
+					const UDataLayerInstance* DataLayerInstance = DataLayerManager->GetDataLayerInstance(DataLayerAsset.Get());
+					if (DataLayerInstance != nullptr)
+					{
+						LoadedDataLayers.Emplace(DataLayerInstance);
+					}
+				}
+			}
+		}
 	}
-	
-	return TArray<FName>();
+
+	return LoadedDataLayers;
+}
+
+TArray<UDataLayerInstance*> UWorldPartitionEditorPerProjectUserSettings::GetWorldDataLayersLoadedInEditor(UWorld* InWorld)
+{
+	TArray<const UDataLayerInstance*> LoadedDataLayers = AsConst(*this).GetWorldDataLayersLoadedInEditor(InWorld);
+	TArray<UDataLayerInstance*> LoadedDataLayersMutable;
+	Algo::Transform(LoadedDataLayers, LoadedDataLayersMutable, [](const UDataLayerInstance* DataLayerInstance) { return const_cast<UDataLayerInstance*>(DataLayerInstance); });
+	return LoadedDataLayersMutable;
 }
 
 const FWorldPartitionPerWorldSettings* UWorldPartitionEditorPerProjectUserSettings::GetWorldPartitionPerWorldSettings(UWorld* InWorld) const
diff --git a/Engine/Source/Runtime/Engine/Private/WorldSettings.cpp b/Engine/Source/Runtime/Engine/Private/WorldSettings.cpp
--- a/Engine/Source/Runtime/Engine/Private/WorldSettings.cpp
+++ b/Engine/Source/Runtime/Engine/Private/WorldSettings.cpp
@@ -30,6 +30,7 @@
 #include "WorldSettingsCustomVersion.h"
 #include "Materials/Material.h"
 #include "ComponentRecreateRenderStateContext.h"
+#include "Algo/Transform.h"
 
 #if WITH_EDITOR
 #include "Editor.h"
@@ -551,7 +552,13 @@ void AWorldSettings::SaveDefaultWorldPartitionSettings()
 
 		if (const UDataLayerManager* DataLayerManager = UDataLayerManager::GetDataLayerManager(GetWorld()))
 		{
-			DataLayerManager->GetUserLoadedInEditorStates(DefaultWorldPartitionSettings.LoadedDataLayers, DefaultWorldPartitionSettings.NotLoadedDataLayers);
+			TArray<const UDataLayerInstance*> LoadedDataLayersPtrArray;
+			TArray<const UDataLayerInstance*> NotLoadedDataLayersPtrArray;
+			DataLayerManager->GetUserLoadedInEditorStates(LoadedDataLayersPtrArray, NotLoadedDataLayersPtrArray);
+			DefaultWorldPartitionSettings.LoadedDataLayers.Reset(LoadedDataLayersPtrArray.Num());
+			DefaultWorldPartitionSettings.NotLoadedDataLayers.Reset(NotLoadedDataLayersPtrArray.Num());
+			Algo::Transform(LoadedDataLayersPtrArray, DefaultWorldPartitionSettings.LoadedDataLayers, [](const UDataLayerInstance* DataLayerInstance) { return DataLayerInstance->GetAsset(); });
+			Algo::Transform(NotLoadedDataLayersPtrArray, DefaultWorldPartitionSettings.NotLoadedDataLayers, [](const UDataLayerInstance* DataLayerInstance) { return DataLayerInstance->GetAsset(); });
 		}
 	}
 }
diff --git a/Engine/Source/Runtime/Engine/Public/WorldPartition/DataLayer/DataLayerManager.h b/Engine/Source/Runtime/Engine/Public/WorldPartition/DataLayer/DataLayerManager.h
--- a/Engine/Source/Runtime/Engine/Public/WorldPartition/DataLayer/DataLayerManager.h
+++ b/Engine/Source/Runtime/Engine/Public/WorldPartition/DataLayer/DataLayerManager.h
@@ -215,7 +215,7 @@ private:
 
 	//~ Begin User settings
 	ENGINE_API void UpdateDataLayerEditorPerProjectUserSettings() const;
-	ENGINE_API void GetUserLoadedInEditorStates(TArray<FName>& OutDataLayersLoadedInEditor, TArray<FName>& OutDataLayersNotLoadedInEditor) const;
+	ENGINE_API void GetUserLoadedInEditorStates(TArray<const UDataLayerInstance*>& OutDataLayersLoadedInEditor, TArray<const UDataLayerInstance*>& OutDataLayersNotLoadedInEditor) const;
 	//~ End User settings
 
 	/** Used by Editor Context */
diff --git a/Engine/Source/Runtime/Engine/Public/WorldPartition/WorldPartitionEditorPerProjectUserSettings.h b/Engine/Source/Runtime/Engine/Public/WorldPartition/WorldPartitionEditorPerProjectUserSettings.h
--- a/Engine/Source/Runtime/Engine/Public/WorldPartition/WorldPartitionEditorPerProjectUserSettings.h
+++ b/Engine/Source/Runtime/Engine/Public/WorldPartition/WorldPartitionEditorPerProjectUserSettings.h
@@ -37,10 +37,10 @@ struct FWorldPartitionPerWorldSettings
 	TArray<FName> LoadedEditorLocationVolumes;
 
 	UPROPERTY()
-	TArray<FName> NotLoadedDataLayers;
+	TArray<TSoftObjectPtr<const UDataLayerAsset>> NotLoadedDataLayers;
 
 	UPROPERTY()
-	TArray<FName> LoadedDataLayers;
+	TArray<TSoftObjectPtr<const UDataLayerAsset>> LoadedDataLayers;
 #endif
 };
 
@@ -125,10 +125,13 @@ public:
 		SaveConfig();
 	}
 
-	ENGINE_API TArray<FName> GetWorldDataLayersNotLoadedInEditor(UWorld* InWorld) const;
-	ENGINE_API TArray<FName> GetWorldDataLayersLoadedInEditor(UWorld* InWorld) const;
-	
-	ENGINE_API void SetWorldDataLayersNonDefaultEditorLoadStates(UWorld* InWorld, const TArray<FName>& InDataLayersLoadedInEditor, const TArray<FName>& InDataLayersNotLoadedInEditor);
+	ENGINE_API TArray<const UDataLayerInstance*> GetWorldDataLayersNotLoadedInEditor(UWorld* InWorld) const;
+	ENGINE_API TArray<const UDataLayerInstance*> GetWorldDataLayersLoadedInEditor(UWorld* InWorld) const;
+
+	ENGINE_API TArray<UDataLayerInstance*> GetWorldDataLayersNotLoadedInEditor(UWorld* InWorld);
+	ENGINE_API TArray<UDataLayerInstance*> GetWorldDataLayersLoadedInEditor(UWorld* InWorld);
+
+	ENGINE_API void SetWorldDataLayersNonDefaultEditorLoadStates(UWorld* InWorld, const TArray<const UDataLayerInstance*>& InDataLayersLoadedInEditor, const TArray<const UDataLayerInstance*>& InDataLayersNotLoadedInEditor);
 		
 private:
 	const FWorldPartitionPerWorldSettings* GetWorldPartitionPerWorldSettings(UWorld* InWorld) const;
