[World Partition] Relax some asserts with regards to external actors to accommodate for use cases the weren't valid, but exists in some workflows.

diff --git a/Engine/Source/Runtime/Engine/Classes/GameFramework/Actor.h b/Engine/Source/Runtime/Engine/Classes/GameFramework/Actor.h
--- a/Engine/Source/Runtime/Engine/Classes/GameFramework/Actor.h
+++ b/Engine/Source/Runtime/Engine/Classes/GameFramework/Actor.h
@@ -1038,8 +1038,7 @@ protected:
 	UPROPERTY(VisibleAnywhere, AdvancedDisplay, Category = DataLayers)
 	TArray<FActorDataLayer> DataLayers;
 
-	// There is currently an issue where if we allow property override of DataLayerAssets and it contains Private datalayers
-	// then it will always serialize a diff since those are outered to the instanced level and will get remapped differently between the Override instance and Archetype instance we are comparing against
+	/** DataLayers assets the actor belongs to.*/
 	UPROPERTY(EditAnywhere, AdvancedDisplay, Category = DataLayers, meta=(DisableLevelInstancePropertyOverride))
 	TArray<TSoftObjectPtr<UDataLayerAsset>> DataLayerAssets;
 
@@ -2568,7 +2567,7 @@ public:
 
 	void SetForceExternalActorLevelReferenceForPIE(bool bValue)
 	{
-		if (ensure(IsPackageExternal()))
+		if (IsPackageExternal())
 		{
 			bForceExternalActorLevelReferenceForPIE = bValue;
 		}
diff --git a/Engine/Source/Runtime/Engine/Private/Level.cpp b/Engine/Source/Runtime/Engine/Private/Level.cpp
--- a/Engine/Source/Runtime/Engine/Private/Level.cpp
+++ b/Engine/Source/Runtime/Engine/Private/Level.cpp
@@ -650,7 +650,7 @@ void ULevel::Serialize( FArchive& Ar )
 			{
 #if WITH_EDITOR
 				// When duplicating transient actors for PIE, allow 
-				if (ensure(!Actor->IsPackageExternal()) && Actor->HasAnyFlags(RF_NonPIEDuplicateTransient) && (Ar.GetPortFlags() & PPF_DuplicateForPIE))
+				if (!Actor->IsPackageExternal() && Actor->HasAnyFlags(RF_NonPIEDuplicateTransient) && (Ar.GetPortFlags() & PPF_DuplicateForPIE))
 				{
 					return !bIsPartitioned;
 				}
