[Material Editor] New material translator push (experimental code only).

- Made FValue caching more peformant by using a hashmap.
- Reroute node.
- Some Parameters support (incomplete).

diff --git a/Engine/Source/Runtime/Engine/Private/Materials/MaterialExpressionsIR.cpp b/Engine/Source/Runtime/Engine/Private/Materials/MaterialExpressionsIR.cpp
--- a/Engine/Source/Runtime/Engine/Private/Materials/MaterialExpressionsIR.cpp
+++ b/Engine/Source/Runtime/Engine/Private/Materials/MaterialExpressionsIR.cpp
@@ -261,43 +261,133 @@ void UMaterialExpression::Build(MIR::FEmitter& Emitter)
 void UMaterialExpressionConstant::Build(MIR::FEmitter& Emitter)
 {
 	MIR::FValue* Value = Emitter.EmitConstantFloat1(R);
-	Emitter.Put(GetOutput(0), Value);
+	Emitter.Output(GetOutput(0), Value);
 }
 
 void UMaterialExpressionConstant2Vector::Build(MIR::FEmitter& Emitter)
 {
 	MIR::FValue* Value = Emitter.EmitConstantFloat2({ R, G });
-	Emitter.Put(GetOutput(0), Value);
+	Emitter.Output(GetOutput(0), Value);
 	for (int i = 0; i < 2; ++i)
 	{
-		Emitter.Put(GetOutput(i + 1), Emitter.EmitSubscript(Value, i));
+		Emitter.Output(GetOutput(i + 1), Emitter.EmitSubscript(Value, i));
 	}
 }
 
 void UMaterialExpressionConstant3Vector::Build(MIR::FEmitter& Emitter)
 {
 	MIR::FValue* Value = Emitter.EmitConstantFloat3({ Constant.R, Constant.G, Constant.B });
-	Emitter.Put(GetOutput(0), Value);
+	Emitter.Output(GetOutput(0), Value);
 	for (int i = 0; i < 3; ++i)
 	{
-		Emitter.Put(GetOutput(i + 1), Emitter.EmitSubscript(Value, i));
+		Emitter.Output(GetOutput(i + 1), Emitter.EmitSubscript(Value, i));
 	}
 }
 
 void UMaterialExpressionConstant4Vector::Build(MIR::FEmitter& Emitter)
 {
 	MIR::FValue* Value = Emitter.EmitConstantFloat4(Constant);
-	Emitter.Put(GetOutput(0), Value);
+	Emitter.Output(GetOutput(0), Value);
 	for (int i = 0; i < 4; ++i)
 	{
-		Emitter.Put(GetOutput(i + 1), Emitter.EmitSubscript(Value, i));
+		Emitter.Output(GetOutput(i + 1), Emitter.EmitSubscript(Value, i));
 	}
 }
 
+static MIR::FValue* BuildMaterialExpressionParameter(UMaterialExpressionParameter* ParameterExpr, UE::MIR::FEmitter& Emitter)
+{
+	FMaterialParameterMetadata Metadata;
+	if (!ParameterExpr->GetParameterValue(Metadata))
+	{
+		Emitter.Error(TEXT("Could not get parameter value."));
+		return nullptr;
+	}
+
+	// Is parameter a custom primitive data?
+	if (Metadata.PrimitiveDataIndex != INDEX_NONE)
+	{
+		if (Metadata.Value.Type == EMaterialParameterType::Scalar)
+		{
+			return Emitter.GetCustomPrimitiveData(Metadata.PrimitiveDataIndex);
+		}
+		else
+		{
+			MIR::FValue* X = Emitter.GetCustomPrimitiveData(Metadata.PrimitiveDataIndex + 0);
+			MIR::FValue* Y = Emitter.GetCustomPrimitiveData(Metadata.PrimitiveDataIndex + 1);
+			MIR::FValue* Z = Emitter.GetCustomPrimitiveData(Metadata.PrimitiveDataIndex + 2);
+			MIR::FValue* W = Emitter.GetCustomPrimitiveData(Metadata.PrimitiveDataIndex + 3);
+			return Emitter.EmitVector4(X, Y, Z, W);
+		}
+	}
+
+	switch (Metadata.Value.Type)
+	{
+		case EMaterialParameterType::Scalar:
+			return Emitter.EmitConstantFloat1(Metadata.Value.Float[0]);
+			
+		case EMaterialParameterType::Vector:
+			return Emitter.EmitConstantFloat4(FVector4f{ Metadata.Value.Float[0], Metadata.Value.Float[1], Metadata.Value.Float[2], Metadata.Value.Float[3] });
+
+		 case EMaterialParameterType::Texture:
+		 	return Emitter.EmitTextureObject(Metadata.Value.Texture, SAMPLERTYPE_Color);
+
+		case EMaterialParameterType::StaticSwitch:
+			return Emitter.EmitConstantBool1(Metadata.Value.Bool[0]);
+
+		default:
+			UE_MIR_TODO();
+	}
+}
+
+void UMaterialExpressionParameter::Build(UE::MIR::FEmitter& Emitter)
+{
+	Emitter.Output(0, BuildMaterialExpressionParameter(this, Emitter));
+}
+
+void UMaterialExpressionStaticSwitch::Build(UE::MIR::FEmitter& Emitter)
+{
+	Emitter.DefaultToBool(&Value, (bool)DefaultValue);
+	bool bCondition = Emitter.InputConstantBool(&Value);
+
+	UE_MIR_CHECKPOINT();
+
+	Emitter.Output(0, Emitter.Input(bCondition ? &A : &B));
+}
+
+void UMaterialExpressionStaticSwitchParameter::Build(UE::MIR::FEmitter& Emitter)
+{
+	MIR::FValue* Value = BuildMaterialExpressionParameter(this, Emitter);
+
+	UE_MIR_CHECKPOINT();
+
+	bool bCondition = Emitter.ToConstantBool(Value);
+
+	UE_MIR_CHECKPOINT();
+
+	MIR::FValue* Output;
+	if (DynamicBranch)
+	{
+		MIR::FValue* True = Emitter.Input(&A);
+		MIR::FValue* False = Emitter.Input(&B);
+
+		UE_MIR_CHECKPOINT();
+
+		Output = Emitter.EmitBranchForced(Emitter.EmitConstantBool1(bCondition), True, False);
+	}
+	else
+	{
+		Output = Emitter.Input(bCondition ? &A : &B);
+	}
+
+	UE_MIR_CHECKPOINT();
+
+	Emitter.Output(0 ,Output);
+}
+
 void UMaterialExpressionAppendVector::Build(MIR::FEmitter& Emitter)
 {
-	MIR::FValue* AValue = Emitter.GetPrimitive(&A);
-	MIR::FValue* BValue = Emitter.TryGet(&B) ? Emitter.GetPrimitive(&B) : nullptr;
+	MIR::FValue* AValue = Emitter.InputPrimitive(&A);
+	MIR::FValue* BValue = Emitter.TryInput(&B) ? Emitter.InputPrimitive(&B) : nullptr;
 
 	if (Emitter.IsInvalid())
 	{
@@ -365,7 +455,7 @@ void UMaterialExpressionAppendVector::Build(MIR::FEmitter& Emitter)
 		Output = Emitter.EmitVector4(Components[0], Components[1], Components[2], Components[3]);
 	}
 
-	Emitter.Put(0, Output);
+	Emitter.Output(0, Output);
 }
 
 /* Mathematical Operations */
@@ -379,13 +469,10 @@ static void BuildBinaryArithmeticOperator(
 	float RhsConst)
 {
 	// Default inputs to their relative constants if disconnected, then get each input after checking it has primitive type.
-	MIR::FValue* LhsValue = Emitter.DefaultTo(LhsInput, LhsConst).GetPrimitive(LhsInput);
-	MIR::FValue* RhsValue = Emitter.DefaultTo(RhsInput, RhsConst).GetPrimitive(RhsInput);
+	MIR::FValue* LhsValue = Emitter.DefaultToFloat(LhsInput, LhsConst).InputPrimitive(LhsInput);
+	MIR::FValue* RhsValue = Emitter.DefaultToFloat(RhsInput, RhsConst).InputPrimitive(RhsInput);
 
-	if (Emitter.IsInvalid())
-	{
-		return;
-	}
+	UE_MIR_CHECKPOINT();
 
 	// Determine operation input/output type by looking at the first connected input and picking float1 otherwise.
 	MIR::FTypePtr ResultType = LhsValue ? LhsValue->Type
@@ -396,16 +483,13 @@ static void BuildBinaryArithmeticOperator(
 	LhsValue = Emitter.EmitConstruct(ResultType, LhsValue);
 	RhsValue = Emitter.EmitConstruct(ResultType, RhsValue);
 
-	if (Emitter.IsInvalid())
-	{
-		return;
-	}
+	UE_MIR_CHECKPOINT();
 
 	// Finally emit the binary operator.
 	MIR::FValue* Value = Emitter.EmitBinaryOperator(Op, LhsValue, RhsValue);
 
 	// And flow it out of the expression's only output.
-	Emitter.Put(0, Value);
+	Emitter.Output(0, Value);
 }
 
 void UMaterialExpressionAdd::Build(MIR::FEmitter& Emitter)
@@ -432,30 +516,26 @@ void UMaterialExpressionIf::Build(MIR::FEmitter& Emitter)
 {
 	// Create default values flowing into disconnected inputs
 	Emitter.DefaultToFloatZero(&A);
-	Emitter.DefaultTo(&B, ConstB);
+	Emitter.DefaultToFloat(&B, ConstB);
 	Emitter.DefaultToFloatZero(&AGreaterThanB);
 	Emitter.DefaultToFloatZero(&AEqualsB);
 	Emitter.DefaultToFloatZero(&ALessThanB);
 
+	UE_MIR_CHECKPOINT();
+
 	// Get input values and check their types are what we expect.
-	MIR::FValue* AValue = Emitter.GetScalar(&A);
-	MIR::FValue* BValue = Emitter.GetScalar(&B);
-	MIR::FValue* AGreaterThanBValue = Emitter.GetPrimitive(&AGreaterThanB);
-	MIR::FValue* AEqualsBValue = Emitter.GetPrimitive(&AEqualsB);
-	MIR::FValue* ALessThanBValue = Emitter.GetPrimitive(&ALessThanB);
+	MIR::FValue* AValue = Emitter.InputScalar(&A);
+	MIR::FValue* BValue = Emitter.InputScalar(&B);
+	MIR::FValue* AGreaterThanBValue = Emitter.InputPrimitive(&AGreaterThanB);
+	MIR::FValue* AEqualsBValue = Emitter.InputPrimitive(&AEqualsB);
+	MIR::FValue* ALessThanBValue = Emitter.InputPrimitive(&ALessThanB);
 
-	if (Emitter.IsInvalid())
-	{
-		return;
-	}
+	UE_MIR_CHECKPOINT();
 
 	// Get the primitive common type between the conditional arguments (e.g. if inputs are int and float, it will return float).
 	MIR::FPrimitiveTypePtr ConditionArgsType = Emitter.GetCommonPrimitiveType(AValue->Type->AsPrimitive(), BValue->Type->AsPrimitive());
 
-	if (Emitter.IsInvalid())
-	{
-		return;
-	}
+	UE_MIR_CHECKPOINT();
 
 	// Convert both conditional argument values to the common type. This is expected to work as the common type could be found.
 	AValue = Emitter.EmitConstruct(ConditionArgsType, AValue);
@@ -465,20 +545,14 @@ void UMaterialExpressionIf::Build(MIR::FEmitter& Emitter)
 	MIR::FPrimitiveTypePtr OutputType = Emitter.GetCommonPrimitiveType(AGreaterThanBValue->Type->AsPrimitive(), AEqualsBValue->Type->AsPrimitive());
 	OutputType = Emitter.GetCommonPrimitiveType(OutputType, ALessThanBValue->Type->AsPrimitive());
 
-	if (Emitter.IsInvalid())
-	{
-		return;
-	}
+	UE_MIR_CHECKPOINT();
 
 	// Convert result values to the common result type.
 	AGreaterThanBValue = Emitter.EmitConstruct(OutputType, AGreaterThanBValue);
 	AEqualsBValue = Emitter.EmitConstruct(OutputType, AEqualsBValue);
 	ALessThanBValue = Emitter.EmitConstruct(OutputType, ALessThanBValue);
 
-	if (Emitter.IsInvalid())
-	{
-		return;
-	}
+	UE_MIR_CHECKPOINT();
 
 	// Emit the comparison expressions.
 	MIR::FValue* ALessThanBConditionValue = Emitter.EmitBinaryOperator(MIR::BO_LowerThan, AValue, BValue);
@@ -488,46 +562,80 @@ void UMaterialExpressionIf::Build(MIR::FEmitter& Emitter)
 	MIR::FValue* OutputValue = Emitter.EmitBranch(AEqualsBConditionValue, AEqualsBValue, AGreaterThanBValue);
 	OutputValue = Emitter.EmitBranch(ALessThanBConditionValue, ALessThanBValue, OutputValue);
 
-	Emitter.Put(GetOutput(0), OutputValue);
+	Emitter.Output(0, OutputValue);
 }
 
-void UMaterialExpressionTextureSample::Build(UE::MIR::FEmitter& Emitter)
+void UMaterialExpressionTextureObject::Build(UE::MIR::FEmitter& Emitter)
+{
+	Emitter.Output(0, Emitter.EmitTextureObject(Texture, SamplerType));
+}
+
+static void BuildTextureSampleExpression(UTexture* Texture, UMaterialExpressionTextureSample* Expr, UE::MIR::FEmitter& Emitter)
 {
-	UTexture* InputTexture = Texture;
-	if (MIR::FValue* InputTextureValue = Emitter.TryGet(&TextureObject))
+	if (MIR::FValue* InputTextureValue = Emitter.TryInput(&Expr->TextureObject))
 	{
-		InputTexture = InputTextureValue->GetTexture();
-		if (!InputTexture)
+		auto InputTextureSampler = InputTextureValue->GetTexture();
+		if (!InputTextureSampler.IsSet())
 		{
 			Emitter.Error(TEXT("Value flowing into texture input is not a texture."));
 			return;
 		}
+
+		if (InputTextureSampler->Get<EMaterialSamplerType>() != SAMPLERTYPE_Color)
+		{
+			Emitter.Error(TEXT("Input texture sampler type must be color."));
+			return;
+		}
+
+		Texture = InputTextureSampler->Get<UTexture*>();
 	}
-	else if (!Texture)
+	
+	if (!Texture)
 	{
 		Emitter.Error("Unspecified texture.");
 		return;
 	}
 
-	MIR::FValue* TexCoordsValue = Emitter.TryGet(&Coordinates);
+	MIR::FValue* TexCoordsValue = Emitter.TryInput(&Expr->Coordinates);
 	if (!TexCoordsValue)
 	{
-		TexCoordsValue = Emitter.GetExternalInput(MIR::TexCoordIndexToExternalInput(ConstCoordinate));
+		TexCoordsValue = Emitter.GetExternalInput(MIR::TexCoordIndexToExternalInput(Expr->ConstCoordinate));
 	}
 
-	if (Emitter.IsInvalid())
+	UE_MIR_CHECKPOINT();
+
+	MIR::FValue* OutputValue = Emitter.EmitTextureSample(Texture, TexCoordsValue, Expr->SamplerSource, Expr->MipValueMode, Expr->SamplerType);
+	
+	Emitter.Output(0, Emitter.EmitSwizzle(OutputValue, MIR::FSwizzleMask::XYZ()));
+	Emitter.Output(1, Emitter.EmitSubscript(OutputValue, 0)); 
+	Emitter.Output(2, Emitter.EmitSubscript(OutputValue, 1)); 
+	Emitter.Output(3, Emitter.EmitSubscript(OutputValue, 2));
+	Emitter.Output(4, Emitter.EmitSubscript(OutputValue, 3));
+	Emitter.Output(5, OutputValue); 
+}
+
+void UMaterialExpressionTextureSample::Build(UE::MIR::FEmitter& Emitter)
+{
+	BuildTextureSampleExpression(Texture, this, Emitter);
+}
+
+void UMaterialExpressionTextureSampleParameter::Build(UE::MIR::FEmitter& Emitter)
+{
+	FMaterialParameterMetadata Param;
+	if (!GetParameterValue(Param))
 	{
+		Emitter.Error(TEXT("Failed to get parameter value"));
 		return;
 	}
 
-	MIR::FValue* OutputValue = Emitter.EmitTextureSample(InputTexture, TexCoordsValue, SamplerSource, MipValueMode, SamplerType);
-	
-	Emitter.Put(GetOutput(0), Emitter.EmitSwizzle(OutputValue, MIR::FSwizzleMask::XYZ()));
-	Emitter.Put(GetOutput(1), Emitter.EmitSubscript(OutputValue, 0)); 
-	Emitter.Put(GetOutput(2), Emitter.EmitSubscript(OutputValue, 1)); 
-	Emitter.Put(GetOutput(3), Emitter.EmitSubscript(OutputValue, 2));
-	Emitter.Put(GetOutput(4), Emitter.EmitSubscript(OutputValue, 3));
-	Emitter.Put(GetOutput(5), OutputValue); 
+	UTexture* TheTexture = Cast<UTexture>(Param.Value.AsTextureObject());
+	if (!TheTexture)
+	{
+		Emitter.Error(TEXT("Parameter does not have a texture object"));
+		return;
+	}
+
+	BuildTextureSampleExpression(TheTexture, this, Emitter);
 }
 
 void UMaterialExpressionTextureCoordinate::Build(UE::MIR::FEmitter& Emitter)
@@ -542,7 +650,14 @@ void UMaterialExpressionTextureCoordinate::Build(UE::MIR::FEmitter& Emitter)
 	// Multiply the UV input by the UV tiling constants
 	OutputValue = Emitter.EmitBinaryOperator(MIR::BO_Multiply, OutputValue, Emitter.EmitConstantFloat2({ UTiling, VTiling }));
 	
-	Emitter.Put(GetOutput(0), OutputValue);
+	Emitter.Output(GetOutput(0), OutputValue);
+}
+
+void UMaterialExpressionReroute::Build(UE::MIR::FEmitter& Emitter)
+{
+	MIR::FValue* Value = Emitter.Input(&Input);
+	Emitter.Output(0, Value);
 }
 
+
 #endif // WITH_EDITOR
diff --git a/Engine/Source/Runtime/Engine/Private/Materials/MaterialIR.cpp b/Engine/Source/Runtime/Engine/Private/Materials/MaterialIR.cpp
--- a/Engine/Source/Runtime/Engine/Private/Materials/MaterialIR.cpp
+++ b/Engine/Source/Runtime/Engine/Private/Materials/MaterialIR.cpp
@@ -15,7 +15,7 @@ const TCHAR* ValueKindToString(EValueKind Kind)
 	{
 		case VK_Constant: return TEXT("Constant");
 		case VK_ExternalInput: return TEXT("ExternalInput");
-		case VK_MaterialParameter: return TEXT("MaterialParameter");
+		case VK_TextureObject: return TEXT("TextureObject");
 		case VK_Dimensional: return TEXT("Dimensional");
 		case VK_SetMaterialOutput: return TEXT("SetMaterialOutput");
 		case VK_BinaryOperator: return TEXT("BinaryOperator");
@@ -33,7 +33,7 @@ uint32 FValue::GetSizeInBytes() const
 	{
 		case VK_Constant: return sizeof(FConstant);
 		case VK_ExternalInput: return sizeof(FExternalInput);
-		case VK_MaterialParameter: return sizeof(FMaterialParameter);
+		case VK_TextureObject: return sizeof(FTextureObject);
 		case VK_SetMaterialOutput: return sizeof(FSetMaterialOutput);
 		case VK_BinaryOperator: return sizeof(FBinaryOperator);
 		case VK_Dimensional: return sizeof(FDimensional) + sizeof(FValue*) * static_cast<const FDimensional*>(this)->GetComponents().Num();
@@ -57,20 +57,11 @@ const FInstruction* FValue::AsInstruction() const
 
 bool FValue::Equals(const FValue* Other) const
 {
-	// Trivial case, pointers match.
 	if (this == Other)
 	{
 		return true;
 	}
 
-	// Todo: we are guaranteeing that the IR values are unique. Are we sure we need this?
-
-	// If kinds are different the two values are surely different.
-	if (Kind != Other->Kind)
-	{
-		return false;
-	}
-
 	// Get the size of this value in bytes. It should match that of Other, since the value kinds are the same.
 	uint32 SizeInBytes = GetSizeInBytes();
 	if (SizeInBytes != Other->GetSizeInBytes())
@@ -200,18 +191,14 @@ bool FValue::IsNearlyOne() const
 	return false;
 }
 
-UTexture* FValue::GetTexture()
+TOptional<TPair<UTexture*, EMaterialSamplerType>> FValue::GetTexture()
 {
-	if (auto Parameter = As<FMaterialParameter>())
+	using FPair = TPair<UTexture*, EMaterialSamplerType>;
+	if (auto TextureObject = As<MIR::FTextureObject>())
 	{
-		UObject* TextureObject = Parameter->Metadata.Value.AsTextureObject();
-		if (TextureObject)
-		{
-			return static_cast<UTexture*>(TextureObject);
-		}
+		return FPair{ TextureObject->Texture, TextureObject->SamplerType };
 	}
-
-	return nullptr;
+	return {};
 }
 
 bool FConstant::IsBool() const
@@ -233,14 +220,14 @@ const TCHAR* ExternalInputToString(EExternalInput Input)
 {
 	switch (Input)
 	{
-		case EExternalInput::TexCoord0: return TEXT("TexCoord0");
-		case EExternalInput::TexCoord1: return TEXT("TexCoord1");
-		case EExternalInput::TexCoord2: return TEXT("TexCoord2");
-		case EExternalInput::TexCoord3: return TEXT("TexCoord3");
-		case EExternalInput::TexCoord4: return TEXT("TexCoord4");
-		case EExternalInput::TexCoord5: return TEXT("TexCoord5");
-		case EExternalInput::TexCoord6: return TEXT("TexCoord6");
-		case EExternalInput::TexCoord7: return TEXT("TexCoord7");
+		case EExternalInput::TexCoord0:     return TEXT("TexCoord0");
+		case EExternalInput::TexCoord1:     return TEXT("TexCoord1");
+		case EExternalInput::TexCoord2:     return TEXT("TexCoord2");
+		case EExternalInput::TexCoord3:     return TEXT("TexCoord3");
+		case EExternalInput::TexCoord4:     return TEXT("TexCoord4");
+		case EExternalInput::TexCoord5:     return TEXT("TexCoord5");
+		case EExternalInput::TexCoord6:     return TEXT("TexCoord6");
+		case EExternalInput::TexCoord7:     return TEXT("TexCoord7");
 		case EExternalInput::TexCoord0_Ddx: return TEXT("TexCoord0_Ddx");
 		case EExternalInput::TexCoord1_Ddx: return TEXT("TexCoord1_Ddx");
 		case EExternalInput::TexCoord2_Ddx: return TEXT("TexCoord2_Ddx");
diff --git a/Engine/Source/Runtime/Engine/Private/Materials/MaterialIREmitter.cpp b/Engine/Source/Runtime/Engine/Private/Materials/MaterialIREmitter.cpp
--- a/Engine/Source/Runtime/Engine/Private/Materials/MaterialIREmitter.cpp
+++ b/Engine/Source/Runtime/Engine/Private/Materials/MaterialIREmitter.cpp
@@ -4,9 +4,12 @@
 #include "Materials/MaterialIRModule.h"
 #include "Materials/MaterialIRModuleBuilder.h"
 #include "Materials/MaterialIRInternal.h"
+#include "Materials/MaterialIRInternal.h"
 #include "Shader/ShaderTypes.h"
 #include "MaterialShared.h"
 #include "MaterialExpressionIO.h"
+#include "Engine/Texture.h"
+#include "Engine/Texture2D.h"
 
 #if WITH_EDITOR
 
@@ -66,16 +69,8 @@ struct FEmitter::FPrivate
 	// Looks for an existing value in the module that matches `Prototype` and returns it if found.
 	static FValue* FindValue(FEmitter* Emitter, const FValue* Prototype)
 	{
-		/* todo: improve this with hashmap */
-		for (FValue* CurrValue : Emitter->Module->Values)
-		{
-			if (CurrValue->Equals(Prototype))
-			{
-				return CurrValue;
-			}
-		}
-
-		return nullptr;
+		FValue** Value = Emitter->ValueSet.Find(Prototype);
+		return Value ? *Value : nullptr;
 	}
 
 	// Allocates a block of memory using the emitter's allocator.
@@ -90,6 +85,7 @@ struct FEmitter::FPrivate
 	static void PushNewValue(FEmitter* Emitter, FValue* Value)
 	{
 		Emitter->Module->Values.Add(Value);
+		Emitter->ValueSet.Add(Value);
 	}
 };
 
@@ -154,11 +150,11 @@ static FValue* EmitPrototype(FEmitter* Emitter, const TValueType& Prototype)
 }
 
 // Finds the expression input index. Although the implementation has O(n) complexity, it is only used for error reporting.
-static int SlowFindExpressionInputIndex(UMaterialExpression* Expression, const FExpressionInput* Input)
+static int SlowFindExpressionInputIndex(UMaterialExpression* Expression, const FExpressionInput* InInput)
 {
 	for (FExpressionInputIterator It{ Expression }; It; ++It)
 	{
-		if (It.Input == Input)
+		if (It.Input == InInput)
 		{
 			return It.Index;
 		}
@@ -168,9 +164,9 @@ static int SlowFindExpressionInputIndex(UMaterialExpression* Expression, const F
 }
 
 // Finds the expression input name. Although the implementation has O(n) complexity, it is only used for error reporting.
-static FName SlowFindInputName(UMaterialExpression* Expression, const FExpressionInput* Input)
+static FName SlowFindInputName(UMaterialExpression* Expression, const FExpressionInput* InInput)
 {
-	int InputIndex = SlowFindExpressionInputIndex(Expression, Input);
+	int InputIndex = SlowFindExpressionInputIndex(Expression, InInput);
 	return Expression->GetInputName(InputIndex);
 }
 
@@ -183,61 +179,83 @@ FEmitter::FEmitter(FMaterialIRModuleBuilderImpl* InBuilder, UMaterial* InMateria
 	Module = InModule;
 }
 
-FValue*  FEmitter::TryGet(const FExpressionInput* Input)
+FEmitter& FEmitter::DefaultToFloatZero(const FExpressionInput* InInput)
 {
-	FValue* Value = MIR::Internal::GetInputValue(Builder, Input);
-	return Value ? Value : nullptr;
+	return DefaultToFloat(InInput, 0.0f);
 }
 
-FValue* FEmitter::Get(const FExpressionInput* Input)
+FEmitter& FEmitter::DefaultToBool(const FExpressionInput* InInput, bool Value)
 {
-	FValue* Value = TryGet(Input);
-	if (!Value)
+	if (!InInput->IsConnected())
 	{
-		Errorf(TEXT("Input '%s' has no input value."), *SlowFindInputName(Expression, Input).ToString());
-		return nullptr;
+		Internal::SetInputValue(Builder, InInput, EmitConstantBool1(Value));
 	}
-	return Value;
+
+	return *this;
 }
-	
-FEmitter& FEmitter::Put(int OutputIndex, FValue* Value)
+
+FEmitter& FEmitter::DefaultToFloat(const FExpressionInput* InInput, TFloat Value)
 {
-	Put(Expression->GetOutput(OutputIndex), Value);
+	if (!InInput->IsConnected())
+	{
+		Internal::SetInputValue(Builder, InInput, EmitConstantFloat1(Value));
+	}
+
 	return *this;
 }
 
-FEmitter& FEmitter::Put(const FExpressionOutput* Output, FValue* Value)
+FValue*  FEmitter::TryInput(const FExpressionInput* InInput)
 {
-	check(Output);
-	MIR::Internal::SetOutputValue(Builder, Output, Value);
-	return *this;
+	FValue* Value = MIR::Internal::GetInputValue(Builder, InInput);
+	return Value ? Value : nullptr;
 }
 
-FEmitter& FEmitter::DefaultToFloatZero(const FExpressionInput* Input)
+FValue* FEmitter::Input(const FExpressionInput* InInput)
 {
-	return DefaultTo(Input, 0.0f);
+	FValue* Value = TryInput(InInput);
+	if (!Value)
+	{
+		Errorf(TEXT("Input '%s' has no input value."), *SlowFindInputName(Expression, InInput).ToString());
+		return nullptr;
+	}
+	return Value;
 }
 
-FEmitter& FEmitter::DefaultTo(const FExpressionInput* Input, TFloat Float)
+bool FEmitter::InputConstantBool(const FExpressionInput* InInput)
 {
-	if (!Input->IsConnected())
+	FValue* Value = InputBoolScalar(InInput);
+	if (!Value)
 	{
-		Internal::SetInputValue(Builder, Input, EmitConstantFloat1(Float));
+		return false;
 	}
 
-	return *this;
+	auto Constant = Value->As<FConstant>();
+	if (!Constant)
+	{
+		Errorf(TEXT("Expected constant bool value through '%s' input."), *SlowFindInputName(Expression, InInput).ToString());
+		return false;
+	}
+
+	return Constant->Boolean;
+}
+
+FValue* FEmitter::InputBoolScalar(const FExpressionInput* InInput)
+{
+	FValue* Value = Input(InInput);
+	CheckInputIsScalar(InInput, Value, SK_Bool);
+	return Value;
 }
 
-FValue* FEmitter::GetFloat(const FExpressionInput* Input)
+FValue* FEmitter::InputFloat1(const FExpressionInput* InInput)
 {
-	FValue* Value = Get(Input);
-	CheckInputIsScalar(Input, Value, SK_Float);
+	FValue* Value = Input(InInput);
+	CheckInputIsScalar(InInput, Value, SK_Float);
 	return Value;
 }
 
-FValue* FEmitter::GetScalar(const FExpressionInput* Input)
+FValue* FEmitter::InputScalar(const FExpressionInput* InInput)
 {
-	FValue* Value = GetPrimitive(Input);
+	FValue* Value = InputPrimitive(InInput);
 	if (!Value)
 	{
 		return nullptr;
@@ -247,59 +265,104 @@ FValue* FEmitter::GetScalar(const FExpressionInput* Input)
 	return EmitConstruct(ScalarType, Value);
 }
 
-FValue* FEmitter::GetPrimitive(const FExpressionInput* Input)
+FValue* FEmitter::InputPrimitive(const FExpressionInput* InInput)
 {
-	FValue* Value = Get(Input);
-	CheckInputTypeIs(Input, Value, TK_Primitive);
+	FValue* Value = Input(InInput);
+	CheckInputTypeIs(InInput, Value, TK_Primitive);
 	return Value;
 }
 
-FValue* FEmitter::GetOfType(const FExpressionInput* Input, ETypeKind Kind)
+FValue* FEmitter::InputOfType(const FExpressionInput* InInput, ETypeKind Kind)
 {
-	FValue* Value = Get(Input);
+	FValue* Value = Input(InInput);
 	if (!Value)
 	{
 		return nullptr;
 	}
 
-	CheckInputTypeIs(Input, Value, Kind);
+	CheckInputTypeIs(InInput, Value, Kind);
 	return Value;
 }
 
-bool FEmitter::CheckValueValid(const FValue* Value)
+FEmitter& FEmitter::Output(int OutputIndex, FValue* Value)
 {
-	if (!Value)
+	Output(Expression->GetOutput(OutputIndex), Value);
+	return *this;
+}
+
+FEmitter& FEmitter::Output(const FExpressionOutput* Output, FValue* Value)
+{
+	MIR::Internal::SetOutputValue(Builder, Output, Value);
+	return *this;
+}
+
+bool FEmitter::ToConstantBool(MIR::FValue* Value)
+{
+	if (!CheckValueIsScalar(Value, SK_Bool))
 	{
-		Error(TEXT("Unspecified value"));
 		return false;
 	}
-	return true;
+
+	if (MIR::FConstant* Constant = Value->As<MIR::FConstant>())
+	{
+		return Constant->Boolean;
+	}
+
+	Errorf(TEXT("Value is not a constant bool."));
+	return false;
 }
 
-void FEmitter::CheckInputIsScalar(const FExpressionInput* Input, FValue* InputValue)
+void FEmitter::CheckInputIsScalar(const FExpressionInput* InInput, FValue* InputValue)
 {
 	FPrimitiveTypePtr PrimitiveType = InputValue->Type->AsPrimitive();
 	if (!PrimitiveType || !PrimitiveType->IsScalar())
 	{
-		Errorf(TEXT("Input '%s' expected to be a scalar. It is %s instead."), *SlowFindInputName(Expression, Input).ToString(), InputValue->Type->GetSpelling().GetData());
+		Errorf(TEXT("Input '%s' expected to be a scalar. It is %s instead."), *SlowFindInputName(Expression, InInput).ToString(), InputValue->Type->GetSpelling().GetData());
 	}
 }
 
-void FEmitter::CheckInputIsScalar(const FExpressionInput* Input, FValue* InputValue, EScalarKind Kind)
+void FEmitter::CheckInputIsScalar(const FExpressionInput* InInput, FValue* InputValue, EScalarKind Kind)
 {
 	FPrimitiveTypePtr PrimitiveType = InputValue->Type->AsPrimitive();
 	if (!PrimitiveType || !PrimitiveType->IsScalar() || PrimitiveType->ScalarKind != Kind)
 	{
-		Errorf(TEXT("Input '%s' expected to be a %s scalar. It is %s instead."), *SlowFindInputName(Expression, Input).ToString(), ScalarKindToString(Kind), InputValue->Type->GetSpelling().GetData());
+		Errorf(TEXT("Input '%s' expected to be a %s scalar. It is %s instead."), *SlowFindInputName(Expression, InInput).ToString(), ScalarKindToString(Kind), InputValue->Type->GetSpelling().GetData());
 	}
 }
 
-void FEmitter::CheckInputTypeIs(const FExpressionInput* Input, FValue* InputValue, ETypeKind Kind)
+void FEmitter::CheckInputTypeIs(const FExpressionInput* InInput, FValue* InputValue, ETypeKind Kind)
 {
 	if (InputValue && InputValue->Type->Kind != Kind)
 	{
-		Errorf(TEXT("Input '%s' expected to be have type %s. It is %s instead."), *SlowFindInputName(Expression, Input).ToString(), TypeKindToString(Kind), InputValue->Type->GetSpelling().GetData());
+		Errorf(TEXT("Input '%s' expected to be have type %s. It is %s instead."), *SlowFindInputName(Expression, InInput).ToString(), TypeKindToString(Kind), InputValue->Type->GetSpelling().GetData());
+	}
+}
+
+bool FEmitter::CheckValueIsScalar(const FValue* Value, EScalarKind Kind)
+{
+	if (!CheckValueValid(Value))
+	{
+		return false;
 	}
+
+	FPrimitiveTypePtr PrimitiveType = Value->Type->AsPrimitive();
+	if (!PrimitiveType || !PrimitiveType->IsScalar() || PrimitiveType->ScalarKind != Kind)
+	{
+		Errorf(TEXT("Value expected to be a %s scalar. It is %s instead."), ScalarKindToString(Kind), Value->Type->GetSpelling().GetData());
+		return false;
+	}
+
+	return true;
+}
+
+bool FEmitter::CheckValueValid(const FValue* Value)
+{
+	if (!Value) 
+	{
+		Error(TEXT("Unspecified value"));
+		return false;
+	}
+	return true;
 }
 
 FValue* FEmitter::EmitConstantFromShaderValue(const UE::Shader::FValue& InValue)
@@ -521,12 +584,6 @@ FValue* FEmitter::EmitSwizzle(FValue* Value, FSwizzleMask Mask)
 		}
 	}
 
-	// If only one component is requested, we can use EmitSubscript() to return the single component.
-	if (Mask.NumComponents == 1)
-	{
-		return EmitSubscript(Value, (int)Mask.Components[0]);
-	}
-
 	// If the requested number of components is the same as Value and the order in which the components
 	// are specified in the mask is sequential (e.g. x, y, z) then this is a no op, simply return Value as is.
 	if (Mask.NumComponents == PrimitiveType->GetNumComponents())
@@ -547,6 +604,12 @@ FValue* FEmitter::EmitSwizzle(FValue* Value, FSwizzleMask Mask)
 		}
 	}
 	
+	// If only one component is requested, we can use EmitSubscript() to return the single component.
+	if (Mask.NumComponents == 1)
+	{
+		return EmitSubscript(Value, (int)Mask.Components[0]);
+	}
+
 	// Make the result vector type.
 	FPrimitiveTypePtr ResultType = FPrimitiveType::GetVector(PrimitiveType->ScalarKind, Mask.NumComponents);
 	FDimensional* Result = NewDimensionalValue(this, ResultType);
@@ -559,17 +622,30 @@ FValue* FEmitter::EmitSwizzle(FValue* Value, FSwizzleMask Mask)
 	return Result;
 }
 
-FValue* FEmitter::GetParameter(FName Name, const FMaterialParameterMetadata& Metadata)
+FValue* FEmitter::EmitTextureObject(UTexture* Texture, EMaterialSamplerType SamplerType)
 {
-	FMaterialParameterInfo Info{ MoveTemp(Name), EMaterialParameterAssociation::GlobalParameter, INDEX_NONE };
-	
-	MIR::FMaterialParameter Proto = MakePrototype<MIR::FMaterialParameter>(MIR::FTextureType::Get());
-	Proto.Info = Info;
-	Proto.Metadata = Metadata;
+	if (Texture->GetMaterialType() != MCT_Texture2D)
+	{
+		Error(TEXT("Only Textures 2D are supported for now."));
+		return nullptr;
+	}
 
+	FTextureObject Proto = MakePrototype<FTextureObject>(FObjectType::GetTexture2D());
+	Proto.Texture = Texture;
+	Proto.SamplerType = SamplerType;
 	return EmitPrototype(this, Proto);
 }
 
+FValue* FEmitter::GetParameter(FName Name, const FMaterialParameterMetadata& Metadata)
+{
+	UE_MIR_UNREACHABLE();
+}
+
+FValue* FEmitter::GetCustomPrimitiveData(uint32 PrimitiveDataIndex)
+{
+	UE_MIR_TODO();
+}
+
 FSetMaterialOutput* FEmitter::EmitSetMaterialOutput(EMaterialProperty InProperty, FValue* InArgValue)
 {
 	FSetMaterialOutput* Instr = new FSetMaterialOutput;
@@ -851,6 +927,14 @@ FValue* FEmitter::EmitBranch(FValue* Condition, FValue* True, FValue* False)
 		return ConstCondition->Boolean ? True : False;
 	}
 
+	return EmitBranchForced(Condition, True, False);
+}
+
+FValue* FEmitter::EmitBranchForced(FValue* Condition, FValue* True, FValue* False)
+{
+	// Condition must be of type bool
+	check(Condition->Type->IsBoolScalar());
+
 	// If the condition is not static, the types of the true and false operands must match.
 	// The resulting type will be that of the branch instruction created.
 	check(True->Type == False->Type);
@@ -1104,6 +1188,16 @@ void FEmitter::Initialize()
 	ConstantFalse = EmitPrototype(this, Temp);
 }
 
+bool FEmitter::FValueKeyFuncs::Matches(KeyInitType A, KeyInitType B)
+{
+	return A->Equals(B);
+}
+
+uint32 FEmitter::FValueKeyFuncs::GetKeyHash(KeyInitType Key)
+{
+	return Internal::HashBytes((char*)Key, Key->GetSizeInBytes());
+}
+
 } // namespace UE::MIR
 
 #endif // #if WITH_EDITOR
diff --git a/Engine/Source/Runtime/Engine/Private/Materials/MaterialIRInternal.cpp b/Engine/Source/Runtime/Engine/Private/Materials/MaterialIRInternal.cpp
--- a/Engine/Source/Runtime/Engine/Private/Materials/MaterialIRInternal.cpp
+++ b/Engine/Source/Runtime/Engine/Private/Materials/MaterialIRInternal.cpp
@@ -91,6 +91,34 @@ EMaterialTextureParameterType TextureMaterialValueTypeToParameterType(EMaterialV
 	}
 }
 
+uint32 HashBytes(const char* Ptr, uint32 Size)
+{
+	uint32 Hash = 0;
+
+	const uint32* UPtr = (const uint32*)Ptr;
+	uint32 i = 0; 
+	for (uint32 SizeIn32Bits = Size / 4; i < SizeIn32Bits; ++i)
+	{
+		Hash = HashCombineFast(Hash, UPtr[i]);
+	}
+
+	uint32 Modulo = Size % 4;
+	if (Modulo == 1)
+	{
+		Hash = HashCombineFast(Hash, UPtr[i] & 0xff000000);
+	}
+	else if (Modulo == 2)
+	{
+		Hash = HashCombineFast(Hash, UPtr[i] & 0xffff0000);
+	}
+	else if (Modulo == 3)
+	{
+		Hash = HashCombineFast(Hash, UPtr[i] & 0xffffff00);
+	}
+
+	return Hash;
+}
+
 } // namespace UE::MIR::Internal
 
 #endif
diff --git a/Engine/Source/Runtime/Engine/Private/Materials/MaterialIRInternal.h b/Engine/Source/Runtime/Engine/Private/Materials/MaterialIRInternal.h
--- a/Engine/Source/Runtime/Engine/Private/Materials/MaterialIRInternal.h
+++ b/Engine/Source/Runtime/Engine/Private/Materials/MaterialIRInternal.h
@@ -36,6 +36,9 @@ void SetInputValue(FMaterialIRModuleBuilderImpl* Builder, const FExpressionInput
 //
 void SetOutputValue(FMaterialIRModuleBuilderImpl* Builder, const FExpressionOutput* Output, FValue* Value);
 
+//
+uint32 HashBytes(const char* Ptr, uint32 Size);
+
 } // namespace UE::MIR::Internal
 
 #endif // #if WITH_EDITOR
diff --git a/Engine/Source/Runtime/Engine/Private/Materials/MaterialIRModuleBuilder.cpp b/Engine/Source/Runtime/Engine/Private/Materials/MaterialIRModuleBuilder.cpp
--- a/Engine/Source/Runtime/Engine/Private/Materials/MaterialIRModuleBuilder.cpp
+++ b/Engine/Source/Runtime/Engine/Private/Materials/MaterialIRModuleBuilder.cpp
@@ -161,17 +161,16 @@ struct FMaterialIRModuleBuilderImpl
 		// Flow the value into this expression's inputs from their connected outputs.
 		for (FExpressionInputIterator It{ Emitter->Expression}; It; ++It)
 		{
-			FExpressionOutput* ConnectedOutput = It->GetConnectedOutput();
-			if (!ConnectedOutput)
-			{
-				break;
-			}
-
 			// Fetch the value flowing through connected output.
-			if (MIR::FValue** ValuePtr = CurrContext.OutputValues.Find(ConnectedOutput))
+			FExpressionOutput* ConnectedOutput = It->GetConnectedOutput();
+			if (ConnectedOutput)
 			{
-				// ...and flow it into this input.
-				CurrContext.InputValues.Add(It.Input, *ValuePtr);
+				MIR::FValue** ValuePtr = CurrContext.OutputValues.Find(ConnectedOutput);
+				if (ValuePtr)
+				{
+					// ...and flow it into this input.
+					CurrContext.InputValues.Add(It.Input, *ValuePtr);
+				}
 			}
 		}
 
@@ -203,7 +202,7 @@ struct FMaterialIRModuleBuilderImpl
 		for (int i = 0; i < Call->FunctionInputs.Num(); ++i)
 		{
 			FFunctionExpressionInput& FunctionInput = Call->FunctionInputs[i];
-			MIR::FValue* Value = Emitter->Get(Call->GetInput(i));
+			MIR::FValue* Value = Emitter->Input(Call->GetInput(i));
 			if (Value)
 			{
 				MIR::FTypePtr Type = MIR::FType::FromMaterialValueType((EMaterialValueType)FunctionInput.ExpressionInput->GetInputType(0));
@@ -260,7 +259,7 @@ struct FMaterialIRModuleBuilderImpl
 			FFunctionExpressionOutput& FunctionOutput = Call->FunctionOutputs[i];
 
 			// Get the value flowing into the function output inside the function in the current context.
-			MIR::FValue* Value = Emitter->Get(FunctionOutput.ExpressionOutput->GetInput(0));
+			MIR::FValue* Value = Emitter->Input(FunctionOutput.ExpressionOutput->GetInput(0));
 
 			// Get the function output type.
 			MIR::FTypePtr OutputType = MIR::FType::FromMaterialValueType((EMaterialValueType)FunctionOutput.ExpressionOutput->GetOutputType(0));
@@ -366,7 +365,7 @@ struct FMaterialIRModuleBuilderImpl
 		}
 	}
 
-	void Step_PopulateBlocks()
+	void Step_LinkInstructions()
 	{
 		// This function walks the instruction graph and puts each instruction into the inner most possible block.
 		InstructionStack.Empty(InstructionStack.Max());
@@ -504,7 +503,7 @@ bool FMaterialIRModuleBuilder::Build(FMaterialIRModule* TargetModule)
 
 	Impl.Step_FlowValuesIntoMaterialOutputs();
 	Impl.Step_AnalyzeIRGraph();
-	Impl.Step_PopulateBlocks();
+	Impl.Step_LinkInstructions();
 	Impl.Step_Finalize();
 
 	UE::MIR::DebugDumpIRUseGraph(*TargetModule);
diff --git a/Engine/Source/Runtime/Engine/Private/Materials/MaterialIRToHLSLTranslator.cpp b/Engine/Source/Runtime/Engine/Private/Materials/MaterialIRToHLSLTranslator.cpp
--- a/Engine/Source/Runtime/Engine/Private/Materials/MaterialIRToHLSLTranslator.cpp
+++ b/Engine/Source/Runtime/Engine/Private/Materials/MaterialIRToHLSLTranslator.cpp
@@ -40,12 +40,6 @@ struct FHLSLPrinter
 	bool bFirstListItem = false;
 	int32 Tabs = 0;
 
-	template <typename... Types>
-	void Appendf(UE::Core::TCheckedFormatString<FString::FmtCharType, Types...> Format, Types... Args)
-	{
-		Buffer.Appendf(Format, Args...);
-	}
-
 	FHLSLPrinter& operator<<(const TCHAR* Text)
 	{
 		Buffer.Append(Text);
@@ -308,7 +302,6 @@ struct FTranslator : FMaterialIRToHLSLTranslation
 		{
 			case MIR::VK_Constant: LowerConstant(static_cast<const MIR::FConstant*>(InValue)); break;
 			case MIR::VK_ExternalInput: LowerExternalInput(static_cast<const MIR::FExternalInput*>(InValue)); break;
-			case MIR::VK_MaterialParameter: LowerMaterialParameter(static_cast<const MIR::FMaterialParameter*>(InValue)); break;
 			default: UE_MIR_UNREACHABLE();
 		}
 
@@ -353,17 +346,17 @@ struct FTranslator : FMaterialIRToHLSLTranslation
 		if (MIR::IsExternalInputTexCoord(ExternalInput->Id))
 		{
 			int Index = ExternalInputIndex - (int)MIR::EExternalInput::TexCoord0;
-			Printer.Appendf(TEXT("Parameters.TexCoords[%d]"), Index);
+			Printer.Buffer.Appendf(TEXT("Parameters.TexCoords[%d]"), Index);
 		}
 		else if (MIR::IsExternalInputTexCoordDdx(ExternalInput->Id))
 		{
 			int Index = ExternalInputIndex - (int)MIR::EExternalInput::TexCoord0_Ddx;
-			Printer.Appendf(TEXT("Parameters.TexCoords_DDX[%d]"), Index);
+			Printer.Buffer.Appendf(TEXT("Parameters.TexCoords_DDX[%d]"), Index);
 		}
 		else if (MIR::IsExternalInputTexCoordDdy(ExternalInput->Id))
 		{
 			int Index = ExternalInputIndex - (int)MIR::EExternalInput::TexCoord0_Ddy;
-			Printer.Appendf(TEXT("Parameters.TexCoords_DDY[%d]"), Index);
+			Printer.Buffer.Appendf(TEXT("Parameters.TexCoords_DDY[%d]"), Index);
 		}
 		else
 		{
@@ -371,11 +364,6 @@ struct FTranslator : FMaterialIRToHLSLTranslation
 		}
 	}
 
-	void LowerMaterialParameter(const MIR::FMaterialParameter* Parameter)
-	{
-		UE_MIR_UNREACHABLE();
-	}
-	
 	void LowerDimensional(const MIR::FDimensional* Dimensional) 
 	{
 		MIR::FPrimitiveTypePtr ArithmeticType = Dimensional->Type->AsPrimitive();
diff --git a/Engine/Source/Runtime/Engine/Private/Materials/MaterialIRTypes.cpp b/Engine/Source/Runtime/Engine/Private/Materials/MaterialIRTypes.cpp
--- a/Engine/Source/Runtime/Engine/Private/Materials/MaterialIRTypes.cpp
+++ b/Engine/Source/Runtime/Engine/Private/Materials/MaterialIRTypes.cpp
@@ -80,6 +80,23 @@ FTypePtr FType::FromMaterialValueType(EMaterialValueType Type)
 	}
 }
 
+FTypePtr FType::FromMaterialParameterType(EMaterialParameterType Type)
+{
+	switch (Type)
+	{
+		case EMaterialParameterType::Scalar: return FPrimitiveType::GetFloat1();
+		case EMaterialParameterType::Vector: return FPrimitiveType::GetFloat4();
+		case EMaterialParameterType::DoubleVector: UE_MIR_TODO();
+		case EMaterialParameterType::Texture: return FObjectType::GetTexture2D();
+		case EMaterialParameterType::TextureCollection: UE_MIR_TODO();
+		case EMaterialParameterType::Font: UE_MIR_TODO();
+		case EMaterialParameterType::RuntimeVirtualTexture: UE_MIR_TODO();
+		case EMaterialParameterType::SparseVolumeTexture: UE_MIR_TODO();
+		case EMaterialParameterType::StaticSwitch: return FPrimitiveType::GetBool1();
+		default: UE_MIR_UNREACHABLE();
+	}
+}
+
 FTypePtr FType::GetVoid()
 {
 	static FType Type{ TK_Void };
@@ -129,6 +146,10 @@ UE::Shader::EValueType FType::ToValueType() const
 			default: UE_MIR_UNREACHABLE();
 		}
 	}
+	else if (FObjectTypePtr ObjectType = AsObject())
+	{
+		return EValueType::Object;
+	}
 	
 	UE_MIR_UNREACHABLE();
 }
@@ -143,6 +164,11 @@ FPrimitiveTypePtr FType::AsPrimitive() const
 	return Kind == TK_Primitive ? static_cast<FPrimitiveTypePtr>(this) : nullptr; 
 }
 
+FObjectTypePtr FType::AsObject() const
+{
+	return Kind == TK_Object ? static_cast<FObjectTypePtr>(this) : nullptr; 
+}
+
 FPrimitiveTypePtr FType::AsScalar() const
 {
 	FPrimitiveTypePtr Type = AsPrimitive();
@@ -287,9 +313,9 @@ FPrimitiveTypePtr FPrimitiveType::ToScalar() const
 	return FPrimitiveType::GetScalar(ScalarKind);
 }
 
-FTypePtr FTextureType::Get()
+FTypePtr FObjectType::GetTexture2D()
 {
-	static FTextureType Instance { TK_Texture };
+	static FObjectType Instance { TK_Object };
 	return &Instance;
 }
 
diff --git a/Engine/Source/Runtime/Engine/Public/Materials/MaterialExpressionParameter.h b/Engine/Source/Runtime/Engine/Public/Materials/MaterialExpressionParameter.h
--- a/Engine/Source/Runtime/Engine/Public/Materials/MaterialExpressionParameter.h
+++ b/Engine/Source/Runtime/Engine/Public/Materials/MaterialExpressionParameter.h
@@ -59,6 +59,7 @@ class UMaterialExpressionParameter : public UMaterialExpression
 		return true;
 	}
 
+	virtual void Build(UE::MIR::FEmitter& Emitter) override;
 	virtual bool GenerateHLSLExpression(FMaterialHLSLGenerator& Generator, UE::HLSLTree::FScope& Scope, int32 OutputIndex, UE::HLSLTree::FExpression const*& OutExpression) const override;
 #endif
 	//~ End UMaterialExpression Interface
diff --git a/Engine/Source/Runtime/Engine/Public/Materials/MaterialExpressionReroute.h b/Engine/Source/Runtime/Engine/Public/Materials/MaterialExpressionReroute.h
--- a/Engine/Source/Runtime/Engine/Public/Materials/MaterialExpressionReroute.h
+++ b/Engine/Source/Runtime/Engine/Public/Materials/MaterialExpressionReroute.h
@@ -20,6 +20,7 @@ class UMaterialExpressionReroute : public UMaterialExpressionRerouteBase
 
 	//~ Begin UMaterialExpression Interface
 #if WITH_EDITOR
+	ENGINE_API virtual void Build(UE::MIR::FEmitter& Emitter) override;
 	ENGINE_API virtual int32 Compile(class FMaterialCompiler* Compiler, int32 OutputIndex) override;
 	ENGINE_API virtual int32 CompilePreview(FMaterialCompiler* Compiler, int32 OutputIndex) override;
 	ENGINE_API virtual void GetCaption(TArray<FString>& OutCaptions) const override;
diff --git a/Engine/Source/Runtime/Engine/Public/Materials/MaterialExpressionStaticSwitch.h b/Engine/Source/Runtime/Engine/Public/Materials/MaterialExpressionStaticSwitch.h
--- a/Engine/Source/Runtime/Engine/Public/Materials/MaterialExpressionStaticSwitch.h
+++ b/Engine/Source/Runtime/Engine/Public/Materials/MaterialExpressionStaticSwitch.h
@@ -30,6 +30,7 @@ class UMaterialExpressionStaticSwitch : public UMaterialExpression
 
 	//~ Begin UMaterialExpression Interface
 #if WITH_EDITOR
+	virtual void Build(UE::MIR::FEmitter& Emitter) override;
 	virtual int32 Compile(class FMaterialCompiler* Compiler, int32 OutputIndex) override;
 	virtual void GetCaption(TArray<FString>& OutCaptions) const override;
 	virtual FName GetInputName(int32 InputIndex) const override;
diff --git a/Engine/Source/Runtime/Engine/Public/Materials/MaterialExpressionStaticSwitchParameter.h b/Engine/Source/Runtime/Engine/Public/Materials/MaterialExpressionStaticSwitchParameter.h
--- a/Engine/Source/Runtime/Engine/Public/Materials/MaterialExpressionStaticSwitchParameter.h
+++ b/Engine/Source/Runtime/Engine/Public/Materials/MaterialExpressionStaticSwitchParameter.h
@@ -23,6 +23,7 @@ class UMaterialExpressionStaticSwitchParameter : public UMaterialExpressionStati
 
 	//~ Begin UMaterialExpression Interface
 #if WITH_EDITOR
+	virtual void Build(UE::MIR::FEmitter& Emitter) override;
 	virtual FName GetInputName(int32 InputIndex) const override;
 	virtual int32 Compile(class FMaterialCompiler* Compiler, int32 OutputIndex) override;
 	virtual void GetCaption(TArray<FString>& OutCaptions) const override;
diff --git a/Engine/Source/Runtime/Engine/Public/Materials/MaterialExpressionTextureObject.h b/Engine/Source/Runtime/Engine/Public/Materials/MaterialExpressionTextureObject.h
--- a/Engine/Source/Runtime/Engine/Public/Materials/MaterialExpressionTextureObject.h
+++ b/Engine/Source/Runtime/Engine/Public/Materials/MaterialExpressionTextureObject.h
@@ -24,6 +24,7 @@ class UMaterialExpressionTextureObject : public UMaterialExpressionTextureBase
 #endif // WITH_EDITOR
 	//~ Begin UMaterialExpression Interface
 #if WITH_EDITOR
+	virtual void Build(UE::MIR::FEmitter& Emitter) override;
 	virtual int32 Compile(class FMaterialCompiler* Compiler, int32 OutputIndex) override;
 	virtual int32 CompilePreview(class FMaterialCompiler* Compiler, int32 OutputIndex) override;
 	virtual void GetCaption(TArray<FString>& OutCaptions) const override;
diff --git a/Engine/Source/Runtime/Engine/Public/Materials/MaterialExpressionTextureSampleParameter.h b/Engine/Source/Runtime/Engine/Public/Materials/MaterialExpressionTextureSampleParameter.h
--- a/Engine/Source/Runtime/Engine/Public/Materials/MaterialExpressionTextureSampleParameter.h
+++ b/Engine/Source/Runtime/Engine/Public/Materials/MaterialExpressionTextureSampleParameter.h
@@ -38,6 +38,7 @@ class UMaterialExpressionTextureSampleParameter : public UMaterialExpressionText
 
 	//~ Begin UMaterialExpression Interface
 #if WITH_EDITOR
+	ENGINE_API virtual void Build(UE::MIR::FEmitter& Emitter) override;
 	ENGINE_API virtual int32 Compile(class FMaterialCompiler* Compiler, int32 OutputIndex) override;
 	ENGINE_API virtual void GetCaption(TArray<FString>& OutCaptions) const override;
 	ENGINE_API virtual bool MatchesSearchQuery(const TCHAR* SearchQuery) override;
diff --git a/Engine/Source/Runtime/Engine/Public/Materials/MaterialIR.h b/Engine/Source/Runtime/Engine/Public/Materials/MaterialIR.h
--- a/Engine/Source/Runtime/Engine/Public/Materials/MaterialIR.h
+++ b/Engine/Source/Runtime/Engine/Public/Materials/MaterialIR.h
@@ -14,7 +14,7 @@ enum EValueKind
 
 	VK_Constant,
 	VK_ExternalInput,
-	VK_MaterialParameter,
+	VK_TextureObject,
 
 	/* Instructions */
 
@@ -38,8 +38,8 @@ const TCHAR* ValueKindToString(EValueKind Kind);
 enum EValueFlags
 {
 	VF_None = 0,
-	VF_ValueAnalyzed = 1,
-	VF_InstructionAnalyzed = 2,
+	VF_ValueAnalyzed = 1,       // todo
+	VF_InstructionAnalyzed = 2, // todo
 };
 
 //
@@ -98,7 +98,7 @@ struct FValue
 	bool IsNearlyOne() const;
 	
 	// Returns this value's textures if it has one (nullptr otherwise).
-	UTexture* GetTexture();
+	TOptional<TPair<UTexture*, EMaterialSamplerType>> GetTexture();
 
 	// Tries to cast this value to specified type T and returns the casted pointer, if possible (nullptr otherwise).
 	template <typename T>
@@ -187,28 +187,41 @@ enum class EExternalInput
 	Count,
 };
 
+//
 static constexpr int TexCoordMaxNum = 8;
 
+//
 const TCHAR* ExternalInputToString(EExternalInput Input);
+
+//
 MIR::EExternalInput TexCoordIndexToExternalInput(int TexCoordIndex);
 FTypePtr GetExternalInputType(EExternalInput Id);
+
+//
 bool IsExternalInputTexCoord(EExternalInput Id);
+
+//
 bool IsExternalInputTexCoordDdx(EExternalInput Id);
+
+//
 bool IsExternalInputTexCoordDdy(EExternalInput Id);
 
+//
 struct FExternalInput : TValue<VK_ExternalInput>
 {
 	EExternalInput Id;
 };
 
-struct FMaterialParameter : TValue<VK_MaterialParameter>
-{	
-	FMaterialParameterInfo Info;
-	FMaterialParameterMetadata Metadata;
+//
+struct FTextureObject : TValue<VK_TextureObject>
+{
+	UTexture* Texture;
+	EMaterialSamplerType SamplerType;
 };
 
 /* Instructions */
 
+//
 struct FBlock
 {
 	FBlock* Parent{};
@@ -218,6 +231,7 @@ struct FBlock
 	FBlock* FindCommonParentWith(MIR::FBlock* Other);
 };
 
+//
 struct FInstruction : FValue
 {
 	FInstruction* Next{};
@@ -234,6 +248,7 @@ struct TInstruction : FInstruction
 	static constexpr EValueKind TypeKind = TTypeKind;
 };
 
+//
 struct FDimensional : TInstruction<VK_Dimensional>
 {
 	static constexpr int MaxNumComponents = 16;
@@ -254,6 +269,7 @@ struct TDimensional : FDimensional
 	FValue* Components[TDimension];
 };
 
+//
 struct FSetMaterialOutput : TInstruction<VK_SetMaterialOutput>
 {
 	EMaterialProperty Property;
@@ -279,10 +295,16 @@ enum EBinaryOperator
 	BO_NotEquals,
 };
 
+//
 bool IsArithmeticOperator(EBinaryOperator Op);
+
+//
 bool IsComparisonOperator(EBinaryOperator Op);
+
+//
 const TCHAR* BinaryOperatorToString(EBinaryOperator Op);
 
+//
 struct FBinaryOperator : TInstruction<VK_BinaryOperator>
 {
 	EBinaryOperator Operator = BO_Invalid;
@@ -290,6 +312,7 @@ struct FBinaryOperator : TInstruction<VK_BinaryOperator>
 	FValue* RhsArg{};
 };
 
+//
 struct FBranch : TInstruction<VK_Branch>
 {
 	FValue* ConditionArg{};
@@ -299,17 +322,20 @@ struct FBranch : TInstruction<VK_Branch>
 	FBlock FalseBlock{};
 };
 
+//
 struct FSubscript : TInstruction<VK_Subscript>
 {
 	FValue* Arg;
 	int Index;
 };
 
+//
 struct FCast : TInstruction<VK_Cast>
 {
 	FValue* Arg{};
 };
 
+//
 struct FTextureSample : TInstruction<VK_TextureSample>
 {
 	FValue* TexCoordArg;
@@ -325,7 +351,5 @@ struct FTextureSample : TInstruction<VK_TextureSample>
 	int TextureParameterIndex;
 };
 
-UTexture* GetTextureFromValue(FValue* TextureValue);
-
 } // namespace UE::MIR
 #endif
diff --git a/Engine/Source/Runtime/Engine/Public/Materials/MaterialIRCommon.h b/Engine/Source/Runtime/Engine/Public/Materials/MaterialIRCommon.h
--- a/Engine/Source/Runtime/Engine/Public/Materials/MaterialIRCommon.h
+++ b/Engine/Source/Runtime/Engine/Public/Materials/MaterialIRCommon.h
@@ -41,9 +41,11 @@ namespace UE::MIR
 /* Types*/
 struct FType;
 struct FPrimitiveType;
+struct FObjectType;
 
 using FTypePtr = const FType*;
 using FPrimitiveTypePtr = const FPrimitiveType*;
+using FObjectTypePtr = const FObjectType*;
 
 /* IR */
 struct FValue;
@@ -57,9 +59,14 @@ enum class EExternalInput;
 class FEmitter;
 struct FBlock;
 
+namespace Internal
+{
+	uint32 HashBytes(const char* Ptr, uint32 Size);
+}
+
 }
 
 #define UE_MIR_UNREACHABLE() { check(!"Unreachable"); UE_ASSUME(false); }
-#define UE_MIR_TODO() checkNoEntry()
+#define UE_MIR_TODO() UE_MIR_UNREACHABLE()
 
 #endif // #if WITH_EDITOR
diff --git a/Engine/Source/Runtime/Engine/Public/Materials/MaterialIREmitter.h b/Engine/Source/Runtime/Engine/Public/Materials/MaterialIREmitter.h
--- a/Engine/Source/Runtime/Engine/Public/Materials/MaterialIREmitter.h
+++ b/Engine/Source/Runtime/Engine/Public/Materials/MaterialIREmitter.h
@@ -7,6 +7,8 @@
 
 #if WITH_EDITOR
 
+#define UE_MIR_CHECKPOINT() if (Emitter.IsInvalid()) return
+
 namespace UE::MIR {
 
 enum class EVectorComponent : uint8_t
@@ -42,48 +44,62 @@ class FEmitter
 public:
 	FEmitter(FMaterialIRModuleBuilderImpl* InBuilder, UMaterial* InMaterial, FMaterialIRModule* InModule);
 	
-	/* IO */
+	/* Defaults */
+	
+	//
+	FEmitter& DefaultToFloatZero(const FExpressionInput* Input);
+	
+	//
+	FEmitter& DefaultToBool(const FExpressionInput* Input, bool Value);
+
+	//
+	FEmitter& DefaultToFloat(const FExpressionInput* Input, TFloat Value);
+
+	//
+
+	/* Input */
 
 	// Tries to get a value from specified input. If it has no value (for instance
 	// because it is not connected), it returns null.
-	FValue* TryGet(const FExpressionInput* Input);
+	FValue* TryInput(const FExpressionInput* Input);
 
 	// Gets and returns the value flowing into input with specified `InputIndex`.
 	// If no value is flowing in (e.g. bacause it is not connected) it reports
 	// an error and returns null.
-	FValue* Get(const FExpressionInput* Input);
+	FValue* Input(const FExpressionInput* Input);
 	
-	// Flows specified `Value` out of output with specified `OutputIndex`.
-	FEmitter& Put(int OutputIndex, FValue* Value);
-
-	// Flows specified `Value` out of specified expression `Output`.
-	FEmitter& Put(const FExpressionOutput* Output, FValue* Value);
-	
-	/* IO Helpers */
-
 	//
-	FEmitter& DefaultToFloatZero(const FExpressionInput* Input);
-
+	bool InputConstantBool(const FExpressionInput* Input);
+	
 	//
-	FEmitter& DefaultTo(const FExpressionInput* Input, TFloat Float);
+	FValue* InputBoolScalar(const FExpressionInput* Input);
 
 	// It gets the value flowing into it and checks that its type is float scalar.
-	FValue* GetFloat(const FExpressionInput* Input);
+	FValue* InputFloat1(const FExpressionInput* Input);
 
 	//
-	FValue* GetScalar(const FExpressionInput* Input);
+	FValue* InputScalar(const FExpressionInput* Input);
 	
 	//
-	FValue* GetPrimitive(const FExpressionInput* Input);
+	FValue* InputPrimitive(const FExpressionInput* Input);
 
 	// Gets the value flowing into `Input` and returns it after checking that its
 	// type matches `Kind`.
-	FValue* GetOfType(const FExpressionInput* Input, ETypeKind Kind);
+	FValue* InputOfType(const FExpressionInput* Input, ETypeKind Kind);
 
-	/* Error Checking */
+	/* Output */
+	
+	// Flows specified `Value` out of output with specified `OutputIndex`.
+	FEmitter& Output(int OutputIndex, FValue* Value);
 
-	//
-	bool CheckValueValid(const FValue* Value);
+	// Flows specified `Value` out of specified expression `Output`.
+	FEmitter& Output(const FExpressionOutput* Output, FValue* Value);
+	
+	/* Value */
+
+	bool ToConstantBool(MIR::FValue* Value);
+
+	/* Error Checking */
 
 	//
 	void CheckInputIsScalar(const FExpressionInput* Input, FValue* InputValue);
@@ -95,6 +111,12 @@ public:
 	// specified type `Kind. If it isn't it reports an error. You may check whether
 	// an error occurred with `IsInvalid()`.
 	void CheckInputTypeIs(const FExpressionInput* Input, FValue* InputValue, ETypeKind Kind);
+	
+	//
+	bool CheckValueIsScalar(const FValue* Value, EScalarKind Kind);
+
+	//
+	bool CheckValueValid(const FValue* Value);
 
 	/* Constants */
 
@@ -120,13 +142,16 @@ public:
 
 	FValue* EmitSubscript(FValue* Value, int ComponentIndex);
 	FValue* EmitSwizzle(FValue* Value, FSwizzleMask Mask);
+	FValue* EmitTextureObject(UTexture* Texture, EMaterialSamplerType SamplerType);
 	FValue* GetParameter(FName Name, const FMaterialParameterMetadata& Metadata);
+	FValue* GetCustomPrimitiveData(uint32 PrimitiveDataIndex);
 
 	/* Instructions */
 
 	FSetMaterialOutput* EmitSetMaterialOutput(EMaterialProperty InProperty, FValue* InArgValue);
 	FValue* EmitBinaryOperator(EBinaryOperator Operator, FValue* Lhs, FValue* Rhs);
 	FValue* EmitBranch(FValue* Condition, FValue* True, FValue* False);
+	FValue* EmitBranchForced(FValue* Condition, FValue* True, FValue* False);
 	FValue* EmitConstruct(FTypePtr Type, FValue* Initializer);
 	FValue* EmitTextureSample(UTexture* Texture, FValue* TexCoord, ESamplerSourceMode SamplerSourceMode, ETextureMipValueMode MipValueMode, EMaterialSamplerType SamplerType);
 
@@ -152,6 +177,12 @@ public:
 private:
 	void Initialize();
 
+	struct FValueKeyFuncs : DefaultKeyFuncs<FValue*>
+	{
+		static bool Matches(KeyInitType A, KeyInitType B);
+		static uint32 GetKeyHash(KeyInitType Key);
+	};
+
 	UMaterial* Material{};
 	FMaterialIRModule* Module{};
 	UMaterialExpression* Expression{};
@@ -159,6 +190,7 @@ private:
 	bool bHasExprBuildError = false;
 	FValue* ConstantTrue;
 	FValue* ConstantFalse;
+	TSet<FValue*, FValueKeyFuncs> ValueSet;
 };
 
 } // namespace UE::MIR
diff --git a/Engine/Source/Runtime/Engine/Public/Materials/MaterialIRTypes.h b/Engine/Source/Runtime/Engine/Public/Materials/MaterialIRTypes.h
--- a/Engine/Source/Runtime/Engine/Public/Materials/MaterialIRTypes.h
+++ b/Engine/Source/Runtime/Engine/Public/Materials/MaterialIRTypes.h
@@ -4,6 +4,7 @@
 #include "Materials/MaterialIRCommon.h"
 #include "Shader/ShaderTypes.h"
 #include "MaterialValueType.h"
+#include "MaterialTypes.h"
 
 #if WITH_EDITOR
 
@@ -14,7 +15,7 @@ enum ETypeKind
 {
 	TK_Void,
 	TK_Primitive,
-	TK_Texture,
+	TK_Object,
 };
 
 const TCHAR* TypeKindToString(ETypeKind Kind);
@@ -30,6 +31,9 @@ struct FType
 	// Returns the type matching specified EMaterialValueType.
 	static FTypePtr FromMaterialValueType(EMaterialValueType Type);
 
+	// Returns the type matching specified EMaterialParameterType.
+	static FTypePtr FromMaterialParameterType(EMaterialParameterType Type);
+
 	// Returns the `void` type.
 	static FTypePtr GetVoid();
 
@@ -39,6 +43,9 @@ struct FType
 	// Returns this type upcast this type to ArithmeticType if it is one. Otherwise it returns nullptr.
 	FPrimitiveTypePtr AsPrimitive() const;
 
+	//
+	FObjectTypePtr AsObject() const;
+
 	// Returns this type upcast this type to ArithmeticType if it's a scalar. Otherwise it returns nullptr.
 	FPrimitiveTypePtr AsScalar() const;
 
@@ -90,9 +97,9 @@ struct FPrimitiveType : FType
 	FPrimitiveTypePtr ToScalar() const;
 };
 
-struct FTextureType : FType
+struct FObjectType : FType
 {
-	static FTypePtr Get();
+	static FTypePtr GetTexture2D();
 };
 
 } // namespace UE::MIR
