In game, removed unnecessary call to PrecachePSOs in UStaticMeshComponent::OnRegister if function was previously called (Postload) and component world transform doesn't affect culling mode (negative determinant).
Saves ~32% of AddToWorld when streaming in Big_City_LVL (on PC DX12).

diff --git a/Engine/Source/Runtime/Engine/Classes/Components/StaticMeshComponent.h b/Engine/Source/Runtime/Engine/Classes/Components/StaticMeshComponent.h
--- a/Engine/Source/Runtime/Engine/Classes/Components/StaticMeshComponent.h
+++ b/Engine/Source/Runtime/Engine/Classes/Components/StaticMeshComponent.h
@@ -189,6 +189,9 @@ protected:
 	/** Whether mesh painting is supported. This is can be set by derived component types that don't support mesh painting. */
 	uint8 bSupportMeshPainting : 1;
 
+	/** Whether PrecachedPSOs was called. */
+	uint8 bPrecachedPSOs : 1;
+
 public:
 
 #if WITH_EDITORONLY_DATA
@@ -553,6 +556,7 @@ protected:
 	ENGINE_API virtual bool AllowsAsyncPhysicsStateCreationAndDestruction() const override;
 
 public:
+	ENGINE_API virtual void PrecachePSOs() override;
 	ENGINE_API virtual void InvalidateLightingCacheDetailed(bool bInvalidateBuildEnqueuedLighting, bool bTranslationOnly) override;
 	ENGINE_API virtual UObject const* AdditionalStatObject() const override;
 #if WITH_EDITOR
diff --git a/Engine/Source/Runtime/Engine/Private/Components/StaticMeshComponent.cpp b/Engine/Source/Runtime/Engine/Private/Components/StaticMeshComponent.cpp
--- a/Engine/Source/Runtime/Engine/Private/Components/StaticMeshComponent.cpp
+++ b/Engine/Source/Runtime/Engine/Private/Components/StaticMeshComponent.cpp
@@ -197,6 +197,7 @@ UStaticMeshComponent::UStaticMeshComponent(const FObjectInitializer& ObjectIniti
 	WireframeColorOverride = FColor(255, 255, 255, 255);
 
 	MinLOD = 0;
+	bPrecachedPSOs = false;
 	bOverrideLightMapRes = false;
 	OverriddenLightMapRes = 64;
 	SubDivisionStepSize = 32;
@@ -855,13 +856,34 @@ void UStaticMeshComponent::OnRegister()
 
 	Super::OnRegister();
 
-	// World transform might have changes causing negative determinant which changes the culling mode
-	PrecachePSOs();
+#if UE_WITH_PSO_PRECACHING
+	if (FApp::CanEverRender() && IsComponentPSOPrecachingEnabled())
+	{
+		bool bCanSkipPrecachePSOs = false;
+		if (bPrecachedPSOs && GetWorld() && GetWorld()->IsGameWorld())
+		{
+			// World transform might have changes causing negative determinant which changes the culling mode
+			const bool bIsLocalToWorldDeterminantNegative = GetRenderMatrix().Determinant() < 0;
+			bCanSkipPrecachePSOs = !bIsLocalToWorldDeterminantNegative;
+		}
+		if (!bCanSkipPrecachePSOs)
+		{
+			PrecachePSOs();
+		}
+	}
+#endif
 }
 
 void UStaticMeshComponent::OnUnregister()
 {
 	Super::OnUnregister();
+	bPrecachedPSOs = false;
+}
+
+void UStaticMeshComponent::PrecachePSOs()
+{
+	Super::PrecachePSOs();
+	bPrecachedPSOs = true;
 }
 
 void UStaticMeshComponent::BeginPlay()
