Reset ActorDescInstance internal flag bIsRegisteringOrUnregistering when deferred actor unregistration is skipped because actor was destroyed.

diff --git a/Engine/Source/Runtime/Engine/Private/WorldPartition/WorldPartitionHandle.cpp b/Engine/Source/Runtime/Engine/Private/WorldPartition/WorldPartitionHandle.cpp
--- a/Engine/Source/Runtime/Engine/Private/WorldPartition/WorldPartitionHandle.cpp
+++ b/Engine/Source/Runtime/Engine/Private/WorldPartition/WorldPartitionHandle.cpp
@@ -251,6 +251,11 @@ void FWorldPartitionLoadingContext::FDeferred::UnregisterActor(FWorldPartitionAc
 
 		bNeedsClearTransactions |= Actor->GetTypedOuter<UWorld>()->HasAnyFlags(RF_Transactional);
 	}
+	else
+	{
+		check(InActorDescInstance->bIsRegisteringOrUnregistering);
+		InActorDescInstance->bIsRegisteringOrUnregistering = false;
+	}
 }
 
 /**
