Anim: Add SetAnimInstanceClass_Internal as the Setter for AnimClass.

Follow up to 36563604, SetAnimClass was deprecated in favor of SetAnimInstanceClass

#rnx

diff --git a/Engine/Source/Runtime/Engine/Classes/Components/SkeletalMeshComponent.h b/Engine/Source/Runtime/Engine/Classes/Components/SkeletalMeshComponent.h
--- a/Engine/Source/Runtime/Engine/Classes/Components/SkeletalMeshComponent.h
+++ b/Engine/Source/Runtime/Engine/Classes/Components/SkeletalMeshComponent.h
@@ -350,7 +350,7 @@ public:
 	TObjectPtr<class UAnimBlueprintGeneratedClass> AnimBlueprintGeneratedClass;
 
 	/* The AnimBlueprint class to use. Use 'SetAnimInstanceClass' to change at runtime. */
-	UPROPERTY(EditAnywhere, BlueprintReadOnly, Category = Animation, meta=(EditCondition = bEnableAnimation))
+	UPROPERTY(EditAnywhere, BlueprintReadOnly, Setter = SetAnimInstanceClass_Internal, Category = Animation, meta=(EditCondition = bEnableAnimation))
 	class TSubclassOf<UAnimInstance> AnimClass;
 
 	/** The active animation graph program instance. */
@@ -975,7 +975,7 @@ public:
 	UFUNCTION(BlueprintCallable, Category = "Components|SkeletalMesh", meta = (Keywords = "AnimBlueprint", DisplayName = "Set Anim Instance Class"))
 	ENGINE_API virtual void SetAnimInstanceClass(class UClass* NewClass);
 
-	/** 
+	/**
 	 * Returns the animation instance that is driving the class (if available). This is typically an instance of
 	 * the class set as AnimBlueprintGeneratedClass (generated by an animation blueprint)
 	 * Since this instance is transient, it is not safe to be used during construction script
@@ -983,6 +983,12 @@ public:
 	UFUNCTION(BlueprintCallable, Category="Components|SkeletalMesh", meta=(Keywords = "AnimBlueprint", UnsafeDuringActorConstruction = "true"))
 	ENGINE_API class UAnimInstance * GetAnimInstance() const;
 	
+private:
+	/* Setter function for anim instance class */
+	void SetAnimInstanceClass_Internal(const TSubclassOf<UAnimInstance>& InAnimClass);
+
+public:
+
 	/**
 	 * Returns the active post process instance is one is available. This is set on the mesh that this
 	 * component is using, and is evaluated immediately after the main instance.
diff --git a/Engine/Source/Runtime/Engine/Private/Components/SkeletalMeshComponent.cpp b/Engine/Source/Runtime/Engine/Private/Components/SkeletalMeshComponent.cpp
--- a/Engine/Source/Runtime/Engine/Private/Components/SkeletalMeshComponent.cpp
+++ b/Engine/Source/Runtime/Engine/Private/Components/SkeletalMeshComponent.cpp
@@ -3265,6 +3265,11 @@ UAnimInstance* USkeletalMeshComponent::GetAnimInstance() const
 	return AnimScriptInstance;
 }
 
+void USkeletalMeshComponent::SetAnimInstanceClass_Internal(const TSubclassOf<UAnimInstance>& InAnimClass)
+{
+	SetAnimInstanceClass(InAnimClass.Get());
+}
+
 UAnimInstance* USkeletalMeshComponent::GetPostProcessInstance() const
 {
 	return PostProcessAnimInstance;
