Prevent AnimInstanceProxy doing draw debug on server.
AnimInstanceProxy DrawDebugCapsule now take DepthPriority like other DebugDraw functions.

diff --git a/Engine/Source/Runtime/Engine/Private/Animation/AnimInstanceProxy.cpp b/Engine/Source/Runtime/Engine/Private/Animation/AnimInstanceProxy.cpp
--- a/Engine/Source/Runtime/Engine/Private/Animation/AnimInstanceProxy.cpp
+++ b/Engine/Source/Runtime/Engine/Private/Animation/AnimInstanceProxy.cpp
@@ -685,27 +685,30 @@ void FAnimInstanceProxy::PostUpdate(UAnimInstance* InAnimInstance) const
 #if ENABLE_ANIM_DRAW_DEBUG
 void FAnimInstanceProxy::FlushQueuedDebugDrawItems(AActor* InActor, UWorld* InWorld) const
 {
-	for (const FQueuedDrawDebugItem& DebugItem : QueuedDrawDebugItems)
+	if (InActor == nullptr || InActor->GetLocalRole() != ROLE_Authority)
 	{
-		switch (DebugItem.ItemType)
+		for (const FQueuedDrawDebugItem& DebugItem : QueuedDrawDebugItems)
 		{
-			case EDrawDebugItemType::OnScreenMessage: GEngine->AddOnScreenDebugMessage(INDEX_NONE, 0.f, DebugItem.Color, DebugItem.Message, false, DebugItem.TextScale); break;
-			case EDrawDebugItemType::InWorldMessage: DrawDebugString(InWorld, DebugItem.StartLoc, DebugItem.Message, InActor, DebugItem.Color, DebugItem.LifeTime, false /*bDrawShadow*/, DebugItem.TextScale.X); break;
-			case EDrawDebugItemType::DirectionalArrow: DrawDebugDirectionalArrow(InWorld, DebugItem.StartLoc, DebugItem.EndLoc.Value, DebugItem.Size, DebugItem.Color, DebugItem.bPersistentLines, DebugItem.LifeTime, DebugItem.DepthPriority, DebugItem.Thickness); break;
-			case EDrawDebugItemType::Sphere: DrawDebugSphere(InWorld, DebugItem.Center, DebugItem.Radius, DebugItem.Segments, DebugItem.Color, DebugItem.bPersistentLines, DebugItem.LifeTime, DebugItem.DepthPriority, DebugItem.Thickness); break;
-			case EDrawDebugItemType::Line: DrawDebugLine(InWorld, DebugItem.StartLoc, DebugItem.EndLoc.Value, DebugItem.Color, DebugItem.bPersistentLines, DebugItem.LifeTime, DebugItem.DepthPriority, DebugItem.Thickness); break;
-			case EDrawDebugItemType::CoordinateSystem: DrawDebugCoordinateSystem(InWorld, DebugItem.StartLoc, DebugItem.Rotation, DebugItem.Size, DebugItem.bPersistentLines, DebugItem.LifeTime, DebugItem.DepthPriority, DebugItem.Thickness); break;
-			case EDrawDebugItemType::Point: DrawDebugPoint(InWorld, DebugItem.StartLoc, DebugItem.Size, DebugItem.Color, DebugItem.bPersistentLines, DebugItem.LifeTime, DebugItem.DepthPriority); break;
-			case EDrawDebugItemType::Circle:
+			switch (DebugItem.ItemType)
 			{
-				const FMatrix RotationMatrix = FRotationMatrix::MakeFromZ(DebugItem.Direction.Value);
-				const FVector ForwardVector = RotationMatrix.GetScaledAxis(EAxis::X);
-				const FVector RightVector = RotationMatrix.GetScaledAxis(EAxis::Y);
-				DrawDebugCircle(InWorld, DebugItem.Center, DebugItem.Radius, DebugItem.Segments, DebugItem.Color, DebugItem.bPersistentLines, DebugItem.LifeTime, DebugItem.DepthPriority, DebugItem.Thickness, ForwardVector, RightVector, false);
-				break;
+				case EDrawDebugItemType::OnScreenMessage: GEngine->AddOnScreenDebugMessage(INDEX_NONE, 0.f, DebugItem.Color, DebugItem.Message, false, DebugItem.TextScale); break;
+				case EDrawDebugItemType::InWorldMessage: DrawDebugString(InWorld, DebugItem.StartLoc, DebugItem.Message, InActor, DebugItem.Color, DebugItem.LifeTime, false /*bDrawShadow*/, DebugItem.TextScale.X); break;
+				case EDrawDebugItemType::DirectionalArrow: DrawDebugDirectionalArrow(InWorld, DebugItem.StartLoc, DebugItem.EndLoc.Value, DebugItem.Size, DebugItem.Color, DebugItem.bPersistentLines, DebugItem.LifeTime, DebugItem.DepthPriority, DebugItem.Thickness); break;
+				case EDrawDebugItemType::Sphere: DrawDebugSphere(InWorld, DebugItem.Center, DebugItem.Radius, DebugItem.Segments, DebugItem.Color, DebugItem.bPersistentLines, DebugItem.LifeTime, DebugItem.DepthPriority, DebugItem.Thickness); break;
+				case EDrawDebugItemType::Line: DrawDebugLine(InWorld, DebugItem.StartLoc, DebugItem.EndLoc.Value, DebugItem.Color, DebugItem.bPersistentLines, DebugItem.LifeTime, DebugItem.DepthPriority, DebugItem.Thickness); break;
+				case EDrawDebugItemType::CoordinateSystem: DrawDebugCoordinateSystem(InWorld, DebugItem.StartLoc, DebugItem.Rotation, DebugItem.Size, DebugItem.bPersistentLines, DebugItem.LifeTime, DebugItem.DepthPriority, DebugItem.Thickness); break;
+				case EDrawDebugItemType::Point: DrawDebugPoint(InWorld, DebugItem.StartLoc, DebugItem.Size, DebugItem.Color, DebugItem.bPersistentLines, DebugItem.LifeTime, DebugItem.DepthPriority); break;
+				case EDrawDebugItemType::Circle:
+				{
+					const FMatrix RotationMatrix = FRotationMatrix::MakeFromZ(DebugItem.Direction.Value);
+					const FVector ForwardVector = RotationMatrix.GetScaledAxis(EAxis::X);
+					const FVector RightVector = RotationMatrix.GetScaledAxis(EAxis::Y);
+					DrawDebugCircle(InWorld, DebugItem.Center, DebugItem.Radius, DebugItem.Segments, DebugItem.Color, DebugItem.bPersistentLines, DebugItem.LifeTime, DebugItem.DepthPriority, DebugItem.Thickness, ForwardVector, RightVector, false);
+					break;
+				}
+				case EDrawDebugItemType::Cone: DrawDebugCone(InWorld, DebugItem.Center, DebugItem.Direction.Value, DebugItem.Length, DebugItem.AngleWidth, DebugItem.AngleHeight, DebugItem.Segments, DebugItem.Color, DebugItem.bPersistentLines, DebugItem.LifeTime, DebugItem.DepthPriority, DebugItem.Thickness); break;
+				case EDrawDebugItemType::Capsule: DrawDebugCapsule(InWorld, DebugItem.Center, DebugItem.Size, DebugItem.Radius, DebugItem.Rotation.Quaternion(), DebugItem.Color, DebugItem.bPersistentLines, DebugItem.LifeTime, DebugItem.DepthPriority, DebugItem.Thickness); break;
 			}
-			case EDrawDebugItemType::Cone: DrawDebugCone(InWorld, DebugItem.Center, DebugItem.Direction.Value, DebugItem.Length, DebugItem.AngleWidth, DebugItem.AngleHeight, DebugItem.Segments, DebugItem.Color, DebugItem.bPersistentLines, DebugItem.LifeTime, DebugItem.DepthPriority, DebugItem.Thickness); break;
-			case EDrawDebugItemType::Capsule: DrawDebugCapsule(InWorld, DebugItem.Center, DebugItem.Size, DebugItem.Radius, DebugItem.Rotation.Quaternion(), DebugItem.Color, DebugItem.bPersistentLines, DebugItem.LifeTime, 0, DebugItem.Thickness); break;
 		}
 	}
 
@@ -2378,7 +2381,7 @@ void FAnimInstanceProxy::AnimDrawDebugCone(const FVector& Center, float Length,
 	QueuedDrawDebugItems.Add(DrawDebugItem);
 }
 
-void FAnimInstanceProxy::AnimDrawDebugCapsule(const FVector& Center, float HalfHeight, float Radius, const FRotator& Rotation, const FColor& Color, bool bPersistentLines /*= false*/, float LifeTime /*= -1.f*/, float Thickness /*= 0.f*/)
+void FAnimInstanceProxy::AnimDrawDebugCapsule(const FVector& Center, float HalfHeight, float Radius, const FRotator& Rotation, const FColor& Color, bool bPersistentLines /*= false*/, float LifeTime /*= -1.f*/, float Thickness /*= 0.f*/, ESceneDepthPriorityGroup DepthPriority /*= SDPG_World*/)
 {
 	FQueuedDrawDebugItem DrawDebugItem;
 
@@ -2391,6 +2394,7 @@ void FAnimInstanceProxy::AnimDrawDebugCapsule(const FVector& Center, float HalfH
 	DrawDebugItem.bPersistentLines = bPersistentLines;
 	DrawDebugItem.LifeTime = LifeTime;
 	DrawDebugItem.Thickness = Thickness;
+	DrawDebugItem.DepthPriority = DepthPriority;
 
 	QueuedDrawDebugItems.Add(DrawDebugItem);
 }
diff --git a/Engine/Source/Runtime/Engine/Public/Animation/AnimInstanceProxy.h b/Engine/Source/Runtime/Engine/Public/Animation/AnimInstanceProxy.h
--- a/Engine/Source/Runtime/Engine/Public/Animation/AnimInstanceProxy.h
+++ b/Engine/Source/Runtime/Engine/Public/Animation/AnimInstanceProxy.h
@@ -464,7 +464,7 @@ public:
 	ENGINE_API void AnimDrawDebugPoint(const FVector& Loc, float Size, const FColor& Color, bool bPersistentLines = false, float LifeTime = -1.f, ESceneDepthPriorityGroup DepthPriority = SDPG_World);
 	ENGINE_API void AnimDrawDebugCircle(const FVector& Center, float Radius, int32 Segments, const FColor& Color, const FVector& UpVector = FVector::UpVector, bool bPersistentLines = false, float LifeTime = -1.f, ESceneDepthPriorityGroup DepthPriority = SDPG_World, float Thickness = 0.f);
 	ENGINE_API void AnimDrawDebugCone(const FVector& Center, float Radius, const FVector& Direction, float AngleWidth, float AngleHeight, int32 Segments, const FColor& Color, bool bPersistentLines = false, float LifeTime = -1.f, ESceneDepthPriorityGroup DepthPriority = SDPG_World, float Thickness = 0.f);
-	ENGINE_API void AnimDrawDebugCapsule(const FVector& Center, float HalfHeight, float Radius, const FRotator& Rotation, const FColor& Color, bool bPersistentLines = false, float LifeTime = -1.f, float Thickness = 0.f);
+	ENGINE_API void AnimDrawDebugCapsule(const FVector& Center, float HalfHeight, float Radius, const FRotator& Rotation, const FColor& Color, bool bPersistentLines = false, float LifeTime = -1.f, float Thickness = 0.f, ESceneDepthPriorityGroup DepthPriority = SDPG_World);
 #else
 	void AnimDrawDebugOnScreenMessage(const FString& DebugMessage, const FColor& Color, const FVector2D& TextScale = FVector2D::UnitVector, ESceneDepthPriorityGroup DepthPriority = SDPG_World) {}
 	void AnimDrawDebugInWorldMessage(const FString& DebugMessage, const FVector& TextLocation, const FColor& Color, float TextScale) {}
@@ -476,7 +476,7 @@ public:
 	void AnimDrawDebugPoint(const FVector& Loc, float Size, const FColor& Color, bool bPersistentLines = false, float LifeTime = -1.f, ESceneDepthPriorityGroup DepthPriority = SDPG_World) {}
 	void AnimDrawDebugCircle(const FVector& Center, float Radius, int32 Segments, const FColor& Color, const FVector& UpVector = FVector::UpVector, bool bPersistentLines = false, float LifeTime=-1.f, ESceneDepthPriorityGroup DepthPriority = SDPG_World, float Thickness = 0.f) {}
 	void AnimDrawDebugCone(const FVector& Center, float Radius, const FVector& Direction, float AngleWidth, float AngleHeight, int32 Segments, const FColor& Color, bool bPersistentLines = false, float LifeTime = -1.f, ESceneDepthPriorityGroup DepthPriority = SDPG_World, float Thickness = 0.f) {}
-	ENGINE_API void AnimDrawDebugCapsule(const FVector& Center, float HalfHeight, float Radius, const FRotator& Rotation, const FColor& Color, bool bPersistentLines = false, float LifeTime = -1.f, float Thickness = 0.f) {}
+	ENGINE_API void AnimDrawDebugCapsule(const FVector& Center, float HalfHeight, float Radius, const FRotator& Rotation, const FColor& Color, bool bPersistentLines = false, float LifeTime = -1.f, float Thickness = 0.f, ESceneDepthPriorityGroup DepthPriority = SDPG_World) {}
 #endif // ENABLE_ANIM_DRAW_DEBUG
 
 #if ENABLE_ANIM_LOGGING
