Fixing bug with AHUD::GetActorsInSelectionRectangle that was incorrectly allowing partial matches when bActorMustBeFullyEnclosed is true

Before this change if bActorMustBeFullyEnclosed is true and SelectionRectangle.IsInside(ActorBox2D) is false,  it would still consider the else if condition SelectionRectangle.Intersect(ActorBox2D)  but that should not be the case. If bActorMustBeFullyEnclosed we do not want to consider the Partial intersection case.

diff --git a/Engine/Source/Runtime/Engine/Private/HUD.cpp b/Engine/Source/Runtime/Engine/Private/HUD.cpp
--- a/Engine/Source/Runtime/Engine/Private/HUD.cpp
+++ b/Engine/Source/Runtime/Engine/Private/HUD.cpp
@@ -1092,9 +1092,12 @@ void AHUD::GetActorsInSelectionRectangle(TSubclassOf<class AActor> ClassFilter,
 		if (ActorBox2D.bIsValid)
 		{
 			//Selection Box must fully enclose the Projected Actor Bounds
-			if (bActorMustBeFullyEnclosed && SelectionRectangle.IsInside(ActorBox2D))
+			if (bActorMustBeFullyEnclosed)
 			{
-				OutActors.Add(EachActor);
+				if (SelectionRectangle.IsInside(ActorBox2D))
+				{
+					OutActors.Add(EachActor);
+				}
 			}
 			//Partial Intersection with Projected Actor Bounds
 			else if (SelectionRectangle.Intersect(ActorBox2D))
