SkeletalMesh: Fix GetViewRelevance not taking editor effects into account. This fixes a bug where colored overlays don't work for skeletal meshes.

diff --git a/Engine/Source/Runtime/Engine/Private/SkeletalMeshSceneProxy.cpp b/Engine/Source/Runtime/Engine/Private/SkeletalMeshSceneProxy.cpp
--- a/Engine/Source/Runtime/Engine/Private/SkeletalMeshSceneProxy.cpp
+++ b/Engine/Source/Runtime/Engine/Private/SkeletalMeshSceneProxy.cpp
@@ -1051,7 +1051,7 @@ FPrimitiveViewRelevance FSkeletalMeshSceneProxy::GetViewRelevance(const FSceneVi
 	//only check these in the editor
 	if (Result.bStaticRelevance)
 	{
-		Result.bEditorStaticSelectionRelevance = (IsSelected() || IsHovered());
+		Result.bEditorStaticSelectionRelevance = (IsSelected() || IsHovered() || WantsEditorEffects());
 
 		Result.bEditorVisualizeLevelInstanceRelevance = IsEditingLevelInstanceChild();
 	}
