Unshelved from pending changelist '36623991':

Fix ClusterUnion not triggering events

Note: You need to enable hit events on the cluster union as well as the child geometry collection(s)

diff --git a/Engine/Source/Runtime/Engine/Private/PhysicsEngine/ClusterUnionComponent.cpp b/Engine/Source/Runtime/Engine/Private/PhysicsEngine/ClusterUnionComponent.cpp
--- a/Engine/Source/Runtime/Engine/Private/PhysicsEngine/ClusterUnionComponent.cpp
+++ b/Engine/Source/Runtime/Engine/Private/PhysicsEngine/ClusterUnionComponent.cpp
@@ -733,6 +733,11 @@ void UClusterUnionComponent::OnCreatePhysicsState()
 	if (FPhysScene_Chaos* Scene = GetChaosScene())
 	{
 		Scene->AddObject(this, PhysicsProxy);
+
+		if (BodyInstance.bNotifyRigidBodyCollision)
+		{
+			Scene->RegisterForCollisionEvents(this);
+		}
 	}
 	
 	if (bUseClusterUnionAccelerationStructure)
