World Partition - Spline partitionning - Reworked some of the logic in UControlPointMeshComponent / USplineMeshComponent IsEditorOnly()
* Only make these components Editor Only if they were processed by the spline partitioning builder

diff --git a/Engine/Source/Runtime/Engine/Private/Components/SplineMeshComponent.cpp b/Engine/Source/Runtime/Engine/Private/Components/SplineMeshComponent.cpp
--- a/Engine/Source/Runtime/Engine/Private/Components/SplineMeshComponent.cpp
+++ b/Engine/Source/Runtime/Engine/Private/Components/SplineMeshComponent.cpp
@@ -832,7 +832,7 @@ bool USplineMeshComponent::IsRelevantForSplinePartitioning() const
 
 bool USplineMeshComponent::IsEditorOnly() const
 {
-	if (!IsRelevantForSplinePartitioning())
+	if (Super::IsEditorOnly())
 	{
 		return true;
 	}
@@ -841,7 +841,8 @@ bool USplineMeshComponent::IsEditorOnly() const
 	const ALandscapeSplineActor* SplineActor = Cast<ALandscapeSplineActor>(GetOwner());
 	if (SplineActor && SplineActor->HasGeneratedLandscapeSplineMeshesActors())
 	{
-		return true;
+		// Treat as editor only if the component was included in the spline partitioning
+		return IsRelevantForSplinePartitioning();
 	}
 
 	return false;
diff --git a/Engine/Source/Runtime/Landscape/Private/LandscapeSplines.cpp b/Engine/Source/Runtime/Landscape/Private/LandscapeSplines.cpp
--- a/Engine/Source/Runtime/Landscape/Private/LandscapeSplines.cpp
+++ b/Engine/Source/Runtime/Landscape/Private/LandscapeSplines.cpp
@@ -1736,7 +1736,7 @@ bool UControlPointMeshComponent::IsRelevantForSplinePartitioning() const
 
 bool UControlPointMeshComponent::IsEditorOnly() const
 {
-	if (!IsRelevantForSplinePartitioning())
+	if (Super::IsEditorOnly())
 	{
 		return true;
 	}
@@ -1745,7 +1745,8 @@ bool UControlPointMeshComponent::IsEditorOnly() const
 	ALandscapeSplineActor* SplineActor = Cast<ALandscapeSplineActor>(GetOwner());
 	if (SplineActor && SplineActor->HasGeneratedLandscapeSplineMeshesActors())
 	{
-		return true;
+		// Treat as editor only if the component was included in the spline partitioning
+		return IsRelevantForSplinePartitioning();
 	}
 
 	return false;
