[BUGFIX] Assets copied lose the ability to reimport
- Changed AssetImportData to serialize using absolute paths.

diff --git a/Engine/Source/Runtime/Engine/Classes/EditorFramework/AssetImportData.h b/Engine/Source/Runtime/Engine/Classes/EditorFramework/AssetImportData.h
--- a/Engine/Source/Runtime/Engine/Classes/EditorFramework/AssetImportData.h
+++ b/Engine/Source/Runtime/Engine/Classes/EditorFramework/AssetImportData.h
@@ -49,11 +49,11 @@ struct FAssetImportInfo
 		FString DisplayLabelName;
 	};
 
-	/** Convert this import information to JSON */
-	ENGINE_API FString ToJson() const;
+	/** Convert this import information to JSON by converting the paths to absolute paths resolved with the Outer object provided*/
+	ENGINE_API FString ToJson(const UPackage* Outermost = nullptr) const;
 
-	/** Attempt to parse an asset import structure from the specified json string. */
-	static ENGINE_API TOptional<FAssetImportInfo> FromJson(FString InJsonString);
+	/** Attempt to parse an asset import structure from the specified json string. If outermost is provided, the paths would be converted to relative paths. */
+	static ENGINE_API TOptional<FAssetImportInfo> FromJson(FString InJsonString, const UPackage* Outermost = nullptr);
 
 	/** Insert information pertaining to a new source file to this structure */
 	void Insert(const FSourceFile& InSourceFile) { SourceFiles.Add(InSourceFile); }
@@ -188,7 +188,7 @@ protected:
 
 	/** Overridden serialize function to write out the underlying data as json */
 	ENGINE_API virtual void Serialize(FStructuredArchive::FRecord Record) override;
-	
+
 #endif		// WITH_EDITORONLY_DATA
 };
 
diff --git a/Engine/Source/Runtime/Engine/Private/EditorFramework/AssetImportData.cpp b/Engine/Source/Runtime/Engine/Private/EditorFramework/AssetImportData.cpp
--- a/Engine/Source/Runtime/Engine/Private/EditorFramework/AssetImportData.cpp
+++ b/Engine/Source/Runtime/Engine/Private/EditorFramework/AssetImportData.cpp
@@ -38,7 +38,7 @@ void UAssetImportData::ScriptedAddFilename(const FString& InPath, int32 Index, F
 
 FOnImportDataChanged UAssetImportData::OnImportDataChanged;
 
-FString FAssetImportInfo::ToJson() const
+FString FAssetImportInfo::ToJson(const UPackage* Outermost/*= nullptr*/) const
 {
 	FString Json;
 	Json.Reserve(1024);
@@ -47,11 +47,11 @@ FString FAssetImportInfo::ToJson() const
 	for (int32 Index = 0; Index < SourceFiles.Num(); ++Index)
 	{
 		Json += FString::Printf(TEXT("{ \"RelativeFilename\" : \"%s\", \"Timestamp\" : \"%" INT64_FMT "\", \"FileMD5\" : \"%s\", \"DisplayLabelName\" : \"%s\" }"),
-			*SourceFiles[Index].RelativeFilename,
+			*UAssetImportData::ResolveImportFilename(SourceFiles[Index].RelativeFilename, Outermost),
 			SourceFiles[Index].Timestamp.ToUnixTimestamp(),
 			*LexToString(SourceFiles[Index].FileHash),
 			*SourceFiles[Index].DisplayLabelName
-			);
+		);
 
 		if (Index != SourceFiles.Num() - 1)
 		{
@@ -63,7 +63,7 @@ FString FAssetImportInfo::ToJson() const
 	return Json;
 }
 
-TOptional<FAssetImportInfo> FAssetImportInfo::FromJson(FString InJsonString)
+TOptional<FAssetImportInfo> FAssetImportInfo::FromJson(FString InJsonString, const UPackage* Outermost/*= nullptr*/)
 {
 	// Load json
 	TSharedRef<TJsonReader<TCHAR>> Reader = FJsonStringReader::Create(MoveTemp(InJsonString));
@@ -73,7 +73,7 @@ TOptional<FAssetImportInfo> FAssetImportInfo::FromJson(FString InJsonString)
 	{
 		return TOptional<FAssetImportInfo>();
 	}
-	
+
 	FAssetImportInfo Info;
 
 	for (const auto& Value : JSONSourceFiles)
@@ -95,6 +95,8 @@ TOptional<FAssetImportInfo> FAssetImportInfo::FromJson(FString InJsonString)
 			continue;
 		}
 
+		RelativeFilename = UAssetImportData::SanitizeImportFilename(RelativeFilename, Outermost);
+		
 		int64 UnixTimestamp = 0;
 		LexFromString(UnixTimestamp, *TimestampString);
 
@@ -428,7 +430,7 @@ void UAssetImportData::Serialize(FStructuredArchive::FRecord Record)
 			if (BaseArchive.IsLoading())
 			{
 				Record << SA_VALUE(TEXT("Json"), Json);
-				TOptional<FAssetImportInfo> Copy = FAssetImportInfo::FromJson(MoveTemp(Json));
+				TOptional<FAssetImportInfo> Copy = FAssetImportInfo::FromJson(MoveTemp(Json), GetOutermost());
 				if (Copy.IsSet())
 				{
 					SourceData = MoveTemp(Copy.GetValue());
@@ -436,7 +438,7 @@ void UAssetImportData::Serialize(FStructuredArchive::FRecord Record)
 			}
 			else if (BaseArchive.IsSaving())
 			{
-				Json = SourceData.ToJson();
+				Json = SourceData.ToJson(GetOutermost());
 				Record << SA_VALUE(TEXT("Json"), Json);
 			}
 		}
