Provide RetargetSourceAsset getters and expose to BP, as after deprecatiing the direct use of the retarget source asset, there is no way to perform certain actions

diff --git a/Engine/Source/Runtime/Engine/Classes/Animation/AnimSequence.h b/Engine/Source/Runtime/Engine/Classes/Animation/AnimSequence.h
--- a/Engine/Source/Runtime/Engine/Classes/Animation/AnimSequence.h
+++ b/Engine/Source/Runtime/Engine/Classes/Animation/AnimSequence.h
@@ -325,7 +325,7 @@ public:
 
 #if WITH_EDITORONLY_DATA
 	/** If RetargetSource is set to Default (None), this is asset for the base pose to use when retargeting. Transform data will be saved in RetargetSourceAssetReferencePose. */
-	UE_DEPRECATED(5.5, "Direct access to RetargetSourceAsset has been deprecated. Please use member SetRetargetSourceAssetData instead.")
+	UE_DEPRECATED(5.5, "Direct access to RetargetSourceAsset has been deprecated. Please use members GetRetargetSourceAsset & SetRetargetSourceAsset instead.")
 	UPROPERTY(EditAnywhere, AssetRegistrySearchable, Category=Animation, meta = (DisallowedClasses = "/Script/ApexDestruction.DestructibleMesh"))
 	TSoftObjectPtr<USkeletalMesh> RetargetSourceAsset;
 #endif
@@ -492,14 +492,30 @@ protected:
 	
 	UE_DEPRECATED(5.3, "Please use UpdateCompressedCurveName that takes FNames.")
 	void UpdateCompressedCurveName(SmartName::UID_Type CurveUID, const struct FSmartName& NewCurveName) {}
+#endif
 
 public:
+#if WITH_EDITOR
+	// Assigns the passed skeletal mesh to the retarget source
+	UFUNCTION(BlueprintCallable, Category = "Animation")
 	ENGINE_API void SetRetargetSourceAsset(USkeletalMesh* InRetargetSourceAsset);
-	// Update the retarget data post from the source, if it exist, else clears the retarget data pose
+
+	// Resets the retarget source asset
+	UFUNCTION(BlueprintCallable, Category = "Animation")
+	ENGINE_API void ClearRetargetSourceAsset();
+
+	// Returns the retarget source asset soft object pointer.
+	UFUNCTION(BlueprintPure, Category = "Animation")
+	ENGINE_API const TSoftObjectPtr<USkeletalMesh>& GetRetargetSourceAsset() const;
+
+	// Update the retarget data pose from the source, if it exist, else clears the retarget data pose saved in RetargetSourceAssetReferencePose.
 	// Warning : This function calls LoadSynchronous at the retarget source asset soft object pointer, so it can not be used at PostLoad
+	UFUNCTION(BlueprintCallable, Category = "Animation")
 	ENGINE_API void UpdateRetargetSourceAssetData();
+#endif
 
 private:
+#if WITH_EDITORONLY_DATA
 	/** Updates the stored sampling frame-rate using the sequence length and number of sampling keys */
 	UE_DEPRECATED(5.0, "UpdateFrameRate has been deprecated see UAnimDataController::SetFrameRate")
 	ENGINE_API void UpdateFrameRate();
diff --git a/Engine/Source/Runtime/Engine/Classes/Animation/PoseAsset.h b/Engine/Source/Runtime/Engine/Classes/Animation/PoseAsset.h
--- a/Engine/Source/Runtime/Engine/Classes/Animation/PoseAsset.h
+++ b/Engine/Source/Runtime/Engine/Classes/Animation/PoseAsset.h
@@ -205,7 +205,7 @@ public:
 
 #if WITH_EDITORONLY_DATA
 	/** If RetargetSource is set to Default (None), this is asset for the base pose to use when retargeting. Transform data will be saved in RetargetSourceAssetReferencePose. */
-	UE_DEPRECATED(5.5, "Direct access to RetargetSourceAsset has been deprecated. Please use member SetRetargetSourceAsset instead.")
+	UE_DEPRECATED(5.5, "Direct access to RetargetSourceAsset has been deprecated. Please use members GetRetargetSourceAsset & SetRetargetSourceAsset instead.")
 	UPROPERTY(EditAnywhere, AssetRegistrySearchable, Category=Animation, meta = (DisallowedClasses = "/Script/ApexDestruction.DestructibleMesh"))
 	TSoftObjectPtr<USkeletalMesh> RetargetSourceAsset;
 #endif
@@ -423,10 +423,22 @@ private:
 	bool RemoveInvalidTracks();
 
 public:
-#if WITH_EDITORONLY_DATA
+#if WITH_EDITOR
+	// Assigns the passed skeletal mesh to the retarget source
+	UFUNCTION(BlueprintCallable, Category = "Animation")
 	ENGINE_API void SetRetargetSourceAsset(USkeletalMesh* InRetargetSourceAsset);
-	// Update the retarget data post from the source, if it exist, else clears the retarget data pose
+
+	// Resets the retarget source asset
+	UFUNCTION(BlueprintCallable, Category = "Animation")
+	ENGINE_API void ClearRetargetSourceAsset();
+
+	// Returns the retarget source asset soft object pointer.
+	UFUNCTION(BlueprintPure, Category = "Animation")
+	ENGINE_API const TSoftObjectPtr<USkeletalMesh>& GetRetargetSourceAsset() const;
+
+	// Update the retarget data pose from the source, if it exist, else clears the retarget data pose saved in RetargetSourceAssetReferencePose.
 	// Warning : This function calls LoadSynchronous at the retarget source asset soft object pointer, so it can not be used at PostLoad
+	UFUNCTION(BlueprintCallable, Category = "Animation")
 	ENGINE_API void UpdateRetargetSourceAssetData();
 #endif
 
diff --git a/Engine/Source/Runtime/Engine/Private/Animation/AnimSequence.cpp b/Engine/Source/Runtime/Engine/Private/Animation/AnimSequence.cpp
--- a/Engine/Source/Runtime/Engine/Private/Animation/AnimSequence.cpp
+++ b/Engine/Source/Runtime/Engine/Private/Animation/AnimSequence.cpp
@@ -1956,7 +1956,7 @@ void UAnimSequence::UpdateCompressedCurveName(const FName& OldCurveName, const F
 }
 #endif // WITH_EDITOR
 
-#if WITH_EDITORONLY_DATA
+#if WITH_EDITOR
 void UAnimSequence::SetRetargetSourceAsset(USkeletalMesh* InRetargetSourceAsset)
 {
 	if (InRetargetSourceAsset != nullptr && InRetargetSourceAsset->HasAnyFlags(RF_Transient))
@@ -1973,6 +1973,20 @@ void UAnimSequence::SetRetargetSourceAsset(USkeletalMesh* InRetargetSourceAsset)
 	PRAGMA_ENABLE_DEPRECATION_WARNINGS
 }
 
+const TSoftObjectPtr<USkeletalMesh>& UAnimSequence::GetRetargetSourceAsset() const
+{
+	PRAGMA_DISABLE_DEPRECATION_WARNINGS
+	return RetargetSourceAsset;
+	PRAGMA_ENABLE_DEPRECATION_WARNINGS
+}
+
+void UAnimSequence::ClearRetargetSourceAsset()
+{
+	PRAGMA_DISABLE_DEPRECATION_WARNINGS
+	RetargetSourceAsset.Reset();
+	PRAGMA_ENABLE_DEPRECATION_WARNINGS
+}
+
 void UAnimSequence::UpdateRetargetSourceAssetData()
 {
 	PRAGMA_DISABLE_DEPRECATION_WARNINGS
@@ -1989,7 +2003,7 @@ void UAnimSequence::UpdateRetargetSourceAssetData()
 		RetargetSourceAssetReferencePose.Empty();
 	}
 }
-#endif // WITH_EDITORONLY_DATA
+#endif // WITH_EDITOR
 
 const TArray<FTransform>& UAnimSequence::GetRetargetTransforms() const
 {
diff --git a/Engine/Source/Runtime/Engine/Private/Animation/PoseAsset.cpp b/Engine/Source/Runtime/Engine/Private/Animation/PoseAsset.cpp
--- a/Engine/Source/Runtime/Engine/Private/Animation/PoseAsset.cpp
+++ b/Engine/Source/Runtime/Engine/Private/Animation/PoseAsset.cpp
@@ -1806,7 +1806,7 @@ bool UPoseAsset::RemoveInvalidTracks()
 	return InitialNumTracks != PoseContainer.Tracks.Num();
 }
 
-#if WITH_EDITORONLY_DATA
+#if WITH_EDITOR
 void UPoseAsset::SetRetargetSourceAsset(USkeletalMesh* InRetargetSourceAsset)
 {
 	if (InRetargetSourceAsset != nullptr && InRetargetSourceAsset->HasAnyFlags(RF_Transient))
@@ -1823,6 +1823,20 @@ void UPoseAsset::SetRetargetSourceAsset(USkeletalMesh* InRetargetSourceAsset)
 	PRAGMA_ENABLE_DEPRECATION_WARNINGS
 }
 
+void UPoseAsset::ClearRetargetSourceAsset()
+{
+	PRAGMA_DISABLE_DEPRECATION_WARNINGS
+	RetargetSourceAsset.Reset();
+	PRAGMA_ENABLE_DEPRECATION_WARNINGS
+}
+
+const TSoftObjectPtr<USkeletalMesh>& UPoseAsset::GetRetargetSourceAsset() const
+{
+	PRAGMA_DISABLE_DEPRECATION_WARNINGS
+	return RetargetSourceAsset;
+	PRAGMA_ENABLE_DEPRECATION_WARNINGS
+}
+
 void UPoseAsset::UpdateRetargetSourceAssetData()
 {
 	PRAGMA_DISABLE_DEPRECATION_WARNINGS
