[World Partition] Make sure to output streaming generation logs at the good location when running on build machines so they are treated as artifacts on Horde jobs.

diff --git a/Engine/Source/Runtime/Engine/Private/WorldPartition/WorldPartitionStreamingGeneration.cpp b/Engine/Source/Runtime/Engine/Private/WorldPartition/WorldPartitionStreamingGeneration.cpp
--- a/Engine/Source/Runtime/Engine/Private/WorldPartition/WorldPartitionStreamingGeneration.cpp
+++ b/Engine/Source/Runtime/Engine/Private/WorldPartition/WorldPartitionStreamingGeneration.cpp
@@ -1753,8 +1753,9 @@ public:
 
 	static TUniquePtr<FArchive> CreateDumpStateLogArchive(const TCHAR* Suffix)
 	{
-		FString StateLogOutputFilename = FPaths::ProjectLogDir() / TEXT("WorldPartition") / FString::Printf(TEXT("StreamingGeneration-%s"), Suffix);
-		StateLogOutputFilename += FString::Printf(TEXT("-%08x-%s.log"), FPlatformProcess::GetCurrentProcessId(), *FDateTime::Now().ToIso8601().Replace(TEXT(":"), TEXT(".")));
+		const FString LogDir = GIsBuildMachine ? FPaths::Combine(*FPaths::EngineDir(), TEXT("Programs"), TEXT("AutomationTool"), TEXT("Saved"), TEXT("Logs")) : FPaths::ProjectLogDir();
+		const FString StateLogOutputFilename = FPaths::Combine(LogDir, TEXT("WorldPartition"), 
+			FString::Printf(TEXT("StreamingGeneration-%s-%08x-%s.log"), Suffix, FPlatformProcess::GetCurrentProcessId(), *FDateTime::Now().ToIso8601().Replace(TEXT(":"), TEXT("."))));
 		return TUniquePtr<FArchive>(IFileManager::Get().CreateFileWriter(*StateLogOutputFilename));
 	}
 
