Updated FISMComponentBatcher so it can be expanded by subclasses to support new component types.

#rnx

diff --git a/Engine/Source/Runtime/Engine/Public/ISMPartition/ISMComponentBatcher.h b/Engine/Source/Runtime/Engine/Public/ISMPartition/ISMComponentBatcher.h
--- a/Engine/Source/Runtime/Engine/Public/ISMPartition/ISMComponentBatcher.h
+++ b/Engine/Source/Runtime/Engine/Public/ISMPartition/ISMComponentBatcher.h
@@ -73,7 +73,7 @@ public:
 		return Hash;
 	}
 
-private:
+protected:
 	void ENGINE_API AddInternal(const UActorComponent* InComponent, TOptional<TFunctionRef<FTransform(const FTransform&)>> InTransformFunc);
 	void ENGINE_API ComputeHash() const;
 
