Alternate Skin Weights: remove GSkinnedMeshSkinWeightProfileEarlyExit to get back to the previous behavior

NOTE: ClearSkinWeightProfile has to be called when initilizing to make sure that the alternative skin weights are properly set when calling SetSkinWeightProfile later

#rnx

diff --git a/Engine/Source/Runtime/Engine/Private/Components/SkinnedMeshComponent.cpp b/Engine/Source/Runtime/Engine/Private/Components/SkinnedMeshComponent.cpp
--- a/Engine/Source/Runtime/Engine/Private/Components/SkinnedMeshComponent.cpp
+++ b/Engine/Source/Runtime/Engine/Private/Components/SkinnedMeshComponent.cpp
@@ -104,13 +104,6 @@ static FAutoConsoleVariableRef CVarSkinnedMeshRenderNanite(
 	)
 );
 
-static bool GSkinnedMeshSkinWeightProfileEarlyExit = false;
-static FAutoConsoleVariableRef CVarSkinnedMeshSkinWeightProfileEarlyExit(
-	TEXT("r.SkinnedMesh.SkinWeightProfileEarlyExit"),
-   GSkinnedMeshSkinWeightProfileEarlyExit,
-   TEXT("When set to true, the new profile stack update is not re-done if it's similar to the current skin weight profile.")
-   );
-
 bool ShouldRenderNaniteSkinnedMeshes()
 {
 	return NaniteSkinnedMeshesSupported() && GSkinnedMeshRenderNanite != 0;
@@ -5047,10 +5040,7 @@ bool USkinnedMeshComponent::SetSkinWeightProfileStack(const FSkinWeightProfileSt
 	// If we're not actually changing anything, then just return and say we did, as if the operation succeeded.
 	if (InProfileStack.Normalized() == FSkinWeightProfileStack{CurrentSkinWeightProfileLayers}.Normalized())
 	{
-		if (GSkinnedMeshSkinWeightProfileEarlyExit)
-		{
-			return true;
-		}
+		return true;
 	}
 
 	bool bChanged = false;
