[WorldPartition] Revert behavior to stop tracking dirty actors outside of transactions
- This was causing issues in PCG where actors are created/deleted/mocified mostly outside of transactions
- PCG: Remove dummy transaction introduced for 5.4 to paliate this behavior
- Keep NonDirtyActors that can't be pinned, this fixes a bug with Landscape location volumes being unloaded after modify/undo (known case are newly created always loaded actors)
- Get a reference in OnRemoveNonDirtyActor instead of waiting for tick
- Introduce new disable tracking scope for use cases where we don't want the new actors to remain pinned after a save (Landscape use case)

#rnx

diff --git a/Engine/Plugins/PCG/Source/PCG/Private/Helpers/PCGActorHelpers.cpp b/Engine/Plugins/PCG/Source/PCG/Private/Helpers/PCGActorHelpers.cpp
--- a/Engine/Plugins/PCG/Source/PCG/Private/Helpers/PCGActorHelpers.cpp
+++ b/Engine/Plugins/PCG/Source/PCG/Private/Helpers/PCGActorHelpers.cpp
@@ -21,7 +21,6 @@
 #include "Engine/Level.h"
 #include "Editor/IPCGEditorModule.h"
 #include "Modules/ModuleManager.h"
-#include "ScopedTransaction.h"
 #include "Subsystems/ActorEditorContextSubsystem.h"
 #include "WorldPartition/DataLayer/DataLayerInstance.h"
 #include "WorldPartition/DataLayer/ExternalDataLayerAsset.h"
@@ -405,18 +404,6 @@ bool UPCGActorHelpers::DeleteActors(UWorld* World, const TArray<TSoftObjectPtr<A
 			EditorModule->SetOutlinerUIRefreshDelay(UE::PCGActorHelpers::Local::OutlinerUIRefreshDelay);
 		}
 
-		// Create TX so that dirty actor packages are tracked
-		// 
-		// Without tracking deleted actor packages will get unloaded on the next GC with no chance to save them first 
-		// See FWorldPartitionExternalDirtyActorsTracker::OnAddDirtyActor (Since CL 32133290)
-		//
-		// The Actors not being referenced by the Dirty Tracker will prevent them from being collected in UWorldPartition::AddReferencedObjects
-		// then in UWorldPartition::OnGCPostReachabilityAnalysis all unreachable actors will get processed to remove RF_Standalone flags allowing 
-		// the next GC to collect those packages
-		//
-		// The fix here is to create a dummy transaction so that the deleted actors are tracked but since we don't want to actually push a transaction, we cancel it after the DestroyActor calls
-		const bool bShouldTransact = World && !World->IsGameWorld() && ActorsToDelete.Num() > 0 && !GEditor->IsTransactionActive() && !GIsTransacting;
-		FScopedTransaction DummyTransaction(NSLOCTEXT("PCGActorHelpers", "DummyTransaction", "DummyTransaction"), bShouldTransact);
 #endif
 
 		// Not in editor, really unlikely to happen but might be slow
@@ -431,11 +418,6 @@ bool UPCGActorHelpers::DeleteActors(UWorld* World, const TArray<TSoftObjectPtr<A
 				}
 			}
 		}
-
-#if WITH_EDITOR
-		// Cancel the Dummy transaction so that it can't be undone.
-		DummyTransaction.Cancel();
-#endif
 	}
 
 	return true;
diff --git a/Engine/Source/Editor/LandscapeEditor/Private/LandscapeEditorDetailCustomization_NewLandscape.cpp b/Engine/Source/Editor/LandscapeEditor/Private/LandscapeEditorDetailCustomization_NewLandscape.cpp
--- a/Engine/Source/Editor/LandscapeEditor/Private/LandscapeEditorDetailCustomization_NewLandscape.cpp
+++ b/Engine/Source/Editor/LandscapeEditor/Private/LandscapeEditorDetailCustomization_NewLandscape.cpp
@@ -1250,7 +1250,7 @@ FReply FLandscapeEditorDetailCustomization_NewLandscape::OnCreateButtonClicked()
 
 			// ensures all the final height textures have been updated.
 			LandscapeInfo->ForceLayersFullUpdate();
-			LandscapeEditorUtils::SaveLandscapeProxies(MakeArrayView(CreatedStreamingProxies));
+			LandscapeEditorUtils::SaveLandscapeProxies(World, MakeArrayView(CreatedStreamingProxies));
 			LandscapeBounds += LandscapeInfo->GetCompleteBounds();
 
 			Progress.EnterProgressFrame(1.0f , FText::Format(LOCTEXT("LandscapeCreateRegion", "Creating Landscape Editor Regions ({0}, {1})"), RegionCoordinate.X, RegionCoordinate.Y));
@@ -1284,7 +1284,7 @@ FReply FLandscapeEditorDetailCustomization_NewLandscape::OnCreateButtonClicked()
 			return true;
 		});
 
-		LandscapeEditorUtils::SaveLandscapeProxies(MakeArrayView(AllProxies));
+		LandscapeEditorUtils::SaveLandscapeProxies(World, MakeArrayView(AllProxies));
 	}
 
 	return FReply::Handled();
diff --git a/Engine/Source/Editor/LandscapeEditor/Private/LandscapeEditorUtils.cpp b/Engine/Source/Editor/LandscapeEditor/Private/LandscapeEditorUtils.cpp
--- a/Engine/Source/Editor/LandscapeEditor/Private/LandscapeEditorUtils.cpp
+++ b/Engine/Source/Editor/LandscapeEditor/Private/LandscapeEditorUtils.cpp
@@ -87,17 +87,12 @@ namespace LandscapeEditorUtils
 
 	}
 
-	void SaveLandscapeProxies(TArrayView<ALandscapeProxy*> Proxies)
+	void SaveLandscapeProxies(UWorld* World, TArrayView<ALandscapeProxy*> Proxies)
 	{
 		// Save the proxies
-		{
-			TRACE_CPUPROFILER_EVENT_SCOPE(SaveCreatedActors);
-			LandscapeEditorUtils::SaveObjects(Proxies);
-		}
-
-		// Grab references to proxies so they get unloaded after this function returns
-		TArray<FWorldPartitionReference> ProxyReferences;
-		Algo::Transform(Proxies, ProxyReferences, [](ALandscapeProxy* Proxy) { return FWorldPartitionReference(FWorldPartitionHelpers::GetWorldPartition(Proxy), Proxy->GetActorGuid()); });
+		TRACE_CPUPROFILER_EVENT_SCOPE(SaveCreatedActors);
+		UWorldPartition::FDisableNonDirtyActorTrackingScope Scope(World->GetWorldPartition(), true);
+		LandscapeEditorUtils::SaveObjects(Proxies);
 	}
 }
 
diff --git a/Engine/Source/Editor/LandscapeEditor/Private/LandscapeRegionUtils.cpp b/Engine/Source/Editor/LandscapeEditor/Private/LandscapeRegionUtils.cpp
--- a/Engine/Source/Editor/LandscapeEditor/Private/LandscapeRegionUtils.cpp
+++ b/Engine/Source/Editor/LandscapeEditor/Private/LandscapeRegionUtils.cpp
@@ -135,7 +135,7 @@ void ForEachRegion_LoadProcessUnload(ULandscapeInfo* InLandscapeInfo, const FInt
 
 		InLandscapeInfo->ForceLayersFullUpdate();
 
-		LandscapeEditorUtils::SaveLandscapeProxies(MakeArrayView(LandscapeProxies));
+		LandscapeEditorUtils::SaveLandscapeProxies(InWorld, MakeArrayView(LandscapeProxies));
 
 		Region->Unload();
 
diff --git a/Engine/Source/Editor/LandscapeEditor/Public/LandscapeEditorUtils.h b/Engine/Source/Editor/LandscapeEditor/Public/LandscapeEditorUtils.h
--- a/Engine/Source/Editor/LandscapeEditor/Public/LandscapeEditorUtils.h
+++ b/Engine/Source/Editor/LandscapeEditor/Public/LandscapeEditorUtils.h
@@ -25,5 +25,5 @@ namespace LandscapeEditorUtils
 		UEditorLoadingAndSavingUtils::SavePackages(Packages, /* bOnlyDirty = */ false);
 	}
 
-	void SaveLandscapeProxies(TArrayView<ALandscapeProxy*> Proxies);
+	void SaveLandscapeProxies(UWorld* InWorld, TArrayView<ALandscapeProxy*> Proxies);
 }
diff --git a/Engine/Source/Runtime/Engine/Private/WorldPartition/WorldPartition.cpp b/Engine/Source/Runtime/Engine/Private/WorldPartition/WorldPartition.cpp
--- a/Engine/Source/Runtime/Engine/Private/WorldPartition/WorldPartition.cpp
+++ b/Engine/Source/Runtime/Engine/Private/WorldPartition/WorldPartition.cpp
@@ -259,6 +259,24 @@ public:
 	}
 };
 
+UWorldPartition::FDisableNonDirtyActorTrackingScope::FDisableNonDirtyActorTrackingScope(UWorldPartition* InWorldPartition, bool bInDisableTracking)
+	: WorldPartition(InWorldPartition)
+{
+	if (WorldPartition && WorldPartition->ExternalDirtyActorsTracker)
+	{
+		bPreviousValue = WorldPartition->ExternalDirtyActorsTracker->IsNonDirtyTrackingDisabled();
+		WorldPartition->ExternalDirtyActorsTracker->SetNonDirtyTrackingDisabled(bInDisableTracking);
+	}
+}
+
+UWorldPartition::FDisableNonDirtyActorTrackingScope::~FDisableNonDirtyActorTrackingScope()
+{
+	if (WorldPartition && WorldPartition->ExternalDirtyActorsTracker)
+	{
+		WorldPartition->ExternalDirtyActorsTracker->SetNonDirtyTrackingDisabled(bPreviousValue);
+	}
+}
+
 UWorldPartition::FWorldPartitionExternalDirtyActorsTracker::FWorldPartitionExternalDirtyActorsTracker()
 	: Super(nullptr, nullptr)
 {}
@@ -270,7 +288,20 @@ UWorldPartition::FWorldPartitionExternalDirtyActorsTracker::FWorldPartitionExter
 void UWorldPartition::FWorldPartitionExternalDirtyActorsTracker::OnRemoveNonDirtyActor(TWeakObjectPtr<AActor> InActor, FWorldPartitionReference& InValue)
 {
 	check(InActor.IsValid());
-	NonDirtyActors.Emplace({ InActor, InValue });
+
+	FWorldPartitionReference NonDirtyReference = InValue;
+	
+	// Grab Reference if it isn't valid. This means we saved a new actor.
+	if (!NonDirtyReference.IsValid() && InActor.IsValid())
+	{
+		NonDirtyReference = FWorldPartitionReference(Owner, InActor->GetActorGuid());
+	}
+
+	// If Tracking is disabled, the reference will get released and actor will get unloaded (if not referenced elsewhere)
+	if (!bIsNonDirtyTrackingDisabled)
+	{
+		NonDirtyActors.Emplace({ InActor, NonDirtyReference });
+	}
 }
 
 void UWorldPartition::FWorldPartitionExternalDirtyActorsTracker::Tick(float InDeltaSeconds)
@@ -279,10 +310,10 @@ void UWorldPartition::FWorldPartitionExternalDirtyActorsTracker::Tick(float InDe
 
 	for (auto& [Actor, Reference] : NonDirtyActors)
 	{
-		// Resolve reference for newly added actors
-		if (!Reference.IsValid() && Actor.IsValid())
+		// If Actor was dirtied since last tick ignore it
+		if (DirtyActors.Contains(Actor))
 		{
-			Reference = FWorldPartitionReference(Owner, Actor->GetActorGuid());
+			continue;
 		}
 
 		// Transfer ownership of our last ref if actor can be pinned
@@ -1356,7 +1387,7 @@ void UWorldPartition::OnActorDescInstanceAdded(FWorldPartitionActorDescInstance*
 		ForceLoadedActors->AddActors({ NewActorDescInstance->GetGuid() });
 	}
 
-	bForceRefreshAlwaysLoaded = !NewActorDescInstance->GetIsSpatiallyLoaded();
+	bForceRefreshAlwaysLoaded |= !NewActorDescInstance->GetIsSpatiallyLoaded();
 	bForceRefreshEditor = true;
 }
 
@@ -1374,7 +1405,7 @@ void UWorldPartition::OnActorDescInstanceRemoved(FWorldPartitionActorDescInstanc
 		ForceLoadedActors->RemoveActors({ ActorDescInstance->GetGuid() });
 	}
 
-	bForceRefreshAlwaysLoaded = !ActorDescInstance->GetIsSpatiallyLoaded();
+	bForceRefreshAlwaysLoaded |= !ActorDescInstance->GetIsSpatiallyLoaded();
 	bForceRefreshEditor = true;
 }
 
@@ -1672,11 +1703,7 @@ void UWorldPartition::Tick(float DeltaSeconds)
 		EditorHash->Tick(DeltaSeconds);
 	}
 
-	if (ExternalDirtyActorsTracker)
-	{
-		ExternalDirtyActorsTracker->Tick(DeltaSeconds);
-	}
-
+	// Force refresh needs to happen before dirty tracker tick to allow new always loaded actors to be referenced before releasing NonDirtyActors
 	if (bForceRefreshAlwaysLoaded)
 	{
 		if (AlwaysLoadedActors)
@@ -1687,6 +1714,11 @@ void UWorldPartition::Tick(float DeltaSeconds)
 		bForceRefreshAlwaysLoaded = false;
 	}
 
+	if (ExternalDirtyActorsTracker)
+	{
+		ExternalDirtyActorsTracker->Tick(DeltaSeconds);
+	}
+		
 	if (bForceRefreshEditor)
 	{
 		if (WorldPartitionEditor)
diff --git a/Engine/Source/Runtime/Engine/Public/WorldPartition/WorldPartition.h b/Engine/Source/Runtime/Engine/Public/WorldPartition/WorldPartition.h
--- a/Engine/Source/Runtime/Engine/Public/WorldPartition/WorldPartition.h
+++ b/Engine/Source/Runtime/Engine/Public/WorldPartition/WorldPartition.h
@@ -642,13 +642,15 @@ private:
 		FWorldPartitionExternalDirtyActorsTracker(UWorldPartition* InWorldPartition);
 
 		//~ Begin TExternalDirtyActorsTracker interface
-		virtual bool OnAddDirtyActor(const TWeakObjectPtr<AActor> InActor) override { return !!GUndo; }
 		virtual void OnRemoveNonDirtyActor(const TWeakObjectPtr<AActor> InActor, FWorldPartitionReference& InValue) override;
 		virtual void Tick(float InDeltaTime) override;
 		//~ End TExternalDirtyActorsTracker interface
 
+		void SetNonDirtyTrackingDisabled(bool bInIsNonDirtyTrackingDisabled) { bIsNonDirtyTrackingDisabled = bInIsNonDirtyTrackingDisabled; }
+		bool IsNonDirtyTrackingDisabled() const { return bIsNonDirtyTrackingDisabled; }
 	private:
 		TSet<TPair<TWeakObjectPtr<AActor>, FWorldPartitionReference>> NonDirtyActors;
+		bool bIsNonDirtyTrackingDisabled = false;
 	};
 
 	TUniquePtr<FWorldPartitionExternalDirtyActorsTracker> ExternalDirtyActorsTracker;
@@ -663,6 +665,18 @@ private:
 public:
 	TOptional<bool> bOverrideEnableStreamingInEditor;
 
+	friend class FDisableNonDirtyActorTrackingScope;
+
+	// Use scope around actor package save calls to prevent newly created spatial actors from being pinned (actors will get unloaded instead)
+	class FDisableNonDirtyActorTrackingScope
+	{
+	public:
+		ENGINE_API FDisableNonDirtyActorTrackingScope(UWorldPartition* InWorldPartition, bool bInDisableTracking);
+		ENGINE_API ~FDisableNonDirtyActorTrackingScope();
+	private:
+		UWorldPartition* WorldPartition = nullptr;
+		bool bPreviousValue = false;
+	};
 private:
 #endif
 
