Removed viewport render target indirection as it results in dangling pointers due to not being properly cleared. Instead, just grab the render target from the viewport directly.

diff --git a/Engine/Source/Runtime/Engine/Classes/Engine/Engine.h b/Engine/Source/Runtime/Engine/Classes/Engine/Engine.h
--- a/Engine/Source/Runtime/Engine/Classes/Engine/Engine.h
+++ b/Engine/Source/Runtime/Engine/Classes/Engine/Engine.h
@@ -2819,7 +2819,7 @@ public:
 	ENGINE_API virtual bool HasMultipleLocalPlayers(UWorld* InWorld);
 
 	/** @return whether we're currently running with stereoscopic 3D enabled for the specified viewport (or globally, if viewport is nullptr) */
-	ENGINE_API bool IsStereoscopic3D(FViewport* InViewport = nullptr);
+	ENGINE_API bool IsStereoscopic3D(const FViewport* InViewport = nullptr) const;
 
 	/**
 	 * Adds a world location as a secondary view location for purposes of texture streaming.
diff --git a/Engine/Source/Runtime/Engine/Private/Slate/SceneViewport.cpp b/Engine/Source/Runtime/Engine/Private/Slate/SceneViewport.cpp
--- a/Engine/Source/Runtime/Engine/Private/Slate/SceneViewport.cpp
+++ b/Engine/Source/Runtime/Engine/Private/Slate/SceneViewport.cpp
@@ -1312,11 +1312,19 @@ void FSceneViewport::OnViewportDeactivated(const FWindowActivateEvent& InActivat
 }
 
 FSlateShaderResource* FSceneViewport::GetViewportRenderTargetTexture() const
-{ 
-	check(IsThreadSafeForSlateRendering());
+{
+	if (IsInRenderingThread())
+	{
+		return RenderThreadSlateTexture;
+	}
 	return (BufferedSlateHandles.Num() != 0) ? BufferedSlateHandles[CurrentBufferedTargetIndex] : nullptr;
 }
 
+bool FSceneViewport::IsStereoscopic3D() const
+{
+	return GEngine->IsStereoscopic3D(this);
+}
+
 void FSceneViewport::SetDebugCanvas(TSharedPtr<SDebugCanvas> InDebugCanvas)
 {
 	DebugCanvas = InDebugCanvas;
@@ -1684,15 +1692,6 @@ const FTextureRHIRef& FSceneViewport::GetRenderTargetTexture() const
 	return 	RenderTargetTextureRHI;
 }
 
-FSlateShaderResource* FSceneViewport::GetViewportRenderTargetTexture()
-{
-	if (IsInRenderingThread())
-	{
-		return RenderThreadSlateTexture;
-	}
-	return (BufferedSlateHandles.Num() != 0) ? BufferedSlateHandles[CurrentBufferedTargetIndex] : nullptr;
-}
-
 void FSceneViewport::SetRenderTargetTextureRenderThread(FTextureRHIRef& RT)
 {
 	check(IsInRenderingThread());
@@ -1752,7 +1751,6 @@ void FSceneViewport::UpdateViewportRHI(bool bDestroyed, uint32 NewSizeX, uint32
 			{
 				TSharedPtr<SWindow> Window = FSlateApplication::Get().FindWidgetWindow(PinnedViewport.ToSharedRef());
 
-				WindowRenderTargetUpdate(Renderer, Window.Get());
 				if (UseSeparateRenderTarget())
 				{
 					uint32 TexSizeX = SizeX, TexSizeY = SizeY;
@@ -1943,19 +1941,6 @@ void FSceneViewport::Tick( const FGeometry& AllottedGeometry, double InCurrentTi
 
 void FSceneViewport::OnPlayWorldViewportSwapped( const FSceneViewport& OtherViewport )
 {
-	// We need to call WindowRenderTargetUpdate() to make sure the Slate renderer is updated to render
-	// to the viewport client we'll be using for PIE/SIE.  Otherwise if stereo rendering is enabled, Slate
-	// could render the HMD mirror to a game viewport client which is not visible on screen!
-	TSharedPtr<SWidget> PinnedViewport = ViewportWidget.Pin();
-	if( PinnedViewport.IsValid() )
-	{
-		FSlateRenderer* Renderer = FSlateApplication::Get().GetRenderer();
-
-		TSharedPtr<SWindow> Window = FSlateApplication::Get().FindWidgetWindow( PinnedViewport.ToSharedRef() );
-
-		WindowRenderTargetUpdate( Renderer, Window.Get() );
-	}
-
 	// Play world viewports should always be the same size.  Resize to other viewports size
 	if( GetSizeXY() != OtherViewport.GetSizeXY() )
 	{
@@ -1990,17 +1975,6 @@ void FSceneViewport::SwapStatCommands( const FSceneViewport& OtherViewport )
 	}
 }
 
-/** Queue an update to the Window's RT on the Renderthread */
-void FSceneViewport::WindowRenderTargetUpdate(FSlateRenderer* Renderer, SWindow* Window)
-{
-	check(IsInGameThread());
-	if (Renderer && Window)
-	{
-		// The viewport provider interface is ONLY used for stereo VR compositing to blit the intermediate viewport target over to the VR swap chain.
-		Renderer->SetWindowRenderTarget(*Window, UseSeparateRenderTarget() && GEngine->IsStereoscopic3D(this) ? this : nullptr);
-	}
-}
-
 void FSceneViewport::OnWindowBackBufferResourceDestroyed(void* Backbuffer)
 {
 	check(IsInGameThread());
diff --git a/Engine/Source/Runtime/Engine/Private/UnrealEngine.cpp b/Engine/Source/Runtime/Engine/Private/UnrealEngine.cpp
--- a/Engine/Source/Runtime/Engine/Private/UnrealEngine.cpp
+++ b/Engine/Source/Runtime/Engine/Private/UnrealEngine.cpp
@@ -4332,7 +4332,7 @@ bool UEngine::HasMultipleLocalPlayers(UWorld* InWorld)
 }
 
 /** @return whether we're currently running with stereoscopic 3D enabled */
-bool UEngine::IsStereoscopic3D(FViewport* InViewport)
+bool UEngine::IsStereoscopic3D(const FViewport* InViewport) const
 {
 	return (!InViewport || InViewport->IsStereoRenderingAllowed()) &&
 		(StereoRenderingDevice.IsValid() && StereoRenderingDevice->IsStereoEnabled());
diff --git a/Engine/Source/Runtime/Engine/Public/Slate/SceneViewport.h b/Engine/Source/Runtime/Engine/Public/Slate/SceneViewport.h
--- a/Engine/Source/Runtime/Engine/Public/Slate/SceneViewport.h
+++ b/Engine/Source/Runtime/Engine/Public/Slate/SceneViewport.h
@@ -31,7 +31,7 @@ extern const FName NAME_SceneViewport;
 /**
  * A viewport for use with Slate SViewport widgets.
  */
-class FSceneViewport : public FViewportFrame, public FViewport, public ISlateViewport, public IViewportRenderTargetProvider
+class FSceneViewport : public FViewportFrame, public FViewport, public ISlateViewport
 {
 public:
 	ENGINE_API FSceneViewport( FViewportClient* InViewportClient, TSharedPtr<SViewport> InViewportWidget );
@@ -261,8 +261,6 @@ public:
 	/** Returns dimensions of RenderTarget texture. Can be called on a game thread. */
 	virtual FIntPoint GetRenderTargetTextureSizeXY() const { return (RTTSize.X != 0) ? RTTSize : GetSizeXY(); }
 
-	ENGINE_API virtual FSlateShaderResource* GetViewportRenderTargetTexture() override;
-
 	/** Returns format for the scene of this viewport. */
 	ENGINE_API EPixelFormat GetSceneTargetFormat() const override { return SceneTargetFormat; }
 
@@ -345,15 +343,18 @@ private:
 	/** Utility function to figure out if we are currently a game viewport */
 	ENGINE_API bool IsCurrentlyGameViewport();
 
-	ENGINE_API void WindowRenderTargetUpdate(FSlateRenderer* Renderer, SWindow* Window);
+	UE_DEPRECATED(5.5, "WindowRenderTargetUpdate is no longer used")
+	void WindowRenderTargetUpdate(FSlateRenderer* Renderer, SWindow* Window) {}
 
 	/** @return Returns true if we should always render to a separate render target (rather than rendering directly to the
 	    viewport backbuffer, taking into account any temporary requirements of head-mounted displays */
-	bool UseSeparateRenderTarget() const
+	bool UseSeparateRenderTarget() const override
 	{
 		return bUseSeparateRenderTarget || bForceSeparateRenderTarget;
 	}
 
+	ENGINE_API bool IsStereoscopic3D() const override;
+
 	/**
 	 * Called right before a slate window is destroyed so we can free up the backbuffer resource before the window backing it is destroyed
 	 */
diff --git a/Engine/Source/Runtime/SlateCore/Public/Rendering/RenderingCommon.h b/Engine/Source/Runtime/SlateCore/Public/Rendering/RenderingCommon.h
--- a/Engine/Source/Runtime/SlateCore/Public/Rendering/RenderingCommon.h
+++ b/Engine/Source/Runtime/SlateCore/Public/Rendering/RenderingCommon.h
@@ -504,6 +504,22 @@ public:
 		return false;
 	}
 
+	/**
+	 * Does the texture returned by GetViewportRenderTargetTexture represent a stereoscopic 3D target?
+	 */
+	virtual bool IsStereoscopic3D() const
+	{
+		return false;
+	}
+	
+	/**
+	 * Is the viewport expected to render to GetViewportRenderTargetTexture? Or to the swapchain?
+	 */
+	virtual bool UseSeparateRenderTarget() const
+	{
+		return false;
+	}
+
 	/**
 	 * Does the texture contain SDR/HDR information
 	 */
diff --git a/Engine/Source/Runtime/SlateCore/Public/Rendering/SlateRenderer.h b/Engine/Source/Runtime/SlateCore/Public/Rendering/SlateRenderer.h
--- a/Engine/Source/Runtime/SlateCore/Public/Rendering/SlateRenderer.h
+++ b/Engine/Source/Runtime/SlateCore/Public/Rendering/SlateRenderer.h
@@ -484,6 +484,7 @@ public:
 	/**
 	 * Pushes the rendering of the specified window to the specified render target
 	 */
+	UE_DEPRECATED(5.5, "SetWindowRenderTarget is no longer used.")
 	virtual void SetWindowRenderTarget(const SWindow& Window, class IViewportRenderTargetProvider* Provider) {}
 
 	/**
diff --git a/Engine/Source/Runtime/SlateCore/Public/Textures/SlateShaderResource.h b/Engine/Source/Runtime/SlateCore/Public/Textures/SlateShaderResource.h
--- a/Engine/Source/Runtime/SlateCore/Public/Textures/SlateShaderResource.h
+++ b/Engine/Source/Runtime/SlateCore/Public/Textures/SlateShaderResource.h
@@ -208,7 +208,7 @@ protected:
 	ResourceType ShaderResource;
 };
 
-class IViewportRenderTargetProvider
+class UE_DEPRECATED(5.5, "This is no longer used.") IViewportRenderTargetProvider
 {
 public:
 	virtual FSlateShaderResource* GetViewportRenderTargetTexture() = 0;
diff --git a/Engine/Source/Runtime/SlateRHIRenderer/Private/SlateRHIRenderer.cpp b/Engine/Source/Runtime/SlateRHIRenderer/Private/SlateRHIRenderer.cpp
--- a/Engine/Source/Runtime/SlateRHIRenderer/Private/SlateRHIRenderer.cpp
+++ b/Engine/Source/Runtime/SlateRHIRenderer/Private/SlateRHIRenderer.cpp
@@ -146,7 +146,6 @@ struct FSlateViewportInfo : public FRenderResource
 {
 	FViewportRHIRef ViewportRHI;
 	void* OSWindow = nullptr;
-	IViewportRenderTargetProvider* RenderTargetProvider = nullptr;
 	FMatrix ProjectionMatrix;
 	FIntPoint Extent = FIntPoint::ZeroValue;
 	FIntPoint ExtentToResizeTo = FIntPoint::ZeroValue;
@@ -162,18 +161,6 @@ struct FSlateViewportInfo : public FRenderResource
 		FRHICommandListExecutor::GetImmediateCommandList().BlockUntilGPUIdle();
 		ViewportRHI.SafeRelease();
 	}
-
-	FTextureRHIRef GetRenderTargetTexture() const
-	{
-		if (RenderTargetProvider)
-		{
-			if (FSlateShaderResource* RenderTargetTexture = RenderTargetProvider->GetViewportRenderTargetTexture())
-			{
-				return static_cast<FSlateRenderTargetRHI*>(RenderTargetTexture)->GetTypedResource();
-			}
-		}
-		return nullptr;
-	}
 };
 
 struct FSlatePostProcessUpdateRequest
@@ -679,14 +666,16 @@ FSlateDrawWindowPassOutputs FSlateRHIRenderer::DrawWindow_RenderThread(FRDGBuild
 		RDG_CSV_STAT_EXCLUSIVE_SCOPE(GraphBuilder, Slate);
 		TRACE_CPUPROFILER_EVENT_SCOPE(Slate::DrawWindow_RenderThread);
 
+		ISlateViewport* SlateViewport = Inputs.Window->GetViewport().Get();
+
 		// The viewport texture is an optional user-allocated render target. This is rendered to if valid.
-		ViewportTextureRHI = ViewportInfo.GetRenderTargetTexture();
+		ViewportTextureRHI = SlateViewport && SlateViewport->UseSeparateRenderTarget() ? static_cast<FSlateRenderTargetRHI*>(SlateViewport->GetViewportRenderTargetTexture())->GetTypedResource() : nullptr;
 
 		// The swap chain is the final output. This is rendered to if no viewport render target is provided.
 		FRHITexture* SwapChainTextureRHI = RHIGetViewportBackBuffer(ViewportInfo.ViewportRHI);
 
 		// Only render to the intermediate viewport render target if stereo rendering is enabled, which we'll then composite later.
-		const bool bCompositeStereoToSwapChain = ViewportTextureRHI && GEngine && GEngine->StereoRenderingDevice.IsValid();
+		const bool bCompositeStereoToSwapChain = ViewportTextureRHI && GEngine && GEngine->StereoRenderingDevice.IsValid() && SlateViewport && SlateViewport->IsStereoscopic3D();
 
 		// The output texture is what we ultimately render or composite slate elements into.
 		OutputTextureRHI = bCompositeStereoToSwapChain ? ViewportTextureRHI : SwapChainTextureRHI;
@@ -1718,15 +1707,6 @@ void FSlateRHIRenderer::RequestResize(const TSharedPtr<SWindow>& Window, uint32
 	}
 }
 
-void FSlateRHIRenderer::SetWindowRenderTarget(const SWindow& Window, IViewportRenderTargetProvider* Provider)
-{
-	FSlateViewportInfo* ViewInfo = WindowToViewportInfo.FindRef(&Window);
-	if (ViewInfo)
-	{
-		ViewInfo->RenderTargetProvider = Provider;
-	}
-}
-
 void FSlateRHIRenderer::AddWidgetRendererUpdate(const FRenderThreadUpdateContext& Context, bool bDeferredRenderTargetUpdate)
 {
 	if (bDeferredRenderTargetUpdate)
diff --git a/Engine/Source/Runtime/SlateRHIRenderer/Private/SlateRHIRenderer.h b/Engine/Source/Runtime/SlateRHIRenderer/Private/SlateRHIRenderer.h
--- a/Engine/Source/Runtime/SlateRHIRenderer/Private/SlateRHIRenderer.h
+++ b/Engine/Source/Runtime/SlateRHIRenderer/Private/SlateRHIRenderer.h
@@ -107,8 +107,6 @@ public:
 	virtual void PrepareToTakeScreenshot(const FIntRect& Rect, TArray<FColor>* OutColorData, SWindow* ScreenshotWindow) override;
 	virtual void PrepareToTakeHDRScreenshot(const FIntRect& Rect, TArray<FLinearColor>* OutColorData, SWindow* ScreenshotWindow) override;
 
-	virtual void SetWindowRenderTarget(const SWindow& Window, class IViewportRenderTargetProvider* Provider) override;
-
 private:
 	FSlateDrawWindowPassOutputs DrawWindow_RenderThread(FRDGBuilder& GraphBuilder, const FSlateDrawWindowPassInputs& Inputs);
 
