[AutoRTFM] Migrate to UE::TScopeLock instead of FTransactionallySafeScopeLock.

These are functionally identical.

Additionally, I've replaced a TUniquePtr<ScopeLock> with a TOptional<ScopeLock> so that locking and unlocking doesn't cause a needless allocation/free.

diff --git a/Engine/Source/Runtime/Engine/Private/TraceFilter.cpp b/Engine/Source/Runtime/Engine/Private/TraceFilter.cpp
--- a/Engine/Source/Runtime/Engine/Private/TraceFilter.cpp
+++ b/Engine/Source/Runtime/Engine/Private/TraceFilter.cpp
@@ -8,7 +8,9 @@
 
 #include "ObjectTrace.h"
 
-#include "Misc/TransactionallySafeScopeLock.h"
+#include "Misc/Optional.h"
+#include "Misc/ScopeLock.h"
+#include "Misc/TransactionallySafeCriticalSection.h"
 
 struct ENGINE_API FTraceFilterObjectAnnotation
 {
@@ -58,7 +60,7 @@ private:
 	void AddAnnotationInternal(const UObjectBase* Object, const FTraceFilterObjectAnnotation& Annotation)
 	{
 		check(Object);
-		FTransactionallySafeScopeLock ScopeLock(&AnnotationMapCritical);
+		UE::TScopeLock ScopeLock(AnnotationMapCritical);
 		AnnotationCacheKey = Object;
 		AnnotationCacheValue = Annotation;
 		if (AnnotationCacheValue.IsDefault())
@@ -90,7 +92,7 @@ public:
 
 	void RemoveAnnotation(const UObjectBase* Object)
 	{
-		FTransactionallySafeScopeLock ScopeLock(&AnnotationMapCritical);
+		UE::TScopeLock ScopeLock(AnnotationMapCritical);
 		check(Object);
 		
 		AnnotationCacheKey = Object;
@@ -106,7 +108,7 @@ public:
 
 	void RemoveAllAnnotations()
 	{
-		FTransactionallySafeScopeLock ScopeLock(&AnnotationMapCritical);
+		UE::TScopeLock ScopeLock(AnnotationMapCritical);
 
 		AnnotationCacheKey = NULL;
 		AnnotationCacheValue = FTraceFilterObjectAnnotation();
@@ -121,7 +123,7 @@ public:
 
 	FORCEINLINE FTraceFilterObjectAnnotation GetAnnotation(const UObjectBase* Object)
 	{
-		FTransactionallySafeScopeLock ScopeLock(&AnnotationMapCritical);
+		UE::TScopeLock ScopeLock(AnnotationMapCritical);
 
 		check(Object);
 		
@@ -148,17 +150,17 @@ public:
 
 	void Lock()
 	{
-		UniqueScopeLock = MakeUnique<FTransactionallySafeScopeLock>(&AnnotationMapCritical);
+		AnnotationMapScopeLock.Emplace(AnnotationMapCritical);
 	}
 
 	void Unlock()
 	{		
-		UniqueScopeLock.Reset();
+		AnnotationMapScopeLock.Reset();
 	}
 
 	bool IsLocked()
 	{
-		return UniqueScopeLock.IsValid();
+		return AnnotationMapScopeLock.IsSet();
 	}
 
 	virtual SIZE_T GetAllocatedSize() const override
@@ -170,7 +172,7 @@ private:
 	TMap<const UObjectBase*, FTraceFilterObjectAnnotation> AnnotationMap;
 	FTransactionallySafeCriticalSection AnnotationMapCritical;
 
-	TUniquePtr<FTransactionallySafeScopeLock> UniqueScopeLock;
+	TOptional<UE::TScopeLock<FTransactionallySafeCriticalSection>> AnnotationMapScopeLock;
 
 	const UObjectBase* AnnotationCacheKey;
 	FTraceFilterObjectAnnotation AnnotationCacheValue;
@@ -278,15 +280,15 @@ template<bool bForceThreadSafe>
 bool FTraceFilter::IsObjectTraceable(const UObject* InObject)
 {
 	return AutoRTFM::Open([&]
+	{
+		if constexpr (!bForceThreadSafe)
 		{
-			if constexpr (!bForceThreadSafe)
-			{
-				check(GObjectFilterAnnotations.IsLocked());
-			}
+			check(GObjectFilterAnnotations.IsLocked());
+		}
 
-			// Object not found in the AnnotationMap means that it is at the default value, which is bIsTraceable == true
-			return GObjectFilterAnnotations.GetAnnotationMap().Find(InObject) == nullptr;
-		});
+		// Object not found in the AnnotationMap means that it is at the default value, which is bIsTraceable == true
+		return GObjectFilterAnnotations.GetAnnotationMap().Find(InObject) == nullptr;
+	});
 }
 
 template bool FTraceFilter::IsObjectTraceable</*bForceThreadSafe = */ false>(const UObject* InObject);
@@ -326,7 +328,7 @@ void FTraceFilter::SetObjectIsTraceable(const UObject* InObject, bool bIsTraceab
 		}
 	});
 
-AutoRTFM::OnAbort([InObject, bIsTraceable]
+	AutoRTFM::OnAbort([InObject, bIsTraceable]
 	{
 		SetObjectIsTraceable(InObject, !bIsTraceable);
 	});
