- Block on SVT blocking disk requests

#rnx

diff --git a/Engine/Source/Runtime/Engine/Private/SparseVolumeTexture/SparseVolumeTextureStreamingManager.cpp b/Engine/Source/Runtime/Engine/Private/SparseVolumeTexture/SparseVolumeTextureStreamingManager.cpp
--- a/Engine/Source/Runtime/Engine/Private/SparseVolumeTexture/SparseVolumeTextureStreamingManager.cpp
+++ b/Engine/Source/Runtime/Engine/Private/SparseVolumeTexture/SparseVolumeTextureStreamingManager.cpp
@@ -1135,6 +1135,9 @@ void FStreamingManager::IssueRequests()
 	DDCRequestsBlocking.Reserve(TileRangesToStream.Num());
 #endif
 
+	// Indicies into requests we should block on after sending them all
+	TArray<int32> BlockingDiskRequests;
+
 	// Process all tile ranges selected for streaming, allocate a slot in the tile data texture for every tile and finally create IO requests for every range.
 	for (FTileRange& TileRange : TileRangesToStream)
 	{
@@ -1333,6 +1336,11 @@ void FStreamingManager::IssueRequests()
 					Batch.Read(BulkData, ReadOffset, ReadSize);
 					Batch.Issue(PendingRequest.RequestBuffer, Priority, [](FBulkDataRequest::EStatus){}, PendingRequest.Request);
 
+					if (PendingRequest.bBlocking)
+					{
+						BlockingDiskRequests.Add(PendingRequestIndex);
+					}
+
 #if WITH_EDITORONLY_DATA
 					PendingRequest.State = FPendingRequest::EState::Disk;
 #endif
@@ -1350,6 +1358,14 @@ void FStreamingManager::IssueRequests()
 		}
 	}
 
+	// Wait on all blocking disk requests
+	for (int32 PendingRequestIndex : BlockingDiskRequests)
+	{
+		FPendingRequest& PendingRequest = PendingRequests[PendingRequestIndex];
+		LOCK_PENDING_REQUEST(PendingRequest);
+		PendingRequest.Request.Wait();
+	}
+
 	// Now we can finally issue the requests
 	{
 #if WITH_EDITORONLY_DATA
