Re-enable quad overdraw mode for MetalRHI now we have atomic support

diff --git a/Engine/Source/Runtime/Engine/Private/DebugViewModeHelpers.cpp b/Engine/Source/Runtime/Engine/Private/DebugViewModeHelpers.cpp
--- a/Engine/Source/Runtime/Engine/Private/DebugViewModeHelpers.cpp
+++ b/Engine/Source/Runtime/Engine/Private/DebugViewModeHelpers.cpp
@@ -99,7 +99,7 @@ bool AllowDebugViewShaderMode(EDebugViewShaderMode ShaderMode, EShaderPlatform P
 	case DVSM_ShaderComplexityContainedQuadOverhead:
 	case DVSM_ShaderComplexityBleedingQuadOverhead:
 	case DVSM_QuadComplexity:
-		return (bForceQuadOverdraw || (PlatformSupportsDebugViewShaders(Platform) && !IsMetalPlatform(Platform))); // Last one to fix for Metal then remove this Metal check.
+		return (bForceQuadOverdraw || (PlatformSupportsDebugViewShaders(Platform)));
 	case DVSM_PrimitiveDistanceAccuracy:
 	case DVSM_MeshUVDensityAccuracy:
 		return FeatureLevel >= ERHIFeatureLevel::SM5 && (bForceStreamingAccuracy || PlatformSupportsDebugViewShaders(Platform));
