[World Partition] Clear actor references to editor-only actors from the persistent level when duplicating the world for PIE.
We already do this for streamed actors through none entries in the instancing context, but actors from the persistent level don't go through that codepath to load actors, they are duplicated.

diff --git a/Engine/Source/Runtime/Engine/Private/WorldPartition/WorldPartition.cpp b/Engine/Source/Runtime/Engine/Private/WorldPartition/WorldPartition.cpp
--- a/Engine/Source/Runtime/Engine/Private/WorldPartition/WorldPartition.cpp
+++ b/Engine/Source/Runtime/Engine/Private/WorldPartition/WorldPartition.cpp
@@ -547,6 +547,87 @@ void UWorldPartition::PostEditChangeChainProperty(FPropertyChangedChainEvent& Pr
 	}
 }
 
+void UWorldPartition::PostDuplicate(EDuplicateMode::Type DuplicateMode)
+{
+	Super::PostDuplicate(DuplicateMode);
+
+	if (DuplicateMode == EDuplicateMode::PIE)
+	{
+		UE_SCOPED_TIMER(TEXT("UWorldPartition::PostDuplicatePIE"), LogWorldPartition, Display);
+		TRACE_CPUPROFILER_EVENT_SCOPE(UWorldPartition::PostDuplicatePIE);
+
+		// Clear actor references to editor-only actors from the persistent level when duplicating the world for PIE. We already do this for streamed actors
+		// through none entries in the instancing context, but actors from the persistent level don't go through that codepath to load actors, they are duplicated.
+		class FArchiveClearEditorOnlyActorReferences : public FArchiveUObject
+		{
+		public:
+			FArchiveClearEditorOnlyActorReferences(AActor* InRoot)
+				: Root(InRoot)
+			{
+				ArIgnoreOuterRef = true;
+				ArIsObjectReferenceCollector = true;
+				ArIsModifyingWeakAndStrongReferences = true;
+				ArShouldSkipBulkData = true;
+
+				SetIsPersistent(true);
+				SetShouldSkipCompilingAssets(true);
+				SetIsSaving(true);
+
+				Root->Serialize(*this);
+			}
+
+			virtual FArchive& operator<<(UObject*& InObj) override
+			{
+				if (InObj && (InObj != Root) && !InObj->IsTemplate())
+				{
+					if (AActor* Actor = Cast<AActor>(InObj))
+					{
+						if (Actor->IsEditorOnly() && !Actor->IsNonPIEEditorOnly())
+						{
+							InObj = nullptr;
+							UE_LOG(LogWorldPartition, Log, TEXT("Cleared actor editor-only reference '%s' -> '%s'."), *Root->GetActorNameOrLabel(), *Actor->GetActorNameOrLabel());							
+						}
+					}
+					else if (InObj->IsInOuter(Root) && !InObj->IsA<UWorldPartition>())
+					{
+						bool bWasAlreadyInSet;
+						SubObjects.Add(InObj, &bWasAlreadyInSet);
+
+						if (!bWasAlreadyInSet)
+						{
+							InObj->Serialize(*this);
+						}
+					}
+				}
+
+				return *this;
+			}
+
+			virtual FArchive& operator<<(FWeakObjectPtr& Value) override
+			{
+				if (UObject* Object = Value.Get())
+				{
+					*this << Object;
+					Value = Object;
+				}
+				return *this;
+			}
+
+		private:
+			AActor* Root;
+			TSet<UObject*> SubObjects;
+		};
+
+		for (AActor* Actor : GetWorld()->PersistentLevel->Actors)
+		{
+			if (Actor)
+			{
+				FArchiveClearEditorOnlyActorReferences ClearEditorOnlyActorReferencesAr(Actor);
+			}
+		}
+	}
+}
+
 FName UWorldPartition::GetWorldPartitionEditorName() const
 {
 	if (SupportsStreaming())
diff --git a/Engine/Source/Runtime/Engine/Public/WorldPartition/WorldPartition.h b/Engine/Source/Runtime/Engine/Public/WorldPartition/WorldPartition.h
--- a/Engine/Source/Runtime/Engine/Public/WorldPartition/WorldPartition.h
+++ b/Engine/Source/Runtime/Engine/Public/WorldPartition/WorldPartition.h
@@ -209,8 +209,9 @@ public:
 #if WITH_EDITOR
 	ENGINE_API virtual bool CanEditChange(const FProperty* InProperty) const override;
 	ENGINE_API virtual void PostEditChangeChainProperty(FPropertyChangedChainEvent& PropertyChangedEvent) override;
+	ENGINE_API virtual void PostDuplicate(EDuplicateMode::Type DuplicateMode) override;
 #endif //WITH_EDITOR
-	ENGINE_API virtual void Serialize(FArchive& Ar) override;
+	ENGINE_API virtual void Serialize(FArchive& Ar) override;	
 	ENGINE_API virtual UWorld* GetWorld() const override;
 	ENGINE_API virtual bool ResolveSubobject(const TCHAR* SubObjectPath, UObject*& OutObject, bool bLoadIfExists) override;
 	ENGINE_API virtual void BeginDestroy() override;
