Fix for debug draws not working in standalone

diff --git a/Engine/Source/Runtime/Engine/Private/Animation/AnimInstanceProxy.cpp b/Engine/Source/Runtime/Engine/Private/Animation/AnimInstanceProxy.cpp
--- a/Engine/Source/Runtime/Engine/Private/Animation/AnimInstanceProxy.cpp
+++ b/Engine/Source/Runtime/Engine/Private/Animation/AnimInstanceProxy.cpp
@@ -685,7 +685,7 @@ void FAnimInstanceProxy::PostUpdate(UAnimInstance* InAnimInstance) const
 #if ENABLE_ANIM_DRAW_DEBUG
 void FAnimInstanceProxy::FlushQueuedDebugDrawItems(AActor* InActor, UWorld* InWorld) const
 {
-	if (InActor == nullptr || InActor->GetLocalRole() != ROLE_Authority)
+	if (GEngine->GetNetMode(InWorld) != NM_DedicatedServer)
 	{
 		for (const FQueuedDrawDebugItem& DebugItem : QueuedDrawDebugItems)
 		{
