[NetPhysics] Setting to opt in and out of applying inputs and states via the NetworkPhysicsComponent for sim-proxies at runtime (not just during resimulation).
- Also marhsal the SettingsComponent from GameThread to PhysicsThread in BeginPlay since the PT Settings might not have been initialized in InitializeComponent.

diff --git a/Engine/Source/Runtime/Engine/Private/PhysicsEngine/NetworkPhysicsComponent.cpp b/Engine/Source/Runtime/Engine/Private/PhysicsEngine/NetworkPhysicsComponent.cpp
--- a/Engine/Source/Runtime/Engine/Private/PhysicsEngine/NetworkPhysicsComponent.cpp
+++ b/Engine/Source/Runtime/Engine/Private/PhysicsEngine/NetworkPhysicsComponent.cpp
@@ -42,6 +42,10 @@ namespace PhysicsReplicationCVars
 		static FAutoConsoleVariableRef CVarResimAllowInputExtrapolation(TEXT("np2.Resim.AllowInputExtrapolation"), bAllowInputExtrapolation, TEXT("When true and not locally controlled, allow inputs to be extrapolated from last known and if there is a gap allow interpolation between two known inputs."));
 		bool bValidateDataOnGameThread = false;
 		static FAutoConsoleVariableRef CVarResimValidateDataOnGameThread(TEXT("np2.Resim.ValidateDataOnGameThread"), bValidateDataOnGameThread, TEXT("When true, perform server-side input validation through FNetworkPhysicsData::ValidateData on the Game Thread. If false, perform the call on the Physics Thread."));
+		bool bApplySimProxyStateAtRuntime = false;
+		static FAutoConsoleVariableRef CVarResimApplySimProxyStateAtRuntime(TEXT("np2.Resim.ApplySimProxyStateAtRuntime"), bApplySimProxyStateAtRuntime, TEXT("When true, call ApplyData on received states for simulated proxies at runtime."));
+		bool bApplySimProxyInputAtRuntime = true;
+		static FAutoConsoleVariableRef CVarResimApplySimProxyInputAtRuntime(TEXT("np2.Resim.ApplySimProxyInputAtRuntime"), bApplySimProxyInputAtRuntime, TEXT("When true, call ApplyData on received inputs for simulated proxies at runtime."));
 	}
 }
 
@@ -465,11 +469,10 @@ void UNetworkPhysicsComponent::InitializeComponent()
 	bEnableReliableFlow = PhysicsReplicationCVars::ResimulationCVars::bEnableReliableFlow;
 	bValidateDataOnGameThread = PhysicsReplicationCVars::ResimulationCVars::bValidateDataOnGameThread;
 
-	UNetworkPhysicsSettingsComponent* SettingsComponent = nullptr;
 	if (AActor* Owner = GetOwner())
 	{
 		// Get settings from NetworkPhysicsSettingsComponent, if there is one
-		SettingsComponent = Owner->FindComponentByClass<UNetworkPhysicsSettingsComponent>();
+		UNetworkPhysicsSettingsComponent* SettingsComponent = Owner->FindComponentByClass<UNetworkPhysicsSettingsComponent>();
 		if (SettingsComponent)
 		{
 			SetNumberOfInputsToNetwork(SettingsComponent->NetworkPhysicsComponentSettings.GetRedundantInputs());
@@ -497,7 +500,6 @@ void UNetworkPhysicsComponent::InitializeComponent()
 			{
 				// Create async component to run on Physics Thread
 				NetworkPhysicsComponent_Internal = Solver->CreateAndRegisterSimCallbackObject_External<FAsyncNetworkPhysicsComponent>();
-				NetworkPhysicsComponent_Internal->SettingsComponent = SettingsComponent ? SettingsComponent->GetNetworkPhysicsSettings_Internal() : nullptr;
 				NetworkPhysicsComponent_Internal->RootPhysicsObject = RootPhysicsObject;
 				NetworkPhysicsComponent_Internal->InputsToNetwork = InputsToNetwork;
 				NetworkPhysicsComponent_Internal->StatesToNetwork = StatesToNetwork;
@@ -952,6 +954,12 @@ void UNetworkPhysicsComponent::UpdateAsyncComponent(const bool bFullUpdate)
 					AsyncInput->NetRole = Owner->GetLocalRole();
 					AsyncInput->PhysicsReplicationMode = Owner->GetPhysicsReplicationMode();
 					AsyncInput->ActorName = AActor::GetDebugName(Owner);
+			
+					UNetworkPhysicsSettingsComponent* SettingsComponent = Owner->FindComponentByClass<UNetworkPhysicsSettingsComponent>();
+					if (SettingsComponent && SettingsComponent->GetNetworkPhysicsSettings_Internal())
+					{
+						AsyncInput->SettingsComponent = SettingsComponent->GetNetworkPhysicsSettings_Internal();
+					}
 				}
 				
 				if (ActorComponent)
@@ -1040,6 +1048,8 @@ FAsyncNetworkPhysicsComponent::FAsyncNetworkPhysicsComponent() : TSimCallbackObj
 	, NetRole(ENetRole::ROLE_Authority)
 	, NetworkPhysicsTickOffset(0)
 	, PhysicsReplicationMode(EPhysicsReplicationMode::Default)
+	, SettingsComponent(nullptr)
+	, RootPhysicsObject(nullptr)
 {
 }
 
@@ -1094,6 +1104,10 @@ void FAsyncNetworkPhysicsComponent::ConsumeAsyncInput(const int32 PhysicsStep)
 
 		/** Onetime setup marshaled data */
 		{
+			if (AsyncInput->SettingsComponent.IsSet())
+			{
+				SettingsComponent = *AsyncInput->SettingsComponent;
+			}
 			if (AsyncInput->bIsLocallyControlled.IsSet())
 			{
 				bIsLocallyControlled = *AsyncInput->bIsLocallyControlled;
@@ -1296,13 +1310,14 @@ void FAsyncNetworkPhysicsComponent::OnPreProcessInputs_Internal(const int32 Phys
 
 	if (ActorComponent.IsValid() && ActorComponent.Get()->IsBeingDestroyed() == false)
 	{
-		// Apply replicated state on clients if we are resimulating
-		if (bIsSolverResim && StateHistory && StateData)
+		// Apply replicated state on clients if we are resimulating or on simulated proxies if setting is enabled
+		const bool bApplySimProxyState = ComponentSettings.GetApplySimProxyStateAtRuntime() && !IsServer() && !IsLocallyControlled();
+		if ((bApplySimProxyState || bIsSolverResim) && StateHistory && StateData)
 		{
 			FNetworkPhysicsData* PhysicsData = StateData.Get();
 			PhysicsData->LocalFrame = PhysicsStep;
 			const bool bExactFrame = PhysicsReplicationCVars::ResimulationCVars::bAllowRewindToClosestState ? !bIsSolverReset : true;
-			if (StateHistory->ExtractData(PhysicsStep, bIsSolverReset, PhysicsData, bExactFrame) && PhysicsData->bReceivedData)
+			if (StateHistory->ExtractData(PhysicsStep, bIsSolverReset, PhysicsData, (bExactFrame && bIsSolverResim)) && PhysicsData->bReceivedData)
 			{
 				PhysicsData->ApplyData(ActorComponent.Get());
 #if DEBUG_NETWORK_PHYSICS
@@ -1319,8 +1334,10 @@ void FAsyncNetworkPhysicsComponent::OnPreProcessInputs_Internal(const int32 Phys
 #endif
 		}
 
-		// Apply replicated inputs on server and simulated proxies (and on local player if we are resimulating)
-		if (InputHistory && InputData && (!IsLocallyControlled() || bIsSolverResim))
+		// Apply replicated inputs on server and simulated proxies if setting is enabled, and on local player if we are resimulating
+		const bool bApplySimProxyInput = ComponentSettings.GetApplySimProxyInputAtRuntime() && !IsServer() && !IsLocallyControlled();
+		const bool bApplyServerInput = IsServer() && !IsLocallyControlled();
+		if ((bApplyServerInput || bApplySimProxyInput || bIsSolverResim) && InputHistory && InputData)
 		{
 			FNetworkPhysicsData* PhysicsData = InputData.Get();
 			int32 NextExpectedLocalFrame = PhysicsData->LocalFrame + 1;
diff --git a/Engine/Source/Runtime/Engine/Public/Physics/NetworkPhysicsComponent.h b/Engine/Source/Runtime/Engine/Public/Physics/NetworkPhysicsComponent.h
--- a/Engine/Source/Runtime/Engine/Public/Physics/NetworkPhysicsComponent.h
+++ b/Engine/Source/Runtime/Engine/Public/Physics/NetworkPhysicsComponent.h
@@ -1035,6 +1035,7 @@ struct FAsyncNetworkPhysicsComponentInput : public Chaos::FSimCallbackInput
 	TOptional<TUniquePtr<FNetworkPhysicsDataHelper>> StateHelper;
 	TOptional<bool> bRegisterDataHistoryInRewindData;
 	TOptional<bool> bUnregisterDataHistoryFromRewindData;
+	TOptional<FNetworkPhysicsSettingsComponentAsync*> SettingsComponent;
 
 	TUniquePtr<Chaos::FBaseRewindHistory> InputData;
 	TUniquePtr<Chaos::FBaseRewindHistory> StateData;
@@ -1055,6 +1056,7 @@ struct FAsyncNetworkPhysicsComponentInput : public Chaos::FSimCallbackInput
 		StateHelper.Reset();
 		bRegisterDataHistoryInRewindData.Reset();
 		bUnregisterDataHistoryFromRewindData.Reset();
+		SettingsComponent.Reset();
 
 		if (InputData)
 		{
diff --git a/Engine/Source/Runtime/Engine/Public/Physics/NetworkPhysicsSettingsComponent.h b/Engine/Source/Runtime/Engine/Public/Physics/NetworkPhysicsSettingsComponent.h
--- a/Engine/Source/Runtime/Engine/Public/Physics/NetworkPhysicsSettingsComponent.h
+++ b/Engine/Source/Runtime/Engine/Public/Physics/NetworkPhysicsSettingsComponent.h
@@ -71,6 +71,8 @@ namespace PhysicsReplicationCVars
 		extern int32 RedundantStates;
 		extern bool bCompareStateToTriggerRewind;
 		extern bool bCompareInputToTriggerRewind;
+		extern bool bApplySimProxyStateAtRuntime;
+		extern bool bApplySimProxyInputAtRuntime;
 	}
 }
 
@@ -406,6 +408,20 @@ struct FNetworkPhysicsSettingsNetworkPhysicsComponent
 	UPROPERTY(config, EditDefaultsOnly, Category = "Overrides", Meta = (EditCondition = "bOverrideValidateDataOnGameThread"))
 	bool bValidateDataOnGameThread = PhysicsReplicationCVars::ResimulationCVars::bValidateDataOnGameThread;
 	const bool GetValidateDataOnGameThread() const { return bOverrideValidateDataOnGameThread ? bValidateDataOnGameThread : PhysicsReplicationCVars::ResimulationCVars::bValidateDataOnGameThread; }
+
+	UPROPERTY(config, EditDefaultsOnly, Category = "Overrides", Meta = (InlineEditConditionToggle))
+	uint32 bOverrideApplySimProxyStateAtRuntime : 1;
+	// Overrides CVar: np2.Resim.ApplySimProxyStateAtRuntime -- When true, call ApplyData on received states for simulated proxies at runtime.
+	UPROPERTY(config, EditDefaultsOnly, Category = "Overrides", Meta = (EditCondition = "bOverrideApplySimProxyStateAtRuntime"))
+	bool bApplySimProxyStateAtRuntime = PhysicsReplicationCVars::ResimulationCVars::bApplySimProxyStateAtRuntime;
+	const bool GetApplySimProxyStateAtRuntime() const { return bOverrideApplySimProxyStateAtRuntime ? bApplySimProxyStateAtRuntime : PhysicsReplicationCVars::ResimulationCVars::bApplySimProxyStateAtRuntime; }
+
+	UPROPERTY(config, EditDefaultsOnly, Category = "Overrides", Meta = (InlineEditConditionToggle))
+	uint32 bOverrideApplySimProxyInputAtRuntime : 1;
+	// Overrides CVar: np2.Resim.ApplySimProxyInputAtRuntime -- When true, call ApplyData on received inputs for simulated proxies at runtime.
+	UPROPERTY(config, EditDefaultsOnly, Category = "Overrides", Meta = (EditCondition = "bOverrideApplySimProxyInputAtRuntime"))
+	bool bApplySimProxyInputAtRuntime = PhysicsReplicationCVars::ResimulationCVars::bApplySimProxyInputAtRuntime;
+	const bool GetApplySimProxyInputAtRuntime() const { return bOverrideApplySimProxyInputAtRuntime ? bApplySimProxyInputAtRuntime : PhysicsReplicationCVars::ResimulationCVars::bApplySimProxyInputAtRuntime; }
 };
 
 /*
