Remove CustomOutput nodes from the node palette in Material Functions

#tests Windows Editor

diff --git a/Engine/Source/Runtime/Engine/Private/Materials/MaterialExpressions.cpp b/Engine/Source/Runtime/Engine/Private/Materials/MaterialExpressions.cpp
--- a/Engine/Source/Runtime/Engine/Private/Materials/MaterialExpressions.cpp
+++ b/Engine/Source/Runtime/Engine/Private/Materials/MaterialExpressions.cpp
@@ -356,7 +356,8 @@ bool IsAllowedExpressionType(const UClass* const Class, const bool bMaterialFunc
 
 	if (bMaterialFunction)
 	{
-		return bSharedAllowed;
+		return bSharedAllowed 
+			&& !(Class && Class->IsChildOf(UMaterialExpressionCustomOutput::StaticClass()));
 	}
 	else
 	{
