Improve shader error with Substrate when mixing Unlit node with other Slab nodes.

diff --git a/Engine/Source/Runtime/Engine/Private/Materials/HLSLMaterialTranslator.cpp b/Engine/Source/Runtime/Engine/Private/Materials/HLSLMaterialTranslator.cpp
--- a/Engine/Source/Runtime/Engine/Private/Materials/HLSLMaterialTranslator.cpp
+++ b/Engine/Source/Runtime/Engine/Private/Materials/HLSLMaterialTranslator.cpp
@@ -13038,7 +13038,9 @@ bool FHLSLMaterialTranslator::FSubstrateCompilationContext::SubstrateGenerateDer
 
 			if (bHasUnlit && bOperatorEncounteredButNotWeight)
 			{
-				Compiler->Errorf(TEXT("Unlit can only be used with coverage operators. See %s (asset: %s).\r\n"), *CompilerMaterial->GetDebugName(), *CompilerMaterial->GetAssetPath().ToString());
+				Compiler->Errorf(TEXT("Unlit can only be used with coverage operators.\r\n \
+					If you want to blend an unlit material with other slabs, replace the unlit node with a slab having a black albedo, black F0, and an emissive value.\r\n \
+					See %s (asset: %s).\r\n"), *CompilerMaterial->GetDebugName(), *CompilerMaterial->GetAssetPath().ToString());
 				return false;
 			}
 
