Fix for out of memory issues when running the retarget source asset fix command in the editor, on projects that have a extremely high amount of animations. Now we can provide include and exclude paths, making the asset list a lot smaller.
Removed unused cvar and made namespace functions static, to avoid future issues.

diff --git a/Engine/Source/Runtime/Engine/Private/Animation/AnimRetargetHelpers.cpp b/Engine/Source/Runtime/Engine/Private/Animation/AnimRetargetHelpers.cpp
--- a/Engine/Source/Runtime/Engine/Private/Animation/AnimRetargetHelpers.cpp
+++ b/Engine/Source/Runtime/Engine/Private/Animation/AnimRetargetHelpers.cpp
@@ -24,12 +24,6 @@ namespace UE::Anim::RetargetHelpers
 namespace Private
 {
 
-int32 GEnablePostLoadRetargetSourceErrorReporting = 1;
-static FAutoConsoleVariableRef CVarStripAdditiveRefPose(
-	TEXT("a.EnablePostLoadRetargetSourceErrorReporting"),
-	GEnablePostLoadRetargetSourceErrorReporting,
-	TEXT("1 = Enables validation of retarget source asset data. 0 = off"));
-
 template<typename AssetType>
 ERetargetSourceAssetStatus CheckRetargetSourceAssetDataImpl(const AssetType* InAsset)
 {
@@ -64,10 +58,10 @@ ERetargetSourceAssetStatus CheckRetargetSourceAssetDataImpl(const AssetType* InA
 			const USkeletalMesh* SourceReferenceMesh = RetargetSourceAsset.LoadSynchronous();
 			if (SourceReferenceMesh == nullptr)
 			{
-					UE_LOG(LogAnimation, Warning, TEXT("Asset [%s] references a missing Retarget Source Asset [%s]. Retarget Reference Pose has [%d] elements. Please, add a correct retarget source asset and resave.")
-						, *InAsset->GetFullName()
-						, *(RetargetSourceAsset.GetLongPackageName() + FString(TEXT("/")) + RetargetSourceAsset.GetAssetName())
-						, InAsset->RetargetSourceAssetReferencePose.Num());
+				UE_LOG(LogAnimation, Warning, TEXT("Asset [%s] references a missing Retarget Source Asset [%s]. Retarget Reference Pose has [%d] elements. Please, add a correct retarget source asset and resave.")
+					, *InAsset->GetFullName()
+					, *(RetargetSourceAsset.GetLongPackageName() + FString(TEXT("/")) + RetargetSourceAsset.GetAssetName())
+					, InAsset->RetargetSourceAssetReferencePose.Num());
 				return ERetargetSourceAssetStatus::RetargetSourceMissing;
 			}
 			else
@@ -80,12 +74,49 @@ ERetargetSourceAssetStatus CheckRetargetSourceAssetDataImpl(const AssetType* InA
 	return ERetargetSourceAssetStatus::NoRetargetDataSet;
 }
 
-void CheckRetargetSourceAssetData(bool bFixAssets)
+void CheckRetargetSourceAssetData(bool bFixAssets, const TArray<FString>& IncludedPaths, const TArray<FString>& ExcludedPaths)
 {
 	TArray<FAssetData> Assets;
+
+	FARFilter AssetFilter;
+	if (IncludedPaths.Num() > 0)
+	{
+		for (const FString IncludedPath : IncludedPaths)
+		{
+			UE_LOG(LogAnimation, Log, TEXT("Check Retarget Source Assets scan folder [%s]."), *IncludedPath);
+			AssetFilter.PackagePaths.AddUnique(*IncludedPath);
+		}
+		AssetFilter.bRecursivePaths = true;
+	}
+	else
+	{
+		UE_LOG(LogAnimation, Warning, TEXT("Check Retarget Source Assets will scan all folders (this might take some time and require a lot of memory)."));
+	}
+	AssetFilter.ClassPaths.Add(UAnimSequence::StaticClass()->GetClassPathName());
+	AssetFilter.ClassPaths.Add(UPoseAsset::StaticClass()->GetClassPathName());
+
 	IAssetRegistry& AssetRegistry = FModuleManager::LoadModuleChecked<FAssetRegistryModule>(TEXT("AssetRegistry")).Get();
-	AssetRegistry.GetAssetsByClass(UAnimSequence::StaticClass()->GetClassPathName(), Assets);
-	AssetRegistry.GetAssetsByClass(UPoseAsset::StaticClass()->GetClassPathName(), Assets);
+	AssetRegistry.GetAssets(AssetFilter, Assets);
+
+	UE_LOG(LogAnimation, Log, TEXT("Check Retarget Source Assets found [%d] assets."), Assets.Num());
+
+	// Run through paths and classes that should be excluded
+	if (Assets.Num() > 0 && ExcludedPaths.Num() > 0)
+	{
+		FARFilter Filter;
+		Filter.bIncludeOnlyOnDiskAssets = true;
+		Filter.bRecursivePaths = true;
+		for (const FString& ExcludedPath : ExcludedPaths)
+		{
+			UE_LOG(LogAnimation, Log, TEXT("Check Retarget Source Assets Excluded folder : [%s]."), *ExcludedPath);
+			Filter.PackagePaths.AddUnique(*ExcludedPath);
+		}
+		TArray<FAssetData> ExcludedAssetList;
+		AssetRegistry.GetAssets(Filter, ExcludedAssetList);
+		Assets.RemoveAll([&ExcludedAssetList](const FAssetData& Asset) {return ExcludedAssetList.Contains(Asset); });
+	}
+
+	UE_LOG(LogAnimation, Log, TEXT("Check Retarget Source Assets after filtering exclusions : [%d] assets."), Assets.Num());
 
 	const int32 NumAssets = Assets.Num();
 	for (int32 Idx = 0; Idx < NumAssets; Idx++)
@@ -127,17 +158,47 @@ void CheckRetargetSourceAssetData(bool bFixAssets)
 
 static FAutoConsoleCommand CheckRetargetSourceAssetDataCmd(
 	TEXT("a.CheckRetargetSourceAssetData"),
-	TEXT("Checks if Anim Sequences and Pose Assets RetargetSourceAsset is valid. Use: 'a.CheckRetargetSourceAssetData' to check or 'a.CheckRetargetSourceAssetData true' to check and fix the assets."),
+	TEXT("Checks if Anim Sequences and Pose Assets RetargetSourceAsset is valid. Type: 'a.CheckRetargetSourceAssetData /Game' to check assets in the Game (Content) folder.  'a.CheckRetargetSourceAssetData /Game true' to check and fix all the assets in the Game (Content) folder."),
 	FConsoleCommandWithArgsDelegate::CreateLambda([](const TArray<FString>& Args)
 		{
-			if (Args.Num() > 0)
+			static const FString IncludedPathsSwitch = TEXT("Include=");
+			static const FString ExcludedPathsSwitch = TEXT("Exclude=");
+			static const TCHAR* ParamDelims[] =
 			{
-				CheckRetargetSourceAssetData(Args[0].ToBool());
-			}
-			else
+				TEXT(";"),
+				TEXT("+"),
+				TEXT(","),
+			};
+
+			TArray<FString> IncludedPaths;
+			TArray<FString> ExcludedPaths;
+
+			bool bWantsFix = false;
+			FString SwitchValue;
+
+			const int32 NumArgs = Args.Num();
+			for (int32 i = 0; i < NumArgs; i++)
 			{
-				CheckRetargetSourceAssetData(false);
+				const FString& Arg = Args[i];
+				if (Arg.ToLower() == FString("FixAssets"))
+				{
+					bWantsFix = true;
+				}
+				else if (FParse::Value(*Arg, *IncludedPathsSwitch, SwitchValue))
+				{
+					SwitchValue.ParseIntoArray(IncludedPaths, ParamDelims, 3);
+				}
+				else if (FParse::Value(*Arg, *ExcludedPathsSwitch, SwitchValue))
+				{
+					SwitchValue.ParseIntoArray(ExcludedPaths, ParamDelims, 3);
+				}
+				else
+				{
+					IncludedPaths.AddUnique(Arg);
+				}
 			}
+
+			CheckRetargetSourceAssetData(bWantsFix, IncludedPaths, ExcludedPaths);
 		}
 	));
 } // end namespace Private
diff --git a/Engine/Source/Runtime/Engine/Public/Animation/AnimRetargetHelpers.h b/Engine/Source/Runtime/Engine/Public/Animation/AnimRetargetHelpers.h
--- a/Engine/Source/Runtime/Engine/Public/Animation/AnimRetargetHelpers.h
+++ b/Engine/Source/Runtime/Engine/Public/Animation/AnimRetargetHelpers.h
@@ -20,10 +20,10 @@ enum class ERetargetSourceAssetStatus
 };
 
 // Warning, calling this on PostLoad generates asset dependenciy check failures
-ERetargetSourceAssetStatus CheckRetargetSourceAssetData(const UAnimSequence* InAsset);
+static ERetargetSourceAssetStatus CheckRetargetSourceAssetData(const UAnimSequence* InAsset);
 
 // Warning, calling this on PostLoad generates asset dependenciy check failures
-ERetargetSourceAssetStatus CheckRetargetSourceAssetData(const UPoseAsset* InAsset);
+static ERetargetSourceAssetStatus CheckRetargetSourceAssetData(const UPoseAsset* InAsset);
 
 #endif // WITH_EDITOR
 
