Fix Linux build : ranged loop creating a string copy

#rnx

diff --git a/Engine/Source/Runtime/Engine/Private/Animation/AnimRetargetHelpers.cpp b/Engine/Source/Runtime/Engine/Private/Animation/AnimRetargetHelpers.cpp
--- a/Engine/Source/Runtime/Engine/Private/Animation/AnimRetargetHelpers.cpp
+++ b/Engine/Source/Runtime/Engine/Private/Animation/AnimRetargetHelpers.cpp
@@ -81,7 +81,7 @@ void CheckRetargetSourceAssetData(bool bFixAssets, const TArray<FString>& Includ
 	FARFilter AssetFilter;
 	if (IncludedPaths.Num() > 0)
 	{
-		for (const FString IncludedPath : IncludedPaths)
+		for (const FString& IncludedPath : IncludedPaths)
 		{
 			UE_LOG(LogAnimation, Log, TEXT("Check Retarget Source Assets scan folder [%s]."), *IncludedPath);
 			AssetFilter.PackagePaths.AddUnique(*IncludedPath);
