[World Partition] Fixed support for NonPIEEditorOnly actors coming from level instances.
Non-spatially loaded actors coming from level instances are not moved into the persistent level in PIE but rather placed into an always loaded cell, because that would imply loading them through an instancing context.
Also fixed streaming hashes not dumping the always loaded cells into the runtime hash streaming generation logs.
Converted World Partition streaming generation unit test to use the new hashset runtime hash.

#virtualized

diff --git a/Engine/Source/Runtime/Engine/Private/WorldPartition/RuntimeHashSet/WorldPartitionRuntimeHashSetStreamingGeneration.cpp b/Engine/Source/Runtime/Engine/Private/WorldPartition/RuntimeHashSet/WorldPartitionRuntimeHashSetStreamingGeneration.cpp
--- a/Engine/Source/Runtime/Engine/Private/WorldPartition/RuntimeHashSet/WorldPartitionRuntimeHashSetStreamingGeneration.cpp
+++ b/Engine/Source/Runtime/Engine/Private/WorldPartition/RuntimeHashSet/WorldPartitionRuntimeHashSetStreamingGeneration.cpp
@@ -175,6 +175,8 @@ bool UWorldPartitionRuntimeHashSet::GenerateStreaming(UWorldPartitionStreamingPo
 			{
 				UWorldPartitionRuntimeCell* RuntimeCell = RuntimeCells.Emplace_GetRef(CreateRuntimeCellFromCellDesc(CellDescInstance, StreamingPolicy->GetRuntimeCellClass(), UWorldPartitionRuntimeCellDataHashSet::StaticClass()));
 				RuntimeCell->SetIsAlwaysLoaded(bIsCellAlwaysLoaded);
+				RuntimeCell->SetIsSpatiallyLoaded(CellDescInstance.bIsSpatiallyLoaded);
+
 				PopulateRuntimeCell(RuntimeCell, CellActorInstances, OutPackagesToGenerate);
 
 				UWorldPartitionRuntimeCellDataHashSet* RuntimeCellData = CastChecked<UWorldPartitionRuntimeCellDataHashSet>(RuntimeCell->RuntimeCellData);
@@ -257,9 +259,24 @@ void UWorldPartitionRuntimeHashSet::DumpStateLog(FHierarchicalLogArchive& Ar) co
 	Ar.Printf(TEXT("----------------------------------------------------------------------------------------------------------------"));
 
 	TArray<const UWorldPartitionRuntimeCell*> StreamingCells;
-	ForEachStreamingCells([&StreamingCells](const UWorldPartitionRuntimeCell* StreamingCell) { if (!StreamingCell->IsAlwaysLoaded()) { StreamingCells.Add(StreamingCell); } return true; });
+	ForEachStreamingCells([&StreamingCells](const UWorldPartitionRuntimeCell* StreamingCell)
+	{
+		if (!StreamingCell->IsAlwaysLoaded() || !IsRunningCookCommandlet())
+		{
+			StreamingCells.Add(StreamingCell);
+		}
+		return true;
+	});
 				
-	StreamingCells.Sort([this](const UWorldPartitionRuntimeCell& A, const UWorldPartitionRuntimeCell& B) { return A.GetFName().LexicalLess(B.GetFName()); });
+	StreamingCells.Sort([this](const UWorldPartitionRuntimeCell& A, const UWorldPartitionRuntimeCell& B)
+	{
+		if (A.IsAlwaysLoaded() ==  B.IsAlwaysLoaded())
+		{
+			return A.GetFName().LexicalLess(B.GetFName());
+		}
+
+		return A.IsAlwaysLoaded() > B.IsAlwaysLoaded();
+	});
 
 	for (const UWorldPartitionRuntimeCell* StreamingCell : StreamingCells)
 	{
@@ -269,4 +286,4 @@ void UWorldPartitionRuntimeHashSet::DumpStateLog(FHierarchicalLogArchive& Ar) co
 
 	Ar.Printf(TEXT(""));
 }
-#endif
+#endif
\ No newline at end of file
diff --git a/Engine/Source/Runtime/Engine/Private/WorldPartition/WorldPartitionRuntimeCell.cpp b/Engine/Source/Runtime/Engine/Private/WorldPartition/WorldPartitionRuntimeCell.cpp
--- a/Engine/Source/Runtime/Engine/Private/WorldPartition/WorldPartitionRuntimeCell.cpp
+++ b/Engine/Source/Runtime/Engine/Private/WorldPartition/WorldPartitionRuntimeCell.cpp
@@ -125,6 +125,8 @@ void UWorldPartitionRuntimeCell::SetDataLayers(const TArray<const UDataLayerInst
 void UWorldPartitionRuntimeCell::DumpStateLog(FHierarchicalLogArchive& Ar) const
 {
 	Ar.Printf(TEXT("Actor Count: %d"), GetActorCount());
+	Ar.Printf(TEXT("Always Loaded: %s"), IsAlwaysLoaded() ? TEXT("True") : TEXT("False"));
+	Ar.Printf(TEXT("Spatially Loaded: %s"), IsSpatiallyLoaded() ? TEXT("True") : TEXT("False"));
 
 	RuntimeCellData->DumpStateLog(Ar);
 }
diff --git a/Engine/Source/Runtime/Engine/Private/WorldPartition/WorldPartitionRuntimeLevelStreamingCell.cpp b/Engine/Source/Runtime/Engine/Private/WorldPartition/WorldPartitionRuntimeLevelStreamingCell.cpp
--- a/Engine/Source/Runtime/Engine/Private/WorldPartition/WorldPartitionRuntimeLevelStreamingCell.cpp
+++ b/Engine/Source/Runtime/Engine/Private/WorldPartition/WorldPartitionRuntimeLevelStreamingCell.cpp
@@ -289,7 +289,9 @@ void UWorldPartitionRuntimeLevelStreamingCell::SetIsAlwaysLoaded(bool bInIsAlway
 #if WITH_EDITOR
 void UWorldPartitionRuntimeLevelStreamingCell::AddActorToCell(const FStreamingGenerationActorDescView& ActorDescView)
 {
-	check(!ActorDescView.GetActorIsEditorOnly());
+	// Non-spatially loaded actors coming from level instances are not moved into the persistent level in PIE but rather placed into an always loaded cell,
+	// because that would imply loading them through an instancing context.
+	check(!ActorDescView.GetActorIsEditorOnly() || ActorDescView.GetActorIsNonPIEEditorOnly());
 
 	const UActorDescContainerInstance* ContainerInstance = ActorDescView.GetContainerInstance();
 	check(ContainerInstance);
diff --git a/Engine/Source/Runtime/Engine/Private/WorldPartition/WorldPartitionRuntimeSpatialHash.cpp b/Engine/Source/Runtime/Engine/Private/WorldPartition/WorldPartitionRuntimeSpatialHash.cpp
--- a/Engine/Source/Runtime/Engine/Private/WorldPartition/WorldPartitionRuntimeSpatialHash.cpp
+++ b/Engine/Source/Runtime/Engine/Private/WorldPartition/WorldPartitionRuntimeSpatialHash.cpp
@@ -263,14 +263,22 @@ void FSpatialHashStreamingGrid::DumpStateLog(FHierarchicalLogArchive& Ar) const
 			{
 				for (const TObjectPtr<UWorldPartitionRuntimeCell>& Cell : LayerCell.GridCells)
 				{
-					if (!Cell->IsAlwaysLoaded())
+					if (!Cell->IsAlwaysLoaded() || !IsRunningCookCommandlet())
 					{
 						Cells.Add(Cell);
 					}
 				}
 			}
 
-			Cells.Sort([this](const UWorldPartitionRuntimeCell& A, const UWorldPartitionRuntimeCell& B) { return A.GetFName().LexicalLess(B.GetFName()); });
+			Cells.Sort([this](const UWorldPartitionRuntimeCell& A, const UWorldPartitionRuntimeCell& B)
+			{
+				if (A.IsAlwaysLoaded() ==  B.IsAlwaysLoaded())
+				{
+					return A.GetFName().LexicalLess(B.GetFName());
+				}
+
+				return A.IsAlwaysLoaded() > B.IsAlwaysLoaded();
+			});
 
 			for (UWorldPartitionRuntimeCell* Cell : Cells)
 			{
@@ -1596,7 +1604,8 @@ bool UWorldPartitionRuntimeSpatialHash::CreateStreamingGrid(const FSpatialHashRu
 			for (const FSquare2DGridHelper::FGridLevel::FGridCellDataChunk& GridCellDataChunk : TempCell.GetDataChunks())
 			{
 				// Cell cannot be treated as always loaded if it has data layers
-				const bool bIsCellAlwaysLoaded = (&TempCell == &PartitionedActors.GetAlwaysLoadedCell()) && !GridCellDataChunk.HasDataLayers() && !GridCellDataChunk.GetContentBundleID().IsValid();
+				const bool bIsCellSpatiallyLoaded = &TempCell != &PartitionedActors.GetAlwaysLoadedCell();
+				const bool bIsCellAlwaysLoaded = !bIsCellSpatiallyLoaded && !GridCellDataChunk.HasDataLayers() && !GridCellDataChunk.GetContentBundleID().IsValid();				
 				
 				TArray<IStreamingGenerationContext::FActorInstance> FilteredActors;
 				if (PopulateCellActorInstances(GridCellDataChunk.GetActorSetInstances(), bIsMainWorldPartition, bIsCellAlwaysLoaded, FilteredActors))
@@ -1615,6 +1624,7 @@ bool UWorldPartitionRuntimeSpatialHash::CreateStreamingGrid(const FSpatialHashRu
 					const bool bIsHLOD = RuntimeGrid.HLODLayer ? true : false;
 					const bool bBlockOnSlowStreaming = ResolveBlockOnSlowStreamingForCell(CurrentStreamingGrid.bBlockOnSlowStreaming, bIsHLOD, DataLayers);
 					StreamingCell->SetIsAlwaysLoaded(bIsCellAlwaysLoaded);
+					StreamingCell->SetIsSpatiallyLoaded(bIsCellSpatiallyLoaded);
 					StreamingCell->SetDataLayers(DataLayers);
 					StreamingCell->SetContentBundleUID(GridCellDataChunk.GetContentBundleID());
 					StreamingCell->SetClientOnlyVisible(CurrentStreamingGrid.bClientOnlyVisible || bAreClientOnlyDataLayers);
diff --git a/Engine/Source/Runtime/Engine/Private/WorldPartition/WorldPartitionStreamingDescriptor.cpp b/Engine/Source/Runtime/Engine/Private/WorldPartition/WorldPartitionStreamingDescriptor.cpp
--- a/Engine/Source/Runtime/Engine/Private/WorldPartition/WorldPartitionStreamingDescriptor.cpp
+++ b/Engine/Source/Runtime/Engine/Private/WorldPartition/WorldPartitionStreamingDescriptor.cpp
@@ -65,6 +65,7 @@ void FStreamingDescriptor::FStreamingActor::DumpStateLog(FHierarchicalLogArchive
 void FStreamingDescriptor::FStreamingCell::DumpStateLog(FHierarchicalLogArchive& Ar)
 {
 	Ar.Printf(TEXT("          Bounds: %s"), *Bounds.ToString());
+	Ar.Printf(TEXT("   Always Loaded: %s"), bIsAlwaysLoaded ? TEXT("1" : TEXT("0")));
 	Ar.Printf(TEXT("Spatially Loaded: %s"), bIsSpatiallyLoaded ? TEXT("1" : TEXT("0")));
 
 	if (DataLayers.Num())
@@ -143,7 +144,8 @@ bool FStreamingDescriptor::GenerateStreamingDescriptor(UWorld* InWorld, FStreami
 				FStreamingDescriptor::FStreamingGrid& StreamingGrid = StreamingGrids.FindChecked(CellDataSpatialHash->GridName);
 				FStreamingDescriptor::FStreamingCell& StreamingCell = *new(StreamingGrid.StreamingCells) FStreamingDescriptor::FStreamingCell;
 				StreamingCell.Bounds = Cell->GetCellBounds();
-				StreamingCell.bIsSpatiallyLoaded = !Cell->IsAlwaysLoaded();
+				StreamingCell.bIsAlwaysLoaded = Cell->IsAlwaysLoaded();
+				StreamingCell.bIsSpatiallyLoaded = Cell->IsSpatiallyLoaded();
 				StreamingCell.CellPackage = Cell->GetLevelPackageName();
 
 				StreamingCell.DataLayers = Cell->GetDataLayers();
@@ -226,7 +228,8 @@ bool FStreamingDescriptor::GenerateStreamingDescriptor(UWorld* InWorld, FStreami
 						FStreamingDescriptor::FStreamingCell& StreamingCell = *new(StreamingGrid.StreamingCells) FStreamingDescriptor::FStreamingCell;
 
 						StreamingCell.Bounds = Cell->GetCellBounds();
-						StreamingCell.bIsSpatiallyLoaded = !Cell->IsAlwaysLoaded();
+						StreamingCell.bIsAlwaysLoaded = Cell->IsAlwaysLoaded();
+						StreamingCell.bIsSpatiallyLoaded = Cell->IsSpatiallyLoaded();
 						StreamingCell.DataLayers = Cell->GetDataLayers();
 						StreamingCell.CellPackage = FName(Cell->GetName());
 
diff --git a/Engine/Source/Runtime/Engine/Public/WorldPartition/WorldPartitionRuntimeCell.h b/Engine/Source/Runtime/Engine/Public/WorldPartition/WorldPartitionRuntimeCell.h
--- a/Engine/Source/Runtime/Engine/Public/WorldPartition/WorldPartitionRuntimeCell.h
+++ b/Engine/Source/Runtime/Engine/Public/WorldPartition/WorldPartitionRuntimeCell.h
@@ -253,7 +253,9 @@ class UWorldPartitionRuntimeCell : public UObject, public IWorldPartitionCell, p
 	ENGINE_API virtual EWorldPartitionRuntimeCellState GetCurrentState() const PURE_VIRTUAL(UWorldPartitionRuntimeCell::GetCurrentState, return EWorldPartitionRuntimeCellState::Unloaded;);
 	virtual FLinearColor GetDebugColor(EWorldPartitionRuntimeCellVisualizeMode VisualizeMode) const { static const FLinearColor DefaultColor = FLinearColor::Black.CopyWithNewOpacity(0.25f); return DefaultColor; }
 	virtual bool IsAlwaysLoaded() const { return bIsAlwaysLoaded; }
+	virtual bool IsSpatiallyLoaded() const { return bIsSpatiallyLoaded; }
 	virtual void SetIsAlwaysLoaded(bool bInIsAlwaysLoaded) { bIsAlwaysLoaded = bInIsAlwaysLoaded; }
+	virtual void SetIsSpatiallyLoaded(bool bInIsSpatiallyLoaded) { bIsSpatiallyLoaded = bInIsSpatiallyLoaded; }
 	ENGINE_API virtual void SetStreamingPriority(int32 InStreamingPriority) const PURE_VIRTUAL(UWorldPartitionRuntimeCell::SetStreamingPriority,);
 	virtual EStreamingStatus GetStreamingStatus() const { return LEVEL_Unloaded; }
 	void SetClientOnlyVisible(bool bInClientOnlyVisible) { bClientOnlyVisible = bInClientOnlyVisible; }
@@ -378,6 +380,9 @@ protected:
 	UPROPERTY()
 	bool bIsAlwaysLoaded;
 
+	UPROPERTY()
+	bool bIsSpatiallyLoaded;
+
 private:
 	UPROPERTY()
 	FDataLayerInstanceNames DataLayers;
diff --git a/Engine/Source/Runtime/Engine/Public/WorldPartition/WorldPartitionStreamingDescriptor.h b/Engine/Source/Runtime/Engine/Public/WorldPartition/WorldPartitionStreamingDescriptor.h
--- a/Engine/Source/Runtime/Engine/Public/WorldPartition/WorldPartitionStreamingDescriptor.h
+++ b/Engine/Source/Runtime/Engine/Public/WorldPartition/WorldPartitionStreamingDescriptor.h
@@ -47,6 +47,9 @@ namespace UE::Private::WorldPartition
 			/** The cell bounds, not necessarily the same size for all cells */
 			FBox Bounds;
 
+			/** If this cell is always loaded or not */
+			bool bIsAlwaysLoaded;
+
 			/** If this cell is spatially loaded or not */
 			bool bIsSpatiallyLoaded;
 
