[World Partition] Fixup invalid actor descriptor bounds containing NaNs.

diff --git a/Engine/Source/Runtime/Engine/Private/WorldPartition/WorldPartitionActorDesc.cpp b/Engine/Source/Runtime/Engine/Private/WorldPartition/WorldPartitionActorDesc.cpp
--- a/Engine/Source/Runtime/Engine/Private/WorldPartition/WorldPartitionActorDesc.cpp
+++ b/Engine/Source/Runtime/Engine/Private/WorldPartition/WorldPartitionActorDesc.cpp
@@ -96,6 +96,7 @@ void FWorldPartitionActorDesc::Init(const AActor* InActor)
 	if (!bIsDefaultActorDesc)
 	{
 		InActor->GetStreamingBounds(RuntimeBounds, EditorBounds);
+		FixupStreamingBounds();
 	}
 
 	RuntimeGrid = InActor->GetRuntimeGrid();
@@ -292,6 +293,8 @@ void FWorldPartitionActorDesc::Init(const FWorldPartitionActorDescInitData& Desc
 	ActorDisplayClassNameString = GetDisplayClassName().ToString();
 
 	Container = nullptr;
+
+	FixupStreamingBounds();
 }
 
 void FWorldPartitionActorDesc::Patch(const FWorldPartitionActorDescInitData& DescData, TArray<uint8>& OutData, FWorldPartitionAssetDataPatcher* InAssetDataPatcher)
@@ -661,7 +664,7 @@ void FWorldPartitionActorDesc::Serialize(FArchive& Ar)
 		RuntimeBounds.Init();
 		EditorBounds.Init();
 	}
-	
+
 	if (Ar.CustomVer(FUE5ReleaseStreamObjectVersion::GUID) < FUE5ReleaseStreamObjectVersion::ConvertedActorGridPlacementToSpatiallyLoadedFlag)
 	{
 		PRAGMA_DISABLE_DEPRECATION_WARNINGS
@@ -931,6 +934,32 @@ bool FWorldPartitionActorDesc::IsRuntimeRelevant(const FWorldPartitionActorDescI
 	return !IsMainWorldOnly() || InActorDescInstance->GetContainerInstance()->GetContainerID().IsMainContainer();
 }
 
+void FWorldPartitionActorDesc::FixupStreamingBounds()
+{
+	const bool bRuntimeBoundsContainsNaN = RuntimeBounds.IsValid && RuntimeBounds.ContainsNaN();
+	const bool bEditorBoundsContainsNaN = EditorBounds.IsValid && EditorBounds.ContainsNaN();
+
+	if (bRuntimeBoundsContainsNaN != bEditorBoundsContainsNaN)
+	{
+		if (bRuntimeBoundsContainsNaN)
+		{
+			UE_LOG(LogWorldPartition, Log, TEXT("Invalid runtime bounds for actor descriptor '%s', overwriting with editor bounds."), *ToString(EToStringMode::Compact));
+			RuntimeBounds = EditorBounds;			
+		}
+		else
+		{
+			UE_LOG(LogWorldPartition, Log, TEXT("Invalid editor bounds for actor descriptor '%s', overwriting with runtime bounds."), *ToString(EToStringMode::Compact));
+			EditorBounds = RuntimeBounds;			
+		}
+	}
+	else if (bRuntimeBoundsContainsNaN && bEditorBoundsContainsNaN)
+	{
+		UE_LOG(LogWorldPartition, Log, TEXT("Invalid streaming bounds for actor descriptor '%s'"), *ToString(EToStringMode::Compact));
+		RuntimeBounds.Init();		
+		EditorBounds.Init();
+	}
+}
+
 #undef LOCTEXT_NAMESPACE
 
 #endif
\ No newline at end of file
diff --git a/Engine/Source/Runtime/Engine/Public/WorldPartition/WorldPartitionActorDesc.h b/Engine/Source/Runtime/Engine/Public/WorldPartition/WorldPartitionActorDesc.h
--- a/Engine/Source/Runtime/Engine/Public/WorldPartition/WorldPartitionActorDesc.h
+++ b/Engine/Source/Runtime/Engine/Public/WorldPartition/WorldPartitionActorDesc.h
@@ -305,6 +305,7 @@ PRAGMA_ENABLE_DEPRECATION_WARNINGS
 
 	UE_DEPRECATED(5.4, "Use FWorldPartitionActorDescInstance version instead")
 	virtual void CheckForErrors(IStreamingGenerationErrorHandler* ErrorHandler) const {}
+
 protected:
 	UE_DEPRECATED(5.4, "Please use FWorldPartitionActorDescInstance::IncSoftRefCount")
 	inline uint32 IncSoftRefCount() const { return 0; }
@@ -448,6 +449,8 @@ protected:
 
 	static ENGINE_API TMap<TSubclassOf<AActor>, FActorDescDeprecator> Deprecators;
 
+private:
+	void FixupStreamingBounds();
 #endif // WITH_EDITOR
 };
 
