Provide support for array of HLOD for the injected level approach

diff --git a/Engine/Source/Runtime/Engine/Private/WorldPartition/RuntimeHashSet/WorldPartitionRuntimeHashSet.cpp b/Engine/Source/Runtime/Engine/Private/WorldPartition/RuntimeHashSet/WorldPartitionRuntimeHashSet.cpp
--- a/Engine/Source/Runtime/Engine/Private/WorldPartition/RuntimeHashSet/WorldPartitionRuntimeHashSet.cpp
+++ b/Engine/Source/Runtime/Engine/Private/WorldPartition/RuntimeHashSet/WorldPartitionRuntimeHashSet.cpp
@@ -699,22 +699,28 @@ FGuid UWorldPartitionRuntimeHashSet::RegisterWorldAssetStreaming(const UWorldPar
 	}
 
 	FName RuntimePartitionName;
-	const URuntimePartition* RuntimePartition = ResolveRuntimePartition(InParams.TargetGrid);
+	const URuntimePartition* RuntimePartition = ResolveRuntimePartition(InParams.WorldAssetDesc.TargetGrid);
 	if (!RuntimePartition)
 	{
-		UE_LOG(LogWorldPartition, Error, TEXT("RegisterWorldAssetStreaming: Unable to resolve TargetGrid '%s'."), *InParams.TargetGrid.ToString());
+		UE_LOG(LogWorldPartition, Error, TEXT("RegisterWorldAssetStreaming: Unable to resolve TargetGrid '%s'."), *InParams.WorldAssetDesc.TargetGrid.ToString());
 		return FGuid();
 	}
 
-	FName HLODRuntimePartitionName;
-	const URuntimePartition* HLODRuntimePartition = nullptr;
-	if (!InParams.WorldAssetHLOD.IsNull() && !InParams.TargetGridHLOD.IsNone())
+	TArray<TPair<const URuntimePartition*, const TSoftObjectPtr<UWorld>>> HLODRuntimePartitions;
+	for (const UWorldPartition::FRegisterWorldAssetStreamingParams::FWorldAssetDesc& HLODWorldAssetDesc : InParams.HLODWorldAssetDescs)
 	{
-		HLODRuntimePartition = ResolveRuntimePartition(InParams.TargetGridHLOD);
-		if (!HLODRuntimePartition)
+		if (!HLODWorldAssetDesc.WorldAsset.IsNull() && !HLODWorldAssetDesc.TargetGrid.IsNone())
 		{
-			UE_LOG(LogWorldPartition, Error, TEXT("RegisterWorldAssetStreaming: Unable to resolve TargetGridHLOD '%s'."), *InParams.TargetGridHLOD.ToString());
-			return FGuid();
+			const URuntimePartition* HLODRuntimePartition = ResolveRuntimePartition(HLODWorldAssetDesc.TargetGrid);
+			if (!HLODRuntimePartition)
+			{
+				UE_LOG(LogWorldPartition, Error, TEXT("RegisterWorldAssetStreaming: Unable to resolve TargetGridHLOD '%s'."), *HLODWorldAssetDesc.TargetGrid.ToString());
+				return FGuid();
+			}
+			else
+			{
+				HLODRuntimePartitions.Add({ HLODRuntimePartition, HLODWorldAssetDesc.WorldAsset });
+			}
 		}
 	}
 
@@ -727,15 +733,10 @@ FGuid UWorldPartitionRuntimeHashSet::RegisterWorldAssetStreaming(const UWorldPar
 
 	FGuid SourceCellGuid;
 		
-	// 2 passes : 1st for target grid, 2nd for HLOD target grid
-	for (const bool bIsHLODPass : { false, true })
+	auto CreateStreamingCell = [&InParams, &SourceCellGuid, StreamingObject, this](const URuntimePartition* TargetPartition, const TSoftObjectPtr<UWorld>& WorldAsset, bool bIsHLODPass) -> bool
 	{
-		const URuntimePartition* TargetPartition = bIsHLODPass ? HLODRuntimePartition : RuntimePartition;
-		if (!TargetPartition)
-		{
-			continue;
-		}
-
+		check(TargetPartition);
+		
 		bool bClientOnlyVisible = TargetPartition->bClientOnlyVisible;
 		bool bBlockOnSlowStreaming = TargetPartition->bBlockOnSlowStreaming;
 
@@ -756,7 +757,6 @@ FGuid UWorldPartitionRuntimeHashSet::RegisterWorldAssetStreaming(const UWorldPar
 		}
 
 		const FString CellName = FString::Printf(TEXT("InjectedCell_%s"), *CellGuid.ToString());
-		const TSoftObjectPtr<UWorld>& WorldAsset = bIsHLODPass ? InParams.WorldAssetHLOD : InParams.WorldAsset;
 
 		if (UWorldPartitionRuntimeLevelStreamingCell* RuntimeCell = Cast<UWorldPartitionRuntimeLevelStreamingCell>(
 			CreateRuntimeCell(UWorldPartitionRuntimeLevelStreamingCell::StaticClass(), UWorldPartitionRuntimeCellDataHashSet::StaticClass(), CellName, InParams.CellInstanceSuffix, StreamingObject)))
@@ -789,12 +789,26 @@ FGuid UWorldPartitionRuntimeHashSet::RegisterWorldAssetStreaming(const UWorldPar
 			else
 			{
 				UE_LOG(LogWorldPartition, Error, TEXT("Error creating streaming cell %s for world asset %s at %s"), *RuntimeCell->GetName(), *WorldAsset.ToString(), *InParams.Transform.ToString());
-				return FGuid();
+				return false;
 			}
 		}		
 		else
 		{
 			UE_LOG(LogWorldPartition, Error, TEXT("Error creating streaming cell %s for world asset %s at %s"), *CellName, *WorldAsset.ToString(), *InParams.Transform.ToString());
+			return false;
+		}
+
+		return true;
+	};
+
+	if (CreateStreamingCell(RuntimePartition, InParams.WorldAssetDesc.WorldAsset, false /* bIsHLODPass */) == false)
+	{
+		return FGuid();
+	}
+	for (const TPair<const URuntimePartition*, const TSoftObjectPtr<UWorld>>& HLODRuntimePartition : HLODRuntimePartitions)
+	{
+		if (CreateStreamingCell(HLODRuntimePartition.Key, HLODRuntimePartition.Value, true /* bIsHLODPass */) == false)
+		{
 			return FGuid();
 		}
 	}
diff --git a/Engine/Source/Runtime/Engine/Private/WorldPartition/WorldPartitionRuntimeSpatialHash.cpp b/Engine/Source/Runtime/Engine/Private/WorldPartition/WorldPartitionRuntimeSpatialHash.cpp
--- a/Engine/Source/Runtime/Engine/Private/WorldPartition/WorldPartitionRuntimeSpatialHash.cpp
+++ b/Engine/Source/Runtime/Engine/Private/WorldPartition/WorldPartitionRuntimeSpatialHash.cpp
@@ -1970,20 +1970,27 @@ FGuid UWorldPartitionRuntimeSpatialHash::RegisterWorldAssetStreaming(const UWorl
 		return FGuid();
 	}
 
-	const FSpatialHashStreamingGrid* TargetGrid = GetStreamingGridByName(InParams.TargetGrid);
+	const FSpatialHashStreamingGrid* TargetGrid = GetStreamingGridByName(InParams.WorldAssetDesc.TargetGrid);
 	if (!TargetGrid)
 	{
-		UE_LOG(LogWorldPartition, Error, TEXT("RegisterWorldAssetStreaming: Unable to resolve TargetGrid '%s'."), *InParams.TargetGrid.ToString());
+		UE_LOG(LogWorldPartition, Error, TEXT("RegisterWorldAssetStreaming: Unable to resolve TargetGrid '%s'."), *InParams.WorldAssetDesc.TargetGrid.ToString());
 		return FGuid();
 	}
 
-	const FSpatialHashStreamingGrid* HLODTargetGrid = nullptr;
-	if (!InParams.WorldAssetHLOD.IsNull() && !InParams.TargetGridHLOD.IsNone())
+	TArray<TPair<const FSpatialHashStreamingGrid*, const TSoftObjectPtr<UWorld>>> HLODTargetGrids;
+	for (const UWorldPartition::FRegisterWorldAssetStreamingParams::FWorldAssetDesc& HLODWorldAssetDesc : InParams.HLODWorldAssetDescs)
 	{
-		HLODTargetGrid = GetStreamingGridByName(InParams.TargetGridHLOD);
-		if (!HLODTargetGrid)
+		if (!HLODWorldAssetDesc.WorldAsset.IsNull() && !HLODWorldAssetDesc.TargetGrid.IsNone())
 		{
-			UE_LOG(LogWorldPartition, Error, TEXT("RegisterWorldAssetStreaming: Unable to resolve TargetGridHLOD '%s'."), *InParams.TargetGridHLOD.ToString());
+			const FSpatialHashStreamingGrid* HLODTargetGrid = GetStreamingGridByName(HLODWorldAssetDesc.TargetGrid);
+			if (!HLODTargetGrid)
+			{
+				UE_LOG(LogWorldPartition, Error, TEXT("RegisterWorldAssetStreaming: Unable to resolve TargetGridHLOD '%s'."), *HLODWorldAssetDesc.TargetGrid.ToString());
+			}
+			else
+			{
+				HLODTargetGrids.Add({ HLODTargetGrid, HLODWorldAssetDesc.WorldAsset });
+			}
 		}
 	}
 
@@ -1996,14 +2003,9 @@ FGuid UWorldPartitionRuntimeSpatialHash::RegisterWorldAssetStreaming(const UWorl
 
 	FGuid SourceCellGuid;
 		
-	// 2 passes : 1st for target grid, 2nd for HLOD target grid
-	for (const bool bIsHLODPass : { false, true })
+	auto CreateStreamingCell = [&InParams, &SourceCellGuid, StreamingObject, this](const FSpatialHashStreamingGrid* SourceGrid, const TSoftObjectPtr<UWorld>& WorldAsset, bool bIsHLODPass) -> bool
 	{
-		const FSpatialHashStreamingGrid* SourceGrid = bIsHLODPass ? HLODTargetGrid : TargetGrid;
-		if (!SourceGrid)
-		{
-			continue;
-		}
+		check(SourceGrid);
 
 		FSquare2DGridHelper CurrentGridHelper(SourceGrid->WorldBounds, SourceGrid->Origin, SourceGrid->CellSize, SourceGrid->Settings.bUseAlignedGridLevels);
 		FSpatialHashStreamingGrid& CurrentGrid = StreamingObject->StreamingGrids.AddDefaulted_GetRef();
@@ -2018,11 +2020,6 @@ FGuid UWorldPartitionRuntimeSpatialHash::RegisterWorldAssetStreaming(const UWorl
 		CurrentGrid.bClientOnlyVisible = SourceGrid->bClientOnlyVisible;
 		CurrentGrid.GridLevels.SetNum(CurrentGridHelper.Levels.Num());
 
-		if (bIsHLODPass && InParams.WorldAssetHLOD.IsNull())
-		{
-			continue;
-		}
-
 		FSquare2DGridHelper::FGridLevel::FGridCell* GridCell = nullptr;
 		if (InParams.bBoundsPlacement)
 		{
@@ -2062,7 +2059,7 @@ FGuid UWorldPartitionRuntimeSpatialHash::RegisterWorldAssetStreaming(const UWorl
 		}
 
 		// Create Cell
-		FName TargetGridName = (bIsHLODPass ? InParams.TargetGridHLOD : InParams.TargetGrid);
+		FName TargetGridName = SourceGrid->GridName;
 		FGuid InstanceGuid = InParams.Guid;
 		FArchiveMD5 ArMD5;
 		ArMD5 << TargetGridName << InstanceGuid;
@@ -2074,7 +2071,6 @@ FGuid UWorldPartitionRuntimeSpatialHash::RegisterWorldAssetStreaming(const UWorl
 		}
 
 		const FString CellName = FString::Printf(TEXT("InjectedCell_%s"), *CellGuid.ToString());
-		const TSoftObjectPtr<UWorld>& WorldAsset = bIsHLODPass ? InParams.WorldAssetHLOD : InParams.WorldAsset;
 
 		if (UWorldPartitionRuntimeLevelStreamingCell* StreamingCell = Cast<UWorldPartitionRuntimeLevelStreamingCell>(
 			CreateRuntimeCell(UWorldPartitionRuntimeLevelStreamingCell::StaticClass(), UWorldPartitionRuntimeCellDataSpatialHash::StaticClass(), CellName, InParams.CellInstanceSuffix, StreamingObject)))
@@ -2119,12 +2115,26 @@ FGuid UWorldPartitionRuntimeSpatialHash::RegisterWorldAssetStreaming(const UWorl
 			else
 			{
 				UE_LOG(LogWorldPartition, Error, TEXT("Error creating streaming cell %s for world asset %s at %s"), *StreamingCell->GetName(), *WorldAsset.ToString(), *InParams.Transform.ToString());
-				return FGuid();
+				return false;
 			}
 		}
 		else
 		{
 			UE_LOG(LogWorldPartition, Error, TEXT("Error creating streaming cell %s for world asset %s at %s"), *CellName, *WorldAsset.ToString(), *InParams.Transform.ToString());
+			return false;
+		}
+
+		return true;
+	};
+
+	if (CreateStreamingCell(TargetGrid, InParams.WorldAssetDesc.WorldAsset, false /* bisHLODPass */) == false)
+	{
+		return FGuid();
+	}
+	for (const TPair<const FSpatialHashStreamingGrid*, const TSoftObjectPtr<UWorld>>& HLODTargetGrid : HLODTargetGrids)
+	{
+		if (CreateStreamingCell(HLODTargetGrid.Key, HLODTargetGrid.Value, true /* bIsHLODPass */) == false)
+		{
 			return FGuid();
 		}
 	}
diff --git a/Engine/Source/Runtime/Engine/Public/WorldPartition/WorldPartition.h b/Engine/Source/Runtime/Engine/Public/WorldPartition/WorldPartition.h
--- a/Engine/Source/Runtime/Engine/Public/WorldPartition/WorldPartition.h
+++ b/Engine/Source/Runtime/Engine/Public/WorldPartition/WorldPartition.h
@@ -772,14 +772,17 @@ public:
 	 */
 	struct FRegisterWorldAssetStreamingParams
 	{
+		struct FWorldAssetDesc
+		{
+			TSoftObjectPtr<UWorld> WorldAsset;
+			FName TargetGrid;
+		};
+
 		FRegisterWorldAssetStreamingParams() 
 		{}
 
-		TSoftObjectPtr<UWorld> WorldAsset;
-		FName TargetGrid;
-
-		TSoftObjectPtr<UWorld> WorldAssetHLOD;
-		FName TargetGridHLOD;
+		FWorldAssetDesc WorldAssetDesc;
+		TArray<FWorldAssetDesc> HLODWorldAssetDescs;
 
 		FGuid Guid;
 		FTransform Transform;
@@ -790,13 +793,12 @@ public:
 
 		bool IsValid() const
 		{
-			return !WorldAsset.IsNull() && !TargetGrid.IsNone() && Guid.IsValid() && Bounds.IsValid;
+			return !WorldAssetDesc.WorldAsset.IsNull() && !WorldAssetDesc.TargetGrid.IsNone() && Guid.IsValid() && Bounds.IsValid;
 		}
 
-		FRegisterWorldAssetStreamingParams& SetWorldAsset(const TSoftObjectPtr<UWorld>& InWorldAsset) { WorldAsset = InWorldAsset; return *this; }
-		FRegisterWorldAssetStreamingParams& SetTargetGrid(const FName& InTargetGrid) { TargetGrid = InTargetGrid; return *this; }
-		FRegisterWorldAssetStreamingParams& SetWorldAssetHLOD(const TSoftObjectPtr<UWorld>& InWorldAssetHLOD) { WorldAssetHLOD = InWorldAssetHLOD; return *this; }
-		FRegisterWorldAssetStreamingParams& SetTargetGridHLOD(const FName& InTargetGridHLOD) { TargetGridHLOD = InTargetGridHLOD; return *this; }
+		FRegisterWorldAssetStreamingParams& SetWorldAsset(const TSoftObjectPtr<UWorld>& InWorldAsset, const FName& InTargetGrid) { WorldAssetDesc.WorldAsset = InWorldAsset; WorldAssetDesc.TargetGrid = InTargetGrid; return *this; }
+		FRegisterWorldAssetStreamingParams& AddHLODWorldAsset(const TSoftObjectPtr<UWorld>& InHLODWorldAsset, const FName& InHLODTargetGrid) { HLODWorldAssetDescs.Add({ InHLODWorldAsset, InHLODTargetGrid }); return *this; }
+		FRegisterWorldAssetStreamingParams& RemoveAllHLODWorldAssets() { HLODWorldAssetDescs.Reset(); return *this; }
 		FRegisterWorldAssetStreamingParams& SetGuid(const FGuid InGuid) { Guid = InGuid; return *this; }
 		FRegisterWorldAssetStreamingParams& SetTransform(const FTransform InTransform) { Transform = InTransform; return *this; }
 		FRegisterWorldAssetStreamingParams& SetBounds(const FBox& InBounds) { Bounds = InBounds; return *this; }
