Provide ability for a FTextureMipDataProvider to modify the texture on the GPU before the allocator finalizes and applies the new texture rhi ref to the texture resource.

diff --git a/Engine/Source/Runtime/Engine/Private/Streaming/Texture2DArrayStreaming.cpp b/Engine/Source/Runtime/Engine/Private/Streaming/Texture2DArrayStreaming.cpp
--- a/Engine/Source/Runtime/Engine/Private/Streaming/Texture2DArrayStreaming.cpp
+++ b/Engine/Source/Runtime/Engine/Private/Streaming/Texture2DArrayStreaming.cpp
@@ -53,11 +53,11 @@ bool FTexture2DArrayMipAllocator_Reallocate::AllocateMips(const FTextureUpdateCo
 		StreamedSliceSize[MipIdx] = MipInfo.DataSize / MipInfo.ArraySize;
 	}
 
-	AdvanceTo(ETickState::FinalizeMips, ETickThread::Render);
+	AdvanceTo(ETickState::UploadMips, ETickThread::Render);
 	return true;
 }
 
-bool FTexture2DArrayMipAllocator_Reallocate::FinalizeMips(const FTextureUpdateContext& Context, const FTextureUpdateSyncOptions& SyncOptions)
+bool FTexture2DArrayMipAllocator_Reallocate::UploadMips(const FTextureUpdateContext& Context, const FTextureUpdateSyncOptions& SyncOptions)
 {
 	const uint32 NumSlices = Context.Resource->GetSizeZ();
 
@@ -113,20 +113,12 @@ bool FTexture2DArrayMipAllocator_Reallocate::FinalizeMips(const FTextureUpdateCo
 	StreamedMipData.Empty();
 	StreamedSliceSize.Empty();
 
-	// Use the new texture resource for the texture asset, must run on the renderthread.
-	Context.Resource->FinalizeStreaming(IntermediateTextureRHI);
-	// No need for the intermediate texture anymore.
-	IntermediateTextureRHI.SafeRelease();
-
-	// Update is complete, nothing more to do.
 	AdvanceTo(ETickState::Done, ETickThread::None);
 	return true;
 }
 
 void FTexture2DArrayMipAllocator_Reallocate::Cancel(const FTextureUpdateSyncOptions& SyncOptions)
 {
-	// Release the intermediate texture. If not null, this will be on the renderthread.
-	IntermediateTextureRHI.SafeRelease();
 	// Release the temporary mip data. Can be run on either renderthread or async threads.
 	StreamedMipData.Empty();
 	StreamedSliceSize.Empty();
@@ -134,13 +126,8 @@ void FTexture2DArrayMipAllocator_Reallocate::Cancel(const FTextureUpdateSyncOpti
 
 FTextureMipAllocator::ETickThread FTexture2DArrayMipAllocator_Reallocate::GetCancelThread() const
 {
-	// If there is an  intermediate texture, it is safer to released on the renderthread.
-	if (IntermediateTextureRHI)
-	{
-		return ETickThread::Render;
-	}
-	// Otherwise, if there are only temporary mip data, they can be freed on any threads.
-	else if (StreamedMipData.Num())
+	// If there are only temporary mip data, they can be freed on any threads.
+	if (StreamedMipData.Num())
 	{
 		return ETickThread::Async;
 	}
diff --git a/Engine/Source/Runtime/Engine/Private/Streaming/Texture2DArrayStreaming.h b/Engine/Source/Runtime/Engine/Private/Streaming/Texture2DArrayStreaming.h
--- a/Engine/Source/Runtime/Engine/Private/Streaming/Texture2DArrayStreaming.h
+++ b/Engine/Source/Runtime/Engine/Private/Streaming/Texture2DArrayStreaming.h
@@ -35,7 +35,7 @@ public:
 	// ********************************************************
 
 	bool AllocateMips(const FTextureUpdateContext& Context, FTextureMipInfoArray& OutMipInfos, const FTextureUpdateSyncOptions& SyncOptions) final override;
-	bool FinalizeMips(const FTextureUpdateContext& Context, const FTextureUpdateSyncOptions& SyncOptions) final override;
+	bool UploadMips(const FTextureUpdateContext& Context, const FTextureUpdateSyncOptions& SyncOptions) final override;
 	void Cancel(const FTextureUpdateSyncOptions& SyncOptions) final override;
 	ETickThread GetCancelThread() const final override;
 
@@ -50,9 +50,6 @@ protected:
 		}
 	};
 
-	// The intermediate texture async created in the update process.
-	FTextureRHIRef IntermediateTextureRHI;
-
 	// The temporary main memory allocations holding the mip data
 	TArray<TUniquePtr<uint8, FMallocDeleter>, TInlineAllocator<MAX_TEXTURE_MIP_COUNT>> StreamedMipData;
 	TArray<uint64, TInlineAllocator<MAX_TEXTURE_MIP_COUNT>> StreamedSliceSize;
diff --git a/Engine/Source/Runtime/Engine/Private/Streaming/Texture2DMipAllocator_AsyncCreate.cpp b/Engine/Source/Runtime/Engine/Private/Streaming/Texture2DMipAllocator_AsyncCreate.cpp
--- a/Engine/Source/Runtime/Engine/Private/Streaming/Texture2DMipAllocator_AsyncCreate.cpp
+++ b/Engine/Source/Runtime/Engine/Private/Streaming/Texture2DMipAllocator_AsyncCreate.cpp
@@ -57,8 +57,8 @@ bool FTexture2DMipAllocator_AsyncCreate::AllocateMips(
 		FinalSizeY = OutMipInfos[PendingFirstLODIdx].SizeY;
 		FinalFormat = OutMipInfos[PendingFirstLODIdx].Format;
 
-		// Once the FTextureMipDataProvider has set the mip data, FinalizeMips can then create the texture in it's step (1).
-		AdvanceTo(ETickState::FinalizeMips, ETickThread::Async);
+		// Once the FTextureMipDataProvider has set the mip data, UploadMips can then create the texture.
+		AdvanceTo(ETickState::UploadMips, ETickThread::Async);
 		return true;
 	}
 	else // No new mips? something is wrong.
@@ -70,7 +70,7 @@ bool FTexture2DMipAllocator_AsyncCreate::AllocateMips(
 // This gets called 2 times :
 // - Async : create the texture with the mips data
 // - Render : swap the results
-bool FTexture2DMipAllocator_AsyncCreate::FinalizeMips(const FTextureUpdateContext& Context, const FTextureUpdateSyncOptions& SyncOptions)
+bool FTexture2DMipAllocator_AsyncCreate::UploadMips(const FTextureUpdateContext& Context, const FTextureUpdateSyncOptions& SyncOptions)
 {
 	FRHITexture* Texture2DRHI = Context.Resource ? Context.Resource->GetTexture2DRHI() : nullptr;
 	if (!Texture2DRHI)
@@ -111,7 +111,7 @@ bool FTexture2DMipAllocator_AsyncCreate::FinalizeMips(const FTextureUpdateContex
 		ReleaseAllocatedMipData();
 
 		// Go to next step, on the renderthread.
-		AdvanceTo(ETickState::FinalizeMips, ETickThread::Render);
+		AdvanceTo(ETickState::UploadMips, ETickThread::Render);
 	}
 	// Step (2) : Copy the non initialized mips on the using RHICopySharedMips, must run on the renderthread.
 	else
@@ -121,12 +121,7 @@ bool FTexture2DMipAllocator_AsyncCreate::FinalizeMips(const FTextureUpdateContex
 			FRHICommandListExecutor::GetImmediateCommandList(),
 			Texture2DRHI,
 			IntermediateTextureRHI);
-		// Use the new texture resource for the texture asset, must run on the renderthread.
-		Context.Resource->FinalizeStreaming(IntermediateTextureRHI);
-		// No need for the intermediate texture anymore.
-		IntermediateTextureRHI.SafeRelease();
-
-		// Update complete, nothing more to do.
+		// Go to next step, on the renderthread.
 		AdvanceTo(ETickState::Done, ETickThread::None);
 	}
 	return true;
@@ -134,21 +129,14 @@ bool FTexture2DMipAllocator_AsyncCreate::FinalizeMips(const FTextureUpdateContex
 
 void FTexture2DMipAllocator_AsyncCreate::Cancel(const FTextureUpdateSyncOptions& SyncOptions)
 {
-	// Release the intermediate texture. If not null, this will be on the renderthread.
-	IntermediateTextureRHI.SafeRelease();
 	// Release the temporary mip data. Can be run on either renderthread or async threads.
 	ReleaseAllocatedMipData();
 }
 
 FTextureMipAllocator::ETickThread FTexture2DMipAllocator_AsyncCreate::GetCancelThread() const
 {
-	// If there is an  intermediate texture, it is safer to released on the renderthread.
-	if (IntermediateTextureRHI)
-	{
-		return ETickThread::Render;
-	}
-	// Otherwise, if there are only temporary mip data, they can be freed on any threads.
-	else if (FinalMipData.Num())
+	// If there are only temporary mip data, they can be freed on any threads.
+	if (FinalMipData.Num())
 	{
 		return ETickThread::Async;
 	}
diff --git a/Engine/Source/Runtime/Engine/Private/Streaming/Texture2DMipAllocator_AsyncCreate.h b/Engine/Source/Runtime/Engine/Private/Streaming/Texture2DMipAllocator_AsyncCreate.h
--- a/Engine/Source/Runtime/Engine/Private/Streaming/Texture2DMipAllocator_AsyncCreate.h
+++ b/Engine/Source/Runtime/Engine/Private/Streaming/Texture2DMipAllocator_AsyncCreate.h
@@ -31,7 +31,7 @@ public:
 	// ********************************************************
 
 	bool AllocateMips(const FTextureUpdateContext& Context, FTextureMipInfoArray& OutMipInfos, const FTextureUpdateSyncOptions& SyncOptions) final override;
-	bool FinalizeMips(const FTextureUpdateContext& Context, const FTextureUpdateSyncOptions& SyncOptions) final override;
+	bool UploadMips(const FTextureUpdateContext& Context, const FTextureUpdateSyncOptions& SyncOptions) final override;
 	void Cancel(const FTextureUpdateSyncOptions& SyncOptions) final override;
 	ETickThread GetCancelThread() const final override;
 
@@ -40,9 +40,6 @@ protected:
 	// Release the temporary buffers referenced in FinalMipData.
 	void ReleaseAllocatedMipData();
 
-	// The intermediate texture async created in the update process.
-	FTextureRHIRef IntermediateTextureRHI;
-
 	// The final resolution x and y of the texture 2d.
 	int32 FinalSizeX = 0;
 	int32 FinalSizeY = 0;
diff --git a/Engine/Source/Runtime/Engine/Private/Streaming/Texture2DMipAllocator_AsyncReallocate.cpp b/Engine/Source/Runtime/Engine/Private/Streaming/Texture2DMipAllocator_AsyncReallocate.cpp
--- a/Engine/Source/Runtime/Engine/Private/Streaming/Texture2DMipAllocator_AsyncReallocate.cpp
+++ b/Engine/Source/Runtime/Engine/Private/Streaming/Texture2DMipAllocator_AsyncReallocate.cpp
@@ -86,12 +86,12 @@ bool FTexture2DMipAllocator_AsyncReallocate::AllocateMips(
 		}
 
 		// New mips are ready to be unlocked by the FTextureMipDataProvider implementation.
-		AdvanceTo(ETickState::FinalizeMips, ETickThread::Render);
+		AdvanceTo(ETickState::UploadMips, ETickThread::Render);
 		return true;
 	}
 }
 
-bool FTexture2DMipAllocator_AsyncReallocate::FinalizeMips(const FTextureUpdateContext& Context, const FTextureUpdateSyncOptions& SyncOptions)
+bool FTexture2DMipAllocator_AsyncReallocate::UploadMips(const FTextureUpdateContext& Context, const FTextureUpdateSyncOptions& SyncOptions)
 {
 	if (!IntermediateTextureRHI)
 	{
@@ -100,12 +100,7 @@ bool FTexture2DMipAllocator_AsyncReallocate::FinalizeMips(const FTextureUpdateCo
 
 	// Unlock the mips so that the texture can be updated.
 	UnlockNewMips();
-	// Use the new texture resource for the texture asset, must run on the renderthread.
-	Context.Resource->FinalizeStreaming(IntermediateTextureRHI);
-	// No need for the intermediate texture anymore.
-	IntermediateTextureRHI.SafeRelease();
 
-	// Update complete, nothing more to do.
 	AdvanceTo(ETickState::Done, ETickThread::None);
 	return true;
 }
@@ -114,22 +109,12 @@ void FTexture2DMipAllocator_AsyncReallocate::Cancel(const FTextureUpdateSyncOpti
 {
 	// Unlock any locked mips.
 	UnlockNewMips();
-	// Release the intermediate texture.
-	IntermediateTextureRHI.SafeRelease();
 }
 
 FTextureMipAllocator::ETickThread FTexture2DMipAllocator_AsyncReallocate::GetCancelThread() const
 {
-	// If there is an intermediate texture and possibly locked mips. Unlock them and release it on the renderthread.
-	if (IntermediateTextureRHI)
-	{
-		return ETickThread::Render;
-	}
 	// Nothing to do.
-	else
-	{
-		return ETickThread::None;
-	}
+	return ETickThread::None;
 }
 
 // ****************************
diff --git a/Engine/Source/Runtime/Engine/Private/Streaming/Texture2DMipAllocator_AsyncReallocate.h b/Engine/Source/Runtime/Engine/Private/Streaming/Texture2DMipAllocator_AsyncReallocate.h
--- a/Engine/Source/Runtime/Engine/Private/Streaming/Texture2DMipAllocator_AsyncReallocate.h
+++ b/Engine/Source/Runtime/Engine/Private/Streaming/Texture2DMipAllocator_AsyncReallocate.h
@@ -29,7 +29,7 @@ public:
 	// ********************************************************
 
 	bool AllocateMips(const FTextureUpdateContext& Context, FTextureMipInfoArray& OutMipInfos, const FTextureUpdateSyncOptions& SyncOptions) final override;
-	bool FinalizeMips(const FTextureUpdateContext& Context, const FTextureUpdateSyncOptions& SyncOptions) final override;
+	bool UploadMips(const FTextureUpdateContext& Context, const FTextureUpdateSyncOptions& SyncOptions) final override;
 	void Cancel(const FTextureUpdateSyncOptions& SyncOptions) final override;
 	ETickThread GetCancelThread() const final override;
 
@@ -38,8 +38,6 @@ protected:
 	// Unlock the mips referenced in LockedMipIndices.
 	void UnlockNewMips();
 
-	// The intermediate texture created with RHIAsyncReallocateTexture2D.
-	FTextureRHIRef IntermediateTextureRHI;
 	// The list of mips that are currently locked.
 	TArray<int32, TInlineAllocator<MAX_TEXTURE_MIP_COUNT>> LockedMipIndices;
 };
diff --git a/Engine/Source/Runtime/Engine/Private/Streaming/TextureMipAllocator.cpp b/Engine/Source/Runtime/Engine/Private/Streaming/TextureMipAllocator.cpp
--- a/Engine/Source/Runtime/Engine/Private/Streaming/TextureMipAllocator.cpp
+++ b/Engine/Source/Runtime/Engine/Private/Streaming/TextureMipAllocator.cpp
@@ -5,6 +5,8 @@ TextureMipAllocator.cpp: Base class for implementing a mip allocation strategy u
 =============================================================================*/
 
 #include "Streaming/TextureMipAllocator.h"
+#include "Streaming/TextureMipDataProvider.h"
+#include "Rendering/StreamableTextureResource.h"
 #include "Engine/Texture.h"
 
 FTextureMipAllocator::FTextureMipAllocator(UTexture* Texture, ETickState InTickState, ETickThread InTickThread)
@@ -15,3 +17,38 @@ FTextureMipAllocator::FTextureMipAllocator(UTexture* Texture, ETickState InTickS
 	, NextTickThread(InTickThread) 
 {
 }
+
+FTextureMipAllocator::FTextureMipAllocator(const FTextureMipAllocator&) = default;
+FTextureMipAllocator::FTextureMipAllocator(FTextureMipAllocator&&) = default;
+FTextureMipAllocator::~FTextureMipAllocator() = default;
+
+bool FTextureMipAllocator::FinalizeMips(const FTextureUpdateContext& Context, const FTextureUpdateSyncOptions& SyncOptions)
+{
+	if (!ensureMsgf(IntermediateTextureRHI.IsValid(), TEXT("FTextureMipAllocator::FinalizeMips - Finalized called but a IntermediateTextureRHI is empty. This must be a valid texture reference before finalization")))
+	{
+		return false;
+	}
+
+	// Use the new texture resource for the texture asset, must run on the renderthread.
+	Context.Resource->FinalizeStreaming(IntermediateTextureRHI);
+
+	// No need for the intermediate texture anymore.
+	IntermediateTextureRHI.SafeRelease();
+	return true;
+}
+
+FTextureMipAllocator::ETickThread FTextureMipAllocator::ExecuteGetCancelThread() const
+{
+	if (IntermediateTextureRHI.IsValid())
+	{
+		return ETickThread::Render;
+	}
+
+	return GetCancelThread();
+}
+
+void FTextureMipAllocator::ExecuteCancel(const FTextureUpdateSyncOptions& SyncOptions)
+{
+	Cancel(SyncOptions);
+	IntermediateTextureRHI.SafeRelease();
+}
diff --git a/Engine/Source/Runtime/Engine/Private/Streaming/TextureMipAllocator.h b/Engine/Source/Runtime/Engine/Private/Streaming/TextureMipAllocator.h
--- a/Engine/Source/Runtime/Engine/Private/Streaming/TextureMipAllocator.h
+++ b/Engine/Source/Runtime/Engine/Private/Streaming/TextureMipAllocator.h
@@ -7,6 +7,7 @@ TextureMipAllocator.h: Base class for implementing a mip allocation strategy use
 #pragma once
 
 #include "CoreMinimal.h"
+#include "RHIFwd.h"
 #include "Streaming/TextureMipDataProvider.h"
 
 /**
@@ -21,7 +22,7 @@ public:
 	enum class ETickState : uint32
 	{
 		AllocateMips,	// Call AllocateMips(), which must provide destination memory for the mips to be written to.
-		FinalizeMips,	// Call FinalizeMips(), which must release temporary buffers and update the texture asset with the new mips.
+		UploadMips,		// Call UploadMips(), which must release temporary buffers and update/upload the texture to the gpu.
 		Done			// Nothing to do.
 	};
 
@@ -34,8 +35,9 @@ public:
 
 	// Constructor, defining the first tick step and thread.
 	FTextureMipAllocator(UTexture* Texture, ETickState InTickState, ETickThread InTickThread);
-
-	virtual ~FTextureMipAllocator() {}
+	FTextureMipAllocator(const FTextureMipAllocator&);
+	FTextureMipAllocator(FTextureMipAllocator&&);
+	virtual ~FTextureMipAllocator();
 
 	// Get the next tick state and thread for the mip allocator. Used by FTextureStreamIn to schedule correctly the update between FTextureMipAllocator and FTextureMipDataProvider.
 	FORCEINLINE ETickState GetNextTickState() const { return NextTickState; }
@@ -54,28 +56,46 @@ public:
 	virtual bool AllocateMips(const FTextureUpdateContext& Context, FTextureMipInfoArray& OutMipInfos, const FTextureUpdateSyncOptions& SyncOptions) = 0;
 	
 	/**
-	* Upload the mip data to the new mips and update the texture asset. Update is completed after this step and can not be cancelled afterward..
+	* Upload the new mip data to the GPU.  The
 	*
 	* @param Context - An update context constant throughout the FTextureStreamIn update. Gives things like which texture asset is updated and what mips are streamed in.
 	* @param SyncOptions - Different sync options to control when the next tick of FTextureStreamIn can be scheduled.
 	*
 	* Return true unless the texture update needs to be aborted because of invalid data or setup.
 	*/
-	virtual bool FinalizeMips(const FTextureUpdateContext& Context, const FTextureUpdateSyncOptions& SyncOptions) = 0;
+	virtual bool UploadMips(const FTextureUpdateContext& Context, const FTextureUpdateSyncOptions& SyncOptions) = 0;
 
 	/**
-	* Cancel the progression and release any temporary resources. 
+	* Update the texture resource with the new RHI texture. Update is completed after this step and can not be canceled afterward..
+	*
+	* @param Context - An update context constant throughout the FTextureStreamIn update. Gives things like which texture asset is updated and what mips are streamed in.
+	* @param SyncOptions - Different sync options to control when the next tick of FTextureStreamIn can be scheduled.
+	*
+	* Return true unless the texture update needs to be aborted because of invalid data or setup.
+	*/
+	bool FinalizeMips(const FTextureUpdateContext& Context, const FTextureUpdateSyncOptions& SyncOptions);
+
+	// Wrapper around virtual Cancel.
+	void ExecuteCancel(const FTextureUpdateSyncOptions& SyncOptions);
+
+	// Wrapper around virtual GetCancelThread.
+	ETickThread ExecuteGetCancelThread() const;
+
+	const FTextureRHIRef& GetIntermediateTextureRHI() const { return IntermediateTextureRHI; }
+
+protected:
+
+	/**
+	* Cancel the progression and release any temporary resources.
 	* Called within the FTextureStreamIn update when the stream in request is aborted or cannot complete correctly
 	* The Cancel() function is called on the thread returned by GetCancelThread().
 	*
 	* @param SyncOptions - Different sync options to control when the next tick of FTextureStreamIn can be scheduled.
 	*/
 	virtual void Cancel(const FTextureUpdateSyncOptions& SyncOptions) = 0;
-	
-	// Returns on which thread Cancel() must be called in the texture update (from FTextureStreamIn) to release ressources safely and correctly. ETickThread::None when ready to delete this.
-	virtual ETickThread GetCancelThread() const = 0;
 
-protected:
+	// Returns on which thread Cancel() must be called in the texture update (from FTextureStreamIn) to release resources safely and correctly. ETickThread::None when ready to delete this.
+	virtual ETickThread GetCancelThread() const = 0;
 
 	// Helper to set the next tick state and thread. Validates that the progression is coherent.
 	FORCEINLINE void AdvanceTo(ETickState InState, ETickThread InThread)
@@ -95,7 +115,10 @@ protected:
 	const int32 CurrentFirstLODIdx = INDEX_NONE;
 	// The requested first LOD resource index. With domain = [0, ResourceState.NumLODs[. NOT THE ASSET LOD INDEX!
 	const int32 PendingFirstLODIdx = INDEX_NONE;
-	
+
+	/**  The intermediate texture async created in the update process. */
+	FTextureRHIRef IntermediateTextureRHI;
+
 private:
 
 	// The next tick function that should be called.
diff --git a/Engine/Source/Runtime/Engine/Private/Streaming/TextureStreamIn.cpp b/Engine/Source/Runtime/Engine/Private/Streaming/TextureStreamIn.cpp
--- a/Engine/Source/Runtime/Engine/Private/Streaming/TextureStreamIn.cpp
+++ b/Engine/Source/Runtime/Engine/Private/Streaming/TextureStreamIn.cpp
@@ -145,7 +145,7 @@ FRenderAssetUpdate::EThreadType FTextureStreamIn::GetCancelThread() const
 	// If the mip data provider is fully cancelled, then cancel the mip allocator.
 	if (MipAllocator)
 	{
-		switch (MipAllocator->GetCancelThread())
+		switch (MipAllocator->ExecuteGetCancelThread())
 		{
 		case FTextureMipAllocator::ETickThread::Async:
 			return TT_Async;
@@ -212,9 +212,32 @@ bool FTextureStreamIn::DoPollMipData(const FContext& Context)
 	return true;
 }
 
-bool FTextureStreamIn::DoFinalizeNewMips(const FContext& Context)
+bool FTextureStreamIn::DoUploadNewMips(const FContext& Context)
 {
 	check(MipAllocator && IsSameThread(MipAllocator->GetNextTickThread(), Context.CurrentThread));
+	return MipAllocator->UploadMips(Context, SyncOptions);
+}
+
+bool FTextureStreamIn::DoPostprocessGPUMipData(const FContext& Context)
+{
+	for (TUniquePtr<FTextureMipDataProvider>& MipDataProvider : MipDataProviders)
+	{
+		check(MipDataProvider);
+		if (MipDataProvider->GetNextTickState() == FTextureMipDataProvider::ETickState::PostprocessGPUMipData
+			&& IsSameThread(MipDataProvider->GetNextTickThread(), Context.CurrentThread))
+		{
+			if (!MipDataProvider->PostprocessGPUMipData(Context, MipAllocator->GetIntermediateTextureRHI(), SyncOptions))
+			{
+				return false;
+			}
+		}
+	}
+	return true;
+}
+
+bool FTextureStreamIn::DoFinalizeNewMips(const FContext& Context)
+{
+	check(MipAllocator && IsSameThread(FTextureMipAllocator::ETickThread::Render, Context.CurrentThread));
 	return MipAllocator->FinalizeMips(Context, SyncOptions);
 }
 
@@ -364,45 +387,108 @@ void FTextureStreamIn::PollMipData(const FContext& Context)
 	else
 	{
 		bIsPollingMipData = false;
-
-		if (!IsCancelled())
+		bool bCancelled = IsCancelled();
+		if (!bCancelled)
 		{
-			NextThread = GetMipAllocatorThread(FTextureMipAllocator::ETickState::FinalizeMips);
+			NextThread = GetMipAllocatorThread(FTextureMipAllocator::ETickState::UploadMips);
 			// All mips must be handled before moving to next stage.
 			if (NextThread != TT_None)
 			{
-				PushTask(Context, NextThread, SRA_UPDATE_CALLBACK(FinalizeNewMips), GetCancelThread(), SRA_UPDATE_CALLBACK(Cancel));
+				PushTask(Context, NextThread, SRA_UPDATE_CALLBACK(UploadNewMips), GetCancelThread(), SRA_UPDATE_CALLBACK(Cancel));
 			}
-			else // Otherwise nothing to execute but schedule next phase (edge case).
+			else // Allocator does not specificy an UploadMips thread. This is a required operation before finalizing the texture.
 			{
-				FinalizeNewMips(Context);
+				MarkAsCancelled();
+				bCancelled = true;
 			}
 		}
-		else
+
+		if(bCancelled)
 		{
 			PushTask(Context, TT_None, nullptr, GetCancelThread(), SRA_UPDATE_CALLBACK(Cancel));
 		}
 	}
 }
 
-void FTextureStreamIn::FinalizeNewMips(const FContext& Context)
+void FTextureStreamIn::UploadNewMips(const FContext& Context)
 {
-	DECLARE_SCOPE_CYCLE_COUNTER(TEXT("FTextureStreamIn::FinalizeNewMips"), STAT_TextureStreamIn_FinalizeNewMips, STATGROUP_StreamingDetails);
+	DECLARE_SCOPE_CYCLE_COUNTER(TEXT("FTextureStreamIn::UploadNewMips"), STAT_TextureStreamIn_UploadNewMips, STATGROUP_StreamingDetails);
 
 	// Execute
-	if (!DoFinalizeNewMips(Context))
+	if (!DoUploadNewMips(Context))
 	{
 		MarkAsCancelled();
 	}
 
 	// Schedule the next update step.
-	EThreadType NextThread = GetMipAllocatorThread(FTextureMipAllocator::ETickState::FinalizeMips);
+	EThreadType NextThread = GetMipAllocatorThread(FTextureMipAllocator::ETickState::UploadMips);
 
 	if (NextThread != TT_None) // Loop on this state.
 	{
-		PushTask(Context, NextThread, SRA_UPDATE_CALLBACK(FinalizeNewMips), GetCancelThread(), SRA_UPDATE_CALLBACK(Cancel));
+		PushTask(Context, NextThread, SRA_UPDATE_CALLBACK(UploadNewMips), GetCancelThread(), SRA_UPDATE_CALLBACK(Cancel));
 	}
 	else if (!IsCancelled())
+	{
+		
+		NextThread = GetMipDataProviderThread(FTextureMipDataProvider::ETickState::PostprocessGPUMipData);
+		// All mips must be handled before moving to next stage.
+		if (NextThread != TT_None)
+		{
+			bIsPostprocessingGPUMipData = true;
+			PushTask(Context, NextThread, SRA_UPDATE_CALLBACK(PostprocessGPUMipData), GetCancelThread(), SRA_UPDATE_CALLBACK(Cancel));
+		}
+		else // Otherwise nothing to execute but schedule next phase (edge case).
+		{
+			PostprocessGPUMipData(Context);
+		}
+	}
+	else
+	{
+		PushTask(Context, TT_None, nullptr, GetCancelThread(), SRA_UPDATE_CALLBACK(Cancel));
+	}
+}
+
+void FTextureStreamIn::PostprocessGPUMipData(const FContext& Context)
+{
+	DECLARE_SCOPE_CYCLE_COUNTER(TEXT("FTextureStreamIn::PostprocessGPUMipData"), STAT_TextureStreamIn_PostprocessGPUMipData, STATGROUP_StreamingDetails);
+
+	// Execute
+	if (!DoPostprocessGPUMipData(Context))
+	{
+		MarkAsCancelled();
+	}
+
+	// Schedule the next update step.
+	EThreadType NextThread = GetMipDataProviderThread(FTextureMipDataProvider::ETickState::PostprocessGPUMipData);
+
+	if (NextThread != TT_None) // Loop on this state.
+	{
+		PushTask(Context, NextThread, SRA_UPDATE_CALLBACK(PostprocessGPUMipData), GetCancelThread(), SRA_UPDATE_CALLBACK(Cancel));
+	}
+	else if (!IsCancelled())
+	{
+		bIsPostprocessingGPUMipData = false;
+		// The finalize task should always fire on the render thread
+		PushTask(Context, EThreadType::TT_Render, SRA_UPDATE_CALLBACK(FinalizeNewMips), GetCancelThread(), SRA_UPDATE_CALLBACK(Cancel));
+	}
+	else
+	{
+		PushTask(Context, TT_None, nullptr, GetCancelThread(), SRA_UPDATE_CALLBACK(Cancel));
+	}
+}
+
+void FTextureStreamIn::FinalizeNewMips(const FContext& Context)
+{
+	DECLARE_SCOPE_CYCLE_COUNTER(TEXT("FTextureStreamIn::FinalizeNewMips"), STAT_TextureStreamIn_FinalizeNewMips, STATGROUP_StreamingDetails);
+
+	// Execute
+	if (!DoFinalizeNewMips(Context))
+	{
+		MarkAsCancelled();
+	}
+
+	// Schedule the next update step.
+	if (!IsCancelled())
 	{
 		// Finalize new mips completed successfully.
 		MarkAsSuccessfullyFinished();
@@ -411,7 +497,7 @@ void FTextureStreamIn::FinalizeNewMips(const FContext& Context)
 		MipAllocator.Reset();
 		MipInfos.Empty();
 
-		NextThread = GetMipDataProviderThread(FTextureMipDataProvider::ETickState::CleanUp);
+		EThreadType NextThread = GetMipDataProviderThread(FTextureMipDataProvider::ETickState::CleanUp);
 		// All mips must be handled before moving to next stage.
 		if (NextThread != TT_None)
 		{
@@ -424,7 +510,7 @@ void FTextureStreamIn::FinalizeNewMips(const FContext& Context)
 	}
 	else
 	{
-		// Mark as successfully finished even if cancelled to keep the CachedSRRState in sync with the RHI texture.
+		// Mark as successfully finished even if canceled to keep the CachedSRRState in sync with the RHI texture.
 		if (MipAllocator->GetNextTickState() == FTextureMipAllocator::ETickState::Done)
 		{
 			MarkAsSuccessfullyFinished();
@@ -488,11 +574,11 @@ void FTextureStreamIn::Cancel(const FContext& Context)
 	// Then cancel the mip allocator.
 	if (MipAllocator)
 	{
-		if (IsSameThread(MipAllocator->GetCancelThread(), Context.CurrentThread))
+		if (IsSameThread(MipAllocator->ExecuteGetCancelThread(), Context.CurrentThread))
 		{
-			MipAllocator->Cancel(SyncOptions);
+			MipAllocator->ExecuteCancel(SyncOptions);
 		}
-		switch (MipAllocator->GetCancelThread())
+		switch (MipAllocator->ExecuteGetCancelThread())
 		{
 		case FTextureMipAllocator::ETickThread::Async:
 			NextThread = TT_Async;
@@ -515,11 +601,11 @@ void FTextureStreamIn::Abort()
 {
 	MarkAsCancelled();
 
-	if (bIsPollingMipData)
+	if (bIsPollingMipData || bIsPostprocessingGPUMipData)
 	{
 		// Prevent the update from being considered done before this is finished.
-		// By checking that it was not already cancelled, we make sure this doesn't get called twice.
-		(new FAsyncAbortPollMipsTask(this))->StartBackgroundTask();
+		// By checking that it was not already canceled, we make sure this doesn't get called twice.
+		(new FAsyncAbortPollMipsTask(this, bIsPollingMipData, bIsPostprocessingGPUMipData))->StartBackgroundTask();
 	}
 }
 
@@ -533,7 +619,14 @@ void FTextureStreamIn::FAbortPollMipsTask::DoWork()
 	for (TUniquePtr<FTextureMipDataProvider>& MipDataProvider : PendingUpdate->MipDataProviders)
 	{
 		check(MipDataProvider);
-		MipDataProvider->AbortPollMips();
+		if(bIsPollingMipData)
+		{
+			MipDataProvider->AbortPollMips();
+		}
+		else if (bIsPostprocessingGPUMipData)
+		{
+			MipDataProvider->AbortPostprocessGPUMipData();
+		}
 	}
 	PendingUpdate->DoUnlock(PreviousTaskState);
 }
diff --git a/Engine/Source/Runtime/Engine/Private/Streaming/TextureStreamIn.h b/Engine/Source/Runtime/Engine/Private/Streaming/TextureStreamIn.h
--- a/Engine/Source/Runtime/Engine/Private/Streaming/TextureStreamIn.h
+++ b/Engine/Source/Runtime/Engine/Private/Streaming/TextureStreamIn.h
@@ -41,7 +41,13 @@ public:
 	// Route FTextureMipDataProvider::PollMips()
 	void PollMipData(const FContext& Context);
 
-	// Route FTextureMipAllocator::FinalizeMips()
+	// Route FTextureMipAllocator::UploadMips()
+	void UploadNewMips(const FContext& Context);
+
+	// Route FTextureMipDataProvider::PostprocessGPUMipData
+	void PostprocessGPUMipData(const FContext& Context);
+
+	// Update the texture resource with the new RHI texture. Update is completed after this step and can not be canceled afterward.
 	void FinalizeNewMips(const FContext& Context);
 
 	// Route FTextureMipDataProvider::CleanUp()
@@ -69,6 +75,8 @@ protected:
 	bool DoAllocateNewMips(const FContext& Context);
 	void DoGetMipData(const FContext& Context);
 	bool DoPollMipData(const FContext& Context);
+	bool DoUploadNewMips(const FContext& Context);
+	bool DoPostprocessGPUMipData(const FContext& Context);
 	bool DoFinalizeNewMips(const FContext& Context);
 	void DoCleanUpMipDataProviders(const FContext& Context);
 
@@ -83,7 +91,12 @@ private:
 	class FAbortPollMipsTask : public FNonAbandonableTask
 	{
 	public:
-		FAbortPollMipsTask(FTextureStreamIn* InPendingUpdate) : PendingUpdate(InPendingUpdate) {}
+		FAbortPollMipsTask(FTextureStreamIn* InPendingUpdate, bool bInIsPollingMipData, bool bInIsPostprocessingGPUMipData) 
+			: PendingUpdate(InPendingUpdate)
+			, bIsPollingMipData(bInIsPollingMipData)
+			, bIsPostprocessingGPUMipData(bInIsPostprocessingGPUMipData)
+		{
+		}
 		void DoWork();
 		FORCEINLINE TStatId GetStatId() const
 		{
@@ -91,9 +104,12 @@ private:
 		}
 	protected:
 		TRefCountPtr<FTextureStreamIn> PendingUpdate;
+		bool bIsPollingMipData;
+		bool bIsPostprocessingGPUMipData;
 	};
 
 	bool bIsPollingMipData = false; // Used to prevent indirecting any pointer in the abort task.
+	bool bIsPostprocessingGPUMipData = false; // Used to prevent indirecting any pointer in the abort task.
 	typedef FAutoDeleteAsyncTask<FAbortPollMipsTask> FAsyncAbortPollMipsTask;
 	friend class FAbortPollMipsTask;
 
diff --git a/Engine/Source/Runtime/Engine/Private/Streaming/TextureStreamOut.cpp b/Engine/Source/Runtime/Engine/Private/Streaming/TextureStreamOut.cpp
--- a/Engine/Source/Runtime/Engine/Private/Streaming/TextureStreamOut.cpp
+++ b/Engine/Source/Runtime/Engine/Private/Streaming/TextureStreamOut.cpp
@@ -82,7 +82,7 @@ FRenderAssetUpdate::EThreadType FTextureStreamOut::GetCancelThread() const
 	// Cancel the mip allocator.
 	if (MipAllocator)
 	{
-		switch (MipAllocator->GetCancelThread())
+		switch (MipAllocator->ExecuteGetCancelThread())
 		{
 		case FTextureMipAllocator::ETickThread::Async:
 			return TT_Async;
@@ -106,9 +106,15 @@ bool FTextureStreamOut::DoAllocateNewMips(const FContext& Context)
 	return MipAllocator->AllocateMips(Context, MipInfos, SyncOptions);
 }
 
-bool FTextureStreamOut::DoFinalizeNewMips(const FContext& Context)
+bool FTextureStreamOut::DoUploadNewMips(const FContext& Context)
 {
 	check(MipAllocator && IsSameThread(MipAllocator->GetNextTickThread(), Context.CurrentThread));
+	return MipAllocator->UploadMips(Context, SyncOptions);
+}
+
+bool FTextureStreamOut::DoFinalizeNewMips(const FContext& Context)
+{
+	check(MipAllocator && IsSameThread(FTextureMipAllocator::ETickThread::Render, Context.CurrentThread));
 	return MipAllocator->FinalizeMips(Context, SyncOptions);
 }
 
@@ -118,7 +124,7 @@ bool FTextureStreamOut::DoFinalizeNewMips(const FContext& Context)
 
 void FTextureStreamOut::AllocateNewMips(const FContext& Context)
 {
-	DECLARE_SCOPE_CYCLE_COUNTER(TEXT("FTextureStreamOut::AllocateNewMips"), STAT_TextureStreamIn_AllocateNewMips, STATGROUP_StreamingDetails);
+	DECLARE_SCOPE_CYCLE_COUNTER(TEXT("FTextureStreamOut::AllocateNewMips"), STAT_TextureStreamOut_AllocateNewMips, STATGROUP_StreamingDetails);
 
 	// Execute
 	if (!DoAllocateNewMips(Context))
@@ -134,15 +140,15 @@ void FTextureStreamOut::AllocateNewMips(const FContext& Context)
 	}
 	else if (!IsCancelled())
 	{
-		NextThread = GetMipAllocatorThread(FTextureMipAllocator::ETickState::FinalizeMips);
+		NextThread = GetMipAllocatorThread(FTextureMipAllocator::ETickState::UploadMips);
 		// All mips must be handled before moving to next stage.
 		if (NextThread != TT_None)
 		{
-			PushTask(Context, NextThread, SRA_UPDATE_CALLBACK(FinalizeNewMips), GetCancelThread(), SRA_UPDATE_CALLBACK(Cancel));
+			PushTask(Context, NextThread, SRA_UPDATE_CALLBACK(UploadNewMips), GetCancelThread(), SRA_UPDATE_CALLBACK(Cancel));
 		}
 		else // Otherwise nothing to execute but schedule next phase (edge case).
 		{
-			FinalizeNewMips(Context);
+			PushTask(Context, EThreadType::TT_Render, SRA_UPDATE_CALLBACK(FinalizeNewMips), GetCancelThread(), SRA_UPDATE_CALLBACK(Cancel));
 		}
 	}
 	else
@@ -151,47 +157,67 @@ void FTextureStreamOut::AllocateNewMips(const FContext& Context)
 	}
 }
 
-void FTextureStreamOut::FinalizeNewMips(const FContext& Context)
+void FTextureStreamOut::UploadNewMips(const FContext& Context)
 {
-	DECLARE_SCOPE_CYCLE_COUNTER(TEXT("FTextureStreamOut::FinalizeNewMips"), STAT_TextureStreamIn_FinalizeNewMips, STATGROUP_StreamingDetails);
+	DECLARE_SCOPE_CYCLE_COUNTER(TEXT("FTextureStreamOut::UploadNewMips"), STAT_TextureStreamOut_UploadNewMips, STATGROUP_StreamingDetails);
 
 	// Execute
-	if (DoFinalizeNewMips(Context))
-	{
-		MarkAsSuccessfullyFinished();
-	}
-	else
+	if (!DoUploadNewMips(Context))
 	{
 		MarkAsCancelled();
 	}
 
 	// Schedule the next update step.
-	EThreadType NextThread = GetMipAllocatorThread(FTextureMipAllocator::ETickState::FinalizeMips);
+	EThreadType NextThread = GetMipAllocatorThread(FTextureMipAllocator::ETickState::UploadMips);
 	if (NextThread != TT_None) // Loop on this state.
 	{
-		PushTask(Context, NextThread, SRA_UPDATE_CALLBACK(FinalizeNewMips), GetCancelThread(), SRA_UPDATE_CALLBACK(Cancel));
+		PushTask(Context, NextThread, SRA_UPDATE_CALLBACK(UploadNewMips), GetCancelThread(), SRA_UPDATE_CALLBACK(Cancel));
 	}
 	else
 	{
-		// Release the mip allocator.
-		MipAllocator.Reset();
+		if (!IsCancelled())
+		{
+			PushTask(Context, EThreadType::TT_Render, SRA_UPDATE_CALLBACK(FinalizeNewMips), GetCancelThread(), SRA_UPDATE_CALLBACK(Cancel));
+		}
+		else
+		{
+			PushTask(Context, TT_None, nullptr, GetCancelThread(), SRA_UPDATE_CALLBACK(Cancel));
+		}
 	}
 }
 
+void FTextureStreamOut::FinalizeNewMips(const FContext& Context)
+{
+	DECLARE_SCOPE_CYCLE_COUNTER(TEXT("FTextureStreamOut::FinalizeNewMips"), STAT_TextureStreamOut_FinalizeNewMips, STATGROUP_StreamingDetails);
+
+	// Execute
+	if (DoFinalizeNewMips(Context))
+	{
+		MarkAsSuccessfullyFinished();
+	}
+	else
+	{
+		MarkAsCancelled();
+	}
+
+	/// Release the mip allocator.
+	MipAllocator.Reset();
+}
+
 void FTextureStreamOut::Cancel(const FContext& Context)
 {
-	DECLARE_SCOPE_CYCLE_COUNTER(TEXT("FTextureStreamOut::Cancel"), STAT_TextureStreamIn_Cancel, STATGROUP_StreamingDetails);
+	DECLARE_SCOPE_CYCLE_COUNTER(TEXT("FTextureStreamOut::Cancel"), STAT_TextureStreamOut_Cancel, STATGROUP_StreamingDetails);
 
 	EThreadType NextThread = TT_None;
 
 	// Then cancel the mip allocator.
 	if (MipAllocator)
 	{
-		if (IsSameThread(MipAllocator->GetCancelThread(), Context.CurrentThread))
+		if (IsSameThread(MipAllocator->ExecuteGetCancelThread(), Context.CurrentThread))
 		{
-			MipAllocator->Cancel(SyncOptions);
+			MipAllocator->ExecuteCancel(SyncOptions);
 		}
-		switch (MipAllocator->GetCancelThread())
+		switch (MipAllocator->ExecuteGetCancelThread())
 		{
 		case FTextureMipAllocator::ETickThread::Async:
 			NextThread = TT_Async;
diff --git a/Engine/Source/Runtime/Engine/Private/Streaming/TextureStreamOut.h b/Engine/Source/Runtime/Engine/Private/Streaming/TextureStreamOut.h
--- a/Engine/Source/Runtime/Engine/Private/Streaming/TextureStreamOut.h
+++ b/Engine/Source/Runtime/Engine/Private/Streaming/TextureStreamOut.h
@@ -24,6 +24,9 @@ public:
 	
 	// Route FTextureMipAllocator::AllocateMips()
 	void AllocateNewMips(const FContext& Context);
+
+	// Route FTextureMipAllocator::UploadMips()
+	void UploadNewMips(const FContext& Context);
 	
 	// Route FTextureMipAllocator::FinalizeMips()
 	void FinalizeNewMips(const FContext& Context);
@@ -39,6 +42,7 @@ protected:
 
 	// The actual implementation of the update steps which redirect to FTextureMipAllocator and FTextureMipDataProvider.
 	bool DoAllocateNewMips(const FContext& Context);
+	bool DoUploadNewMips(const FContext& Context);
 	bool DoFinalizeNewMips(const FContext& Context);
 
 	EThreadType GetMipAllocatorThread(FTextureMipAllocator::ETickState TickState) const;
diff --git a/Engine/Source/Runtime/Engine/Private/Streaming/VolumeTextureStreaming.cpp b/Engine/Source/Runtime/Engine/Private/Streaming/VolumeTextureStreaming.cpp
--- a/Engine/Source/Runtime/Engine/Private/Streaming/VolumeTextureStreaming.cpp
+++ b/Engine/Source/Runtime/Engine/Private/Streaming/VolumeTextureStreaming.cpp
@@ -68,11 +68,11 @@ bool FVolumeTextureMipAllocator_Reallocate::AllocateMips(const FTextureUpdateCon
 		}
 	}
 
-	AdvanceTo(ETickState::FinalizeMips, ETickThread::Render);
+	AdvanceTo(ETickState::UploadMips, ETickThread::Render);
 	return true;
 }
 
-bool FVolumeTextureMipAllocator_Reallocate::FinalizeMips(const FTextureUpdateContext& Context, const FTextureUpdateSyncOptions& SyncOptions)
+bool FVolumeTextureMipAllocator_Reallocate::UploadMips(const FTextureUpdateContext& Context, const FTextureUpdateSyncOptions& SyncOptions)
 {
 	EPixelFormat PixelFormat = Context.Resource->GetPixelFormat();
 
@@ -134,20 +134,12 @@ bool FVolumeTextureMipAllocator_Reallocate::FinalizeMips(const FTextureUpdateCon
 	}
 	StreamedInMipData.Discard();
 
-	// Use the new texture resource for the texture asset, must run on the renderthread.
-	Context.Resource->FinalizeStreaming(IntermediateTextureRHI);
-	// No need for the intermediate texture anymore.
-	IntermediateTextureRHI.SafeRelease();
-
-	// Update is complete, nothing more to do.
 	AdvanceTo(ETickState::Done, ETickThread::None);
 	return true;
 }
 
 void FVolumeTextureMipAllocator_Reallocate::Cancel(const FTextureUpdateSyncOptions& SyncOptions)
 {
-	// Release the intermediate texture. If not null, this will be on the renderthread.
-	IntermediateTextureRHI.SafeRelease();
 	// Release the temporary mip data. Can be run on either renderthread or async threads.
 	StreamedInMipData.Discard();
 }
@@ -155,12 +147,7 @@ void FVolumeTextureMipAllocator_Reallocate::Cancel(const FTextureUpdateSyncOptio
 FTextureMipAllocator::ETickThread FVolumeTextureMipAllocator_Reallocate::GetCancelThread() const
 {
 	// If there is an  intermediate texture, it is safer to released on the renderthread.
-	if (IntermediateTextureRHI)
-	{
-		return ETickThread::Render;
-	}
-	// Otherwise, if there are only temporary mip data, they can be freed on any threads.
-	else if (StreamedInMipData.GetResourceBulkData())
+	if (StreamedInMipData.GetResourceBulkData())
 	{
 		return ETickThread::Async;
 	}
diff --git a/Engine/Source/Runtime/Engine/Private/Streaming/VolumeTextureStreaming.h b/Engine/Source/Runtime/Engine/Private/Streaming/VolumeTextureStreaming.h
--- a/Engine/Source/Runtime/Engine/Private/Streaming/VolumeTextureStreaming.h
+++ b/Engine/Source/Runtime/Engine/Private/Streaming/VolumeTextureStreaming.h
@@ -35,15 +35,11 @@ public:
 	// ********************************************************
 
 	bool AllocateMips(const FTextureUpdateContext& Context, FTextureMipInfoArray& OutMipInfos, const FTextureUpdateSyncOptions& SyncOptions) final override;
-	bool FinalizeMips(const FTextureUpdateContext& Context, const FTextureUpdateSyncOptions& SyncOptions) final override;
+	bool UploadMips(const FTextureUpdateContext& Context, const FTextureUpdateSyncOptions& SyncOptions) final override;
 	void Cancel(const FTextureUpdateSyncOptions& SyncOptions) final override;
 	ETickThread GetCancelThread() const final override;
 
 protected:
-
-	// The intermediate texture async created in the update process.
-	FTextureRHIRef IntermediateTextureRHI;
-
 	// The temporary main memory allocations holding the mip data
 	FVolumeTextureBulkData StreamedInMipData;
 };
diff --git a/Engine/Source/Runtime/Engine/Public/Streaming/TextureMipDataProvider.h b/Engine/Source/Runtime/Engine/Public/Streaming/TextureMipDataProvider.h
--- a/Engine/Source/Runtime/Engine/Public/Streaming/TextureMipDataProvider.h
+++ b/Engine/Source/Runtime/Engine/Public/Streaming/TextureMipDataProvider.h
@@ -10,7 +10,7 @@ TextureMipDataProvider.h: Base class for providing the mip data used by FTexture
 #include "HAL/ThreadSafeCounter.h"
 #include "PixelFormat.h"
 #include "RHIDefinitions.h"
-#include "RHIDefinitions.h"
+#include "RHIFwd.h"
 #include "Streaming/StreamableRenderResourceState.h"
 
 class UTexture;
@@ -94,10 +94,11 @@ public:
 
 	enum class ETickState : uint32
 	{
-		Init,			// Call Init(), which must initialize the mip data provider so that it can efficently provide the mips in GetMips.
-		GetMips,		// Call GetMips(), which must fill mip data from highest to lowest (all mips might not be available thought).
-		PollMips,		// Call PollMips(), which stalls the mip update until the mip have been set.
-		CleanUp,		// Call CleanUp(), which must release all temporary resource created in the update process.
+		Init,					// Call Init(), which must initialize the mip data provider so that it can efficiently provide the mips in GetMips.
+		GetMips,				// Call GetMips(), which must fill mip data from highest to lowest (all mips might not be available thought).
+		PollMips,				// Call PollMips(), which stalls the mip update until the mip have been set.
+		PostprocessGPUMipData,	// Call PostprocessGPUMipData(), which allows for mip updates after the texture has been uploaded to the GPU
+		CleanUp,				// Call CleanUp(), which must release all temporary resource created in the update process.
 		Done
 	};
 
@@ -149,14 +150,30 @@ public:
 	*/
 	virtual bool PollMips(const FTextureUpdateSyncOptions& SyncOptions) = 0;
 
+	/*
+	* Once the texture has been uploaded to the GPU, the mip provider has the opportunity to perform any modifications necessary
+	* before the texture is finalized and applied to the texture resource
+	*/
+	virtual bool PostprocessGPUMipData(const FTextureUpdateContext& Context, const FTextureRHIRef& NewTextureRHI, const FTextureUpdateSyncOptions& SyncOptions) { return true; }
+
 	/**
 	* Abort anything that could be stalling the update in PollMips(). Called from an async task.
 	* Expected behavior is to cancel any pending IO operations. 
-	* Called when the FTextureStreamIn update is cancelled and the current step is to be waiting on upon PollMips().
+	* Called when the FTextureStreamIn update is canceled and the current step is to be waiting on upon PollMips().
 	* This is because Cancel() can only be executed when any pending operations are completed (see FTextureUpdateSyncOptions).
+	* Will be used if executing PollMips
 	*/
 	virtual void AbortPollMips() {}
 
+	/**
+	* Abort anything that could be stalling the update in PostprocessGPUMipData(). Called from an async task.
+	* Expected behavior is to cancel any pending IO operations.
+	* Called when the FTextureStreamIn update is canceled and the current step is to be waiting on upon AbortPostprocessGPUMipData().
+	* This is because Cancel() can only be executed when any pending operations are completed (see FTextureUpdateSyncOptions).
+	* Will be used if executing PostprocessGPUMipData
+	*/
+	virtual void AbortPostprocessGPUMipData() {}
+
 	/**
 	* Release any temporary data and objects that where used for the update. Final step executed after the texture has been updated correctly.
 	*
