[Selection] Allow for BP ISM components to opt into per hit proxy generation, so that PCG workflows can set this on BP ISM components

diff --git a/Engine/Source/Runtime/Engine/Classes/Components/PrimitiveComponent.h b/Engine/Source/Runtime/Engine/Classes/Components/PrimitiveComponent.h
--- a/Engine/Source/Runtime/Engine/Classes/Components/PrimitiveComponent.h
+++ b/Engine/Source/Runtime/Engine/Classes/Components/PrimitiveComponent.h
@@ -470,8 +470,10 @@ public:
 	UPROPERTY(EditAnywhere, AdvancedDisplay, BlueprintReadOnly, Category=TextureStreaming)
 	uint8 bForceMipStreaming:1;
 
-	/** If true a hit-proxy will be generated for each instance of instanced static meshes */
-	UPROPERTY()
+	/** If true a hit-proxy will be generated for each instance of instanced static meshes.
+	 *  This allows for other systems, like selection in viewports, to function at the individual instance.
+	 */
+	UPROPERTY(EditDefaultsOnly, BlueprintReadWrite, AdvancedDisplay, Category=Rendering)
 	uint8 bHasPerInstanceHitProxies:1;
 
 	// Lighting flags
diff --git a/Engine/Source/Runtime/Engine/Private/InstancedStaticMesh.cpp b/Engine/Source/Runtime/Engine/Private/InstancedStaticMesh.cpp
--- a/Engine/Source/Runtime/Engine/Private/InstancedStaticMesh.cpp
+++ b/Engine/Source/Runtime/Engine/Private/InstancedStaticMesh.cpp
@@ -2319,7 +2319,7 @@ FMatrix UInstancedStaticMeshComponent::GetRenderMatrix() const
 
 void UInstancedStaticMeshComponent::CreateHitProxyData(TArray<TRefCountPtr<HHitProxy>>& HitProxies)
 {
-	if (GIsEditor && bHasPerInstanceHitProxies)
+	if (GIsEditor && bHasPerInstanceHitProxies && (CreationMethod == EComponentCreationMethod::Instance))
 	{
 		QUICK_SCOPE_CYCLE_COUNTER(STAT_UInstancedStaticMeshComponent_CreateHitProxyData);
 		
