Addressed issue when an infinite loop can be created when flushing montages (which calls notify state end events if any was active) due to calling montage play on a notify state end event.

#rnx

diff --git a/Engine/Source/Runtime/Engine/Classes/Animation/AnimInstance.h b/Engine/Source/Runtime/Engine/Classes/Animation/AnimInstance.h
--- a/Engine/Source/Runtime/Engine/Classes/Animation/AnimInstance.h
+++ b/Engine/Source/Runtime/Engine/Classes/Animation/AnimInstance.h
@@ -1455,6 +1455,10 @@ public:
 	FGraphTraversalCounter DebugDataCounter;
 
 private:
+
+	// Used to prevent indefinitely flushing montages on a single frame. 
+	uint32 LastMontageFlushFrame = 0u;
+	
 	TMap<FName, FMontageActiveSlotTracker> SlotWeightTracker;
 	TMap<FName, FSimpleMulticastDelegate> ExternalNotifyHandlers;
 
diff --git a/Engine/Source/Runtime/Engine/Private/Animation/AnimInstance.cpp b/Engine/Source/Runtime/Engine/Private/Animation/AnimInstance.cpp
--- a/Engine/Source/Runtime/Engine/Private/Animation/AnimInstance.cpp
+++ b/Engine/Source/Runtime/Engine/Private/Animation/AnimInstance.cpp
@@ -1972,21 +1972,36 @@ void UAnimInstance::ConditionalFlushCompletedMontages()
 	const bool bTickedThisFrame = MeshComp->PoseTickedThisFrame();
 	const bool bShouldTickAnimation = MeshComp->ShouldTickAnimation();
 	const bool bShouldTickPose = MeshComp->ShouldTickPose();
-	const bool bShouldFlush = !bTickedThisFrame && (!bShouldTickAnimation || !bShouldTickPose); 
+	const bool bShouldFlush = !bTickedThisFrame && (!bShouldTickAnimation || !bShouldTickPose);
+	
 	if (bShouldFlush)
 	{
-		for (int32 InstanceIndex = MontageInstances.Num() - 1; InstanceIndex >= 0; InstanceIndex--)
+		uint32 PrevMontageFlushFrame = LastMontageFlushFrame;
+		
+		// Don't care about roll over, just care about uniqueness (and 32-bits should give plenty).
+		LastMontageFlushFrame = static_cast<uint32>(GFrameCounter); 
+
+		// If we already flushed montages this frame, then there is no need to do so again
+		const bool bFlushedMontagesThisFrame = LastMontageFlushFrame == PrevMontageFlushFrame;
+		
+		if (!bFlushedMontagesThisFrame)
 		{
-			FAnimMontageInstance* MontageInstance = MontageInstances[InstanceIndex];
-			if (MontageInstance && MontageInstance->IsValid() && MontageInstance->IsStopped() && MontageInstance->GetBlend().IsComplete())
+			for (int32 InstanceIndex = MontageInstances.Num() - 1; InstanceIndex >= 0; InstanceIndex--)
 			{
-				// Need this to trigger Montage ended events.
-				MontageInstance->Terminate();
-
-				// Make sure we've cleared our references before deleting memory. Terminate might miss this call.
-				ClearMontageInstanceReferences(*MontageInstance);
-				delete MontageInstance;
-				MontageInstances.RemoveAt(InstanceIndex);
+				if (MontageInstances.IsValidIndex(InstanceIndex))
+				{
+					FAnimMontageInstance* MontageInstance = MontageInstances[InstanceIndex];
+					if (MontageInstance && MontageInstance->IsValid() && MontageInstance->IsStopped() && MontageInstance->GetBlend().IsComplete())
+					{
+						// Need this to trigger Montage ended events.
+						MontageInstance->Terminate();
+
+						// Make sure we've cleared our references before deleting memory. Terminate might miss this call.
+						ClearMontageInstanceReferences(*MontageInstance);
+						delete MontageInstance;
+						MontageInstances.RemoveAt(InstanceIndex);
+					}	
+				}
 			}
 		}
 	}
