Force FSoundWavePCMWriter to explicitly update it's codec to PCM on creation
(this avoids attempting to use a deprecated decompression path).

diff --git a/Engine/Source/Runtime/Engine/Private/SampleBufferIO.cpp b/Engine/Source/Runtime/Engine/Private/SampleBufferIO.cpp
--- a/Engine/Source/Runtime/Engine/Private/SampleBufferIO.cpp
+++ b/Engine/Source/Runtime/Engine/Private/SampleBufferIO.cpp
@@ -486,6 +486,8 @@ namespace Audio
 		SerializeSoundWaveToAsset();
 	}
 
+	// todo: this function needs to be revisited to have behavior similar to SerializeSoundWaveToAsset() in non-editor environments
+	// todo: for now this only updates simple members and cleans up data and clears RawPCMData to avoid taking a deprecated decode path
 	void FSoundWavePCMWriter::ApplyBufferToSoundWave()
 	{
 		// Since we just want to replace the PCM data to save it to disk. We don't need to compute anything platformdata related.
@@ -500,10 +502,8 @@ namespace Audio
 		if (CurrentSoundWave->RawPCMData != nullptr)
 		{
 			FMemory::Free(CurrentSoundWave->RawPCMData);
+			CurrentSoundWave->RawPCMData = nullptr;
 		}
-
-		CurrentSoundWave->RawPCMData = (uint8*)FMemory::Malloc(CurrentSoundWave->RawPCMDataSize);
-		FMemory::Memcpy(CurrentSoundWave->RawPCMData, CurrentBuffer.GetData(), CurrentSoundWave->RawPCMDataSize);
 	}
 
 	void FSoundWavePCMWriter::SerializeSoundWaveToAsset()
@@ -706,6 +706,19 @@ namespace Audio
 				break;
 			}
 
+		// force update the compression type
+		if (IsInAudioThread())
+		{
+			Writer->CurrentSoundWave->SetSoundAssetCompressionType(ESoundAssetCompressionType::PCM);
+		}
+		
+		// clear out RawPCMData
+		check(Writer->CurrentSoundWave);
+		if (Writer->CurrentSoundWave->RawPCMData)
+		{
+			FMemory::Free(Writer->CurrentSoundWave->RawPCMData);
+			Writer->CurrentSoundWave->RawPCMData = nullptr;
+		}
 
 		// Capture our callback and perform it on the game thread:
 		const USoundWave* SoundWave = Writer->CurrentSoundWave;
diff --git a/Engine/Source/Runtime/Engine/Private/SoundWave.cpp b/Engine/Source/Runtime/Engine/Private/SoundWave.cpp
--- a/Engine/Source/Runtime/Engine/Private/SoundWave.cpp
+++ b/Engine/Source/Runtime/Engine/Private/SoundWave.cpp
@@ -1424,7 +1424,13 @@ void USoundWave::SetSoundAssetCompressionType(ESoundAssetCompressionType InSound
 		SoundAssetCompressionType = InSoundAssetCompressionType;
 	}
 
-	SoundWaveDataPtr->CacheRuntimeFormatDependentState(SoundWaveDataPtr->FindRuntimeFormat(*this));
+	FName RuntimeFormat = SoundWaveDataPtr->FindRuntimeFormat(*this);
+	if (RuntimeFormat.IsNone())
+	{
+		RuntimeFormat = Audio::ToName(SoundAssetCompressionType);
+	}
+
+	SoundWaveDataPtr->CacheRuntimeFormatDependentState(RuntimeFormat);
 
 	UpdateAsset(bMarkDirty);
 #endif // #if WITH_EDITOR
