Expose `APlayerController::IsPrimaryPlayer()` to BP as `BlueprintPure`

This will allow games to check if a controller is a primary player from Blueprint, without having to try to expose it through BP function libraries, components, or other methods.

diff --git a/Engine/Source/Runtime/Engine/Classes/GameFramework/PlayerController.h b/Engine/Source/Runtime/Engine/Classes/GameFramework/PlayerController.h
--- a/Engine/Source/Runtime/Engine/Classes/GameFramework/PlayerController.h
+++ b/Engine/Source/Runtime/Engine/Classes/GameFramework/PlayerController.h
@@ -2016,6 +2016,7 @@ public:
 	 * Wrapper for determining whether this player is the first player on their console.
 	 * @return	true if this player is not using splitscreen, or is the first player in the split-screen layout.
 	 */
+	UFUNCTION(BlueprintPure, Category = "Game|Player")
 	ENGINE_API bool IsPrimaryPlayer() const;
 
 	/**
