[World Partition Actor Descriptors] Don't append FBox Min/Max vectors when the box is invalid for better visibility/diffs.

#preglight 6705167f387f45e004c14e54

diff --git a/Engine/Source/Runtime/Engine/Private/WorldPartition/WorldPartitionActorDesc.cpp b/Engine/Source/Runtime/Engine/Private/WorldPartition/WorldPartitionActorDesc.cpp
--- a/Engine/Source/Runtime/Engine/Private/WorldPartition/WorldPartitionActorDesc.cpp
+++ b/Engine/Source/Runtime/Engine/Private/WorldPartition/WorldPartitionActorDesc.cpp
@@ -495,8 +495,17 @@ FString FWorldPartitionActorDesc::ToString(EToStringMode Mode) const
 		Append(Name, *Value);
 	};
 
-	auto AppendToString = [&AppendFromString]<typename Type>(const TCHAR* Name, const Type& Value)
+	auto AppendToString = [&Append, &AppendFromString]<typename Type>(const TCHAR* Name, const Type& Value)
 	{
+		if constexpr (std::is_same_v<Type, FBox>)
+		{
+			if (!Value.IsValid)
+			{
+				Append(Name, TEXT("IsValid=false"));
+				return;
+			}
+		}
+
 		AppendFromString(Name, Value.ToString());
 	};
 
