New material translator work.

- Implemented UMaterialExpressionIfThenElse and renamed old control-flow based version to ...Flow.
- Implemented UMaterialExpressionStaticBool.
- Renamed some APIs in the emitter and added some documentation.

diff --git a/Engine/Source/Runtime/Engine/Private/Materials/MaterialExpressionHLSL.cpp b/Engine/Source/Runtime/Engine/Private/Materials/MaterialExpressionHLSL.cpp
--- a/Engine/Source/Runtime/Engine/Private/Materials/MaterialExpressionHLSL.cpp
+++ b/Engine/Source/Runtime/Engine/Private/Materials/MaterialExpressionHLSL.cpp
@@ -4766,7 +4766,7 @@ bool UMaterialExpressionSetLocal::GenerateHLSLStatements(FMaterialHLSLGenerator&
 	return true;
 }
 
-bool UMaterialExpressionIfThenElse::GenerateHLSLStatements(FMaterialHLSLGenerator& Generator, UE::HLSLTree::FScope& Scope) const
+bool UMaterialExpressionIfThenElseFlow::GenerateHLSLStatements(FMaterialHLSLGenerator& Generator, UE::HLSLTree::FScope& Scope) const
 {
 	const UE::HLSLTree::FExpression* ConditionExpression = Condition.AcquireHLSLExpression(Generator, Scope);
 	if (!ConditionExpression)
diff --git a/Engine/Source/Runtime/Engine/Private/Materials/MaterialExpressions.cpp b/Engine/Source/Runtime/Engine/Private/Materials/MaterialExpressions.cpp
--- a/Engine/Source/Runtime/Engine/Private/Materials/MaterialExpressions.cpp
+++ b/Engine/Source/Runtime/Engine/Private/Materials/MaterialExpressions.cpp
@@ -26992,6 +26992,40 @@ FName UMaterialExpressionSceneDepthWithoutWater::GetInputName(int32 InputIndex)
 
 #endif // WITH_EDITOR
 
+
+UMaterialExpressionIfThenElse::UMaterialExpressionIfThenElse(const FObjectInitializer& ObjectInitializer)
+: Super(ObjectInitializer)
+{
+	struct FConstructorStatics
+	{
+		FText NAME_Execution;
+		FConstructorStatics() : NAME_Execution(LOCTEXT("Execution", "Execution")) { }
+	};
+	static FConstructorStatics ConstructorStatics;
+#if WITH_EDITORONLY_DATA
+	MenuCategories.Add(ConstructorStatics.NAME_Execution);
+#endif
+}
+
+#if WITH_EDITOR
+void UMaterialExpressionIfThenElse::GetCaption(TArray<FString>& OutCaptions) const
+{
+	OutCaptions.Add(TEXT("If"));
+}
+
+uint32 UMaterialExpressionIfThenElse::GetInputType(int32 InputIndex)
+{
+	switch (InputIndex)
+	{
+		case 0: return MCT_Bool;
+		case 1: return MCT_Unknown;
+		case 2: return MCT_Unknown;
+		default: checkNoEntry(); return 0u;
+	}
+}
+#endif
+
+
 UMaterialExpressionExecBegin::UMaterialExpressionExecBegin(const FObjectInitializer& ObjectInitializer)
 	: Super(ObjectInitializer)
 {
@@ -27039,7 +27073,7 @@ int32 UMaterialExpressionExecEnd::Compile(class FMaterialCompiler* Compiler, int
 }
 #endif
 
-UMaterialExpressionIfThenElse::UMaterialExpressionIfThenElse(const FObjectInitializer& ObjectInitializer)
+UMaterialExpressionIfThenElseFlow::UMaterialExpressionIfThenElseFlow(const FObjectInitializer& ObjectInitializer)
 	: Super(ObjectInitializer)
 {
 	struct FConstructorStatics
@@ -27057,7 +27091,7 @@ UMaterialExpressionIfThenElse::UMaterialExpressionIfThenElse(const FObjectInitia
 }
 
 #if WITH_EDITOR
-int32 UMaterialExpressionIfThenElse::Compile(class FMaterialCompiler* Compiler, int32 OutputIndex)
+int32 UMaterialExpressionIfThenElseFlow::Compile(class FMaterialCompiler* Compiler, int32 OutputIndex)
 {
 	check(OutputIndex == CompileExecutionOutputIndex);
 	if (!Then.GetExpression())
@@ -27086,12 +27120,12 @@ int32 UMaterialExpressionIfThenElse::Compile(class FMaterialCompiler* Compiler,
 	return INDEX_NONE;
 }
 
-void UMaterialExpressionIfThenElse::GetCaption(TArray<FString>& OutCaptions) const
+void UMaterialExpressionIfThenElseFlow::GetCaption(TArray<FString>& OutCaptions) const
 {
 	OutCaptions.Add(TEXT("If"));
 }
 
-uint32 UMaterialExpressionIfThenElse::GetInputType(int32 InputIndex)
+uint32 UMaterialExpressionIfThenElseFlow::GetInputType(int32 InputIndex)
 {
 	switch (InputIndex)
 	{
diff --git a/Engine/Source/Runtime/Engine/Private/Materials/MaterialExpressionsIR.cpp b/Engine/Source/Runtime/Engine/Private/Materials/MaterialExpressionsIR.cpp
--- a/Engine/Source/Runtime/Engine/Private/Materials/MaterialExpressionsIR.cpp
+++ b/Engine/Source/Runtime/Engine/Private/Materials/MaterialExpressionsIR.cpp
@@ -260,7 +260,7 @@ void UMaterialExpression::Build(MIR::FEmitter& Emitter)
 
 void UMaterialExpressionConstant::Build(MIR::FEmitter& Emitter)
 {
-	MIR::FValue* Value = Emitter.EmitConstantFloat1(R);
+	MIR::FValue* Value = Emitter.EmitConstantFloatScalar(R);
 	Emitter.Output(GetOutput(0), Value);
 }
 
@@ -294,6 +294,11 @@ void UMaterialExpressionConstant4Vector::Build(MIR::FEmitter& Emitter)
 	}
 }
 
+void UMaterialExpressionStaticBool::Build(MIR::FEmitter& Emitter)
+{
+	Emitter.Output(0, Emitter.EmitConstantBoolScalar(Value));
+}
+
 static MIR::FValue* BuildMaterialExpressionParameter(UMaterialExpressionParameter* ParameterExpr, UE::MIR::FEmitter& Emitter)
 {
 	FMaterialParameterMetadata Metadata;
@@ -308,14 +313,14 @@ static MIR::FValue* BuildMaterialExpressionParameter(UMaterialExpressionParamete
 	{
 		if (Metadata.Value.Type == EMaterialParameterType::Scalar)
 		{
-			return Emitter.GetCustomPrimitiveData(Metadata.PrimitiveDataIndex);
+			return Emitter.EmitCustomPrimitiveData(Metadata.PrimitiveDataIndex);
 		}
 		else
 		{
-			MIR::FValue* X = Emitter.GetCustomPrimitiveData(Metadata.PrimitiveDataIndex + 0);
-			MIR::FValue* Y = Emitter.GetCustomPrimitiveData(Metadata.PrimitiveDataIndex + 1);
-			MIR::FValue* Z = Emitter.GetCustomPrimitiveData(Metadata.PrimitiveDataIndex + 2);
-			MIR::FValue* W = Emitter.GetCustomPrimitiveData(Metadata.PrimitiveDataIndex + 3);
+			MIR::FValue* X = Emitter.EmitCustomPrimitiveData(Metadata.PrimitiveDataIndex + 0);
+			MIR::FValue* Y = Emitter.EmitCustomPrimitiveData(Metadata.PrimitiveDataIndex + 1);
+			MIR::FValue* Z = Emitter.EmitCustomPrimitiveData(Metadata.PrimitiveDataIndex + 2);
+			MIR::FValue* W = Emitter.EmitCustomPrimitiveData(Metadata.PrimitiveDataIndex + 3);
 			return Emitter.EmitVector4(X, Y, Z, W);
 		}
 	}
@@ -323,7 +328,7 @@ static MIR::FValue* BuildMaterialExpressionParameter(UMaterialExpressionParamete
 	switch (Metadata.Value.Type)
 	{
 		case EMaterialParameterType::Scalar:
-			return Emitter.EmitConstantFloat1(Metadata.Value.Float[0]);
+			return Emitter.EmitConstantFloatScalar(Metadata.Value.Float[0]);
 			
 		case EMaterialParameterType::Vector:
 			return Emitter.EmitConstantFloat4(FVector4f{ Metadata.Value.Float[0], Metadata.Value.Float[1], Metadata.Value.Float[2], Metadata.Value.Float[3] });
@@ -332,7 +337,7 @@ static MIR::FValue* BuildMaterialExpressionParameter(UMaterialExpressionParamete
 		 	return Emitter.EmitTextureObject(Metadata.Value.Texture, SAMPLERTYPE_Color);
 
 		case EMaterialParameterType::StaticSwitch:
-			return Emitter.EmitConstantBool1(Metadata.Value.Bool[0]);
+			return Emitter.EmitConstantBoolScalar(Metadata.Value.Bool[0]);
 
 		default:
 			UE_MIR_TODO();
@@ -344,10 +349,20 @@ void UMaterialExpressionParameter::Build(UE::MIR::FEmitter& Emitter)
 	Emitter.Output(0, BuildMaterialExpressionParameter(this, Emitter));
 }
 
+void UMaterialExpressionStaticBoolParameter::Build(MIR::FEmitter& Emitter)
+{
+	MIR::FValue* Value = BuildMaterialExpressionParameter(this, Emitter);
+	
+	// Check value is a static bool scalar.
+	Emitter.ToConstantBoolScalar(Value);
+
+	Emitter.Output(0, Value);
+}
+
 void UMaterialExpressionStaticSwitch::Build(UE::MIR::FEmitter& Emitter)
 {
-	Emitter.DefaultToBool(&Value, (bool)DefaultValue);
-	bool bCondition = Emitter.InputConstantBool(&Value);
+	Emitter.DefaultInputToBoolScalar(&Value, (bool)DefaultValue);
+	bool bCondition = Emitter.InputConstantBoolScalar(&Value);
 
 	UE_MIR_CHECKPOINT();
 
@@ -358,30 +373,16 @@ void UMaterialExpressionStaticSwitchParameter::Build(UE::MIR::FEmitter& Emitter)
 {
 	MIR::FValue* Value = BuildMaterialExpressionParameter(this, Emitter);
 
-	UE_MIR_CHECKPOINT();
-
-	bool bCondition = Emitter.ToConstantBool(Value);
+	// Expect it to be a constant bool scalar, get its value.
+	bool bCondition = Emitter.ToConstantBoolScalar(Value);
 
 	UE_MIR_CHECKPOINT();
 
-	MIR::FValue* Output;
-	if (DynamicBranch)
-	{
-		MIR::FValue* True = Emitter.Input(&A);
-		MIR::FValue* False = Emitter.Input(&B);
-
-		UE_MIR_CHECKPOINT();
-
-		Output = Emitter.EmitBranchForced(Emitter.EmitConstantBool1(bCondition), True, False);
-	}
-	else
-	{
-		Output = Emitter.Input(bCondition ? &A : &B);
-	}
-
-	UE_MIR_CHECKPOINT();
+	// Select active input
+	FExpressionInput* ActiveInput = bCondition ? &A : &B;
 
-	Emitter.Output(0 ,Output);
+	// And bypass its value to our output
+	Emitter.Output(0, Emitter.Input(ActiveInput));
 }
 
 void UMaterialExpressionAppendVector::Build(MIR::FEmitter& Emitter)
@@ -469,8 +470,8 @@ static void BuildBinaryArithmeticOperator(
 	float RhsConst)
 {
 	// Default inputs to their relative constants if disconnected, then get each input after checking it has primitive type.
-	MIR::FValue* LhsValue = Emitter.DefaultToFloat(LhsInput, LhsConst).InputPrimitive(LhsInput);
-	MIR::FValue* RhsValue = Emitter.DefaultToFloat(RhsInput, RhsConst).InputPrimitive(RhsInput);
+	MIR::FValue* LhsValue = Emitter.DefaultInputToFloatScalar(LhsInput, LhsConst).InputPrimitive(LhsInput);
+	MIR::FValue* RhsValue = Emitter.DefaultInputToFloatScalar(RhsInput, RhsConst).InputPrimitive(RhsInput);
 
 	UE_MIR_CHECKPOINT();
 
@@ -512,14 +513,49 @@ void UMaterialExpressionDivide::Build(MIR::FEmitter& Emitter)
 	BuildBinaryArithmeticOperator(Emitter, MIR::BO_Divide, &A, ConstA, &B, ConstB);
 }
 
+void UMaterialExpressionIfThenElse::Build(MIR::FEmitter& Emitter)
+{
+	// If Condition isn't connected, default its value to false
+	Emitter.DefaultInputToBoolScalar(&Condition, false);
+
+	// Get the condition value checking it is a bool scalar
+	MIR::FValue* ConditionValue = Emitter.InputBoolScalar(&Condition);
+
+	UE_MIR_CHECKPOINT();
+	
+	// If condition boolean is constant, select which input is active and simply
+	// bypass its value to our output.
+	if (MIR::FConstant* Constant = ConditionValue->As<MIR::FConstant>())
+	{
+		FExpressionInput* ActiveInput = Constant->Boolean ? &True : &False;
+		Emitter.Output(0, Emitter.Input(ActiveInput));
+		return;
+	}
+
+	// The condition isn't static; Get the true and false values.
+	// If any is disconnected, the emitter will report an error.
+	MIR::FValue* ThenValue = Emitter.Input(&True);
+	MIR::FValue* ElseValue = Emitter.Input(&False);
+	
+	UE_MIR_CHECKPOINT();
+
+	// Convert the false type to the true's
+	ElseValue = Emitter.EmitConstruct(ThenValue->Type, ElseValue);
+
+	// Emit the branch instruction
+	MIR::FValue* OutputValue = Emitter.EmitBranch(ConditionValue, ThenValue, ElseValue);
+
+	Emitter.Output(0, OutputValue);
+}
+
 void UMaterialExpressionIf::Build(MIR::FEmitter& Emitter)
 {
 	// Create default values flowing into disconnected inputs
-	Emitter.DefaultToFloatZero(&A);
-	Emitter.DefaultToFloat(&B, ConstB);
-	Emitter.DefaultToFloatZero(&AGreaterThanB);
-	Emitter.DefaultToFloatZero(&AEqualsB);
-	Emitter.DefaultToFloatZero(&ALessThanB);
+	Emitter.DefaultInputToFloatScalarZero(&A);
+	Emitter.DefaultInputToFloatScalar(&B, ConstB);
+	Emitter.DefaultInputToFloatScalarZero(&AGreaterThanB);
+	Emitter.DefaultInputToFloatScalarZero(&AEqualsB);
+	Emitter.DefaultInputToFloatScalarZero(&ALessThanB);
 
 	UE_MIR_CHECKPOINT();
 
@@ -599,7 +635,7 @@ static void BuildTextureSampleExpression(UTexture* Texture, UMaterialExpressionT
 	MIR::FValue* TexCoordsValue = Emitter.TryInput(&Expr->Coordinates);
 	if (!TexCoordsValue)
 	{
-		TexCoordsValue = Emitter.GetExternalInput(MIR::TexCoordIndexToExternalInput(Expr->ConstCoordinate));
+		TexCoordsValue = Emitter.EmitExternalInput(MIR::TexCoordIndexToExternalInput(Expr->ConstCoordinate));
 	}
 
 	UE_MIR_CHECKPOINT();
@@ -645,7 +681,7 @@ void UMaterialExpressionTextureCoordinate::Build(UE::MIR::FEmitter& Emitter)
 		Emitter.Error("unmirroring unsupported");
 	}
 
-	MIR::FValue* OutputValue = Emitter.GetExternalInput(MIR::TexCoordIndexToExternalInput(CoordinateIndex));
+	MIR::FValue* OutputValue = Emitter.EmitExternalInput(MIR::TexCoordIndexToExternalInput(CoordinateIndex));
 	
 	// Multiply the UV input by the UV tiling constants
 	OutputValue = Emitter.EmitBinaryOperator(MIR::BO_Multiply, OutputValue, Emitter.EmitConstantFloat2({ UTiling, VTiling }));
diff --git a/Engine/Source/Runtime/Engine/Private/Materials/MaterialIR.cpp b/Engine/Source/Runtime/Engine/Private/Materials/MaterialIR.cpp
--- a/Engine/Source/Runtime/Engine/Private/Materials/MaterialIR.cpp
+++ b/Engine/Source/Runtime/Engine/Private/Materials/MaterialIR.cpp
@@ -45,6 +45,16 @@ uint32 FValue::GetSizeInBytes() const
 	}
 }
 
+void FValue::SetFlags(EValueFlags InFlags)
+{
+	Flags = (EValueFlags)(Flags | InFlags);
+}
+
+bool FValue::IsA(EValueKind InKind) const
+{
+	return Kind == InKind;
+}
+
 FInstruction* FValue::AsInstruction()
 {
 	return this && (Kind > VK_InstructionBegin && Kind < VK_InstructionEnd) ? static_cast<FInstruction*>(this) : nullptr;
@@ -55,6 +65,11 @@ const FInstruction* FValue::AsInstruction() const
 	return const_cast<FValue*>(this)->AsInstruction();
 }
 
+bool FValue::IsPoison() const
+{
+	return Kind == VK_Poison;
+}
+
 bool FValue::Equals(const FValue* Other) const
 {
 	if (this == Other)
@@ -201,6 +216,12 @@ TOptional<TPair<UTexture*, EMaterialSamplerType>> FValue::GetTexture()
 	return {};
 }
 
+FPoison* FPoison::Get()
+{
+	static FPoison Poison{ VK_Poison, VF_None, FType::GetPoison() };
+	return &Poison;
+}
+
 bool FConstant::IsBool() const
 {
 	return Type->IsBoolScalar();
diff --git a/Engine/Source/Runtime/Engine/Private/Materials/MaterialIREmitter.cpp b/Engine/Source/Runtime/Engine/Private/Materials/MaterialIREmitter.cpp
--- a/Engine/Source/Runtime/Engine/Private/Materials/MaterialIREmitter.cpp
+++ b/Engine/Source/Runtime/Engine/Private/Materials/MaterialIREmitter.cpp
@@ -4,7 +4,6 @@
 #include "Materials/MaterialIRModule.h"
 #include "Materials/MaterialIRModuleBuilder.h"
 #include "Materials/MaterialIRInternal.h"
-#include "Materials/MaterialIRInternal.h"
 #include "Shader/ShaderTypes.h"
 #include "MaterialShared.h"
 #include "MaterialExpressionIO.h"
@@ -179,26 +178,26 @@ FEmitter::FEmitter(FMaterialIRModuleBuilderImpl* InBuilder, UMaterial* InMateria
 	Module = InModule;
 }
 
-FEmitter& FEmitter::DefaultToFloatZero(const FExpressionInput* InInput)
+FEmitter& FEmitter::DefaultInputToFloatScalarZero(const FExpressionInput* InInput)
 {
-	return DefaultToFloat(InInput, 0.0f);
+	return DefaultInputToFloatScalar(InInput, 0.0f);
 }
 
-FEmitter& FEmitter::DefaultToBool(const FExpressionInput* InInput, bool Value)
+FEmitter& FEmitter::DefaultInputToBoolScalar(const FExpressionInput* InInput, bool Value)
 {
 	if (!InInput->IsConnected())
 	{
-		Internal::SetInputValue(Builder, InInput, EmitConstantBool1(Value));
+		Internal::SetInputValue(Builder, InInput, EmitConstantBoolScalar(Value));
 	}
 
 	return *this;
 }
 
-FEmitter& FEmitter::DefaultToFloat(const FExpressionInput* InInput, TFloat Value)
+FEmitter& FEmitter::DefaultInputToFloatScalar(const FExpressionInput* InInput, TFloat Value)
 {
 	if (!InInput->IsConnected())
 	{
-		Internal::SetInputValue(Builder, InInput, EmitConstantFloat1(Value));
+		Internal::SetInputValue(Builder, InInput, EmitConstantFloatScalar(Value));
 	}
 
 	return *this;
@@ -216,15 +215,15 @@ FValue* FEmitter::Input(const FExpressionInput* InInput)
 	if (!Value)
 	{
 		Errorf(TEXT("Input '%s' has no input value."), *SlowFindInputName(Expression, InInput).ToString());
-		return nullptr;
+		return MIR::FPoison::Get();
 	}
 	return Value;
 }
 
-bool FEmitter::InputConstantBool(const FExpressionInput* InInput)
+bool FEmitter::InputConstantBoolScalar(const FExpressionInput* InInput)
 {
 	FValue* Value = InputBoolScalar(InInput);
-	if (!Value)
+	if (Value->IsPoison())
 	{
 		return false;
 	}
@@ -246,7 +245,7 @@ FValue* FEmitter::InputBoolScalar(const FExpressionInput* InInput)
 	return Value;
 }
 
-FValue* FEmitter::InputFloat1(const FExpressionInput* InInput)
+FValue* FEmitter::InputFloatScalar(const FExpressionInput* InInput)
 {
 	FValue* Value = Input(InInput);
 	CheckInputIsScalar(InInput, Value, SK_Float);
@@ -256,9 +255,9 @@ FValue* FEmitter::InputFloat1(const FExpressionInput* InInput)
 FValue* FEmitter::InputScalar(const FExpressionInput* InInput)
 {
 	FValue* Value = InputPrimitive(InInput);
-	if (!Value)
+	if (Value->IsPoison())
 	{
-		return nullptr;
+		return Value;
 	}
 
 	FPrimitiveTypePtr ScalarType = Value->Type->AsPrimitive()->ToScalar();
@@ -268,6 +267,10 @@ FValue* FEmitter::InputScalar(const FExpressionInput* InInput)
 FValue* FEmitter::InputPrimitive(const FExpressionInput* InInput)
 {
 	FValue* Value = Input(InInput);
+	if (Value->IsPoison())
+	{
+		return Value;
+	}
 	CheckInputTypeIs(InInput, Value, TK_Primitive);
 	return Value;
 }
@@ -275,11 +278,10 @@ FValue* FEmitter::InputPrimitive(const FExpressionInput* InInput)
 FValue* FEmitter::InputOfType(const FExpressionInput* InInput, ETypeKind Kind)
 {
 	FValue* Value = Input(InInput);
-	if (!Value)
+	if (Value->IsPoison())
 	{
-		return nullptr;
+		return Value;
 	}
-
 	CheckInputTypeIs(InInput, Value, Kind);
 	return Value;
 }
@@ -296,9 +298,11 @@ FEmitter& FEmitter::Output(const FExpressionOutput* Output, FValue* Value)
 	return *this;
 }
 
-bool FEmitter::ToConstantBool(MIR::FValue* Value)
+bool FEmitter::ToConstantBoolScalar(MIR::FValue* Value)
 {
-	if (!CheckValueIsScalar(Value, SK_Bool))
+	CheckValueIsScalar(Value, SK_Bool);
+
+	if (!bExpressionHasError)
 	{
 		return false;
 	}
@@ -338,31 +342,19 @@ void FEmitter::CheckInputTypeIs(const FExpressionInput* InInput, FValue* InputVa
 	}
 }
 
-bool FEmitter::CheckValueIsScalar(const FValue* Value, EScalarKind Kind)
+void FEmitter::CheckValueIsScalar(const FValue* Value, EScalarKind Kind)
 {
-	if (!CheckValueValid(Value))
+	if (Value->IsPoison())
 	{
-		return false;
+		return;
 	}
 
 	FPrimitiveTypePtr PrimitiveType = Value->Type->AsPrimitive();
 	if (!PrimitiveType || !PrimitiveType->IsScalar() || PrimitiveType->ScalarKind != Kind)
 	{
 		Errorf(TEXT("Value expected to be a %s scalar. It is %s instead."), ScalarKindToString(Kind), Value->Type->GetSpelling().GetData());
-		return false;
+		return;
 	}
-
-	return true;
-}
-
-bool FEmitter::CheckValueValid(const FValue* Value)
-{
-	if (!Value) 
-	{
-		Error(TEXT("Unspecified value"));
-		return false;
-	}
-	return true;
 }
 
 FValue* FEmitter::EmitConstantFromShaderValue(const UE::Shader::FValue& InValue)
@@ -371,12 +363,12 @@ FValue* FEmitter::EmitConstantFromShaderValue(const UE::Shader::FValue& InValue)
 
 	switch (InValue.Type.ValueType)
 	{
-		case UE::Shader::EValueType::Float1: return EmitConstantFloat1(InValue.AsFloatScalar());
+		case UE::Shader::EValueType::Float1: return EmitConstantFloatScalar(InValue.AsFloatScalar());
 		case UE::Shader::EValueType::Float2: return EmitConstantFloat2(FVector2f{ InValue.Component[0].Float, InValue.Component[1].Float });
 		case UE::Shader::EValueType::Float3: return EmitConstantFloat3(FVector3f{ InValue.Component[0].Float, InValue.Component[1].Float, InValue.Component[2].Float });
 		case UE::Shader::EValueType::Float4: return EmitConstantFloat4(FVector4f{ InValue.Component[0].Float, InValue.Component[1].Float, InValue.Component[2].Float, InValue.Component[3].Float });
 
-		case UE::Shader::EValueType::Int1: return EmitConstantInt1(InValue.AsFloatScalar());
+		case UE::Shader::EValueType::Int1: return EmitConstantIntScalar(InValue.AsFloatScalar());
 		case UE::Shader::EValueType::Int2: return EmitConstantInt2(FIntVector2{ InValue.Component[0].Int, InValue.Component[1].Int });
 		case UE::Shader::EValueType::Int3: return EmitConstantInt3(FIntVector3{ InValue.Component[0].Int, InValue.Component[1].Int, InValue.Component[2].Int });
 		case UE::Shader::EValueType::Int4: return EmitConstantInt4(FIntVector4{ InValue.Component[0].Int, InValue.Component[1].Int, InValue.Component[2].Int, InValue.Component[3].Int });
@@ -390,8 +382,8 @@ FValue* FEmitter::EmitConstantScalarZero(EScalarKind Kind)
 	switch (Kind)
 	{
 		case SK_Bool: return EmitConstantFalse();
-		case SK_Int: return EmitConstantInt1(0);
-		case SK_Float: return EmitConstantFloat1(0.0f);
+		case SK_Int: return EmitConstantIntScalar(0);
+		case SK_Float: return EmitConstantFloatScalar(0.0f);
 		default: UE_MIR_UNREACHABLE();
 	}
 }
@@ -406,12 +398,12 @@ FValue* FEmitter::EmitConstantFalse()
 	return ConstantFalse;
 }
 
-FValue* FEmitter::EmitConstantBool1(bool InX)
+FValue* FEmitter::EmitConstantBoolScalar(bool InX)
 {
 	return InX ? EmitConstantTrue() : EmitConstantFalse();
 }
 
-FValue* FEmitter::EmitConstantFloat1(TFloat InX)
+FValue* FEmitter::EmitConstantFloatScalar(TFloat InX)
 {
 	FConstant Scalar = MakePrototype<FConstant>(FPrimitiveType::GetScalar(SK_Float));
 	Scalar.Float = InX;
@@ -420,29 +412,29 @@ FValue* FEmitter::EmitConstantFloat1(TFloat InX)
 
 FValue* FEmitter::EmitConstantFloat2(const FVector2f& InValue)
 {
-	FValue* X = EmitConstantFloat1(InValue.X);
-	FValue* Y = EmitConstantFloat1(InValue.Y);
+	FValue* X = EmitConstantFloatScalar(InValue.X);
+	FValue* Y = EmitConstantFloatScalar(InValue.Y);
 	return EmitVector2(X, Y);
 }
 
 FValue* FEmitter::EmitConstantFloat3(const FVector3f& InValue)
 {
-	FValue* X = EmitConstantFloat1(InValue.X);
-	FValue* Y = EmitConstantFloat1(InValue.Y);
-	FValue* Z = EmitConstantFloat1(InValue.Z);
+	FValue* X = EmitConstantFloatScalar(InValue.X);
+	FValue* Y = EmitConstantFloatScalar(InValue.Y);
+	FValue* Z = EmitConstantFloatScalar(InValue.Z);
 	return EmitVector3(X, Y, Z);
 }
 
 FValue* FEmitter::EmitConstantFloat4(const FVector4f& InValue)
 {
-	FValue* X = EmitConstantFloat1(InValue.X);
-	FValue* Y = EmitConstantFloat1(InValue.Y);
-	FValue* Z = EmitConstantFloat1(InValue.Z);
-	FValue* W = EmitConstantFloat1(InValue.W);
+	FValue* X = EmitConstantFloatScalar(InValue.X);
+	FValue* Y = EmitConstantFloatScalar(InValue.Y);
+	FValue* Z = EmitConstantFloatScalar(InValue.Z);
+	FValue* W = EmitConstantFloatScalar(InValue.W);
 	return EmitVector4(X, Y, Z, W);
 }
 
-FValue* FEmitter::EmitConstantInt1(TInteger InX)
+FValue* FEmitter::EmitConstantIntScalar(TInteger InX)
 {
 	FConstant Scalar = MakePrototype<FConstant>(FPrimitiveType::GetScalar(SK_Int));
 	Scalar.Integer = InX;
@@ -451,25 +443,25 @@ FValue* FEmitter::EmitConstantInt1(TInteger InX)
 
 FValue* FEmitter::EmitConstantInt2(const FIntVector2& InValue)
 {
-	FValue* X = EmitConstantInt1(InValue.X);
-	FValue* Y = EmitConstantInt1(InValue.Y);
+	FValue* X = EmitConstantIntScalar(InValue.X);
+	FValue* Y = EmitConstantIntScalar(InValue.Y);
 	return EmitVector2(X, Y);
 }
 
 FValue* FEmitter::EmitConstantInt3(const FIntVector3& InValue)
 {
-	FValue* X = EmitConstantInt1(InValue.X);
-	FValue* Y = EmitConstantInt1(InValue.Y);
-	FValue* Z = EmitConstantInt1(InValue.Z);
+	FValue* X = EmitConstantIntScalar(InValue.X);
+	FValue* Y = EmitConstantIntScalar(InValue.Y);
+	FValue* Z = EmitConstantIntScalar(InValue.Z);
 	return EmitVector3(X, Y, Z);
 }
 
 FValue* FEmitter::EmitConstantInt4(const FIntVector4& InValue)
 {
-	FValue* X = EmitConstantInt1(InValue.X);
-	FValue* Y = EmitConstantInt1(InValue.Y);
-	FValue* Z = EmitConstantInt1(InValue.Z);
-	FValue* W = EmitConstantInt1(InValue.W);
+	FValue* X = EmitConstantIntScalar(InValue.X);
+	FValue* Y = EmitConstantIntScalar(InValue.Y);
+	FValue* Z = EmitConstantIntScalar(InValue.Z);
+	FValue* W = EmitConstantIntScalar(InValue.W);
 	return EmitVector4(X, Y, Z, W);
 }
 
@@ -518,7 +510,7 @@ FValue* FEmitter::EmitVector4(FValue* InX, FValue* InY, FValue* InZ, FValue* InW
 	return EmitPrototype(this, Vector);
 }
 
-FValue* FEmitter::GetExternalInput(EExternalInput Id)
+FValue* FEmitter::EmitExternalInput(EExternalInput Id)
 {
 	MIR::FExternalInput Prototype = MakePrototype<MIR::FExternalInput>(GetExternalInputType(Id));
 	Prototype.Id = Id;
@@ -531,7 +523,7 @@ FValue* FEmitter::EmitSubscript(FValue* Value, int Index)
 	if (!PrimitiveType)
 	{
 		Errorf(TEXT("Value of type `%s` cannot be subscripted."), Value->Type->GetSpelling().GetData());
-		return nullptr;
+		return MIR::FPoison::Get();
 	}
 
 	// Getting first component and Value is already a scalar, just return itself.
@@ -571,7 +563,7 @@ FValue* FEmitter::EmitSwizzle(FValue* Value, FSwizzleMask Mask)
 	if (!PrimitiveType || PrimitiveType->IsMatrix())
 	{
 		Errorf(TEXT("Cannot swizzle a `%s` value."), Value->Type->GetSpelling().GetData());
-		return nullptr;
+		return MIR::FPoison::Get();
 	}
 
 	// Make sure each component in the mask fits the number of components in Value.
@@ -580,7 +572,7 @@ FValue* FEmitter::EmitSwizzle(FValue* Value, FSwizzleMask Mask)
 		if ((int)Component >= PrimitiveType->NumRows)
 		{
 			Errorf(TEXT("Value of type `%s` has no component `%s`."), PrimitiveType->Spelling.GetData(), VectorComponentToString(Component));
-			return nullptr;
+			return MIR::FPoison::Get();
 		}
 	}
 
@@ -627,7 +619,7 @@ FValue* FEmitter::EmitTextureObject(UTexture* Texture, EMaterialSamplerType Samp
 	if (Texture->GetMaterialType() != MCT_Texture2D)
 	{
 		Error(TEXT("Only Textures 2D are supported for now."));
-		return nullptr;
+		return MIR::FPoison::Get();
 	}
 
 	FTextureObject Proto = MakePrototype<FTextureObject>(FObjectType::GetTexture2D());
@@ -636,12 +628,12 @@ FValue* FEmitter::EmitTextureObject(UTexture* Texture, EMaterialSamplerType Samp
 	return EmitPrototype(this, Proto);
 }
 
-FValue* FEmitter::GetParameter(FName Name, const FMaterialParameterMetadata& Metadata)
+FValue* FEmitter::EmitParameter(FName Name, const FMaterialParameterMetadata& Metadata)
 {
-	UE_MIR_UNREACHABLE();
+	UE_MIR_TODO();
 }
 
-FValue* FEmitter::GetCustomPrimitiveData(uint32 PrimitiveDataIndex)
+FValue* FEmitter::EmitCustomPrimitiveData(uint32 PrimitiveDataIndex)
 {
 	UE_MIR_TODO();
 }
@@ -760,7 +752,7 @@ static FValue* TryFoldBinaryOperatorScalar(FEmitter* Emitter, EBinaryOperator Op
 	const FConstant* RhsConstant = Rhs->As<FConstant>();
 	if (!LhsConstant || !RhsConstant)
 	{
-		return nullptr;
+		return MIR::FPoison::Get();
 	}
 
 	if (IsArithmeticOperator(Operator))
@@ -770,13 +762,13 @@ static FValue* TryFoldBinaryOperatorScalar(FEmitter* Emitter, EBinaryOperator Op
 			case SK_Int:
 			{
 				int Result = FoldScalarArithmeticOperator<TInteger>(Operator, LhsConstant->Integer, RhsConstant->Integer);
-				return Emitter->EmitConstantInt1(Result);
+				return Emitter->EmitConstantIntScalar(Result);
 			}
 
 			case SK_Float:
 			{
 				TFloat Result = FoldScalarArithmeticOperator<TFloat>(Operator, LhsConstant->Float, RhsConstant->Float);
-				return Emitter->EmitConstantFloat1(Result);
+				return Emitter->EmitConstantFloatScalar(Result);
 			}
 
 			default:
@@ -799,7 +791,7 @@ static FValue* TryFoldBinaryOperatorScalar(FEmitter* Emitter, EBinaryOperator Op
 			default:
 				UE_MIR_UNREACHABLE();
 		}
-		return Emitter->EmitConstantBool1(Result);
+		return Emitter->EmitConstantBoolScalar(Result);
 	}
 	else
 	{
@@ -813,7 +805,7 @@ FValue* FEmitter::EmitBinaryOperator(EBinaryOperator Operator, FValue* Lhs, FVal
 	if (Lhs->Type != Rhs->Type)
 	{
 		Errorf(TEXT("operands of binary operator `%s` do not have the same type."), BinaryOperatorToString(Operator));
-		return nullptr;
+		return MIR::FPoison::Get();
 	}
 
 	// Get the operands primitive type.
@@ -821,7 +813,7 @@ FValue* FEmitter::EmitBinaryOperator(EBinaryOperator Operator, FValue* Lhs, FVal
 	if (!PrimitiveType)
 	{
 		Errorf(TEXT("invalid binary operator on non primitive types."));
-		return nullptr;
+		return MIR::FPoison::Get();
 	}
 
 	// Determine the result type. If the operator is arithmetic, the result type will be the same
@@ -927,11 +919,6 @@ FValue* FEmitter::EmitBranch(FValue* Condition, FValue* True, FValue* False)
 		return ConstCondition->Boolean ? True : False;
 	}
 
-	return EmitBranchForced(Condition, True, False);
-}
-
-FValue* FEmitter::EmitBranchForced(FValue* Condition, FValue* True, FValue* False)
-{
 	// Condition must be of type bool
 	check(Condition->Type->IsBoolScalar());
 
@@ -963,8 +950,8 @@ static FValue* CastConstant(FEmitter* Emitter, FConstant* Constant, EScalarKind
 			switch (TargetKind)
 			{
 				case SK_Bool: UE_MIR_TODO();
-				case SK_Int: return Emitter->EmitConstantInt1(Constant->Integer);
-				case SK_Float: return Emitter->EmitConstantFloat1((TFloat)Constant->Integer);
+				case SK_Int: return Emitter->EmitConstantIntScalar(Constant->Integer);
+				case SK_Float: return Emitter->EmitConstantFloatScalar((TFloat)Constant->Integer);
 				default: UE_MIR_UNREACHABLE();
 			}
 		}
@@ -974,7 +961,7 @@ static FValue* CastConstant(FEmitter* Emitter, FConstant* Constant, EScalarKind
 			switch (TargetKind)
 			{
 				case SK_Bool: UE_MIR_TODO();
-				case SK_Int: return Emitter->EmitConstantInt1((int)Constant->Float);
+				case SK_Int: return Emitter->EmitConstantIntScalar((int)Constant->Float);
 				default: UE_MIR_UNREACHABLE();
 			}
 		}
@@ -991,7 +978,7 @@ static FValue* ConstructPrimitiveValue(FEmitter* Emitter, FPrimitiveTypePtr Targ
 	if (!InitializerPrimitiveType)
 	{
 		Emitter->Errorf(TEXT("Cannot construct a '%s' from non primitive type '%s'."), TargetPrimitiveType->Spelling.GetData(), Initializer->Type->GetSpelling().GetData());
-		return nullptr;
+		return MIR::FPoison::Get();
 	}
 
 	// Construct a scalar from another scalar.
@@ -1104,7 +1091,7 @@ static FValue* ConstructPrimitiveValue(FEmitter* Emitter, FPrimitiveTypePtr Targ
 	}
 
 	// Initializer value cannot be used to construct this primitive type.
-	return nullptr;
+	return MIR::FPoison::Get();
 }
 
 FValue* FEmitter::EmitConstruct(FTypePtr Type, FValue* Initializer)
@@ -1173,7 +1160,7 @@ void FEmitter::Error(FString Message)
 	Error.Expression = Expression;
 	Error.Message = MoveTemp(Message);
 	Module->Errors.Push(Error);
-	bHasExprBuildError = true;
+	bExpressionHasError = true;
 }
 
 void FEmitter::Initialize()
diff --git a/Engine/Source/Runtime/Engine/Private/Materials/MaterialIRInternal.cpp b/Engine/Source/Runtime/Engine/Private/Materials/MaterialIRInternal.cpp
--- a/Engine/Source/Runtime/Engine/Private/Materials/MaterialIRInternal.cpp
+++ b/Engine/Source/Runtime/Engine/Private/Materials/MaterialIRInternal.cpp
@@ -62,14 +62,14 @@ MIR::FValue* CreateMaterialAttributeDefaultValue(FEmitter& Emitter, const UMater
 
 	switch (Type)
 	{
-		case MCT_ShadingModel: return Emitter.EmitConstantInt1(Material->GetShadingModels().GetFirstShadingModel());
+		case MCT_ShadingModel: return Emitter.EmitConstantIntScalar(Material->GetShadingModels().GetFirstShadingModel());
 
-		case MCT_Float1: return Emitter.EmitConstantFloat1(  DefaultValue.X );
+		case MCT_Float1: return Emitter.EmitConstantFloatScalar(  DefaultValue.X );
 		case MCT_Float2: return Emitter.EmitConstantFloat2({ DefaultValue.X, DefaultValue.Y });
 		case MCT_Float3: return Emitter.EmitConstantFloat3({ DefaultValue.X, DefaultValue.Y, DefaultValue.Z });
 		case MCT_Float4: return Emitter.EmitConstantFloat4(DefaultValue);
 
-		case MCT_UInt1: return Emitter.EmitConstantInt1(  (int32)DefaultValue.X );
+		case MCT_UInt1: return Emitter.EmitConstantIntScalar(  (int32)DefaultValue.X );
 		case MCT_UInt2: return Emitter.EmitConstantInt2({ (int32)DefaultValue.X, (int32)DefaultValue.Y });
 		case MCT_UInt3: return Emitter.EmitConstantInt3({ (int32)DefaultValue.X, (int32)DefaultValue.Y, (int32)DefaultValue.Z });
 		case MCT_UInt4: return Emitter.EmitConstantInt4({ (int32)DefaultValue.X, (int32)DefaultValue.Y, (int32)DefaultValue.Z, (int32)DefaultValue.W });
diff --git a/Engine/Source/Runtime/Engine/Private/Materials/MaterialIRTypes.cpp b/Engine/Source/Runtime/Engine/Private/Materials/MaterialIRTypes.cpp
--- a/Engine/Source/Runtime/Engine/Private/Materials/MaterialIRTypes.cpp
+++ b/Engine/Source/Runtime/Engine/Private/Materials/MaterialIRTypes.cpp
@@ -105,11 +105,14 @@ FTypePtr FType::GetVoid()
 
 FStringView FType::GetSpelling() const
 {
-	if (auto PrimitiveType = AsPrimitive())
+	if (IsPoison())
+	{
+		return TEXT("Poison");
+	}
+	else if (auto PrimitiveType = AsPrimitive())
 	{
 		return PrimitiveType->Spelling;
 	}
-	
 	UE_MIR_UNREACHABLE();
 }
 
@@ -154,6 +157,17 @@ UE::Shader::EValueType FType::ToValueType() const
 	UE_MIR_UNREACHABLE();
 }
 
+FTypePtr FType::GetPoison()
+{
+	static FPoisonType Type{ TK_Poison };
+	return &Type;
+}
+
+bool FType::IsPoison() const
+{
+	return Kind == TK_Poison;
+}
+
 bool FType::IsBoolScalar() const
 {
 	return this == FPrimitiveType::GetBool1();
diff --git a/Engine/Source/Runtime/Engine/Public/Materials/MaterialExpressionIfThenElse.h b/Engine/Source/Runtime/Engine/Public/Materials/MaterialExpressionIfThenElse.h
--- a/Engine/Source/Runtime/Engine/Public/Materials/MaterialExpressionIfThenElse.h
+++ b/Engine/Source/Runtime/Engine/Public/Materials/MaterialExpressionIfThenElse.h
@@ -9,8 +9,33 @@
 #include "Materials/MaterialExpression.h"
 #include "MaterialExpressionIfThenElse.generated.h"
 
-UCLASS(collapsecategories, hidecategories = Object, MinimalAPI, meta=(MaterialControlFlow))
+
+UCLASS(collapsecategories, hidecategories = Object, MinimalAPI)
 class UMaterialExpressionIfThenElse : public UMaterialExpression
+{
+	GENERATED_UCLASS_BODY()
+	
+	UPROPERTY()
+	FExpressionInput Condition;
+
+	UPROPERTY()
+	FExpressionInput True;
+
+	UPROPERTY()
+	FExpressionInput False;
+	
+	//~ Begin UMaterialExpression Interface
+#if WITH_EDITOR
+	virtual void Build(UE::MIR::FEmitter& Emitter) override;
+	virtual void GetCaption(TArray<FString>& OutCaptions) const override;
+	virtual uint32 GetInputType(int32 InputIndex) override;
+#endif
+	//~ End UMaterialExpression Interface
+};
+
+
+UCLASS(collapsecategories, hidecategories = Object, MinimalAPI, meta=(MaterialControlFlow))
+class UMaterialExpressionIfThenElseFlow : public UMaterialExpression
 {
 	GENERATED_UCLASS_BODY()
 
diff --git a/Engine/Source/Runtime/Engine/Public/Materials/MaterialExpressionStaticBool.h b/Engine/Source/Runtime/Engine/Public/Materials/MaterialExpressionStaticBool.h
--- a/Engine/Source/Runtime/Engine/Public/Materials/MaterialExpressionStaticBool.h
+++ b/Engine/Source/Runtime/Engine/Public/Materials/MaterialExpressionStaticBool.h
@@ -20,6 +20,7 @@ class UMaterialExpressionStaticBool : public UMaterialExpression
 
 	//~ Begin UMaterialExpression Interface
 #if WITH_EDITOR
+	virtual void Build(UE::MIR::FEmitter& Emitter) override;
 	virtual int32 CompilePreview(class FMaterialCompiler* Compiler, int32 OutputIndex) override;
 	virtual int32 Compile(class FMaterialCompiler* Compiler, int32 OutputIndex) override;
 	virtual void GetCaption(TArray<FString>& OutCaptions) const override;
diff --git a/Engine/Source/Runtime/Engine/Public/Materials/MaterialExpressionStaticBoolParameter.h b/Engine/Source/Runtime/Engine/Public/Materials/MaterialExpressionStaticBoolParameter.h
--- a/Engine/Source/Runtime/Engine/Public/Materials/MaterialExpressionStaticBoolParameter.h
+++ b/Engine/Source/Runtime/Engine/Public/Materials/MaterialExpressionStaticBoolParameter.h
@@ -24,6 +24,7 @@ class UMaterialExpressionStaticBoolParameter : public UMaterialExpressionParamet
 public:
 	//~ Begin UMaterialExpression Interface
 #if WITH_EDITOR
+	virtual void Build(UE::MIR::FEmitter& Emitter) override;
 	virtual int32 Compile(class FMaterialCompiler* Compiler, int32 OutputIndex) override;
 	virtual int32 CompilePreview(class FMaterialCompiler* Compiler, int32 OutputIndex) override;
 	virtual void GetCaption(TArray<FString>& OutCaptions) const override;
diff --git a/Engine/Source/Runtime/Engine/Public/Materials/MaterialIR.h b/Engine/Source/Runtime/Engine/Public/Materials/MaterialIR.h
--- a/Engine/Source/Runtime/Engine/Public/Materials/MaterialIR.h
+++ b/Engine/Source/Runtime/Engine/Public/Materials/MaterialIR.h
@@ -12,6 +12,7 @@ enum EValueKind
 {
 	/* Values */
 
+	VK_Poison,
 	VK_Constant,
 	VK_ExternalInput,
 	VK_TextureObject,
@@ -53,10 +54,10 @@ struct FValue
 	uint32 GetSizeInBytes() const;
 	
 	// Sets the specified value flags.
-	void SetFlags(EValueFlags InFlags) { Flags = (EValueFlags)(Flags | InFlags); }
+	void SetFlags(EValueFlags InFlags);
 
 	// Returns whether this value is of specified kind.
-	bool IsA(EValueKind InKind) const { return Kind == InKind; }
+	bool IsA(EValueKind InKind) const;
 
 	// Tries to cast this value to an instruction and returns it (nullptr otherwise).
 	FInstruction* AsInstruction();
@@ -70,6 +71,9 @@ struct FValue
 	// Returns the array of this value's uses. An use is another value referenced by this one (e.g. the operands of a binary expression).
 	TArrayView<FValue*> GetUses();
 
+	// Returns whether this value is poison.
+	bool IsPoison() const;
+
 	// Returns whether this value exactly equals Other.
 	bool Equals(const FValue* Other) const;
 
@@ -115,6 +119,11 @@ struct TValue : FValue
 	static constexpr EValueKind TypeKind = TTypeKind;
 };
 
+struct FPoison : TValue<VK_Poison>
+{
+	static FPoison* Get();
+};
+
 using TInteger = int64_t;
 using TFloat = double;
 
diff --git a/Engine/Source/Runtime/Engine/Public/Materials/MaterialIREmitter.h b/Engine/Source/Runtime/Engine/Public/Materials/MaterialIREmitter.h
--- a/Engine/Source/Runtime/Engine/Public/Materials/MaterialIREmitter.h
+++ b/Engine/Source/Runtime/Engine/Public/Materials/MaterialIREmitter.h
@@ -47,13 +47,13 @@ public:
 	/* Defaults */
 	
 	//
-	FEmitter& DefaultToFloatZero(const FExpressionInput* Input);
+	FEmitter& DefaultInputToFloatScalarZero(const FExpressionInput* Input);
 	
 	//
-	FEmitter& DefaultToBool(const FExpressionInput* Input, bool Value);
+	FEmitter& DefaultInputToBoolScalar(const FExpressionInput* Input, bool Value);
 
 	//
-	FEmitter& DefaultToFloat(const FExpressionInput* Input, TFloat Value);
+	FEmitter& DefaultInputToFloatScalar(const FExpressionInput* Input, TFloat Value);
 
 	//
 
@@ -63,28 +63,33 @@ public:
 	// because it is not connected), it returns null.
 	FValue* TryInput(const FExpressionInput* Input);
 
-	// Gets and returns the value flowing into input with specified `InputIndex`.
+	// Gets and returns the value flowing into input with specified `Input`.
 	// If no value is flowing in (e.g. bacause it is not connected) it reports
 	// an error and returns null.
 	FValue* Input(const FExpressionInput* Input);
 	
-	//
-	bool InputConstantBool(const FExpressionInput* Input);
+	// Gets and returns the value flowing into input with specified `Input`.
+	// Before returning the value, it checks that the value is a constant boolean scalar.
+	bool InputConstantBoolScalar(const FExpressionInput* Input);
 	
-	//
+	// Gets and returns the value flowing into input with specified `Input`.
+	// Before returning the value, it checks that it is a boolean scalar.
 	FValue* InputBoolScalar(const FExpressionInput* Input);
 
-	// It gets the value flowing into it and checks that its type is float scalar.
-	FValue* InputFloat1(const FExpressionInput* Input);
+	// Gets and returns the value flowing into input with specified `Input`.
+	// Before returning the value, it checks that it is a float scalar.
+	FValue* InputFloatScalar(const FExpressionInput* Input);
 
-	//
+	// Gets and returns the value flowing into input with specified `Input`.
+	// Before returning the value, it checks that it is a scalar.
 	FValue* InputScalar(const FExpressionInput* Input);
 	
-	//
+	// Gets and returns the value flowing into input with specified `Input`.
+	// Before returning the value, it checks that it has a primitive type.
 	FValue* InputPrimitive(const FExpressionInput* Input);
 
-	// Gets the value flowing into `Input` and returns it after checking that its
-	// type matches `Kind`.
+	// Gets and returns the value flowing into input with specified `Input`.
+	// Before returning the value, it checks that its type is of specified `Kind`.
 	FValue* InputOfType(const FExpressionInput* Input, ETypeKind Kind);
 
 	/* Output */
@@ -95,9 +100,11 @@ public:
 	// Flows specified `Value` out of specified expression `Output`.
 	FEmitter& Output(const FExpressionOutput* Output, FValue* Value);
 	
-	/* Value */
+	/* Value manipulation */
 
-	bool ToConstantBool(MIR::FValue* Value);
+	// Checks that specified value is a constant bool scalar and returns its value.
+	// If value type is incorrect, it reports an error.
+	bool ToConstantBoolScalar(MIR::FValue* Value);
 
 	/* Error Checking */
 
@@ -113,10 +120,7 @@ public:
 	void CheckInputTypeIs(const FExpressionInput* Input, FValue* InputValue, ETypeKind Kind);
 	
 	//
-	bool CheckValueIsScalar(const FValue* Value, EScalarKind Kind);
-
-	//
-	bool CheckValueValid(const FValue* Value);
+	void CheckValueIsScalar(const FValue* Value, EScalarKind Kind);
 
 	/* Constants */
 
@@ -124,34 +128,33 @@ public:
 	FValue* EmitConstantScalarZero(EScalarKind Kind);
 	FValue* EmitConstantTrue();
 	FValue* EmitConstantFalse();
-	FValue* EmitConstantBool1(bool InX);
-	FValue* EmitConstantFloat1(TFloat InX);
+	FValue* EmitConstantBoolScalar(bool InX);
+	FValue* EmitConstantFloatScalar(TFloat InX);
 	FValue* EmitConstantFloat2(const FVector2f& InValue);
 	FValue* EmitConstantFloat3(const FVector3f& InValue);
 	FValue* EmitConstantFloat4(const FVector4f& InValue);
-	FValue* EmitConstantInt1(TInteger InX);
+	FValue* EmitConstantIntScalar(TInteger InX);
 	FValue* EmitConstantInt2(const FIntVector2& InValue);
 	FValue* EmitConstantInt3(const FIntVector3& InValue);
 	FValue* EmitConstantInt4(const FIntVector4& InValue);
 	FValue* EmitVector2(FValue* InX, FValue* InY);
 	FValue* EmitVector3(FValue* InX, FValue* InY, FValue* InZ);
 	FValue* EmitVector4(FValue* InX, FValue* InY, FValue* InZ, FValue* InW);
-	FValue* GetExternalInput(EExternalInput Id);
+	FValue* EmitExternalInput(EExternalInput Id);
 
 	/* Other Values */
 
 	FValue* EmitSubscript(FValue* Value, int ComponentIndex);
 	FValue* EmitSwizzle(FValue* Value, FSwizzleMask Mask);
 	FValue* EmitTextureObject(UTexture* Texture, EMaterialSamplerType SamplerType);
-	FValue* GetParameter(FName Name, const FMaterialParameterMetadata& Metadata);
-	FValue* GetCustomPrimitiveData(uint32 PrimitiveDataIndex);
+	FValue* EmitParameter(FName Name, const FMaterialParameterMetadata& Metadata);
+	FValue* EmitCustomPrimitiveData(uint32 PrimitiveDataIndex);
 
 	/* Instructions */
 
 	FSetMaterialOutput* EmitSetMaterialOutput(EMaterialProperty InProperty, FValue* InArgValue);
 	FValue* EmitBinaryOperator(EBinaryOperator Operator, FValue* Lhs, FValue* Rhs);
 	FValue* EmitBranch(FValue* Condition, FValue* True, FValue* False);
-	FValue* EmitBranchForced(FValue* Condition, FValue* True, FValue* False);
 	FValue* EmitConstruct(FTypePtr Type, FValue* Initializer);
 	FValue* EmitTextureSample(UTexture* Texture, FValue* TexCoord, ESamplerSourceMode SamplerSourceMode, ETextureMipValueMode MipValueMode, EMaterialSamplerType SamplerType);
 
@@ -161,7 +164,7 @@ public:
 
 	/* Error reporting */
 
-	bool IsInvalid() const { return bHasExprBuildError; }
+	bool IsInvalid() const { return bExpressionHasError; }
 
 	template <typename... TArgs>
 	void Errorf(UE::Core::TCheckedFormatString<FString::FmtCharType, TArgs...> Format, TArgs... Args)
@@ -187,7 +190,7 @@ private:
 	FMaterialIRModule* Module{};
 	UMaterialExpression* Expression{};
 	FMaterialIRModuleBuilderImpl* Builder{};
-	bool bHasExprBuildError = false;
+	bool bExpressionHasError = false;
 	FValue* ConstantTrue;
 	FValue* ConstantFalse;
 	TSet<FValue*, FValueKeyFuncs> ValueSet;
diff --git a/Engine/Source/Runtime/Engine/Public/Materials/MaterialIRTypes.h b/Engine/Source/Runtime/Engine/Public/Materials/MaterialIRTypes.h
--- a/Engine/Source/Runtime/Engine/Public/Materials/MaterialIRTypes.h
+++ b/Engine/Source/Runtime/Engine/Public/Materials/MaterialIRTypes.h
@@ -13,6 +13,7 @@ namespace UE::MIR
 
 enum ETypeKind
 {
+	TK_Poison,
 	TK_Void,
 	TK_Primitive,
 	TK_Object,
@@ -37,6 +38,12 @@ struct FType
 	// Returns the `void` type.
 	static FTypePtr GetVoid();
 
+	// Returns the `void` type.
+	static FTypePtr GetPoison();
+
+	// Returns whether this is the poison type.
+	bool IsPoison() const;
+
 	// Returns whether this type is a `bool` scalar.
 	bool IsBoolScalar() const;
 
@@ -62,6 +69,10 @@ struct FType
 	UE::Shader::EValueType ToValueType() const;
 };
 
+struct FPoisonType : FType
+{
+};
+
 // Primitive types of a single scalar.
 // Note: These are listed in precision order. Converting one to the other is then simply performed taking the max EScalarKind.
 enum EScalarKind
