[Selection] Allow for BP ISM components to opt into per hit proxy generation, so that PCG workflows can set this on BP ISM components

diff --git a/Engine/Source/Runtime/Engine/Classes/Components/PrimitiveComponent.h b/Engine/Source/Runtime/Engine/Classes/Components/PrimitiveComponent.h
--- a/Engine/Source/Runtime/Engine/Classes/Components/PrimitiveComponent.h
+++ b/Engine/Source/Runtime/Engine/Classes/Components/PrimitiveComponent.h
@@ -470,8 +470,10 @@ public:
 	UPROPERTY(EditAnywhere, AdvancedDisplay, BlueprintReadOnly, Category=TextureStreaming)
 	uint8 bForceMipStreaming:1;
 
-	/** If true a hit-proxy will be generated for each instance of instanced static meshes */
-	UPROPERTY()
+	/** If true a hit-proxy will be generated for each instance of instanced static meshes.
+	 *  This allows for other systems, like selection in viewports, to function at the individual instance.
+	 */
+	UPROPERTY(EditDefaultsOnly, BlueprintReadWrite, AdvancedDisplay, Category=Rendering)
 	uint8 bHasPerInstanceHitProxies:1;
 
 	// Lighting flags
