Optimize RichCurve evaluation

diff --git a/Engine/Source/Runtime/Engine/Classes/Curves/RichCurve.h b/Engine/Source/Runtime/Engine/Classes/Curves/RichCurve.h
--- a/Engine/Source/Runtime/Engine/Classes/Curves/RichCurve.h
+++ b/Engine/Source/Runtime/Engine/Classes/Curves/RichCurve.h
@@ -405,7 +405,29 @@ struct FCompressedRichCurve
 	/** Populate RichCurve with decompressed key-data */
 	ENGINE_API void PopulateCurve(FRichCurve& OutCurve) const;
 
+	// Encapsulates the StaticEval arguments
+	struct FStaticEvalArgs
+	{
+		ERichCurveCompressionFormat CompressionFormat = RCCF_Empty;
+		ERichCurveKeyTimeCompressionFormat KeyTimeCompressionFormat = RCKTCF_float32;
+
+		ERichCurveExtrapolation PreInfinityExtrap = RCCE_None;
+		ERichCurveExtrapolation PostInfinityExtrap = RCCE_None;
+
+		TConstantValueNumKeys ConstantValueNumKeys;
+
+		const uint8* CompressedKeys = nullptr;
+
+		float Time = 0.0f;
+
+		float DefaultValue = 0.0f;
+	};
+
+	/** Evaluate this rich curve at the specified time */
+	static ENGINE_API float StaticEval(const FStaticEvalArgs& Args);
+
 	/** Evaluate this rich curve at the specified time */
+	UE_DEPRECATED(5.6, "Please use StaticEval with the packed arguments structure as it has improved performance")
 	static ENGINE_API float StaticEval(ERichCurveCompressionFormat CompressionFormat, ERichCurveKeyTimeCompressionFormat KeyTimeCompressionFormat, ERichCurveExtrapolation PreInfinityExtrap, ERichCurveExtrapolation PostInfinityExtrap, TConstantValueNumKeys ConstantValueNumKeys, const uint8* CompressedKeys, float InTime, float InDefaultValue = 0.0f);
 
 	/** ICPPStructOps interface */
diff --git a/Engine/Source/Runtime/Engine/Private/Animation/AnimCurveCompressionCodec_CompressedRichCurve.cpp b/Engine/Source/Runtime/Engine/Private/Animation/AnimCurveCompressionCodec_CompressedRichCurve.cpp
--- a/Engine/Source/Runtime/Engine/Private/Animation/AnimCurveCompressionCodec_CompressedRichCurve.cpp
+++ b/Engine/Source/Runtime/Engine/Private/Animation/AnimCurveCompressionCodec_CompressedRichCurve.cpp
@@ -136,18 +136,31 @@ void UAnimCurveCompressionCodec_CompressedRichCurve::DecompressCurves(const FCom
 	}
 
 	const uint8* Buffer = AnimSeq.CompressedCurveByteStream.GetData();
-	const FCurveDesc* CurveDescriptions = (const FCurveDesc*)(Buffer);
 
 	auto GetNameFromIndex = [&IndexedCurveNames](int32 InCurveIndex)
 	{
 		return IndexedCurveNames[IndexedCurveNames[InCurveIndex].CurveIndex].CurveName;
 	};
 
-	auto GetValueFromIndex = [&CurveDescriptions, &IndexedCurveNames, Buffer, CurrentTime](int32 InCurveIndex)
+	// Setup our constant arguments, we'll re-use the structure for every curve
+	FCompressedRichCurve::FStaticEvalArgs Args;
+	Args.Time = CurrentTime;
+	Args.DefaultValue = 0.0f;
+
+	auto GetValueFromIndex = [&IndexedCurveNames, Buffer, &Args](int32 InCurveIndex)
 	{
+		const FCurveDesc* CurveDescriptions = (const FCurveDesc*)(Buffer);
 		const FCurveDesc& Curve = CurveDescriptions[IndexedCurveNames[InCurveIndex].CurveIndex];
 		const uint8* CompressedKeys = Buffer + Curve.KeyDataOffset;
-		const float Value = FCompressedRichCurve::StaticEval(Curve.CompressionFormat, Curve.KeyTimeCompressionFormat, Curve.PreInfinityExtrap, Curve.PostInfinityExtrap, Curve.ConstantValueNumKeys, CompressedKeys, CurrentTime);
+		
+		Args.CompressionFormat = Curve.CompressionFormat;
+		Args.KeyTimeCompressionFormat = Curve.KeyTimeCompressionFormat;
+		Args.PreInfinityExtrap = Curve.PreInfinityExtrap;
+		Args.PostInfinityExtrap = Curve.PostInfinityExtrap;
+		Args.ConstantValueNumKeys = Curve.ConstantValueNumKeys;
+		Args.CompressedKeys = CompressedKeys;
+
+		const float Value = FCompressedRichCurve::StaticEval(Args);
 		return Value;
 	};
 	
@@ -168,7 +181,18 @@ float UAnimCurveCompressionCodec_CompressedRichCurve::DecompressCurve(const FCom
 		{
 			const FCurveDesc& Curve = CurveDescriptions[CurveIndex];
 			const uint8* CompressedKeys = Buffer + Curve.KeyDataOffset;
-			const float Value = FCompressedRichCurve::StaticEval(Curve.CompressionFormat, Curve.KeyTimeCompressionFormat, Curve.PreInfinityExtrap, Curve.PostInfinityExtrap, Curve.ConstantValueNumKeys, CompressedKeys, CurrentTime);
+
+			FCompressedRichCurve::FStaticEvalArgs Args;
+			Args.CompressionFormat = Curve.CompressionFormat;
+			Args.KeyTimeCompressionFormat = Curve.KeyTimeCompressionFormat;
+			Args.PreInfinityExtrap = Curve.PreInfinityExtrap;
+			Args.PostInfinityExtrap = Curve.PostInfinityExtrap;
+			Args.ConstantValueNumKeys = Curve.ConstantValueNumKeys;
+			Args.CompressedKeys = CompressedKeys;
+			Args.Time = CurrentTime;
+			Args.DefaultValue = 0.0f;
+
+			const float Value = FCompressedRichCurve::StaticEval(Args);
 			return Value;
 		}
 	}
diff --git a/Engine/Source/Runtime/Engine/Private/Curves/RichCurve.cpp b/Engine/Source/Runtime/Engine/Private/Curves/RichCurve.cpp
--- a/Engine/Source/Runtime/Engine/Private/Curves/RichCurve.cpp
+++ b/Engine/Source/Runtime/Engine/Private/Curves/RichCurve.cpp
@@ -2052,7 +2052,7 @@ FORCEINLINE_DEBUGGABLE static float InterpEval(float InTime, const KeyTimeAdapte
 	const KeyDataHandle KeyValueHandle0 = KeyDataAdapter.GetKeyDataHandle(First - 1);
 	const float KeyValue0 = KeyDataAdapter.GetKeyValue(KeyValueHandle0);
 
-	auto IsItNotWeighted = [KeyDataAdapter](KeyDataHandle Handle0, KeyDataHandle Handle1) -> bool
+	const auto IsItNotWeighted = [&KeyDataAdapter](KeyDataHandle Handle0, KeyDataHandle Handle1) -> bool
 	{
 		const ERichCurveTangentWeightMode Mode0 = KeyDataAdapter.GetKeyTangentWeightMode(Handle0);
 		const ERichCurveTangentWeightMode Mode1 = KeyDataAdapter.GetKeyTangentWeightMode(Handle1);
@@ -2088,9 +2088,9 @@ FORCEINLINE_DEBUGGABLE static float InterpEval(float InTime, const KeyTimeAdapte
 
 			InterpolatedValue = UE::Curves::BezierInterp(P0, P1, P2, P3, Alpha);
 		}
-		else //it's weighted
+		else // it's weighted
 		{
-			return  UE::Curves::WeightedEvalForTwoKeys(
+			return UE::Curves::WeightedEvalForTwoKeys(
 				KeyValue0, KeyTime0, KeyDataAdapter.GetKeyLeaveTangent(KeyValueHandle0), KeyDataAdapter.GetKeyLeaveTangentWeight(KeyValueHandle0), KeyDataAdapter.GetKeyTangentWeightMode(First - 1),
 				KeyValue1, KeyTime1, KeyDataAdapter.GetKeyArriveTangent(KeyValueHandle1), KeyDataAdapter.GetKeyArriveTangentWeight(KeyValueHandle1), KeyDataAdapter.GetKeyTangentWeightMode(First),
 				InTime);
@@ -2104,109 +2104,134 @@ FORCEINLINE_DEBUGGABLE static float InterpEval(float InTime, const KeyTimeAdapte
 	return InterpolatedValue + CycleValueOffset;
 }
 
-static TFunction<float(ERichCurveExtrapolation PreInfinityExtrap, ERichCurveExtrapolation PostInfinityExtrap, FCompressedRichCurve::TConstantValueNumKeys ConstantValueNumKeys, const uint8* CompressedKeys, float InTime, float InDefaultValue)> InterpEvalMap[6][2]
+namespace UE::Anim::RichCurve::Private
 {
+	typedef float(*InterpEvalFun)(const FCompressedRichCurve::FStaticEvalArgs& Args);
+
 	// RCCF_Empty
+	static float StaticEval_Empty(const FCompressedRichCurve::FStaticEvalArgs& Args)
 	{
-		// RCKTCF_uint16
-		[](ERichCurveExtrapolation PreInfinityExtrap, ERichCurveExtrapolation PostInfinityExtrap, FCompressedRichCurve::TConstantValueNumKeys ConstantValueNumKeys, const uint8* CompressedKeys, float InTime, float InDefaultValue)
-		{
-			return ConstantValueNumKeys.ConstantValue == MAX_flt ? InDefaultValue : ConstantValueNumKeys.ConstantValue;
-		},
-		// RCKTCF_float32
-		[](ERichCurveExtrapolation PreInfinityExtrap, ERichCurveExtrapolation PostInfinityExtrap, FCompressedRichCurve::TConstantValueNumKeys ConstantValueNumKeys, const uint8* CompressedKeys, float InTime, float InDefaultValue)
-		{
-			return ConstantValueNumKeys.ConstantValue == MAX_flt ? InDefaultValue : ConstantValueNumKeys.ConstantValue;
-		},
-	},
+		return Args.ConstantValueNumKeys.ConstantValue == MAX_flt ? Args.DefaultValue : Args.ConstantValueNumKeys.ConstantValue;
+	}
+
 	// RCCF_Constant
+	static float StaticEval_Constant(const FCompressedRichCurve::FStaticEvalArgs& Args)
 	{
-		// RCKTCF_uint16
-		[](ERichCurveExtrapolation PreInfinityExtrap, ERichCurveExtrapolation PostInfinityExtrap, FCompressedRichCurve::TConstantValueNumKeys ConstantValueNumKeys, const uint8* CompressedKeys, float InTime, float InDefaultValue) { return ConstantValueNumKeys.ConstantValue; },
-		// RCKTCF_float32
-		[](ERichCurveExtrapolation PreInfinityExtrap, ERichCurveExtrapolation PostInfinityExtrap, FCompressedRichCurve::TConstantValueNumKeys ConstantValueNumKeys, const uint8* CompressedKeys, float InTime, float InDefaultValue) { return ConstantValueNumKeys.ConstantValue; },
-	},
-	// RCCF_Linear
+		return Args.ConstantValueNumKeys.ConstantValue;
+	}
+
+	// RCCF_Linear & RCKTCF_uint16
+	static float StaticEval_Linear_u16(const FCompressedRichCurve::FStaticEvalArgs& Args)
 	{
-		// RCKTCF_uint16
-		[](ERichCurveExtrapolation PreInfinityExtrap, ERichCurveExtrapolation PostInfinityExtrap, FCompressedRichCurve::TConstantValueNumKeys ConstantValueNumKeys, const uint8* CompressedKeys, float InTime, float InDefaultValue)
-		{
-			const int32 KeyTimesOffset = 0;
-			Quantized16BitKeyTimeAdapter KeyTimeAdapter(CompressedKeys, KeyTimesOffset, ConstantValueNumKeys.NumKeys);
-			UniformKeyDataAdapter<RCCF_Linear> KeyDataAdapter(CompressedKeys, KeyTimeAdapter);
-			return InterpEval(InTime, KeyTimeAdapter, KeyDataAdapter, ConstantValueNumKeys.NumKeys, PreInfinityExtrap, PostInfinityExtrap);
-		},
-		// RCKTCF_float32
-		[](ERichCurveExtrapolation PreInfinityExtrap, ERichCurveExtrapolation PostInfinityExtrap, FCompressedRichCurve::TConstantValueNumKeys ConstantValueNumKeys, const uint8* CompressedKeys, float InTime, float InDefaultValue)
-		{
-			const int32 KeyTimesOffset = 0;
-			Float32BitKeyTimeAdapter KeyTimeAdapter(CompressedKeys, KeyTimesOffset, ConstantValueNumKeys.NumKeys);
-			UniformKeyDataAdapter<RCCF_Linear> KeyDataAdapter(CompressedKeys, KeyTimeAdapter);
-			return InterpEval(InTime, KeyTimeAdapter, KeyDataAdapter, ConstantValueNumKeys.NumKeys, PreInfinityExtrap, PostInfinityExtrap);
-		},
-	},
-	// RCCF_Cubic
+		const int32 KeyTimesOffset = 0;
+		Quantized16BitKeyTimeAdapter KeyTimeAdapter(Args.CompressedKeys, KeyTimesOffset, Args.ConstantValueNumKeys.NumKeys);
+		UniformKeyDataAdapter<RCCF_Linear> KeyDataAdapter(Args.CompressedKeys, KeyTimeAdapter);
+		return InterpEval(Args.Time, KeyTimeAdapter, KeyDataAdapter, Args.ConstantValueNumKeys.NumKeys, Args.PreInfinityExtrap, Args.PostInfinityExtrap);
+	}
+
+	// RCCF_Linear & RCKTCF_float32
+	static float StaticEval_Linear_f32(const FCompressedRichCurve::FStaticEvalArgs& Args)
 	{
-		// RCKTCF_uint16
-		[](ERichCurveExtrapolation PreInfinityExtrap, ERichCurveExtrapolation PostInfinityExtrap, FCompressedRichCurve::TConstantValueNumKeys ConstantValueNumKeys, const uint8* CompressedKeys, float InTime, float InDefaultValue)
+		const int32 KeyTimesOffset = 0;
+		Float32BitKeyTimeAdapter KeyTimeAdapter(Args.CompressedKeys, KeyTimesOffset, Args.ConstantValueNumKeys.NumKeys);
+		UniformKeyDataAdapter<RCCF_Linear> KeyDataAdapter(Args.CompressedKeys, KeyTimeAdapter);
+		return InterpEval(Args.Time, KeyTimeAdapter, KeyDataAdapter, Args.ConstantValueNumKeys.NumKeys, Args.PreInfinityExtrap, Args.PostInfinityExtrap);
+	}
+
+	// RCCF_Cubic & RCKTCF_uint16
+	static float StaticEval_Cubic_u16(const FCompressedRichCurve::FStaticEvalArgs& Args)
+	{
+		const int32 KeyTimesOffset = 0;
+		Quantized16BitKeyTimeAdapter KeyTimeAdapter(Args.CompressedKeys, KeyTimesOffset, Args.ConstantValueNumKeys.NumKeys);
+		UniformKeyDataAdapter<RCCF_Cubic> KeyDataAdapter(Args.CompressedKeys, KeyTimeAdapter);
+		return InterpEval(Args.Time, KeyTimeAdapter, KeyDataAdapter, Args.ConstantValueNumKeys.NumKeys, Args.PreInfinityExtrap, Args.PostInfinityExtrap);
+	}
+
+	// RCCF_Cubic & RCKTCF_float32
+	static float StaticEval_Cubic_f32(const FCompressedRichCurve::FStaticEvalArgs& Args)
+	{
+		const int32 KeyTimesOffset = 0;
+		Float32BitKeyTimeAdapter KeyTimeAdapter(Args.CompressedKeys, KeyTimesOffset, Args.ConstantValueNumKeys.NumKeys);
+		UniformKeyDataAdapter<RCCF_Cubic> KeyDataAdapter(Args.CompressedKeys, KeyTimeAdapter);
+		return InterpEval(Args.Time, KeyTimeAdapter, KeyDataAdapter, Args.ConstantValueNumKeys.NumKeys, Args.PreInfinityExtrap, Args.PostInfinityExtrap);
+	}
+
+	// RCCF_Mixed & RCKTCF_uint16
+	static float StaticEval_Mixed_u16(const FCompressedRichCurve::FStaticEvalArgs& Args)
+	{
+		const int32 InterpModesOffset = 0;
+		const int32 KeyTimesOffset = InterpModesOffset + Align(Args.ConstantValueNumKeys.NumKeys * sizeof(uint8), sizeof(uint16));
+		Quantized16BitKeyTimeAdapter KeyTimeAdapter(Args.CompressedKeys, KeyTimesOffset, Args.ConstantValueNumKeys.NumKeys);
+		MixedKeyDataAdapter KeyDataAdapter(Args.CompressedKeys, InterpModesOffset, KeyTimeAdapter);
+		return InterpEval(Args.Time, KeyTimeAdapter, KeyDataAdapter, Args.ConstantValueNumKeys.NumKeys, Args.PreInfinityExtrap, Args.PostInfinityExtrap);
+	}
+
+	// RCCF_Mixed & RCKTCF_float32
+	static float StaticEval_Mixed_f32(const FCompressedRichCurve::FStaticEvalArgs& Args)
+	{
+		const int32 InterpModesOffset = 0;
+		const int32 KeyTimesOffset = InterpModesOffset + Align(Args.ConstantValueNumKeys.NumKeys * sizeof(uint8), sizeof(float));
+		Float32BitKeyTimeAdapter KeyTimeAdapter(Args.CompressedKeys, KeyTimesOffset, Args.ConstantValueNumKeys.NumKeys);
+		MixedKeyDataAdapter KeyDataAdapter(Args.CompressedKeys, InterpModesOffset, KeyTimeAdapter);
+		return InterpEval(Args.Time, KeyTimeAdapter, KeyDataAdapter, Args.ConstantValueNumKeys.NumKeys, Args.PreInfinityExtrap, Args.PostInfinityExtrap);
+	}
+
+	// RCCF_Weighted & RCKTCF_uint16
+	static float StaticEval_Weighted_u16(const FCompressedRichCurve::FStaticEvalArgs& Args)
+	{
+		const int32 InterpModesOffset = 0;
+		const int32 KeyTimesOffset = InterpModesOffset + Align(2 * Args.ConstantValueNumKeys.NumKeys * sizeof(uint8), sizeof(uint16));
+		Quantized16BitKeyTimeAdapter KeyTimeAdapter(Args.CompressedKeys, KeyTimesOffset, Args.ConstantValueNumKeys.NumKeys);
+		WeightedKeyDataAdapter KeyDataAdapter(Args.CompressedKeys, InterpModesOffset, KeyTimeAdapter);
+		return InterpEval(Args.Time, KeyTimeAdapter, KeyDataAdapter, Args.ConstantValueNumKeys.NumKeys, Args.PreInfinityExtrap, Args.PostInfinityExtrap);
+	}
+
+	// RCCF_Weighted & RCKTCF_float32
+	static float StaticEval_Weighted_f32(const FCompressedRichCurve::FStaticEvalArgs& Args)
+	{
+		const int32 InterpModesOffset = 0;
+		const int32 KeyTimesOffset = InterpModesOffset + Align(2 * Args.ConstantValueNumKeys.NumKeys * sizeof(uint8), sizeof(float));
+		Float32BitKeyTimeAdapter KeyTimeAdapter(Args.CompressedKeys, KeyTimesOffset, Args.ConstantValueNumKeys.NumKeys);
+		WeightedKeyDataAdapter KeyDataAdapter(Args.CompressedKeys, InterpModesOffset, KeyTimeAdapter);
+		return InterpEval(Args.Time, KeyTimeAdapter, KeyDataAdapter, Args.ConstantValueNumKeys.NumKeys, Args.PreInfinityExtrap, Args.PostInfinityExtrap);
+	}
+
+	// For performance reasons, we store a map of template optimized functions to evaluate rich curves
+	// Our array is two dimensional: ERichCurveCompressionFormat x ERichCurveKeyTimeCompressionFormat
+	static InterpEvalFun InterpEvalMap[6][2]
+	{
+		// RCCF_Empty
 		{
-			const int32 KeyTimesOffset = 0;
-			Quantized16BitKeyTimeAdapter KeyTimeAdapter(CompressedKeys, KeyTimesOffset, ConstantValueNumKeys.NumKeys);
-			UniformKeyDataAdapter<RCCF_Cubic> KeyDataAdapter(CompressedKeys, KeyTimeAdapter);
-			return InterpEval(InTime, KeyTimeAdapter, KeyDataAdapter, ConstantValueNumKeys.NumKeys, PreInfinityExtrap, PostInfinityExtrap);
+			&StaticEval_Empty,			// RCKTCF_uint16
+			&StaticEval_Empty,			// RCKTCF_float32
 		},
-		// RCKTCF_float32
-		[](ERichCurveExtrapolation PreInfinityExtrap, ERichCurveExtrapolation PostInfinityExtrap, FCompressedRichCurve::TConstantValueNumKeys ConstantValueNumKeys, const uint8* CompressedKeys, float InTime, float InDefaultValue)
+		// RCCF_Constant
 		{
-			const int32 KeyTimesOffset = 0;
-			Float32BitKeyTimeAdapter KeyTimeAdapter(CompressedKeys, KeyTimesOffset, ConstantValueNumKeys.NumKeys);
-			UniformKeyDataAdapter<RCCF_Cubic> KeyDataAdapter(CompressedKeys, KeyTimeAdapter);
-			return InterpEval(InTime, KeyTimeAdapter, KeyDataAdapter, ConstantValueNumKeys.NumKeys, PreInfinityExtrap, PostInfinityExtrap);
+			&StaticEval_Constant,		// RCKTCF_uint16
+			&StaticEval_Constant,		// RCKTCF_float32
 		},
-	},
-	// RCCF_Mixed
-	{
-		// RCKTCF_uint16
-		[](ERichCurveExtrapolation PreInfinityExtrap, ERichCurveExtrapolation PostInfinityExtrap, FCompressedRichCurve::TConstantValueNumKeys ConstantValueNumKeys, const uint8* CompressedKeys, float InTime, float InDefaultValue)
+		// RCCF_Linear
 		{
-			const int32 InterpModesOffset = 0;
-			const int32 KeyTimesOffset = InterpModesOffset + Align(ConstantValueNumKeys.NumKeys * sizeof(uint8), sizeof(uint16));
-			Quantized16BitKeyTimeAdapter KeyTimeAdapter(CompressedKeys, KeyTimesOffset, ConstantValueNumKeys.NumKeys);
-			MixedKeyDataAdapter KeyDataAdapter(CompressedKeys, InterpModesOffset, KeyTimeAdapter);
-			return InterpEval(InTime, KeyTimeAdapter, KeyDataAdapter, ConstantValueNumKeys.NumKeys, PreInfinityExtrap, PostInfinityExtrap);
+			&StaticEval_Linear_u16,		// RCKTCF_uint16
+			&StaticEval_Linear_f32,		// RCKTCF_float32
 		},
-		// RCKTCF_float32
-		[](ERichCurveExtrapolation PreInfinityExtrap, ERichCurveExtrapolation PostInfinityExtrap, FCompressedRichCurve::TConstantValueNumKeys ConstantValueNumKeys, const uint8* CompressedKeys, float InTime, float InDefaultValue)
+		// RCCF_Cubic
 		{
-			const int32 InterpModesOffset = 0;
-			const int32 KeyTimesOffset = InterpModesOffset + Align(ConstantValueNumKeys.NumKeys * sizeof(uint8), sizeof(float));
-			Float32BitKeyTimeAdapter KeyTimeAdapter(CompressedKeys, KeyTimesOffset, ConstantValueNumKeys.NumKeys);
-			MixedKeyDataAdapter KeyDataAdapter(CompressedKeys, InterpModesOffset, KeyTimeAdapter);
-			return InterpEval(InTime, KeyTimeAdapter, KeyDataAdapter, ConstantValueNumKeys.NumKeys, PreInfinityExtrap, PostInfinityExtrap);
+			&StaticEval_Cubic_u16,		// RCKTCF_uint16
+			&StaticEval_Cubic_f32,		// RCKTCF_float32
 		},
-	},
-	// RCCF_Weighted
-	{
-		// RCKTCF_uint16
-		[](ERichCurveExtrapolation PreInfinityExtrap, ERichCurveExtrapolation PostInfinityExtrap, FCompressedRichCurve::TConstantValueNumKeys ConstantValueNumKeys, const uint8* CompressedKeys, float InTime, float InDefaultValue)
+		// RCCF_Mixed
 		{
-			const int32 InterpModesOffset = 0;
-			const int32 KeyTimesOffset = InterpModesOffset + Align(2 * ConstantValueNumKeys.NumKeys * sizeof(uint8), sizeof(uint16));
-			Quantized16BitKeyTimeAdapter KeyTimeAdapter(CompressedKeys, KeyTimesOffset, ConstantValueNumKeys.NumKeys);
-			WeightedKeyDataAdapter KeyDataAdapter(CompressedKeys, InterpModesOffset, KeyTimeAdapter);
-			return InterpEval(InTime, KeyTimeAdapter, KeyDataAdapter, ConstantValueNumKeys.NumKeys, PreInfinityExtrap, PostInfinityExtrap);
+			&StaticEval_Mixed_u16,		// RCKTCF_uint16
+			&StaticEval_Mixed_f32,		// RCKTCF_float32
 		},
-			// RCKTCF_float32
-			[](ERichCurveExtrapolation PreInfinityExtrap, ERichCurveExtrapolation PostInfinityExtrap, FCompressedRichCurve::TConstantValueNumKeys ConstantValueNumKeys, const uint8* CompressedKeys, float InTime, float InDefaultValue)
-			{
-				const int32 InterpModesOffset = 0;
-				const int32 KeyTimesOffset = InterpModesOffset + Align(2 * ConstantValueNumKeys.NumKeys * sizeof(uint8), sizeof(float));
-				Float32BitKeyTimeAdapter KeyTimeAdapter(CompressedKeys, KeyTimesOffset, ConstantValueNumKeys.NumKeys);
-				WeightedKeyDataAdapter KeyDataAdapter(CompressedKeys, InterpModesOffset, KeyTimeAdapter);
-				return InterpEval(InTime, KeyTimeAdapter, KeyDataAdapter, ConstantValueNumKeys.NumKeys, PreInfinityExtrap, PostInfinityExtrap);
-			},
+		// RCCF_Weighted
+		{
+			&StaticEval_Weighted_u16,	// RCKTCF_uint16
+			&StaticEval_Weighted_f32,	// RCKTCF_float32
 		},
-};
+	};
+}
 
 ERichCurveInterpMode GetKeyInterpMode(ERichCurveCompressionFormat Format)
 {
@@ -2398,15 +2423,39 @@ void FCompressedRichCurve::PopulateCurve(FRichCurve& OutCurve) const
 
 float FCompressedRichCurve::Eval(float InTime, float InDefaultValue) const
 {
-	// Dynamic dispatch into a template optimized code path
-	const float Value = InterpEvalMap[CompressionFormat][KeyTimeCompressionFormat](PreInfinityExtrap, PostInfinityExtrap, ConstantValueNumKeys, CompressedKeys.GetData(), InTime, InDefaultValue);
-	return Value;
+	FStaticEvalArgs Args;
+	Args.CompressionFormat = CompressionFormat;
+	Args.KeyTimeCompressionFormat = KeyTimeCompressionFormat;
+	Args.PreInfinityExtrap = PreInfinityExtrap;
+	Args.PostInfinityExtrap = PostInfinityExtrap;
+	Args.ConstantValueNumKeys = ConstantValueNumKeys;
+	Args.CompressedKeys = CompressedKeys.GetData();
+	Args.Time = InTime;
+	Args.DefaultValue = InDefaultValue;
+
+	return StaticEval(Args);
 }
 
 float FCompressedRichCurve::StaticEval(ERichCurveCompressionFormat CompressionFormat, ERichCurveKeyTimeCompressionFormat KeyTimeCompressionFormat, ERichCurveExtrapolation PreInfinityExtrap, ERichCurveExtrapolation PostInfinityExtrap, FCompressedRichCurve::TConstantValueNumKeys ConstantValueNumKeys, const uint8* CompressedKeys, float InTime, float InDefaultValue)
+{
+	FStaticEvalArgs Args;
+	Args.CompressionFormat = CompressionFormat;
+	Args.KeyTimeCompressionFormat = KeyTimeCompressionFormat;
+	Args.PreInfinityExtrap = PreInfinityExtrap;
+	Args.PostInfinityExtrap = PostInfinityExtrap;
+	Args.ConstantValueNumKeys = ConstantValueNumKeys;
+	Args.CompressedKeys = CompressedKeys;
+	Args.Time = InTime;
+	Args.DefaultValue = InDefaultValue;
+
+	return StaticEval(Args);
+}
+
+float FCompressedRichCurve::StaticEval(const FStaticEvalArgs& Args)
 {
 	// Dynamic dispatch into a template optimized code path
-	const float Value = InterpEvalMap[CompressionFormat][KeyTimeCompressionFormat](PreInfinityExtrap, PostInfinityExtrap, ConstantValueNumKeys, CompressedKeys, InTime, InDefaultValue);
+	const UE::Anim::RichCurve::Private::InterpEvalFun EvalFun = UE::Anim::RichCurve::Private::InterpEvalMap[Args.CompressionFormat][Args.KeyTimeCompressionFormat];
+	const float Value = (*EvalFun)(Args);
 	return Value;
 }
 
