Refactor of Verse Material Property Reflection to a common practice inheritance structure for generic property reflection per asset type.

diff --git a/Engine/Source/Runtime/Engine/Public/MaterialTypes.h b/Engine/Source/Runtime/Engine/Public/MaterialTypes.h
--- a/Engine/Source/Runtime/Engine/Public/MaterialTypes.h
+++ b/Engine/Source/Runtime/Engine/Public/MaterialTypes.h
@@ -181,6 +181,7 @@ FORCEINLINE FMaterialParameterInfo::FMaterialParameterInfo(const struct FMemoryI
 // Backwards compat
 using FHashedMaterialParameterInfo = FMemoryImageMaterialParameterInfo;
 
+UENUM()
 enum class EMaterialParameterType : uint8
 {
 	Scalar = 0u,
@@ -234,6 +235,25 @@ static inline bool IsStaticMaterialParameter(EMaterialParameterType InType)
 
 extern ENGINE_API UE::Shader::FType GetShaderValueType(EMaterialParameterType Type);
 
+static inline const UEnum* GetStaticMaterialParameterTypeEnum()
+{
+	static const UEnum* Enum = StaticEnum<EMaterialParameterType>();
+	check(Enum);
+	return Enum;
+}
+
+static inline FString MaterialParameterTypeToString(EMaterialParameterType Type)
+{
+	const UEnum* Enum = GetStaticMaterialParameterTypeEnum();
+	return Enum->GetNameStringByValue(int64(Type));
+}
+
+static inline EMaterialParameterType GetMaterialParameterTypeFromString(FString TypeName)
+{
+	const UEnum* Enum = GetStaticMaterialParameterTypeEnum();
+	return static_cast<EMaterialParameterType>(Enum->GetValueByName(FName(TypeName)));
+}
+
 enum class EMaterialGetParameterValueFlags : uint32
 {
 	None = 0u,
