Fixed zero length sweeps to behave the same as a regular sweep.
Zero length sweeps used to fall back to calling Overlap, now they still go through the sweep path.
Also fixed height fields not populating the face index in sweeps.

diff --git a/Engine/Source/Programs/HeadlessChaos/Private/HeadlessChaosTestSweep.cpp b/Engine/Source/Programs/HeadlessChaos/Private/HeadlessChaosTestSweep.cpp
--- a/Engine/Source/Programs/HeadlessChaos/Private/HeadlessChaosTestSweep.cpp
+++ b/Engine/Source/Programs/HeadlessChaos/Private/HeadlessChaosTestSweep.cpp
@@ -946,9 +946,8 @@ namespace ChaosTest
 			return Result;
 		}
 
-		FSweepResultData SweepSphere(const FTriangleMeshImplicitObject& TriangleMesh, const FVec3& SweepStart, const FVec3& SweepDirection)
+		FSweepResultData SweepSphere(const FTriangleMeshImplicitObject& TriangleMesh, const FVec3& SweepStart, const FVec3& SweepDirection, const FReal SweepLength = 100)
 		{
-			const FReal SweepLength = 100;
 			const FReal SphereRadius = 25;
 			const TSphere<FReal, 3> Sphere(FVec3::ZeroVector, SphereRadius);
 			const FRigidTransform3 StartTM(SweepStart, TRotation<FReal, 3>::Identity);
@@ -959,10 +958,10 @@ namespace ChaosTest
 			return Result;
 		}
 
-		FSweepResultData SweepSphere(FMeshInitData& InitData, const FVec3& SweepStart, const FVec3& SweepDirection)
+		FSweepResultData SweepSphere(FMeshInitData& InitData, const FVec3& SweepStart, const FVec3& SweepDirection, const FReal SweepLength = 100)
 		{
 			FTriangleMeshImplicitObject TriangleMesh(MoveTemp(InitData.TrimeshParticles), MoveTemp(InitData.Indices), MoveTemp(InitData.Materials), nullptr, nullptr, true);
-			return SweepSphere(TriangleMesh, SweepStart, SweepDirection);
+			return SweepSphere(TriangleMesh, SweepStart, SweepDirection, SweepLength);
 		}
 	};
 
@@ -1011,4 +1010,36 @@ namespace ChaosTest
 		EXPECT_VECTOR_NEAR(Result.Normal, FVec3(0, 0, 1), KINDA_SMALL_NUMBER);
 		EXPECT_EQ(Result.FaceId, 1);
 	}
+
+	TEST_F(FTriangleMeshSweepFixture, SphereWithZeroLengthSweep_TestAllAxes_HitIsExpected)
+	{
+		FMeshInitData InitData = CreateSimpleTriInitData();
+		FTriangleMeshImplicitObject TriangleMesh(MoveTemp(InitData.TrimeshParticles), MoveTemp(InitData.Indices), MoveTemp(InitData.Materials), nullptr, nullptr, true);
+
+		// Test each cardinal axis
+		FSweepResultData Result = SweepSphere(TriangleMesh, FVec3(20, 20, 0), FVec3(-1, 0, 0), 0);
+		EXPECT_EQ(Result.bResult, true);
+
+		Result = SweepSphere(TriangleMesh, FVec3(20, 20, 0), FVec3(1, 0, 0), 0);
+		EXPECT_EQ(Result.bResult, true);
+
+		Result = SweepSphere(TriangleMesh, FVec3(20, 20, 0), FVec3(0, -1, 0), 0);
+		EXPECT_EQ(Result.bResult, true);
+
+		Result = SweepSphere(TriangleMesh, FVec3(20, 20, 0), FVec3(0, 1, 0), 0);
+		EXPECT_EQ(Result.bResult, true);
+
+		Result = SweepSphere(TriangleMesh, FVec3(20, 20, 0), FVec3(0, 0, -1), 0);
+		EXPECT_EQ(Result.bResult, true);
+
+		Result = SweepSphere(TriangleMesh, FVec3(20, 20, 0), FVec3(0, 0, 1), 0);
+		EXPECT_EQ(Result.bResult, true);
+
+		// Test the zero vector direction
+		Result = SweepSphere(TriangleMesh, FVec3(20, 20, 0), FVec3(0, 0, 0), 0);
+		EXPECT_EQ(Result.bResult, true);
+
+		Result = SweepSphere(TriangleMesh, FVec3(20, 20, 0), FVec3(0, 0, 0), 1);
+		EXPECT_EQ(Result.bResult, true);
+	}
 }
diff --git a/Engine/Source/Runtime/Engine/Private/Collision/CollisionConversions.cpp b/Engine/Source/Runtime/Engine/Private/Collision/CollisionConversions.cpp
--- a/Engine/Source/Runtime/Engine/Private/Collision/CollisionConversions.cpp
+++ b/Engine/Source/Runtime/Engine/Private/Collision/CollisionConversions.cpp
@@ -294,11 +294,20 @@ EConvertQueryResult ConvertQueryImpactHitImp(const UWorld* World, const THitLoca
 
 		if(pHitShape)
 		{
-			const bool bTriMesh = GetGeometryType(*pHitShape) == ECollisionShapeType::Trimesh;
+			const ECollisionShapeType GeometryType = GetGeometryType(*pHitShape);
+			const bool bTriMesh = GeometryType == ECollisionShapeType::Trimesh;
+			const bool bHeightField = GeometryType == ECollisionShapeType::Heightfield;
 			const bool bValidInternalFace = InternalFaceIndex != GetInvalidPhysicsFaceIndex();
-			if(bReturnFaceIndex && bTriMesh && bValidInternalFace)
+			if(bReturnFaceIndex && bValidInternalFace)
 			{
-				OutResult.FaceIndex = GetTriangleMeshExternalFaceIndex(*pHitShape, InternalFaceIndex);
+				if (bTriMesh)
+				{
+					OutResult.FaceIndex = GetTriangleMeshExternalFaceIndex(*pHitShape, InternalFaceIndex);
+				}
+				else if (bHeightField)
+				{
+					OutResult.FaceIndex = InternalFaceIndex;
+				}
 			}
 		}
 
@@ -409,12 +418,19 @@ EConvertQueryResult ConvertQueryImpactHitImp(const UWorld* World, const THitLoca
 
 	if(GeomType == ECollisionShapeType::Heightfield)
 	{
-		// Lookup physical material for heightfields
-		if (bReturnPhysMat && InternalFaceIndex != GetInvalidPhysicsFaceIndex() && GTActor && GTShape)
+		if (InternalFaceIndex != GetInvalidPhysicsFaceIndex())
 		{
-			if (const FPhysicsMaterial* Material = GetMaterialFromInternalFaceIndex(*GTShape, *GTActor, InternalFaceIndex))
+			// Lookup physical material for heightfields
+			if (bReturnPhysMat && GTActor && GTShape)
+			{
+				if (const FPhysicsMaterial* Material = GetMaterialFromInternalFaceIndex(*GTShape, *GTActor, InternalFaceIndex))
+				{
+					OutResult.PhysMaterial = GetUserData(*Material);
+				}
+			}
+			if (bReturnFaceIndex)
 			{
-				OutResult.PhysMaterial = GetUserData(*Material);
+				OutResult.FaceIndex = InternalFaceIndex;
 			}
 		}
 	}
@@ -559,7 +575,8 @@ static bool ConvertOverlappedShapeToImpactHit(const UWorld* World, const THitLoc
 
 	const FVector HitNormal = GetNormal(Hit);
 	const bool bFiniteNormal = !HitNormal.ContainsNaN();
-	const bool bValidNormal = (GetFlags(Hit) & EHitFlags::Normal) && bFiniteNormal;
+	const bool bHasNormal = (GetFlags(Hit) & EHitFlags::Normal) != EHitFlags::None;
+	const bool bValidNormal = bHasNormal && bFiniteNormal;
 
 	// Use MTD result if possible. We interpret the MTD vector as both the direction to move and the opposing normal.
 	if (bValidNormal)
@@ -600,7 +617,18 @@ static bool ConvertOverlappedShapeToImpactHit(const UWorld* World, const THitLoc
 		if (!bValidNormal)
 		{
 			OutResult.ImpactNormal = (StartLoc - EndLoc).GetSafeNormal();
-			ensure(OutResult.ImpactNormal.IsNormalized());
+			// Since zero-length sweeps are allowed, it is possible that we could compute a zero vector for the normal.
+			// We want to always return a valid normal, however we only want to assert if a normal is requested but not able to be computed.
+			const bool bIsNormalized = OutResult.ImpactNormal.IsNormalized();
+			if (!bIsNormalized)
+			{
+				if (bHasNormal)
+				{
+					ensureMsgf(false, TEXT("Unable to comptue valid normal from start/end when normal was requested. Start: %s, End: %s."), 
+						*StartLoc.ToString(), *EndLoc.ToString());
+				}
+				OutResult.ImpactNormal = FVector::UpVector;
+			}
 		}
 	}
 
diff --git a/Engine/Source/Runtime/Experimental/Chaos/Private/Chaos/TriangleMeshImplicitObject.cpp b/Engine/Source/Runtime/Experimental/Chaos/Private/Chaos/TriangleMeshImplicitObject.cpp
--- a/Engine/Source/Runtime/Experimental/Chaos/Private/Chaos/TriangleMeshImplicitObject.cpp
+++ b/Engine/Source/Runtime/Experimental/Chaos/Private/Chaos/TriangleMeshImplicitObject.cpp
@@ -1613,7 +1613,9 @@ bool FTriangleMeshImplicitObject::SweepGeomImp(const QueryGeomType& QueryGeom, c
 	bool bHit = false;
 	auto LambdaHelper = [&](const auto& Elements)
 	{
-		const FReal CullsBackFaceRaycastCode = bCullsBackFaceRaycast ? GetWindingOrder(TriMeshScale) : 0.f;
+		// When length is zero, the direction of the ray is currently a hardcoded value (since it's computed from the start/end which would be equal). 
+		// We don't want the backface culling check to be arbitrary based upon the hardcoded direction, so we disable it with zero-length sweeps.
+		const FReal CullsBackFaceRaycastCode = (bCullsBackFaceRaycast && Length != 0) ? GetWindingOrder(TriMeshScale) : 0.f;
 		using VisitorType = FTriangleMeshSweepVisitor<QueryGeomType,decltype(Elements[0][0])>;
 		VisitorType SQVisitor(*this, Elements, QueryGeom, ScaledDirNormalized, LengthScale, ScaledStartTM, bComputeMTD, TriMeshScale, CullsBackFaceRaycastCode);
 
diff --git a/Engine/Source/Runtime/Experimental/Chaos/Public/Chaos/GJK.h b/Engine/Source/Runtime/Experimental/Chaos/Public/Chaos/GJK.h
--- a/Engine/Source/Runtime/Experimental/Chaos/Public/Chaos/GJK.h
+++ b/Engine/Source/Runtime/Experimental/Chaos/Public/Chaos/GJK.h
@@ -1387,7 +1387,7 @@ namespace Chaos
 		T& OutTime, TVector<T, 3>& OutPosition, TVector<T, 3>& OutNormal, const T ThicknessA = 0, const TVector<T, 3>& InitialDir = TVector<T, 3>(1, 0, 0), const T ThicknessB = 0)
 	{
 		ensure(FMath::IsNearlyEqual(RayDir.SizeSquared(), (FReal)1, (FReal)UE_KINDA_SMALL_NUMBER));
-		ensure(RayLength > 0);
+		ensure(RayLength >= 0);
 		check(A.IsConvex() && B.IsConvex());
 		const TVector<T, 3> StartPoint = StartTM.GetLocation();
 		int32 VertexIndexA = INDEX_NONE;
@@ -1518,15 +1518,16 @@ namespace Chaos
 	 @OutTime The time along the ray when the objects first overlap
 	 @OutPosition The first point of impact (in A's local space) when the objects first overlap. Invalid if time of impact is 0
 	 @OutNormal The impact normal (in A's local space) when the objects first overlap. Invalid if time of impact is 0
+	 @bComputeMTD If there is an initial overlap, this will compute the output info from the overlap info.
 	 @InitialDir The first direction we use to search the CSO
 	 @return True if the geometries overlap during the sweep, False otherwise
-	 @note If A overlaps B at the start of the ray ("initial overlap" condition) then this function returns true, and sets OutTime = 0, but does not set any other output variables.
+	 @note If A overlaps B at the start of the ray ("initial overlap" condition) then this function returns true, and sets OutTime = 0, but does not set any other output variables unless bComputeMTD is true.
 	 */
 	template<typename T>
 	bool GJKRaycast2ImplSimd(const FGeomGJKHelperSIMD& RESTRICT A, const FGeomGJKHelperSIMD& RESTRICT B, const T& BToARotation, const T& StartPoint, const T& RayDir,
 		FRealSingle RayLength, FRealSingle& OutTime, T& OutPosition, T& OutNormal, bool bComputeMTD, const T& InitialDir)
 	{
-		ensure(RayLength > 0);
+		ensure(RayLength >= 0);
 
 		// Margin selection logic: we only need a small margin for sweeps since we only move the sweeping object
 		// to the point where it just touches.
diff --git a/Engine/Source/Runtime/Experimental/Chaos/Public/Chaos/GeometryQueries.h b/Engine/Source/Runtime/Experimental/Chaos/Public/Chaos/GeometryQueries.h
--- a/Engine/Source/Runtime/Experimental/Chaos/Public/Chaos/GeometryQueries.h
+++ b/Engine/Source/Runtime/Experimental/Chaos/Public/Chaos/GeometryQueries.h
@@ -177,7 +177,8 @@ namespace Chaos
 		ensure(FMath::IsNearlyEqual(Dir.SizeSquared(), (FReal)1, (FReal)UE_KINDA_SMALL_NUMBER)); // Added to help determine cause of this ensure firing in GJKRaycast2.
 		const FVec3 LocalDir = ATM.InverseTransformVectorNoScale(Dir);
 
-		bool bSweepAsRaycast = BType == ImplicitObjectType::Sphere && !bComputeMTD;
+		// We can optimize some sweeps as a raycast. Sweeps can have zero length and most raycast functions don't handle this.
+		bool bSweepAsRaycast = BType == ImplicitObjectType::Sphere && !bComputeMTD && Length != 0;
 		if (bSweepAsRaycast && IsScaled(AType))
 		{
 			const auto& Scaled = TImplicitObjectScaledGeneric<FReal, 3>::AsScaledChecked(A);
diff --git a/Engine/Source/Runtime/Experimental/Chaos/Public/Chaos/ImplicitObjectScaled.h b/Engine/Source/Runtime/Experimental/Chaos/Public/Chaos/ImplicitObjectScaled.h
--- a/Engine/Source/Runtime/Experimental/Chaos/Public/Chaos/ImplicitObjectScaled.h
+++ b/Engine/Source/Runtime/Experimental/Chaos/Public/Chaos/ImplicitObjectScaled.h
@@ -637,7 +637,7 @@ public:
 	template <typename QueryGeomType>
 	bool LowLevelSweepGeom(const QueryGeomType& B, const TRigidTransform<T, d>& BToATM, const TVector<T, d>& LocalDir, const T Length, T& OutTime, TVector<T, d>& LocalPosition, TVector<T, d>& LocalNormal, int32& OutFaceIndex, TVector<T, d>& OutFaceNormal, T Thickness = 0, bool bComputeMTD = false) const
 	{
-		ensure(Length > 0);
+		ensure(Length >= 0);
 		ensure(FMath::IsNearlyEqual(LocalDir.SizeSquared(), (T)1, (T)UE_KINDA_SMALL_NUMBER));
 		ensure(Thickness == 0 || (FMath::IsNearlyEqual(MScale[0], MScale[1]) && FMath::IsNearlyEqual(MScale[0], MScale[2])));
 
diff --git a/Engine/Source/Runtime/PhysicsCore/Public/SQVisitor.h b/Engine/Source/Runtime/PhysicsCore/Public/SQVisitor.h
--- a/Engine/Source/Runtime/PhysicsCore/Public/SQVisitor.h
+++ b/Engine/Source/Runtime/PhysicsCore/Public/SQVisitor.h
@@ -291,11 +291,11 @@ private:
 						WorldNormal = ActorTM.TransformVectorNoScale(LocalNormal);
 					}
 				}
-				else if (SQ == ESQType::Sweep && CurData->CurrentLength > 0 && ensure(QueryGeom))
+				else if (SQ == ESQType::Sweep && ensure(QueryGeom))
 				{
 					bHit = SweepQuery(*Geom, ActorTM, *QueryGeom, StartTM, CurData->Dir, CurData->CurrentLength, Distance, WorldPosition, WorldNormal, FaceIdx, FaceNormal, 0.f, bComputeMTD);
 				}
-				else if ((SQ == ESQType::Overlap || (SQ == ESQType::Sweep && CurData->CurrentLength == 0)) && ensure(QueryGeom))
+				else if (SQ == ESQType::Overlap && ensure(QueryGeom))
 				{
 					if (bComputeMTD)
 					{
@@ -627,21 +627,14 @@ void SweepHelper(const QueryGeomType& QueryGeom, const Chaos::ISpatialAccelerati
 	using namespace ChaosInterface;
 
 	const FAABB3 Bounds = QueryGeom.BoundingBox().TransformedAABB(StartTM);
-	const bool bSweepAsOverlap = DeltaMagnitude == 0;	//question: do we care about tiny sweeps?
 	TSQVisitor<QueryGeomType, TPayload, TSweepHit, std::is_same<TSweepHit, FSweepHit>::value> SweepVisitor(StartTM, Dir, HitBuffer, OutputFlags, QueryFilterData, QueryCallback, QueryGeom, DebugParams);
 
 	HitBuffer.IncFlushCount();
 
-	if (bSweepAsOverlap)
-	{
-		//fallback to overlap
-		SpatialAcceleration.Overlap(Bounds, SweepVisitor);
-	}
-	else
-	{
-		const FVector HalfExtents = Bounds.Extents() * 0.5f;
-		SpatialAcceleration.Sweep(Bounds.GetCenter(), Dir, DeltaMagnitude, HalfExtents, SweepVisitor);
-	}
+	// NOTE: Previously, we'd attempt to perform a sweep as an overlap if the DeltaMagnitude was 0. 
+	// This caused various discrepancies as some of the out info couldn't be computed. Instead, sweep has been updated to handle 0 length sweeps.
+	const FVector HalfExtents = Bounds.Extents() * 0.5f;
+	SpatialAcceleration.Sweep(Bounds.GetCenter(), Dir, DeltaMagnitude, HalfExtents, SweepVisitor);
 
 	HitBuffer.DecFlushCount();
 }
