[World Partition] Fix negative zeros crawling into the cells bounds Z value that were causing unit tests failures following 36907904.

diff --git a/Engine/Source/Runtime/Engine/Private/WorldPartition/RuntimeHashSet/RuntimePartitionLHGrid.cpp b/Engine/Source/Runtime/Engine/Private/WorldPartition/RuntimeHashSet/RuntimePartitionLHGrid.cpp
--- a/Engine/Source/Runtime/Engine/Private/WorldPartition/RuntimeHashSet/RuntimePartitionLHGrid.cpp
+++ b/Engine/Source/Runtime/Engine/Private/WorldPartition/RuntimeHashSet/RuntimePartitionLHGrid.cpp
@@ -190,16 +190,22 @@ bool URuntimePartitionLHGrid::GenerateStreaming(const FGenerateStreamingParams&
 	UWorld* World = WorldPartition->GetWorld();
 	UWorld* OuterWorld = GetTypedOuter<UWorld>();
 	const bool bIsMainWorldPartition = (World == OuterWorld);
-	const FVector SpaceMask = FVector(1, 1, bIs2D ? 0 : 1);
 
 	TMap<FCellCoord, TArray<const IStreamingGenerationContext::FActorSetInstance*>> CellsActorSetInstances;
 	for (const IStreamingGenerationContext::FActorSetInstance* ActorSetInstance : *InParams.ActorSetInstances)
 	{
 		if (ActorSetInstance->bIsSpatiallyLoaded)
 		{
-			const FBox ActorSetInstanceBounds = FBox(ActorSetInstance->Bounds.Min * SpaceMask, ActorSetInstance->Bounds.Max * SpaceMask);
-			const int32 GridLevel = FCellCoord::GetLevelForBox(ActorSetInstanceBounds, CellSize, Origin * SpaceMask);
-			const FCellCoord CellCoord = FCellCoord::GetCellCoords(ActorSetInstanceBounds.GetCenter(), CellSize, GridLevel, Origin * SpaceMask);
+			FBox ActorSetInstanceBounds = ActorSetInstance->Bounds;
+
+			if (bIs2D)
+			{
+				ActorSetInstanceBounds.Min.Z = 0;
+				ActorSetInstanceBounds.Max.Z = 0;
+			}
+
+			const int32 GridLevel = FCellCoord::GetLevelForBox(ActorSetInstanceBounds, CellSize, FVector(Origin.X, Origin.Y, bIs2D ? 0 : Origin.Z));
+			const FCellCoord CellCoord = FCellCoord::GetCellCoords(ActorSetInstanceBounds.GetCenter(), CellSize, GridLevel, FVector(Origin.X, Origin.Y, bIs2D ? 0 : Origin.Z));
 			CellsActorSetInstances.FindOrAdd(CellCoord).Add(ActorSetInstance);
 		}
 		else
@@ -216,9 +222,14 @@ bool URuntimePartitionLHGrid::GenerateStreaming(const FGenerateStreamingParams&
 
 		if (bIsSpatiallyLoaded)
 		{
-			CellDesc.CellBounds = FCellCoord::GetCellBounds(CellCoord, CellSize, Origin * SpaceMask);
-			CellDesc.CellBounds.GetValue().Min *= SpaceMask;
-			CellDesc.CellBounds.GetValue().Max *= SpaceMask;
+			CellDesc.CellBounds = FCellCoord::GetCellBounds(CellCoord, CellSize, FVector(Origin.X, Origin.Y, bIs2D ? 0 : Origin.Z));
+
+			if (bIs2D)
+			{
+				CellDesc.CellBounds.GetValue().Min.Z = 0;
+				CellDesc.CellBounds.GetValue().Max.Z = 0;
+			}
+
 			CellDesc.bIs2D = bIs2D;
 		}
 	}
diff --git a/Engine/Source/Runtime/Engine/Private/WorldPartition/RuntimeHashSet/WorldPartitionRuntimeHashSetStreamingGeneration.cpp b/Engine/Source/Runtime/Engine/Private/WorldPartition/RuntimeHashSet/WorldPartitionRuntimeHashSetStreamingGeneration.cpp
--- a/Engine/Source/Runtime/Engine/Private/WorldPartition/RuntimeHashSet/WorldPartitionRuntimeHashSetStreamingGeneration.cpp
+++ b/Engine/Source/Runtime/Engine/Private/WorldPartition/RuntimeHashSet/WorldPartitionRuntimeHashSetStreamingGeneration.cpp
@@ -180,10 +180,12 @@ bool UWorldPartitionRuntimeHashSet::GenerateStreaming(UWorldPartitionStreamingPo
 				PopulateRuntimeCell(RuntimeCell, CellActorInstances, OutPackagesToGenerate);
 
 				UWorldPartitionRuntimeCellDataHashSet* RuntimeCellData = CastChecked<UWorldPartitionRuntimeCellDataHashSet>(RuntimeCell->RuntimeCellData);
-				const FVector SpaceMask(1, 1, RuntimeCellData->bIs2D ? 0 : 1);
 
-				RuntimeCell->RuntimeCellData->ContentBounds.Min *= SpaceMask;
-				RuntimeCell->RuntimeCellData->ContentBounds.Max *= SpaceMask;
+				if (RuntimeCellData->bIs2D)
+				{
+					RuntimeCell->RuntimeCellData->ContentBounds.Min.Z = 0.0;
+					RuntimeCell->RuntimeCellData->ContentBounds.Max.Z = 0.0;
+				}
 
 				if (CellDescInstance.CellBounds.IsSet())
 				{
@@ -286,4 +288,4 @@ void UWorldPartitionRuntimeHashSet::DumpStateLog(FHierarchicalLogArchive& Ar) co
 
 	Ar.Printf(TEXT(""));
 }
-#endif
\ No newline at end of file
+#endif
