Fix memory leak in virtualized metasounds sounds getting sent parameters.
#rnx

diff --git a/Engine/Source/Runtime/Engine/Private/AudioDevice.cpp b/Engine/Source/Runtime/Engine/Private/AudioDevice.cpp
--- a/Engine/Source/Runtime/Engine/Private/AudioDevice.cpp
+++ b/Engine/Source/Runtime/Engine/Private/AudioDevice.cpp
@@ -5084,6 +5084,7 @@ void FAudioDevice::AddNewActiveSoundInternal(const FActiveSound& InNewActiveSoun
 	{
 		DECLARE_CYCLE_STAT(TEXT("FAudioThreadTask.AddNewActiveSound"), STAT_AudioAddNewActiveSound, STATGROUP_AudioThreadCommands);
 
+		checkf(nullptr == InVirtualLoopToRetrigger, TEXT("Virtual loops can only be retrigger from the audio thread"));
 
 		FAudioThread::RunCommandOnAudioThread([AudioDevice = this, InNewActiveSound, DefaultParams = MoveTemp(InDefaultParams)]() mutable
 		{
@@ -5134,24 +5135,6 @@ void FAudioDevice::AddNewActiveSoundInternal(const FActiveSound& InNewActiveSoun
 	}
 #endif // !UE_BUILD_SHIPPING
 
-	auto InitSoundParams = [this, &Sound](FActiveSound& OutActiveSound, TArray<FAudioParameter>&& DefaultParams)
-	{
-		// Retriggering a virtualized ActiveSound which already have a transmitter
-		// should not be given a new transmitter.
-		if (!OutActiveSound.InstanceTransmitter.IsValid())
-		{
-			Audio::FParameterTransmitterInitParams TransmitterInitParams
-			{
-				Audio::GetTransmitterID(OutActiveSound.GetAudioComponentID(), 0, OutActiveSound.GetPlayOrder()),
-				GetSampleRate(),
-				MoveTemp(DefaultParams),
-				DeviceID
-			};
-
-			OutActiveSound.InstanceTransmitter = Sound->CreateParameterTransmitter(MoveTemp(TransmitterInitParams));
-		}
-	};
-
 	// Determine if sound is loop and eligible for virtualize prior to creating "live" active sound in next Concurrency check step
 	if (!InVirtualLoopToRetrigger)
 	{
@@ -5160,7 +5143,8 @@ void FAudioDevice::AddNewActiveSoundInternal(const FActiveSound& InNewActiveSoun
 		if (FAudioVirtualLoop::Virtualize(InNewActiveSound, *this, bDoRangeCheck, VirtualLoop))
 		{
 			UE_LOG(LogAudio, Verbose, TEXT("New ActiveSound %s Virtualizing: Failed to pass initial audible range check"), *Sound->GetName());
-			InitSoundParams(VirtualLoop.GetActiveSound(), MoveTemp(InDefaultParams));
+			constexpr bool bIsVirtualLoopRealizing = false;
+			InitSoundParams(VirtualLoop.GetActiveSound(), MoveTemp(InDefaultParams), bIsVirtualLoopRealizing);
 			AddVirtualLoop(VirtualLoop);
 			return;
 		}
@@ -5186,7 +5170,8 @@ void FAudioDevice::AddNewActiveSoundInternal(const FActiveSound& InNewActiveSoun
 			if (FAudioVirtualLoop::Virtualize(InNewActiveSound, *this, bDoRangeCheck, VirtualLoop))
 			{
 				UE_LOG(LogAudioConcurrency, Verbose, TEXT("New ActiveSound %s Virtualizing: Failed to pass concurrency"), *Sound->GetName());
-				InitSoundParams(VirtualLoop.GetActiveSound(), MoveTemp(InDefaultParams));
+				constexpr bool bIsVirtualLoopRealizing = false;
+				InitSoundParams(VirtualLoop.GetActiveSound(), MoveTemp(InDefaultParams), bIsVirtualLoopRealizing);
 				AddVirtualLoop(VirtualLoop);
 			}
 			else
@@ -5252,7 +5237,7 @@ void FAudioDevice::AddNewActiveSoundInternal(const FActiveSound& InNewActiveSoun
 		VirtualActiveSound.ClearAudioComponent();
 	}
 
-	InitSoundParams(*ActiveSound, MoveTemp(InDefaultParams));
+	InitSoundParams(*ActiveSound, MoveTemp(InDefaultParams), InVirtualLoopToRetrigger != nullptr /*bIsVirtualLoopRealzing*/);
 	ActiveSounds.Add(ActiveSound);
 	if (ActiveSound->GetAudioComponentID() > 0)
 	{
@@ -5263,6 +5248,58 @@ void FAudioDevice::AddNewActiveSoundInternal(const FActiveSound& InNewActiveSoun
 	NotifySubsystemsActiveSoundCreated(*ActiveSound);
 }
 
+void FAudioDevice::InitSoundParams(FActiveSound& InOutActiveSound, TArray<FAudioParameter> InDefaultParams, bool bInIsVirtualLoopRealizing) const
+{
+	check(IsInAudioThread());
+
+	// Retriggering a virtualized ActiveSound which already have a transmitter
+	// should not be given a new transmitter.
+	if (!InOutActiveSound.InstanceTransmitter.IsValid())
+	{
+		const USoundBase* Sound = InOutActiveSound.GetSound();
+		if (ensure(Sound))
+		{
+			Audio::FParameterTransmitterInitParams TransmitterInitParams
+			{
+				Audio::GetTransmitterID(InOutActiveSound.GetAudioComponentID(), 0, InOutActiveSound.GetPlayOrder()),
+				GetSampleRate(),
+				MoveTemp(InDefaultParams),
+				DeviceID
+			};
+
+			InOutActiveSound.InstanceTransmitter = Sound->CreateParameterTransmitter(MoveTemp(TransmitterInitParams));
+		}
+	}
+	else if (bInIsVirtualLoopRealizing)
+	{
+		// Gather parameters from audio component to set on virtual loop when it restarts
+		TArray<FAudioParameter> InstanceParams;
+		uint64 AudioComponentID = InOutActiveSound.GetAudioComponentID();
+		if (AudioComponentID > 0)
+		{
+			if (UAudioComponent* AudioComponent = UAudioComponent::GetAudioComponentFromID(AudioComponentID))
+			{
+				if (ensure(AudioComponent->Sound != nullptr))
+				{
+					// Use instance parameters which will not contain transient parameters
+					InstanceParams = AudioComponent->GetInstanceParameters();
+
+					// create proxies from the UObject references in InstanceParams
+					AudioComponent->Sound->InitParameters(InstanceParams);
+				}
+			}
+		}
+		
+		FAudioParameter::Merge(MoveTemp(InDefaultParams), InstanceParams);
+		
+		InOutActiveSound.InstanceTransmitter->OnRealizeVirtualizedActiveSound(MoveTemp(InstanceParams));
+	}
+	else
+	{
+		InOutActiveSound.InstanceTransmitter->SetParameters(MoveTemp(InDefaultParams));
+	}
+}
+
 void FAudioDevice::ReportSoundFailedToStart(const uint64 AudioComponentID, FAudioVirtualLoop* VirtualLoop)
 {
 	check(IsInAudioThread());
@@ -5333,6 +5370,11 @@ void FAudioDevice::AddVirtualLoop(const FAudioVirtualLoop& InVirtualLoop)
 	FActiveSound& ActiveSound = VirtualLoop.GetActiveSound();
 	check(!VirtualLoops.Contains(&ActiveSound));
 
+	if (Audio::IParameterTransmitter* ParameterTransmitter = ActiveSound.GetTransmitter())
+	{
+		ParameterTransmitter->OnVirtualizeActiveSound();
+	}
+
 	// If associated with an AudioComponent, add the virtualizing ActiveSound pointer to the VirtualLoop system, 
 	// and ensure it is in the AudioComponentIDToActiveSoundMap so updates from the AudioComponent are still tracked.
 	const int64 ComponentID = ActiveSound.GetAudioComponentID();
diff --git a/Engine/Source/Runtime/Engine/Private/AudioVirtualLoop.cpp b/Engine/Source/Runtime/Engine/Private/AudioVirtualLoop.cpp
--- a/Engine/Source/Runtime/Engine/Private/AudioVirtualLoop.cpp
+++ b/Engine/Source/Runtime/Engine/Private/AudioVirtualLoop.cpp
@@ -125,11 +125,6 @@ bool FAudioVirtualLoop::Virtualize(const FActiveSound& InActiveSound, FAudioDevi
 	OutVirtualLoop.ActiveSound = ActiveSound;
 	OutVirtualLoop.CalculateUpdateInterval();
 
-	if (Audio::IParameterTransmitter* ParameterTransmitter = ActiveSound->GetTransmitter())
-	{
-		ParameterTransmitter->OnVirtualizeActiveSound();
-	}
-
 #if UE_AUDIO_PROFILERTRACE_ENABLED
 	const bool bChannelEnabled = UE_TRACE_CHANNELEXPR_IS_ENABLED(AudioChannel);
 	if (bChannelEnabled)
@@ -313,27 +308,7 @@ bool FAudioVirtualLoop::Update(float DeltaTime, bool bForceUpdate)
 		return false;
 	}
 
-	// Update parameters on virtual loop realization
-	if (Audio::IParameterTransmitter* ParameterTransmitter = ActiveSound->GetTransmitter())
-	{
-		uint64 AudioComponentID = ActiveSound->GetAudioComponentID();
-		if (AudioComponentID > 0)
-		{
-			if (UAudioComponent* AudioComponent = UAudioComponent::GetAudioComponentFromID(AudioComponentID))
-			{
-				if (ensure(AudioComponent->Sound != nullptr))
-				{
-					// Use instance parameters which will not contain transient parameters
-					TArray<FAudioParameter> InstanceParams = AudioComponent->GetInstanceParameters();
-
-					// create proxies from the UObject references in InstanceParams
-					AudioComponent->Sound->InitParameters(InstanceParams);
-
-					ParameterTransmitter->OnRealizeVirtualizedActiveSound(MoveTemp(InstanceParams));
-				}
-			}
-		}
-	}
+
 
 #if UE_AUDIO_PROFILERTRACE_ENABLED
 	if (bChannelEnabled)
diff --git a/Engine/Source/Runtime/Engine/Public/AudioDevice.h b/Engine/Source/Runtime/Engine/Public/AudioDevice.h
--- a/Engine/Source/Runtime/Engine/Public/AudioDevice.h
+++ b/Engine/Source/Runtime/Engine/Public/AudioDevice.h
@@ -1546,6 +1546,11 @@ private:
 	 */
 	void AddNewActiveSoundInternal(const FActiveSound& InActiveSound, TArray<FAudioParameter>&& InDefaultParams, FAudioVirtualLoop* InVirtualLoopToRetrigger = nullptr);
 
+	/**
+	 * Initializes the parameters on the active sound.
+	 */
+	void InitSoundParams(FActiveSound& InOutActiveSound, TArray<FAudioParameter> InDefaultParams, bool bInIsVirtualLoopRealzing) const;
+
 	/**
 	 * Reports if a sound fails to start when attempting to create a new active sound.
 	 */
