Iris: Lower log level for "Created new GroupIndex"

#rnx

diff --git a/Engine/Source/Runtime/Engine/Private/Net/Iris/ReplicationSystem/EngineReplicationBridge.cpp b/Engine/Source/Runtime/Engine/Private/Net/Iris/ReplicationSystem/EngineReplicationBridge.cpp
--- a/Engine/Source/Runtime/Engine/Private/Net/Iris/ReplicationSystem/EngineReplicationBridge.cpp
+++ b/Engine/Source/Runtime/Engine/Private/Net/Iris/ReplicationSystem/EngineReplicationBridge.cpp
@@ -902,7 +902,7 @@ void UEngineReplicationBridge::AddActorToLevelGroup(const AActor* Actor)
 		{
 			LevelGroup = CreateLevelGroup(Level, PackageName);
 
-			UE_LOG_ACTORREPLICATIONBRIDGE(Log, TEXT("Created new GroupIndex: %u for Level: %s"), LevelGroup.GetGroupIndex(), ToCStr(PackageName.ToString()));
+			UE_LOG_ACTORREPLICATIONBRIDGE(Verbose, TEXT("Created new GroupIndex: %u for Level: %s"), LevelGroup.GetGroupIndex(), ToCStr(PackageName.ToString()));
 
 			// Update the filtering status of the group based on current level visibility for all connections
 			NetDriver->UpdateGroupFilterStatusForLevel(Level, LevelGroup);
