[Engine]
* Changed so ParticleSystemSceneProxy can be created without a UParticleSystemComponent

diff --git a/Engine/Source/Runtime/Engine/Private/Particles/ParticleSystemRender.cpp b/Engine/Source/Runtime/Engine/Private/Particles/ParticleSystemRender.cpp
--- a/Engine/Source/Runtime/Engine/Private/Particles/ParticleSystemRender.cpp
+++ b/Engine/Source/Runtime/Engine/Private/Particles/ParticleSystemRender.cpp
@@ -10,6 +10,7 @@
 #include "Materials/Material.h"
 #include "Materials/MaterialRenderProxy.h"
 #include "Particles/Orientation/ParticleModuleOrientationAxisLock.h"
+#include "PrimitiveSceneProxyDesc.h"
 #include "UObject/UObjectIterator.h"
 #include "MaterialDomain.h"
 #include "MaterialShared.h"
@@ -6492,46 +6493,74 @@ int32 FDynamicAnimTrailEmitterData::FillVertexData(struct FAsyncBufferFillData&
 	return TrianglesToRender;
 }
 
+struct FParticleSystemSceneProxyDesc : public FPrimitiveSceneProxyDesc
+{
+	FParticleSystemSceneProxyDesc(UParticleSystemComponent& Component, FParticleDynamicData* InDynamicData, bool InbCanBeOccluded)
+	:	FPrimitiveSceneProxyDesc(&Component)
+	{
+		SystemAsset = Component.Template;
+		DynamicData = InDynamicData;
+
+		if (Component.GetCurrentLODIndex() >= 0 && Component.GetCurrentLODIndex() < Component.CachedViewRelevanceFlags.Num())
+		{
+			MaterialRelevance = Component.CachedViewRelevanceFlags[Component.GetCurrentLODIndex()];
+		}
+		else if (Component.GetCurrentLODIndex() == -1 && Component.CachedViewRelevanceFlags.Num() >= 1)
+		{
+			MaterialRelevance = Component.CachedViewRelevanceFlags[0];
+		}
+
+		VisualizeLODIndex = Component.GetCurrentLODIndex();
+		LODMethod = Component.LODMethod;
+		bCanBeOccluded = InbCanBeOccluded;
+		bManagingSignificance = Component.ShouldManageSignificance();
+		bAlwaysHasVelocity = SystemAsset && SystemAsset->DoesAnyEmitterHaveMotionBlur(Component.GetCurrentLODIndex());
+	}
+
+	UParticleSystem* SystemAsset;
+	FParticleDynamicData* DynamicData;
+	FMaterialRelevance MaterialRelevance;
+	int32 VisualizeLODIndex; // Only used in the LODColoration view mode.
+	int32 LODMethod;
+	uint8 bCanBeOccluded : 1;
+	uint8 bManagingSignificance : 1;
+	uint8 bAlwaysHasVelocity : 1;
+};
+
 ///////////////////////////////////////////////////////////////////////////////
 //	ParticleSystemSceneProxy
 ///////////////////////////////////////////////////////////////////////////////
 /** Initialization constructor. */
-FParticleSystemSceneProxy::FParticleSystemSceneProxy(UParticleSystemComponent* Component, FParticleDynamicData* InDynamicData, bool InbCanBeOccluded)
-	: FPrimitiveSceneProxy(Component, Component->Template ? Component->Template->GetFName() : NAME_None)
-	, Owner(Component->GetOwner())
-	, bCastShadow(Component->CastShadow)
-	, bManagingSignificance(Component->ShouldManageSignificance())
-	, bCanBeOccluded(InbCanBeOccluded)
+
+FParticleSystemSceneProxy::FParticleSystemSceneProxy(const FParticleSystemSceneProxyDesc& Desc)
+	: FPrimitiveSceneProxy(Desc, Desc.SystemAsset ? Desc.SystemAsset->GetFName() : NAME_None)
+	, bCastShadow(Desc.CastShadow)
+	, bManagingSignificance(Desc.bManagingSignificance)
+	, bCanBeOccluded(Desc.bCanBeOccluded)
 	, bHasCustomOcclusionBounds(false)
 	, FeatureLevel(GetScene().GetFeatureLevel())
-	, MaterialRelevance(
-		((Component->GetCurrentLODIndex() >= 0) && (Component->GetCurrentLODIndex() < Component->CachedViewRelevanceFlags.Num())) ?
-			Component->CachedViewRelevanceFlags[Component->GetCurrentLODIndex()] :
-		((Component->GetCurrentLODIndex() == -1) && (Component->CachedViewRelevanceFlags.Num() >= 1)) ?
-			Component->CachedViewRelevanceFlags[0] :
-			FMaterialRelevance()
-		)
-	, DynamicData(InDynamicData)
+	, MaterialRelevance(Desc.MaterialRelevance)
+	, DynamicData(Desc.DynamicData)
 	, LastDynamicData(NULL)
 	, DeselectedWireframeMaterialInstance(new FColoredMaterialRenderProxy(
 		GEngine->WireframeMaterial ? GEngine->WireframeMaterial->GetRenderProxy() : NULL,
 		GetSelectionColor(FLinearColor(1.0f, 0.0f, 0.0f, 1.0f),false,false)
 		))
 	, PendingLODDistance(0.0f)
-	, VisualizeLODIndex(Component->GetCurrentLODIndex())
+	, VisualizeLODIndex(Desc.VisualizeLODIndex)
 	, LastFramePreRendered(-1)
 	, FirstFreeMeshBatch(0)
 {
 	SetWireframeColor(FLinearColor(3.0f, 0.0f, 0.0f));
 
-	LODMethod = Component->LODMethod;
+	LODMethod = Desc.LODMethod;
 
 	// Particle systems intrinsically always have motion, but is this motion relevant to systems external to particle systems?
-	bAlwaysHasVelocity = Component->Template && Component->Template->DoesAnyEmitterHaveMotionBlur(Component->GetCurrentLODIndex());
+	bAlwaysHasVelocity = Desc.bAlwaysHasVelocity;
 
-	if (bCanBeOccluded && Component->Template && (Component->Template->OcclusionBoundsMethod == EPSOBM_CustomBounds))
+	if (bCanBeOccluded && Desc.SystemAsset && (Desc.SystemAsset->OcclusionBoundsMethod == EPSOBM_CustomBounds))
 	{
-		OcclusionBounds = FBoxSphereBounds(Component->Template->CustomOcclusionBounds);
+		OcclusionBounds = FBoxSphereBounds(Desc.SystemAsset->CustomOcclusionBounds);
 		bHasCustomOcclusionBounds = true;
 	}
 }
@@ -6983,11 +7012,11 @@ FPrimitiveSceneProxy* UParticleSystemComponent::CreateSceneProxy()
 		if (CanBeOccluded())
 		{
 			Template->CustomOcclusionBounds.IsValid = true;
-			NewProxy = ::new FParticleSystemSceneProxy(this,ParticleDynamicData, true);
+			NewProxy = ::new FParticleSystemSceneProxy(FParticleSystemSceneProxyDesc(*this,ParticleDynamicData, true));
 		}
 		else
 		{
-			NewProxy = ::new FParticleSystemSceneProxy(this,ParticleDynamicData, false);
+			NewProxy = ::new FParticleSystemSceneProxy(FParticleSystemSceneProxyDesc(*this,ParticleDynamicData, false));
 		}
 		check (NewProxy);
 	}
diff --git a/Engine/Source/Runtime/Engine/Public/ParticleHelper.h b/Engine/Source/Runtime/Engine/Public/ParticleHelper.h
--- a/Engine/Source/Runtime/Engine/Public/ParticleHelper.h
+++ b/Engine/Source/Runtime/Engine/Public/ParticleHelper.h
@@ -29,11 +29,10 @@
 class FColoredMaterialRenderProxy;
 class FParticleSystemSceneProxy;
 class UParticleModuleRequired;
-class UParticleSystemComponent;
-class USkeletalMeshComponent;
 class UStaticMesh;
 struct FBaseParticle;
 struct FParticleMeshEmitterInstance;
+struct FParticleSystemSceneProxyDesc;
 struct FStaticMeshLODResources;
 struct FGlobalDynamicIndexBufferAllocation;
 struct FGlobalDynamicVertexBufferAllocation;
@@ -145,10 +144,7 @@ class UParticleModuleTypeDataMesh;
 class UParticleModuleTypeDataBeam2;
 class UParticleModuleTypeDataTrail2;
 
-class UStaticMeshComponent;
-
 class UParticleSystem;
-class UParticleSystemComponent;
 
 class UParticleModuleBeamSource;
 class UParticleModuleBeamTarget;
@@ -163,8 +159,6 @@ class UParticleModuleOrientationAxisLock;
 
 class UParticleLODLevel;
 
-class USkeletalMeshComponent;
-
 class FParticleSystemSceneProxy;
 struct FDynamicBeam2EmitterData;
 struct FDynamicTrail2EmitterData;
@@ -2317,7 +2311,7 @@ public:
 	SIZE_T GetTypeHash() const override;
 
 	/** Initialization constructor. */
-	FParticleSystemSceneProxy(UParticleSystemComponent* Component, FParticleDynamicData* InDynamicData, bool bCanBeOccluded);
+	FParticleSystemSceneProxy(const FParticleSystemSceneProxyDesc& Desc);
 	virtual ~FParticleSystemSceneProxy();
 
 	// FPrimitiveSceneProxy interface.
@@ -2438,8 +2432,6 @@ protected:
 	 */
 	void ReleaseRenderThreadResourcesForEmitterData();
 
-	AActor* Owner;
-
 #if STATS
 	double LastStatCaptureTime;
 	uint8 bCountedThisFrame:1;
