ShaderCompiler: Ensure that AutogenShaderHeaders.ush exists and is up-to-date in CompileGlobalShaderMap.

This works around an issue that manifested as the FTranslucentLightingInject source hash changing in the DDC key. Because it contains both global and material shader types, the global type would be loaded and hashed prior to the existence of the autogen header, resulting in the subsequent material compiles reusing the incorrect cached hash.

diff --git a/Engine/Source/Runtime/Engine/Private/ShaderCompiler/ShaderCompiler.cpp b/Engine/Source/Runtime/Engine/Private/ShaderCompiler/ShaderCompiler.cpp
--- a/Engine/Source/Runtime/Engine/Private/ShaderCompiler/ShaderCompiler.cpp
+++ b/Engine/Source/Runtime/Engine/Private/ShaderCompiler/ShaderCompiler.cpp
@@ -5116,6 +5116,9 @@ void CompileGlobalShaderMap(EShaderPlatform Platform, const ITargetPlatform* Tar
 #if WITH_EDITOR
 		if (!bLoadedFromCacheFile && bAllowShaderCompiling)
 		{
+			// Ensure we've generated AutogenShaderHeaders.ush
+			FShaderCompileUtilities::GenerateBrdfHeaders(Platform);
+
 			// If we didn't find cooked shaders, we can try loading from the DDC or compiling them if supported by the current configuration
 			FGlobalShaderMapId ShaderMapId(Platform, TargetPlatform);
 
