* Add Substrate code versioning (in addition to the auto-conversion versionning).
* Change Glint density parameterization. Glint density is now a normalized value (0..1), with 0 meaning very low glint density, and 1 very dense glints.
* Change Glint default value from 0 to 1. A value of 1 means very dense glints, which is equivalent to a regular single specular lobe (aggregating all the glint highlights). Now, when glint density is set to 1, it means no glint, just a single (aggregating) specular lobe.

#virtualized

diff --git a/Engine/Shaders/Private/DebugProbes.usf b/Engine/Shaders/Private/DebugProbes.usf
--- a/Engine/Shaders/Private/DebugProbes.usf
+++ b/Engine/Shaders/Private/DebugProbes.usf
@@ -137,7 +137,7 @@ void WriteOutMaterial(
 		/*float FuzzAmount						*/0.0f,
 		/*float3 FuzzColor						*/0.0f,
 		/*float FuzzRoughness					*/0.0f,
-		/*float GlintValue						*/0.0f,
+		/*float GlintValue						*/1.0f,
 		/*float GlintUV							*/0.0f.xx,
 		/*float SpecularProfile					*/0.0f,
 		/*float Thickness						*/1.0f,
diff --git a/Engine/Shaders/Private/PathTracing/Material/PathTracingSubstrateCommon.ush b/Engine/Shaders/Private/PathTracing/Material/PathTracingSubstrateCommon.ush
--- a/Engine/Shaders/Private/PathTracing/Material/PathTracingSubstrateCommon.ush
+++ b/Engine/Shaders/Private/PathTracing/Material/PathTracingSubstrateCommon.ush
@@ -45,7 +45,7 @@ void SubstrateGlint(FPathTracingPayload Payload, float3 InTangentV, float3 InTan
 	// For now use the GGX lobe sampling. This will be changed with glint importance sampling later on.
 #if SUBSTRATE_PATH_TRACING_GLINTS
 	BRANCH
-	if (Payload.GlintValue > 0)
+	if (Payload.GlintValue < 1.f)
 	{
 		const float2 GGXRoughnessXY = float2(sqrtFast(InAlphaXY.x), sqrtFast(InAlphaXY.y));
 		FBeckmannDesc Beckmann = GGXToBeckmann(GGXRoughnessXY);
diff --git a/Engine/Shaders/Private/PathTracing/PathTracingCommon.ush b/Engine/Shaders/Private/PathTracing/PathTracingCommon.ush
--- a/Engine/Shaders/Private/PathTracing/PathTracingCommon.ush
+++ b/Engine/Shaders/Private/PathTracing/PathTracingCommon.ush
@@ -780,10 +780,10 @@ FPackedPathTracingPayload PackPathTracingPayload(FPathTracingPayload Input)
 	Output.PackedData[17] = Pack10F(Input.Ior);
 	Output.PackedData[17]|= Pack10F(Input.CoverageAboveAlongN) << 10;
 	Output.PackedData[17]|= Pack10F(Input.SpecularProfileId) << 20;
-	Output.PackedData[17]|= (Input.GlintValue > 0 ? 1u : 0u) << 30; // Bit flag for glint. 
+	Output.PackedData[17]|= (Input.GlintValue < 1 ? 1u : 0u) << 30; // Bit flag for glint. 
 
 	// In order to not increase the data payload size, we alias SSS/Fuzz parameters when glint are enabled
-	if (Input.GlintValue > 0)
+	if (Input.GlintValue < 1)
 	{
 		const uint2 PackedGlintsData = PackGlints(Input.GlintValue, Input.GlintUV);
 		Output.PackedData[11] = PackedGlintsData.x;
diff --git a/Engine/Shaders/Private/PathTracing/PathTracingMaterialHitShader.usf b/Engine/Shaders/Private/PathTracing/PathTracingMaterialHitShader.usf
--- a/Engine/Shaders/Private/PathTracing/PathTracingMaterialHitShader.usf
+++ b/Engine/Shaders/Private/PathTracing/PathTracingMaterialHitShader.usf
@@ -660,7 +660,7 @@ RAY_TRACING_ENTRY_CLOSEST_HIT(PathTracingMaterialCHS,
 						Payload.PhaseG       = 0;
 					}
 					Payload.SpecularProfileId = 0;
-					Payload.GlintValue = 0;
+					Payload.GlintValue = 1.f;
 
 #if SUBSTRATE_COMPLEXSPECIALPATH
 					if (BSDF_GETHASSPECPROFILE(BSDF))
diff --git a/Engine/Shaders/Private/Substrate/Glint/GlintThirdParty.ush b/Engine/Shaders/Private/Substrate/Glint/GlintThirdParty.ush
--- a/Engine/Shaders/Private/Substrate/Glint/GlintThirdParty.ush
+++ b/Engine/Shaders/Private/Substrate/Glint/GlintThirdParty.ush
@@ -173,7 +173,8 @@ float P22_M(
 	FGlintLutParams GlintLutParams,
 	float2 slope_h, int l, int s0, int t0,
 	float2 slope_dx, float2 slope_dy,
-	float3 sigma_x_y_rho, float l_dist)
+	float3 sigma_x_y_rho, float l_dist, 
+	float MicrofacetRelativeArea)
 {
 	// Coherent index
 	int twoToTheL = int(pow(2., float(l)));
@@ -185,6 +186,13 @@ float P22_M(
 
 	float uDensityRandomisation = hashIQ(rngSeed * 2171U);
 
+	// Discard cells by using microfacet relative area
+	float uMicrofacetRelativeArea = hashIQ(rngSeed * 13U);
+	if (uMicrofacetRelativeArea > MicrofacetRelativeArea)
+	{
+		return 0.f;
+	}
+
 	// Fix density randomisation to 2 to have better appearance
 	float densityRandomisation = 2.;
 
@@ -338,7 +346,7 @@ float P22__glint_discrete_LOD(
 	FGlintLutParams GlintLutParams,
 	int l, float2 slope_h, float2 st, float2 dst0,
 	float2 dst1, float2 slope_dx, float2 slope_dy,
-	float3 sigma_x_y_rho, float l_dist)
+	float3 sigma_x_y_rho, float l_dist, float MicrofacetRelativeArea)
 {
 	// Convert surface coordinates to appropriate scale for level
 	float pyrSize = pyramidSize(GlintLutParams, l);
@@ -393,7 +401,7 @@ float P22__glint_discrete_LOD(
 					GlintLutParams,
 					slope_h, l, is, it,
 					slope_dx, slope_dy,
-					sigma_x_y_rho, l_dist) * W_P;
+					sigma_x_y_rho, l_dist, MicrofacetRelativeArea) * W_P;
 				sumWts += W_P;
 			}
 			nbrOfIter++;
@@ -411,7 +419,7 @@ float P22__glint_discrete_LOD(
 //=============================================================================
 //======== Evaluation of the procedural physically based glinty BRDF ==========
 //=============================================================================
-float f_P(float3 wo, float3 wi, float3 sigmas_rho, float LogMicrofacetDensity, float2 TexCoord, float2 TexCoordDDX, float2 TexCoordDDY, bool bIncludeGeometryTerm=false)
+float f_P(float3 wo, float3 wi, float3 sigmas_rho, float GlintDensity, float2 TexCoord, float2 TexCoordDDX, float2 TexCoordDDY, bool bIncludeGeometryTerm=false)
 {
 	if (wo.z <= 0.)
 		return 0;
@@ -492,6 +500,13 @@ float f_P(float3 wo, float3 wi, float3 sigmas_rho, float LogMicrofacetDensity, f
 	}
 	else
 	{
+		// Reparameterize the glint density [0..1] into a LogMicrofacetDensity and MicrofacetRelativeArea value.
+		// LogMicrofacetDensity has a visible range of [~10..21]. With this we do the following remapping:
+		// * Range in [10..21] remaps to the expressable density
+		// * Range in [0..10] uses MicrofacetRelativeArea to fade progressively the glints.
+		const float LogMicrofacetDensity = lerp(0.f, 21.f, saturate(GlintDensity));
+		const float MicrofacetRelativeArea = saturate(LogMicrofacetDensity / 10.f);
+
 		// Choose LOD
 		float l = max(0.0, GlintLutParams.NLevels - 1. + log2(minorLength));
 		int il = int(floor(l));
@@ -515,14 +530,16 @@ float f_P(float3 wo, float3 wi, float3 sigmas_rho, float LogMicrofacetDensity, f
 		LOD_dist_il   = clamp(LOD_dist_il   + GlintLutParams.LevelBias, 0, max(0, GlintLutParams.NLevels - GlintLutParams.LevelMin));
 		LOD_dist_ilp1 = clamp(LOD_dist_ilp1 + GlintLutParams.LevelBias, 0, max(0, GlintLutParams.NLevels - GlintLutParams.LevelMin));
 
-		if (int(round(LOD_dist_il)) < GlintLutParams.NLevels)
+		bool opti = MicrofacetRelativeArea > 0.99;
+
+		if (int(round(LOD_dist_il)) < GlintLutParams.NLevels || !opti)
 		{
-			P22_P_il = P22__glint_discrete_LOD(GlintLutParams, il, slope_h, TexCoord, dst0, dst1, slope_dx, slope_dy, sigmas_rho, LOD_dist_il);
+			P22_P_il = P22__glint_discrete_LOD(GlintLutParams, il, slope_h, TexCoord, dst0, dst1, slope_dx, slope_dy, sigmas_rho, LOD_dist_il, MicrofacetRelativeArea);
 		}
 
-		if (int(round(LOD_dist_ilp1)) < GlintLutParams.NLevels)
+		if (int(round(LOD_dist_ilp1)) < GlintLutParams.NLevels || !opti)
 		{
-			P22_P_ilp1 = P22__glint_discrete_LOD(GlintLutParams, il + 1, slope_h, TexCoord, dst0, dst1, slope_dx, slope_dy, sigmas_rho, LOD_dist_ilp1);
+			P22_P_ilp1 = P22__glint_discrete_LOD(GlintLutParams, il + 1, slope_h, TexCoord, dst0, dst1, slope_dx, slope_dy, sigmas_rho, LOD_dist_ilp1, MicrofacetRelativeArea);
 		}
 
 		P22_P = lerp(P22_P_il, P22_P_ilp1, w);
diff --git a/Engine/Shaders/Private/Substrate/Substrate.ush b/Engine/Shaders/Private/Substrate/Substrate.ush
--- a/Engine/Shaders/Private/Substrate/Substrate.ush
+++ b/Engine/Shaders/Private/Substrate/Substrate.ush
@@ -1415,7 +1415,7 @@ FSubstrateData InternalGetSubstrateSlabBSDF(
 		ThicknessCm = SUBSTRATE_SIMPLEVOLUME_THICKNESS_CM;
 	}
 
-	const bool bHasGlint = GlintValue > 0.0f;
+	const bool bHasGlint = GlintValue < 1.0f;
 	const uint bHasHaziness = (((Haziness.Roughness != Roughness || bHasGlint)  && Haziness.Weight > 0.0f) || Haziness.bSimpleClearCoat) ? 1 : 0;
 	const uint bHasAnisotropy = Anisotropy != 0 ? 1 : 0;
 	const uint bHasFuzz = FuzzAmount > 0.0f ? 1 : 0;
@@ -2287,7 +2287,7 @@ void FSubstratePixelHeader::SubstrateConvertFromDBuffer(
 		0.f,				// FuzzAmount
 		0.f,				// FuzzColor
 		In.Roughness,		// FuzzRoughness
-		0.f,				// GlintValue, 
+		1.f,				// GlintValue, 
 		(GLINT_UV_TYPE)0,	// GlintUV,
 		0.f, 				// SpecularProfileId
 		DefaultThicknessCm,
diff --git a/Engine/Shaders/Private/Substrate/SubstrateLegacyConversion.ush b/Engine/Shaders/Private/Substrate/SubstrateLegacyConversion.ush
--- a/Engine/Shaders/Private/Substrate/SubstrateLegacyConversion.ush
+++ b/Engine/Shaders/Private/Substrate/SubstrateLegacyConversion.ush
@@ -124,7 +124,7 @@ FSubstrateData CreateLegacySlab(
 		In.FuzzAmount,										// FuzzAmount
 		In.FuzzColor,										// FuzzColor
 		In.Roughness,										// FuzzRoughness
-		0.0f, (GLINT_UV_TYPE)0,								// GlintData
+		1.0f, (GLINT_UV_TYPE)0,								// GlintData
 		0.0f,												// Specular profile
 		In.Thickness,
 		In.bIsThinSurface,
diff --git a/Engine/Source/Runtime/Engine/Private/Materials/Material.cpp b/Engine/Source/Runtime/Engine/Private/Materials/Material.cpp
--- a/Engine/Source/Runtime/Engine/Private/Materials/Material.cpp
+++ b/Engine/Source/Runtime/Engine/Private/Materials/Material.cpp
@@ -3326,6 +3326,9 @@ static int32 GetSubstrateConversionVersion()	{ return  1; }
 static int32 GetSubstrateNoConversionVersion()	{ return  0; }
 // Invalid version - This means that the material has never run ConvertMaterialToSubstrateMaterial()
 static int32 GetSubstrateInvalidVersion()		{ return -1; }
+// Substrate version
+// 0: Glint conversion
+static int32 GetSubstrateVersion()				{ return 0; }
 
 #if WITH_EDITOR
 
@@ -3354,11 +3357,11 @@ bool UMaterial::ConvertMaterialToSubstrateMaterial(bool bAllowEmptyMaterialUpdat
 
 	UMaterialEditorOnlyData* EditorOnly = GetEditorOnlyData();
 
-	// * If the current material has already been converted, skip the conversion. 
+	// * If the current material has already been converted & versioned, skip the conversion. 
 	// * Store the current version of Subtrate's auto-conversion
 	//   This allows to version the conversion, and safely update auto-converted materials if they have been saved.
 	// * If no conversion was needed, the conversion version is set to 0, to avoid rerunning this function
-	if (EditorOnly->SubstrateConversionVersion == GetSubstrateConversionVersion() || EditorOnly->SubstrateConversionVersion == GetSubstrateNoConversionVersion())
+	if ((EditorOnly->SubstrateConversionVersion == GetSubstrateConversionVersion() || EditorOnly->SubstrateConversionVersion == GetSubstrateNoConversionVersion()) && EditorOnly->SubstrateVersion == GetSubstrateVersion())
 	{
 		return false;
 	}
@@ -4007,6 +4010,54 @@ bool UMaterial::ConvertMaterialToSubstrateMaterial(bool bAllowEmptyMaterialUpdat
 		}
 	}
 
+	// Glint conversion parametrization
+	if (EditorOnly->SubstrateVersion < 0)
+	{
+		for (UMaterialExpression* Expression : GetExpressions())
+		{
+			if (UMaterialExpressionSubstrateSlabBSDF* SlabExpression = Cast<UMaterialExpressionSubstrateSlabBSDF>(Expression))
+			{
+				if (SlabExpression->GlintValue.IsConnected())
+				{
+					// Gather custom output for thin translucency
+					// Empty material: Create by default a slab node
+					UMaterialFunction* ConversionMF = LoadObject<UMaterialFunction>(nullptr, TEXT("/Engine/Functions/Substrate/SMF_GlintConversion.SMF_GlintConversion")); 
+					if (ConversionMF)
+					{
+						ConversionMF->UpdateFromFunctionResource();
+						ConversionMF->PostEditChange();
+						ConversionMF->ConditionalPostLoad();
+
+						UMaterialExpressionMaterialFunctionCall* MFCallNode = NewObject<UMaterialExpressionMaterialFunctionCall>(this);
+						if (MFCallNode->SetMaterialFunction(ConversionMF))
+						{
+							// This is needed for input/output expressions to be set correctly, otherwise compilation will fail.
+							EditorOnly->ExpressionCollection.AddExpression(MFCallNode);
+							SetPosXAndMoveReferenceToTheRight(MFCallNode);
+
+							// Connect GlintValue graph to the MF input
+							ConnectionTo(*SlabExpression->GetInput(16)/*GlintValue*/, MFCallNode, 0/*Density*/);
+
+							// Connect the MF output to the GlintValue input
+							SlabExpression->GlintValue.Connect(0, MFCallNode);
+
+							/** Connects output of InExpression to this input */
+							ENGINE_API void Connect( int32 InOutputIndex, class UMaterialExpression* InExpression );
+
+							MFCallNode->UpdateFromFunctionResource();
+							MFCallNode->PostEditChange();
+							MFCallNode->ConditionalPostLoad();
+
+						}
+					}
+				}
+			}
+		}
+
+		// Even if no glint conversion was one, we invalidate the shader in order to save the new Substrate's version (from -1 to 0)
+		bInvalidateShader = true;
+	}
+
 	if (bInvalidateShader)
 	{
 		// Now force the material to recompile and we use a hash of the original StateId.
@@ -4027,10 +4078,25 @@ bool UMaterial::ConvertMaterialToSubstrateMaterial(bool bAllowEmptyMaterialUpdat
 		UpdateCachedExpressionData();
 	}
 
-	// Store the conversion version, only for converted shaders
-	if (bInvalidateShader && !bEmptyShader)
+	// Store the versions, 
+	if (bInvalidateShader)
 	{
-		EditorOnly->SubstrateConversionVersion = GetSubstrateConversionVersion();
+		// 1. Store current Substrate version
+		EditorOnly->SubstrateVersion = GetSubstrateVersion();
+
+		// 2. Store current Substrate conversion version
+		const bool bIsUpToDate = EditorOnly->SubstrateConversionVersion == GetSubstrateNoConversionVersion() || EditorOnly->SubstrateConversionVersion == GetSubstrateConversionVersion();
+		if (!bIsUpToDate)
+		{
+			if (bEmptyShader) // == new shader
+			{
+				EditorOnly->SubstrateConversionVersion = GetSubstrateNoConversionVersion();
+			}
+			else // == Converted shader
+			{
+				EditorOnly->SubstrateConversionVersion = GetSubstrateConversionVersion();
+			}
+		}
 	}
 
 	return bInvalidateShader;
@@ -7749,6 +7815,7 @@ UMaterialEditorOnlyData::UMaterialEditorOnlyData()
 	Refraction.Constant = FMaterialAttributeDefinitionMap::GetDefaultValue(MP_Refraction).X;
 	SurfaceThickness.Constant = FMaterialAttributeDefinitionMap::GetDefaultValue(MP_SurfaceThickness).X;
 	SubstrateConversionVersion = GetSubstrateInvalidVersion();
+	SubstrateVersion = GetSubstrateInvalidVersion();
 }
 
 void UMaterialEditorOnlyData::ResetSubstrateConversionVersion()
diff --git a/Engine/Source/Runtime/Engine/Private/Materials/MaterialExpressionSubstrate.cpp b/Engine/Source/Runtime/Engine/Private/Materials/MaterialExpressionSubstrate.cpp
--- a/Engine/Source/Runtime/Engine/Private/Materials/MaterialExpressionSubstrate.cpp
+++ b/Engine/Source/Runtime/Engine/Private/Materials/MaterialExpressionSubstrate.cpp
@@ -799,7 +799,7 @@ int32 UMaterialExpressionSubstrateSlabBSDF::Compile(class FMaterialCompiler* Com
 	int32 FuzzAmountCodeChunk = CompileWithDefaultFloat1(Compiler, FuzzAmount, 0.0f);
 	int32 FuzzColorCodeChunk = CompileWithDefaultFloat3(Compiler, FuzzColor, 0.0f, 0.0f, 0.0f);
 	int32 FuzzRoughnessCodeChunk = HasFuzzRoughness() ? CompileWithDefaultFloat1(Compiler, FuzzRoughness, 0.5f) : RoughnessCodeChunk;
-	int32 GlintValueCodeChunk = CompileWithDefaultFloat1(Compiler, GlintValue, 0.0f);
+	int32 GlintValueCodeChunk = CompileWithDefaultFloat1(Compiler, GlintValue, 1.0f);
 	int32 GlintUVCodeChunk = CompileWithDefaultFloat2(Compiler, GlintUV, 0.0f, 0.0f);
 
 	// Disable some features if requested by the simplification process
@@ -1306,7 +1306,7 @@ int32 UMaterialExpressionSubstrateSimpleClearCoatBSDF::Compile(class FMaterialCo
 		Compiler->Constant3(0.0f, 0.0f, 0.0f),					// FuzzColor
 		RoughnessCodeChunk,										// FuzzRoughness
 		ThicknessCodeChunk,										// Thickness
-		Compiler->Constant(0.0f),								// GlintValue
+		Compiler->Constant(1.0f),								// GlintValue
 		Compiler->Constant2(0.0f, 0.0f),						// GlintUV
 		Compiler->Constant(0.0f),								// SpecularProfile
 		false,													// bIsAtTheBottomOfTopology, always false for SimpleClearCoat
diff --git a/Engine/Source/Runtime/Engine/Public/Materials/Material.h b/Engine/Source/Runtime/Engine/Public/Materials/Material.h
--- a/Engine/Source/Runtime/Engine/Public/Materials/Material.h
+++ b/Engine/Source/Runtime/Engine/Public/Materials/Material.h
@@ -405,6 +405,10 @@ public:
 	UPROPERTY()
 	int32 SubstrateConversionVersion;
 
+	/** Store the version of the Substrate code used by this material (default, no version =-1). */
+	UPROPERTY()
+	int32 SubstrateVersion;
+
 	/** Reset the version of the Substrate's auto-conversion applied to this material */
 	ENGINE_API void ResetSubstrateConversionVersion();
 };
diff --git a/Engine/Source/Runtime/Engine/Public/Materials/MaterialExpressionSubstrate.h b/Engine/Source/Runtime/Engine/Public/Materials/MaterialExpressionSubstrate.h
--- a/Engine/Source/Runtime/Engine/Public/Materials/MaterialExpressionSubstrate.h
+++ b/Engine/Source/Runtime/Engine/Public/Materials/MaterialExpressionSubstrate.h
@@ -319,7 +319,7 @@ class UMaterialExpressionSubstrateSlabBSDF : public UMaterialExpressionSubstrate
 	FExpressionInput FuzzColor;
 
 	/**
-	 * This represent the logarithm of the micro facet density. Only used when `r.Substrate.Glints=1`. Defaults to 0.
+	 * This represent the micro facet density. Only used when `r.Substrate.Glints=1`. 0 = only very sparse glints, 1 = fully covered with glints (which is equivalent to a regular specular lobe). Defaults to 1 (=no glint).
 	 */
 	UPROPERTY()
 	FExpressionInput GlintValue;
