Fix missing default value when creating slab causing it to have glint by default.

diff --git a/Engine/Source/Runtime/Engine/Private/Materials/MaterialExpressionSubstrate.cpp b/Engine/Source/Runtime/Engine/Private/Materials/MaterialExpressionSubstrate.cpp
--- a/Engine/Source/Runtime/Engine/Private/Materials/MaterialExpressionSubstrate.cpp
+++ b/Engine/Source/Runtime/Engine/Private/Materials/MaterialExpressionSubstrate.cpp
@@ -825,7 +825,7 @@ int32 UMaterialExpressionSubstrateSlabBSDF::Compile(class FMaterialCompiler* Com
 	}
 	if (!SubstrateOperator.Has(ESubstrateBsdfFeature_Glint))
 	{
-		GlintValueCodeChunk = Compiler->Constant(0.0f);
+		GlintValueCodeChunk = Compiler->Constant(1.0f);
 	}
 	if (!SubstrateOperator.Has(ESubstrateBsdfFeature_SpecularProfile))
 	{
