Made UVolumeTexture::SetDefaultSource2DTileSize safer.

#tests Tested in editor by importing a 1024x4 texture.
#rnx

diff --git a/Engine/Source/Runtime/Engine/Private/VolumeTexture.cpp b/Engine/Source/Runtime/Engine/Private/VolumeTexture.cpp
--- a/Engine/Source/Runtime/Engine/Private/VolumeTexture.cpp
+++ b/Engine/Source/Runtime/Engine/Private/VolumeTexture.cpp
@@ -583,8 +583,8 @@ void UVolumeTexture::SetDefaultSource2DTileSize()
 		if (SourceSizeX > 0 && SourceSizeY > 0)
 		{
 			const int32 NumPixels = SourceSizeX * SourceSizeY;
-			const int32 TileSize = FMath::RoundToInt(FMath::Pow((float)NumPixels, 1.f / 3.f));
-			const int32 NumTilesBySide = FMath::RoundToInt(FMath::Sqrt((float)(SourceSizeX / TileSize) * (SourceSizeY / TileSize)));
+			const int32 TileSize = FMath::Max(1, FMath::RoundToInt(FMath::Pow((float)NumPixels, 1.f / 3.f)));
+			const int32 NumTilesBySide = FMath::Max(1, FMath::RoundToInt(FMath::Sqrt((float)(SourceSizeX / TileSize) * (SourceSizeY / TileSize))));
 			Source2DTileSizeX = SourceSizeX / NumTilesBySide;
 			Source2DTileSizeY = SourceSizeY / NumTilesBySide;
 		}
