Speculative fix for MRQ crash.

diff --git a/Engine/Source/Runtime/Engine/Private/SampleBufferIO.cpp b/Engine/Source/Runtime/Engine/Private/SampleBufferIO.cpp
--- a/Engine/Source/Runtime/Engine/Private/SampleBufferIO.cpp
+++ b/Engine/Source/Runtime/Engine/Private/SampleBufferIO.cpp
@@ -504,6 +504,9 @@ namespace Audio
 			FMemory::Free(CurrentSoundWave->RawPCMData);
 			CurrentSoundWave->RawPCMData = nullptr;
 		}
+
+		CurrentSoundWave->RawPCMData = (uint8*)FMemory::Malloc(CurrentSoundWave->RawPCMDataSize);
+		FMemory::Memcpy(CurrentSoundWave->RawPCMData, CurrentBuffer.GetData(), CurrentSoundWave->RawPCMDataSize);
 	}
 
 	void FSoundWavePCMWriter::SerializeSoundWaveToAsset()
@@ -713,8 +716,7 @@ namespace Audio
 		}
 		
 		// clear out RawPCMData
-		check(Writer->CurrentSoundWave);
-		if (Writer->CurrentSoundWave->RawPCMData)
+		if (Writer->CurrentSoundWave && Writer->CurrentSoundWave->RawPCMData)
 		{
 			FMemory::Free(Writer->CurrentSoundWave->RawPCMData);
 			Writer->CurrentSoundWave->RawPCMData = nullptr;
