Fixed unsafe removal of layers in SGameLayerManager. If a player gets garbage collected before missing players are removed, the widget won't be found. Changed the clean-up to remove widgets with null players explicitly.

diff --git a/Engine/Source/Runtime/Engine/Private/Slate/SGameLayerManager.cpp b/Engine/Source/Runtime/Engine/Private/Slate/SGameLayerManager.cpp
--- a/Engine/Source/Runtime/Engine/Private/Slate/SGameLayerManager.cpp
+++ b/Engine/Source/Runtime/Engine/Private/Slate/SGameLayerManager.cpp
@@ -578,33 +578,27 @@ TSharedPtr<SGameLayerManager::FPlayerLayer> SGameLayerManager::FindOrCreatePlaye
 
 void SGameLayerManager::RemoveMissingPlayerLayers(const TArray<ULocalPlayer*>& GamePlayers)
 {
-	TArray<ULocalPlayer*> ToRemove;
-
 	// Find the player layers for players that no longer exist
-	for (auto& PlayerLayerEntry : PlayerLayers)
+	for (auto PlayerLayerIt = PlayerLayers.CreateIterator(); PlayerLayerIt; ++PlayerLayerIt)
 	{
-		ULocalPlayer* Key = Cast<ULocalPlayer>(PlayerLayerEntry.Key.ResolveObjectPtr());
-		if ( !GamePlayers.Contains(Key) )
+		ULocalPlayer* Key = Cast<ULocalPlayer>(PlayerLayerIt.Key().ResolveObjectPtr());
+		if ((Key == nullptr) || !GamePlayers.Contains(Key))
 		{
-			ToRemove.Add(Key);
+			const TSharedPtr<FPlayerLayer> Layer = PlayerLayerIt.Value();
+			PlayerCanvas->RemoveSlot(Layer->Widget.ToSharedRef());
+			PlayerLayerIt.RemoveCurrent();
 		}
 	}
-
-	// Remove the missing players
-	for ( ULocalPlayer* Player : ToRemove )
-	{
-		RemovePlayerWidgets(Player);
-	}
 }
 
 void SGameLayerManager::RemovePlayerWidgets(ULocalPlayer* LocalPlayer)
 {
-	FObjectKey LocalPlayerKey(LocalPlayer);
-
-	TSharedPtr<FPlayerLayer> Layer = PlayerLayers.FindRef(LocalPlayerKey);
-	PlayerCanvas->RemoveSlot(Layer->Widget.ToSharedRef());
-
-	PlayerLayers.Remove(LocalPlayerKey);
+	TSharedPtr<FPlayerLayer> Layer;
+	const FObjectKey LocalPlayerKey(LocalPlayer);
+	if (PlayerLayers.RemoveAndCopyValue(LocalPlayerKey, Layer))
+	{
+		PlayerCanvas->RemoveSlot(Layer->Widget.ToSharedRef());
+	}
 }
 
 void SGameLayerManager::AddOrUpdatePlayerLayers(const FGeometry& AllottedGeometry, UGameViewportClient* ViewportClient, const TArray<ULocalPlayer*>& GamePlayers)
