Improve tooltip for render target HDR setting.

diff --git a/Engine/Source/Runtime/Engine/Classes/Engine/TextureRenderTarget2DArray.h b/Engine/Source/Runtime/Engine/Classes/Engine/TextureRenderTarget2DArray.h
--- a/Engine/Source/Runtime/Engine/Classes/Engine/TextureRenderTarget2DArray.h
+++ b/Engine/Source/Runtime/Engine/Classes/Engine/TextureRenderTarget2DArray.h
@@ -38,12 +38,14 @@ class UTextureRenderTarget2DArray : public UTextureRenderTarget
 	UPROPERTY()
 	FLinearColor ClearColor;
 
-	/** The format of the texture data.											*/
-	/** Normally the format is derived from bHDR, this allows code to set the format explicitly. */
+	/** Specifies the format of the texture data.
+	* When OverrideFormat is set to the default (PF_Unknown), the format is determined by bHDR.
+	* Use OverrideFormat if you need to set the format explicitly from code instead. */
 	UPROPERTY()
 	TEnumAsByte<enum EPixelFormat> OverrideFormat;
 
-	/** Whether to support storing HDR values, which requires more memory. */
+	/** Determines the format of the render target.
+	* When enabled, the format is 16-bit RGBA. When disabled, the format is 8-bit BGRA. */
 	UPROPERTY(EditAnywhere, BlueprintReadOnly, Category=TextureRenderTarget2DArray, AssetRegistrySearchable)
 	uint8 bHDR:1;
 
diff --git a/Engine/Source/Runtime/Engine/Classes/Engine/TextureRenderTargetCube.h b/Engine/Source/Runtime/Engine/Classes/Engine/TextureRenderTargetCube.h
--- a/Engine/Source/Runtime/Engine/Classes/Engine/TextureRenderTargetCube.h
+++ b/Engine/Source/Runtime/Engine/Classes/Engine/TextureRenderTargetCube.h
@@ -30,13 +30,14 @@ class UTextureRenderTargetCube : public UTextureRenderTarget
 	UPROPERTY()
 	FLinearColor ClearColor;
 
-	/** The format of the texture data.											
-	* Normally the format is derived from bHDR, this allows code to set the format explicitly. */
+	/** Specifies the format of the texture data.
+	* When OverrideFormat is set to the default (PF_Unknown), the format is determined by bHDR.
+	* Use OverrideFormat if you need to set the format explicitly from code instead. */
 	UPROPERTY()
 	TEnumAsByte<enum EPixelFormat> OverrideFormat;
 
-	/** If OverrideFormat is not set, bHDR chooses the format of the RT.
-	With bHDR on it is RGBA16F , off is BGRA8 */
+	/** Determines the format of the render target.
+	* When enabled, the format is 16-bit RGBA. When disabled, the format is 8-bit BGRA. */
 	UPROPERTY(EditAnywhere, BlueprintReadOnly, Category=TextureRenderTargetCube, AssetRegistrySearchable)
 	uint8 bHDR:1;
 
diff --git a/Engine/Source/Runtime/Engine/Classes/Engine/TextureRenderTargetVolume.h b/Engine/Source/Runtime/Engine/Classes/Engine/TextureRenderTargetVolume.h
--- a/Engine/Source/Runtime/Engine/Classes/Engine/TextureRenderTargetVolume.h
+++ b/Engine/Source/Runtime/Engine/Classes/Engine/TextureRenderTargetVolume.h
@@ -38,13 +38,14 @@ class UTextureRenderTargetVolume : public UTextureRenderTarget
 	UPROPERTY()
 	FLinearColor ClearColor;
 	
-	/** The format of the texture data.											
-	* Normally the format is derived from bHDR, this allows code to set the format explicitly. */
+	/** Specifies the format of the texture data.											
+	* When OverrideFormat is set to the default (PF_Unknown), the format is determined by bHDR.
+	* Use OverrideFormat if you need to set the format explicitly from code instead. */
 	UPROPERTY()
 	TEnumAsByte<enum EPixelFormat> OverrideFormat;
 
-	/** If OverrideFormat is not set, bHDR chooses the format of the RT.
-	With bHDR on it is RGBA16F , off is BGRA8 */
+	/** Determines the format of the render target. 
+	* When enabled, the format is 16-bit RGBA. When disabled, the format is 8-bit BGRA. */
 	UPROPERTY(EditAnywhere, BlueprintReadOnly, Category=TextureRenderTargetVolume, AssetRegistrySearchable)
 	uint8 bHDR:1;
 
