[CIS] Static analysis fix, making sure nullptr check happens before first access.

diff --git a/Engine/Source/Runtime/Engine/Private/SampleBufferIO.cpp b/Engine/Source/Runtime/Engine/Private/SampleBufferIO.cpp
--- a/Engine/Source/Runtime/Engine/Private/SampleBufferIO.cpp
+++ b/Engine/Source/Runtime/Engine/Private/SampleBufferIO.cpp
@@ -709,17 +709,21 @@ namespace Audio
 				break;
 			}
 
-		// force update the compression type
-		if (IsInAudioThread())
+		check (Writer);
+		if (Writer->CurrentSoundWave)
 		{
-			Writer->CurrentSoundWave->SetSoundAssetCompressionType(ESoundAssetCompressionType::PCM);
-		}
-		
-		// clear out RawPCMData
-		if (Writer->CurrentSoundWave && Writer->CurrentSoundWave->RawPCMData)
-		{
-			FMemory::Free(Writer->CurrentSoundWave->RawPCMData);
-			Writer->CurrentSoundWave->RawPCMData = nullptr;
+			// force update the compression type
+			if (IsInAudioThread())
+			{
+				Writer->CurrentSoundWave->SetSoundAssetCompressionType(ESoundAssetCompressionType::PCM);
+			}
+			
+			// clear out RawPCMData
+			if (Writer->CurrentSoundWave->RawPCMData)
+			{
+				FMemory::Free(Writer->CurrentSoundWave->RawPCMData);
+				Writer->CurrentSoundWave->RawPCMData = nullptr;
+			}
 		}
 
 		// Capture our callback and perform it on the game thread:
