[Engine]
* Moved ParticleSystemSceneProxy to its own file
* Moved ParticleSystemSceneProxyDesc to same file
* Removed includes (deprecated old)

diff --git a/Engine/Plugins/Experimental/ChaosNiagara/Source/ChaosNiagara/Private/NiagaraDataInterfaceGeometryCollection.cpp b/Engine/Plugins/Experimental/ChaosNiagara/Source/ChaosNiagara/Private/NiagaraDataInterfaceGeometryCollection.cpp
--- a/Engine/Plugins/Experimental/ChaosNiagara/Source/ChaosNiagara/Private/NiagaraDataInterfaceGeometryCollection.cpp
+++ b/Engine/Plugins/Experimental/ChaosNiagara/Source/ChaosNiagara/Private/NiagaraDataInterfaceGeometryCollection.cpp
@@ -2,6 +2,7 @@
 
 #include "NiagaraDataInterfaceGeometryCollection.h"
 #include "CoreMinimal.h"
+#include "Misc/LargeWorldRenderPosition.h"
 #include "NiagaraCompileHashVisitor.h"
 #include "NiagaraRenderer.h"
 #include "NiagaraShaderParametersBuilder.h"
diff --git a/Engine/Plugins/FX/Niagara/Source/Niagara/Public/NiagaraRenderer.h b/Engine/Plugins/FX/Niagara/Source/Niagara/Public/NiagaraRenderer.h
--- a/Engine/Plugins/FX/Niagara/Source/Niagara/Public/NiagaraRenderer.h
+++ b/Engine/Plugins/FX/Niagara/Source/Niagara/Public/NiagaraRenderer.h
@@ -6,6 +6,7 @@ NiagaraRenderer.h: Base class for Niagara render modules
 #pragma once
 
 #include "CoreMinimal.h"
+#include "DynamicBufferAllocator.h"
 #include "NiagaraCommon.h"
 #include "Materials/MaterialInterface.h"
 #include "UniformBuffer.h"
@@ -24,6 +25,9 @@ class FNiagaraDataSet;
 class FNiagaraSceneProxy;
 class FNiagaraGPURendererCount;
 class FNiagaraSystemInstanceController;
+class FRayTracingInstanceCollector;
+struct FRayTracingInstance;
+struct FRayTracingMaterialGatheringContext;
 
 /** Struct used to pass dynamic data from game thread to render thread */
 struct FNiagaraDynamicDataBase
diff --git a/Engine/Plugins/FX/Niagara/Source/Niagara/Public/NiagaraRendererLights.h b/Engine/Plugins/FX/Niagara/Source/Niagara/Public/NiagaraRendererLights.h
--- a/Engine/Plugins/FX/Niagara/Source/Niagara/Public/NiagaraRendererLights.h
+++ b/Engine/Plugins/FX/Niagara/Source/Niagara/Public/NiagaraRendererLights.h
@@ -6,10 +6,12 @@ NiagaraRendererLights.h: Renderer for rendering Niagara particles as Lights.
 #pragma once
 
 #include "NiagaraRenderer.h"
+#include "PrimitiveSceneProxy.h"
 
 /**
 * NiagaraRendererLights renders an FNiagaraEmitterInstance as simple lights
 */
+
 class FNiagaraRendererLights : public FNiagaraRenderer
 {
 public:
diff --git a/Engine/Plugins/FX/Niagara/Source/Niagara/Public/NiagaraRendererSprites.h b/Engine/Plugins/FX/Niagara/Source/Niagara/Public/NiagaraRendererSprites.h
--- a/Engine/Plugins/FX/Niagara/Source/Niagara/Public/NiagaraRendererSprites.h
+++ b/Engine/Plugins/FX/Niagara/Source/Niagara/Public/NiagaraRendererSprites.h
@@ -10,6 +10,7 @@ NiagaraRendererSprites.h: Renderer for rendering Niagara particles as sprites.
 #include "NiagaraRenderer.h"
 #include "NiagaraSpriteRendererProperties.h"
 #include "NiagaraSpriteVertexFactory.h"
+#include "SceneManagement.h"
 
 struct FNiagaraDynamicDataSprites;
 
diff --git a/Engine/Source/Editor/Cascade/Private/CascadeEmitterCanvasClient.cpp b/Engine/Source/Editor/Cascade/Private/CascadeEmitterCanvasClient.cpp
--- a/Engine/Source/Editor/Cascade/Private/CascadeEmitterCanvasClient.cpp
+++ b/Engine/Source/Editor/Cascade/Private/CascadeEmitterCanvasClient.cpp
@@ -6,6 +6,7 @@
 #include "EngineGlobals.h"
 #include "Engine/Engine.h"
 #include "CascadeConfiguration.h"
+#include "Materials/MaterialInstance.h"
 #include "Particles/ParticleSystem.h"
 #include "Cascade.h"
 #include "SCascadeEmitterCanvas.h"
diff --git a/Engine/Source/Runtime/Engine/Classes/Particles/TypeData/ParticleModuleTypeDataBase.h b/Engine/Source/Runtime/Engine/Classes/Particles/TypeData/ParticleModuleTypeDataBase.h
--- a/Engine/Source/Runtime/Engine/Classes/Particles/TypeData/ParticleModuleTypeDataBase.h
+++ b/Engine/Source/Runtime/Engine/Classes/Particles/TypeData/ParticleModuleTypeDataBase.h
@@ -4,6 +4,7 @@
 #pragma once
 
 #include "CoreMinimal.h"
+#include "PSOPrecache.h"
 #include "UObject/ObjectMacros.h"
 #include "Particles/ParticleModule.h"
 #include "ParticleModuleTypeDataBase.generated.h"
diff --git a/Engine/Source/Runtime/Engine/Private/Particles/ParticleComponents.cpp b/Engine/Source/Runtime/Engine/Private/Particles/ParticleComponents.cpp
--- a/Engine/Source/Runtime/Engine/Private/Particles/ParticleComponents.cpp
+++ b/Engine/Source/Runtime/Engine/Private/Particles/ParticleComponents.cpp
@@ -35,6 +35,7 @@
 #include "Particles/SubUVAnimation.h"
 #include "StaticMeshResources.h"
 #include "Particles/EmitterCameraLensEffectBase.h"
+#include "ParticleSystemSceneProxy.h"
 #include "Logging/MessageLog.h"
 #include "Misc/UObjectToken.h"
 #include "Misc/MapErrors.h"
diff --git a/Engine/Source/Runtime/Engine/Private/Particles/ParticleGpuSimulation.cpp b/Engine/Source/Runtime/Engine/Private/Particles/ParticleGpuSimulation.cpp
--- a/Engine/Source/Runtime/Engine/Private/Particles/ParticleGpuSimulation.cpp
+++ b/Engine/Source/Runtime/Engine/Private/Particles/ParticleGpuSimulation.cpp
@@ -23,6 +23,7 @@
 #include "Particles/ParticleSortingGPU.h"
 #include "Particles/ParticleCurveTexture.h"
 #include "ParticleResources.h"
+#include "ParticleSystemSceneProxy.h"
 #include "Particles/ParticleModule.h"
 #include "ShaderParameterStruct.h"
 #include "VectorFieldVisualization.h"
diff --git a/Engine/Source/Runtime/Engine/Private/Particles/ParticleSystemRender.cpp b/Engine/Source/Runtime/Engine/Private/Particles/ParticleSystemRender.cpp
--- a/Engine/Source/Runtime/Engine/Private/Particles/ParticleSystemRender.cpp
+++ b/Engine/Source/Runtime/Engine/Private/Particles/ParticleSystemRender.cpp
@@ -20,6 +20,7 @@
 #include "Particles/ParticleModule.h"
 #include "StaticMeshResources.h"
 #include "ParticleResources.h"
+#include "ParticleSystemSceneProxy.h"
 #include "Particles/ParticlePerfStats.h"
 #include "Particles/TypeData/ParticleModuleTypeDataBeam2.h"
 #include "Particles/ParticleSpriteEmitter.h"
@@ -6493,39 +6494,38 @@ int32 FDynamicAnimTrailEmitterData::FillVertexData(struct FAsyncBufferFillData&
 	return TrianglesToRender;
 }
 
-struct FParticleSystemSceneProxyDesc : public FPrimitiveSceneProxyDesc
+FParticleSystemSceneProxyDesc::FParticleSystemSceneProxyDesc()
 {
-	FParticleSystemSceneProxyDesc(UParticleSystemComponent& Component, FParticleDynamicData* InDynamicData, bool InbCanBeOccluded)
-	:	FPrimitiveSceneProxyDesc(&Component)
-	{
-		SystemAsset = Component.Template;
-		DynamicData = InDynamicData;
+	SystemAsset = nullptr;
+	DynamicData = nullptr;
+	VisualizeLODIndex = 0;
+	LODMethod = 0;
+	bCanBeOccluded = true;
+	bManagingSignificance = false;
+	bAlwaysHasVelocity = false;
+}
 
-		if (Component.GetCurrentLODIndex() >= 0 && Component.GetCurrentLODIndex() < Component.CachedViewRelevanceFlags.Num())
-		{
-			MaterialRelevance = Component.CachedViewRelevanceFlags[Component.GetCurrentLODIndex()];
-		}
-		else if (Component.GetCurrentLODIndex() == -1 && Component.CachedViewRelevanceFlags.Num() >= 1)
-		{
-			MaterialRelevance = Component.CachedViewRelevanceFlags[0];
-		}
+FParticleSystemSceneProxyDesc::FParticleSystemSceneProxyDesc(UParticleSystemComponent& Component, FParticleDynamicData* InDynamicData, bool InbCanBeOccluded)
+:	FPrimitiveSceneProxyDesc(&Component)
+{
+	SystemAsset = Component.Template;
+	DynamicData = InDynamicData;
 
-		VisualizeLODIndex = Component.GetCurrentLODIndex();
-		LODMethod = Component.LODMethod;
-		bCanBeOccluded = InbCanBeOccluded;
-		bManagingSignificance = Component.ShouldManageSignificance();
-		bAlwaysHasVelocity = SystemAsset && SystemAsset->DoesAnyEmitterHaveMotionBlur(Component.GetCurrentLODIndex());
+	if (Component.GetCurrentLODIndex() >= 0 && Component.GetCurrentLODIndex() < Component.CachedViewRelevanceFlags.Num())
+	{
+		MaterialRelevance = Component.CachedViewRelevanceFlags[Component.GetCurrentLODIndex()];
+	}
+	else if (Component.GetCurrentLODIndex() == -1 && Component.CachedViewRelevanceFlags.Num() >= 1)
+	{
+		MaterialRelevance = Component.CachedViewRelevanceFlags[0];
 	}
 
-	UParticleSystem* SystemAsset;
-	FParticleDynamicData* DynamicData;
-	FMaterialRelevance MaterialRelevance;
-	int32 VisualizeLODIndex; // Only used in the LODColoration view mode.
-	int32 LODMethod;
-	uint8 bCanBeOccluded : 1;
-	uint8 bManagingSignificance : 1;
-	uint8 bAlwaysHasVelocity : 1;
-};
+	VisualizeLODIndex = Component.GetCurrentLODIndex();
+	LODMethod = Component.LODMethod;
+	bCanBeOccluded = InbCanBeOccluded;
+	bManagingSignificance = Component.ShouldManageSignificance();
+	bAlwaysHasVelocity = SystemAsset && SystemAsset->DoesAnyEmitterHaveMotionBlur(Component.GetCurrentLODIndex());
+}
 
 ///////////////////////////////////////////////////////////////////////////////
 //	ParticleSystemSceneProxy
diff --git a/Engine/Source/Runtime/Engine/Public/ParticleEmitterInstances.h b/Engine/Source/Runtime/Engine/Public/ParticleEmitterInstances.h
--- a/Engine/Source/Runtime/Engine/Public/ParticleEmitterInstances.h
+++ b/Engine/Source/Runtime/Engine/Public/ParticleEmitterInstances.h
@@ -8,6 +8,7 @@
 #pragma once
 
 #include "CoreMinimal.h"
+#include "Materials/MaterialRelevance.h"
 #include "Stats/Stats.h"
 #include "ProfilingDebugging/ResourceSize.h"
 #include "Distributions.h"
diff --git a/Engine/Source/Runtime/Engine/Public/ParticleHelper.h b/Engine/Source/Runtime/Engine/Public/ParticleHelper.h
--- a/Engine/Source/Runtime/Engine/Public/ParticleHelper.h
+++ b/Engine/Source/Runtime/Engine/Public/ParticleHelper.h
@@ -6,33 +6,31 @@
 
 #pragma once
 
-#include "CoreMinimal.h"
-#include "Containers/IndirectArray.h"
-#include "Stats/Stats.h"
-#include "UObject/ObjectMacros.h"
+#include "Math/Quat.h"
 #include "Math/RandomStream.h"
-#include "RHI.h"
-#include "RenderResource.h"
-#include "UniformBuffer.h"
-#include "PrimitiveViewRelevance.h"
-#include "PrimitiveSceneProxy.h"
-#include "Materials/MaterialInterface.h"
-#include "MeshBatch.h"
-#include "MeshParticleVertexFactory.h"
-#include "PrimitiveSceneProxy.h"
+#include "Math/Vector4.h"
+#include "ParticleVertexFactory.h"
 #include "Particles/ParticlePerfStats.h"
-
 #include "ParticleHelper.generated.h"
 
 #define _ENABLE_PARTICLE_LOD_INGAME_
 
-class FColoredMaterialRenderProxy;
+class FGlobalDynamicIndexBuffer;
+class FGlobalDynamicVertexBuffer;
+class FMaterialRenderProxy;
+class FMeshElementCollector;
+class FMeshParticleVertexFactory;
 class FParticleSystemSceneProxy;
+class FPrimitiveDrawInterface;
+class FSceneView;
+class FSceneViewFamily;
+class FSimpleLightArray;
+class UMaterialInterface;
 class UParticleModuleRequired;
 class UStaticMesh;
+namespace ERHIFeatureLevel { enum Type : int; }
 struct FBaseParticle;
 struct FParticleMeshEmitterInstance;
-struct FParticleSystemSceneProxyDesc;
 struct FStaticMeshLODResources;
 struct FGlobalDynamicIndexBufferAllocation;
 struct FGlobalDynamicVertexBufferAllocation;
@@ -159,7 +157,6 @@ class UParticleModuleOrientationAxisLock;
 
 class UParticleLODLevel;
 
-class FParticleSystemSceneProxy;
 struct FDynamicBeam2EmitterData;
 struct FDynamicTrail2EmitterData;
 
@@ -2301,189 +2298,6 @@ public:
 #endif
 };
 
-//
-//	Scene Proxies
-//
-
-class FParticleSystemSceneProxy final : public FPrimitiveSceneProxy
-{
-public:
-	SIZE_T GetTypeHash() const override;
-
-	/** Initialization constructor. */
-	FParticleSystemSceneProxy(const FParticleSystemSceneProxyDesc& Desc);
-	virtual ~FParticleSystemSceneProxy();
-
-	// FPrimitiveSceneProxy interface.
-	virtual bool CanBeOccluded() const override
-	{
-		return  bCanBeOccluded ? !MaterialRelevance.bDisableDepthTest : false;
-	}
-
-	/**
-	*	Returns whether the proxy utilizes custom occlusion bounds or not
-	*
-	*	@return	bool		true if custom occlusion bounds are used, false if not;
-	*/
-	virtual bool HasCustomOcclusionBounds() const override
-	{
-		return bCanBeOccluded ? bHasCustomOcclusionBounds : FPrimitiveSceneProxy::HasCustomOcclusionBounds();
-	}
-
-	/**
-	*	Return the custom occlusion bounds for this scene proxy.
-	*
-	*	@return	FBoxSphereBounds		The custom occlusion bounds.
-	*/
-	virtual FBoxSphereBounds GetCustomOcclusionBounds() const override
-	{
-		return bCanBeOccluded ? OcclusionBounds.TransformBy(GetLocalToWorld()) : FPrimitiveSceneProxy::GetCustomOcclusionBounds();
-	}
-
-	virtual void GetDynamicMeshElements(const TArray<const FSceneView*>& Views, const FSceneViewFamily& ViewFamily, uint32 VisibilityMap, FMeshElementCollector& Collector) const override;
-	virtual FPrimitiveViewRelevance GetViewRelevance(const FSceneView* View) const override;
-	virtual void OnTransformChanged(FRHICommandListBase& RHICmdList) override;
-
-	/** Gathers simple lights for this emitter. */
-	virtual void GatherSimpleLights(const FSceneViewFamily& ViewFamily, FSimpleLightArray& OutParticleLights) const override;
-
-	/**
-	 *	Called when the rendering thread adds the proxy to the scene.
-	 *	This function allows for generating renderer-side resources.
-	 */
-	virtual void CreateRenderThreadResources(FRHICommandListBase& RHICmdList) override;
-
-	/**
-	 *	Called when the rendering thread removes the dynamic data from the scene.
-	 */
-	void ReleaseRenderThreadResources();
-
-	void UpdateData(FParticleDynamicData* NewDynamicData);
-	void UpdateData_RenderThread(FParticleDynamicData* NewDynamicData);
-
-	FParticleDynamicData* GetDynamicData()
-	{
-		return DynamicData;
-	}
-
-	FParticleDynamicData* GetLastDynamicData()
-	{
-		return LastDynamicData;
-	}
-
-	void SetLastDynamicData(FParticleDynamicData* InLastDynamicData)
-	{
-		LastDynamicData  = InLastDynamicData;
-	}
-
-	virtual uint32 GetMemoryFootprint( void ) const override { return( sizeof( *this ) + GetAllocatedSize() ); }
-	uint32 GetAllocatedSize( void ) const 
-	{ 
-		uint32 AdditionalSize = (uint32)FPrimitiveSceneProxy::GetAllocatedSize();
-
-		return( AdditionalSize ); 
-	}
-
-	// @param FrameNumber from ViewFamily.FrameNumber
-	void DetermineLODDistance(const FSceneView* View, int32 FrameNumber);
-
-	/**
-	 * Called by dynamic emitter data during initialization to make sure the
-	 * world space primitive uniform buffer is up-to-date.
-	 * Only called in the rendering thread.
-	 */
-	void UpdateWorldSpacePrimitiveUniformBuffer(FRHICommandListBase& RHICmdList) const;
-
-	/** Object position in post projection space. */
-	void GetObjectPositionAndScale(const FSceneView& View, FVector2D& ObjectNDCPosition, FVector2D& ObjectMacroUVScales) const;
-
-	// While this isn't good OO design, access to everything is made public.
-	// This is to allow custom emitter instances to easily be written when extending the engine.
-	FMatrix GetWorldToLocal() const		{	return GetLocalToWorld().Inverse();	}
-	bool GetCastShadow() const			{	return bCastShadow;				}
-	const FMaterialRelevance& GetMaterialRelevance() const
-	{
-		return MaterialRelevance;
-	}
-	float GetPendingLODDistance() const	{	return PendingLODDistance;		}
-	void SetVisualizeLODIndex(int32 InVisualizeLODIndex) { VisualizeLODIndex = InVisualizeLODIndex; }
-	int32  GetVisualizeLODIndex() const { return VisualizeLODIndex; }
-
-	inline FRHIUniformBuffer* GetWorldSpacePrimitiveUniformBuffer() const { return WorldSpacePrimitiveUniformBuffer.GetUniformBufferRHI(); }
-
-	const FColoredMaterialRenderProxy* GetDeselectedWireframeMatInst() const { return DeselectedWireframeMaterialInstance; }
-
-	/** Gets a mesh batch from the pool. */
-	FMeshBatch* GetPooledMeshBatch();
-
-	// persistent proxy storage for mesh emitter LODs; need to store these here, because GDME needs to calc the index,
-	// but VF needs to be init'ed with the correct LOD, and DynamicData goes away every frame
-	mutable TArray<int32> MeshEmitterLODIndices;
-	ERHIFeatureLevel::Type GetFeatureLevel() const { return FeatureLevel;  }
-protected:
-
-	/**
-	 * Allows dynamic emitter data to create render thread resources.
-	 */
-	void CreateRenderThreadResourcesForEmitterData();
-
-	/**
-	 * Allows dynamic emitter data to release render thread resources.
-	 */
-	void ReleaseRenderThreadResourcesForEmitterData();
-
-#if STATS
-	double LastStatCaptureTime;
-	uint8 bCountedThisFrame:1;
-#endif
-
-	uint8 bCastShadow : 1;
-	uint8 bManagingSignificance : 1;
-
-private:
-	uint8	bCanBeOccluded : 1;
-	uint8	bHasCustomOcclusionBounds : 1;
-
-protected:
-	TEnumAsByte<ERHIFeatureLevel::Type> FeatureLevel;
-
-	FMaterialRelevance MaterialRelevance;
-
-	FParticleDynamicData* DynamicData;			// RENDER THREAD USAGE ONLY
-	FParticleDynamicData* LastDynamicData;		// RENDER THREAD USAGE ONLY
-
-	FColoredMaterialRenderProxy* DeselectedWireframeMaterialInstance;
-
-	int32 LODMethod;
-	float PendingLODDistance;
-	int32 VisualizeLODIndex; // Only used in the LODColoration view mode.
-
-	// from ViewFamily.FrameNumber
-	int32 LastFramePreRendered;
-
-	/** The primitive's uniform buffer.  Mutable because it is cached state during GetDynamicMeshElements. */
-	mutable TUniformBuffer<FPrimitiveUniformShaderParameters> WorldSpacePrimitiveUniformBuffer;
-	mutable uint32 WorldSpaceUBHash = 0;
-	mutable UE::FMutex WorldSpacePrimitiveUniformBufferMutex;
-
-	/** Pool for holding FMeshBatches to reduce allocations. */
-	TIndirectArray<FMeshBatch, TInlineAllocator<4> > MeshBatchPool;
-	int32 FirstFreeMeshBatch;
-
-private:
-	/** Bounds for occlusion rendering. */
-	FBoxSphereBounds OcclusionBounds;
-
-protected:
-	mutable TArray<FDynamicEmitterDataBase*> DynamicDataForThisFrame;
-
-	friend struct FDynamicSpriteEmitterDataBase;
-
-#if WITH_PARTICLE_PERF_STATS
-public:
-	FParticlePerfStatsContext PerfStatContext;
-#endif
-};
 
 #if STATS
 /*-----------------------------------------------------------------------------
@@ -2650,3 +2464,21 @@ struct FParticleSystemCustomVersion
 private:
 	FParticleSystemCustomVersion() {}
 };
+
+#if UE_ENABLE_INCLUDE_ORDER_DEPRECATED_IN_5_6
+#include "CoreMinimal.h"
+#include "Containers/IndirectArray.h"
+#include "Stats/Stats.h"
+#include "UObject/ObjectMacros.h"
+#include "Math/RandomStream.h"
+#include "RHI.h"
+#include "RenderResource.h"
+#include "UniformBuffer.h"
+#include "PrimitiveViewRelevance.h"
+#include "PrimitiveSceneProxy.h"
+#include "Materials/MaterialInterface.h"
+#include "MeshBatch.h"
+#include "MeshParticleVertexFactory.h"
+#include "PrimitiveSceneProxy.h"
+#include "ParticleSystemSceneProxy.h"
+#endif
diff --git a/Engine/Source/Runtime/Engine/Public/ParticleSystemSceneProxy.h b/Engine/Source/Runtime/Engine/Public/ParticleSystemSceneProxy.h
new file mode 100644
--- /dev/null
+++ b/Engine/Source/Runtime/Engine/Public/ParticleSystemSceneProxy.h
@@ -0,0 +1,215 @@
+// Copyright Epic Games, Inc. All Rights Reserved.
+
+#pragma once
+
+#include "Materials/MaterialRelevance.h"
+#include "Particles/ParticlePerfStats.h"
+#include "PrimitiveSceneProxy.h"
+#include "PrimitiveSceneProxyDesc.h"
+
+class FColoredMaterialRenderProxy;
+class FParticleDynamicData;
+class UParticleSystemComponent;
+struct FDynamicEmitterDataBase;
+struct FParticleSystemSceneProxyDesc;
+
+////////////////////////////////////////////////////////////////////////////////////////////////////
+
+struct FParticleSystemSceneProxyDesc : public FPrimitiveSceneProxyDesc
+{
+	FParticleSystemSceneProxyDesc();
+	FParticleSystemSceneProxyDesc(UParticleSystemComponent& Component, FParticleDynamicData* InDynamicData, bool InbCanBeOccluded);
+
+	UParticleSystem* SystemAsset;
+	FParticleDynamicData* DynamicData;
+	FMaterialRelevance MaterialRelevance;
+	int32 VisualizeLODIndex; // Only used in the LODColoration view mode.
+	int32 LODMethod;
+	uint8 bCanBeOccluded : 1;
+	uint8 bManagingSignificance : 1;
+	uint8 bAlwaysHasVelocity : 1;
+};
+
+////////////////////////////////////////////////////////////////////////////////////////////////////
+
+class FParticleSystemSceneProxy final : public FPrimitiveSceneProxy
+{
+public:
+	SIZE_T GetTypeHash() const override;
+
+	/** Initialization constructor. */
+	FParticleSystemSceneProxy(const FParticleSystemSceneProxyDesc& Desc);
+	virtual ~FParticleSystemSceneProxy();
+
+	// FPrimitiveSceneProxy interface.
+	virtual bool CanBeOccluded() const override
+	{
+		return  bCanBeOccluded ? !MaterialRelevance.bDisableDepthTest : false;
+	}
+
+	/**
+	*	Returns whether the proxy utilizes custom occlusion bounds or not
+	*
+	*	@return	bool		true if custom occlusion bounds are used, false if not;
+	*/
+	virtual bool HasCustomOcclusionBounds() const override
+	{
+		return bCanBeOccluded ? bHasCustomOcclusionBounds : FPrimitiveSceneProxy::HasCustomOcclusionBounds();
+	}
+
+	/**
+	*	Return the custom occlusion bounds for this scene proxy.
+	*
+	*	@return	FBoxSphereBounds		The custom occlusion bounds.
+	*/
+	virtual FBoxSphereBounds GetCustomOcclusionBounds() const override
+	{
+		return bCanBeOccluded ? OcclusionBounds.TransformBy(GetLocalToWorld()) : FPrimitiveSceneProxy::GetCustomOcclusionBounds();
+	}
+
+	virtual void GetDynamicMeshElements(const TArray<const FSceneView*>& Views, const FSceneViewFamily& ViewFamily, uint32 VisibilityMap, FMeshElementCollector& Collector) const override;
+	virtual FPrimitiveViewRelevance GetViewRelevance(const FSceneView* View) const override;
+	virtual void OnTransformChanged(FRHICommandListBase& RHICmdList) override;
+
+	/** Gathers simple lights for this emitter. */
+	virtual void GatherSimpleLights(const FSceneViewFamily& ViewFamily, FSimpleLightArray& OutParticleLights) const override;
+
+	/**
+	 *	Called when the rendering thread adds the proxy to the scene.
+	 *	This function allows for generating renderer-side resources.
+	 */
+	virtual void CreateRenderThreadResources(FRHICommandListBase& RHICmdList) override;
+
+	/**
+	 *	Called when the rendering thread removes the dynamic data from the scene.
+	 */
+	void ReleaseRenderThreadResources();
+
+	void UpdateData(FParticleDynamicData* NewDynamicData);
+	void UpdateData_RenderThread(FParticleDynamicData* NewDynamicData);
+
+	FParticleDynamicData* GetDynamicData()
+	{
+		return DynamicData;
+	}
+
+	FParticleDynamicData* GetLastDynamicData()
+	{
+		return LastDynamicData;
+	}
+
+	void SetLastDynamicData(FParticleDynamicData* InLastDynamicData)
+	{
+		LastDynamicData  = InLastDynamicData;
+	}
+
+	virtual uint32 GetMemoryFootprint( void ) const override { return( sizeof( *this ) + GetAllocatedSize() ); }
+	uint32 GetAllocatedSize( void ) const 
+	{ 
+		uint32 AdditionalSize = (uint32)FPrimitiveSceneProxy::GetAllocatedSize();
+
+		return( AdditionalSize ); 
+	}
+
+	// @param FrameNumber from ViewFamily.FrameNumber
+	void DetermineLODDistance(const FSceneView* View, int32 FrameNumber);
+
+	/**
+	 * Called by dynamic emitter data during initialization to make sure the
+	 * world space primitive uniform buffer is up-to-date.
+	 * Only called in the rendering thread.
+	 */
+	void UpdateWorldSpacePrimitiveUniformBuffer(FRHICommandListBase& RHICmdList) const;
+
+	/** Object position in post projection space. */
+	void GetObjectPositionAndScale(const FSceneView& View, FVector2D& ObjectNDCPosition, FVector2D& ObjectMacroUVScales) const;
+
+	// While this isn't good OO design, access to everything is made public.
+	// This is to allow custom emitter instances to easily be written when extending the engine.
+	FMatrix GetWorldToLocal() const		{	return GetLocalToWorld().Inverse();	}
+	bool GetCastShadow() const			{	return bCastShadow;				}
+	const FMaterialRelevance& GetMaterialRelevance() const
+	{
+		return MaterialRelevance;
+	}
+	float GetPendingLODDistance() const	{	return PendingLODDistance;		}
+	void SetVisualizeLODIndex(int32 InVisualizeLODIndex) { VisualizeLODIndex = InVisualizeLODIndex; }
+	int32  GetVisualizeLODIndex() const { return VisualizeLODIndex; }
+
+	inline FRHIUniformBuffer* GetWorldSpacePrimitiveUniformBuffer() const { return WorldSpacePrimitiveUniformBuffer.GetUniformBufferRHI(); }
+
+	const FColoredMaterialRenderProxy* GetDeselectedWireframeMatInst() const { return DeselectedWireframeMaterialInstance; }
+
+	/** Gets a mesh batch from the pool. */
+	FMeshBatch* GetPooledMeshBatch();
+
+	// persistent proxy storage for mesh emitter LODs; need to store these here, because GDME needs to calc the index,
+	// but VF needs to be init'ed with the correct LOD, and DynamicData goes away every frame
+	mutable TArray<int32> MeshEmitterLODIndices;
+	ERHIFeatureLevel::Type GetFeatureLevel() const { return FeatureLevel;  }
+protected:
+
+	/**
+	 * Allows dynamic emitter data to create render thread resources.
+	 */
+	void CreateRenderThreadResourcesForEmitterData();
+
+	/**
+	 * Allows dynamic emitter data to release render thread resources.
+	 */
+	void ReleaseRenderThreadResourcesForEmitterData();
+
+#if STATS
+	double LastStatCaptureTime;
+	uint8 bCountedThisFrame:1;
+#endif
+
+	uint8 bCastShadow : 1;
+	uint8 bManagingSignificance : 1;
+
+private:
+	uint8	bCanBeOccluded : 1;
+	uint8	bHasCustomOcclusionBounds : 1;
+
+protected:
+	TEnumAsByte<ERHIFeatureLevel::Type> FeatureLevel;
+
+	FMaterialRelevance MaterialRelevance;
+
+	FParticleDynamicData* DynamicData;			// RENDER THREAD USAGE ONLY
+	FParticleDynamicData* LastDynamicData;		// RENDER THREAD USAGE ONLY
+
+	FColoredMaterialRenderProxy* DeselectedWireframeMaterialInstance;
+
+	int32 LODMethod;
+	float PendingLODDistance;
+	int32 VisualizeLODIndex; // Only used in the LODColoration view mode.
+
+	// from ViewFamily.FrameNumber
+	int32 LastFramePreRendered;
+
+	/** The primitive's uniform buffer.  Mutable because it is cached state during GetDynamicMeshElements. */
+	mutable TUniformBuffer<FPrimitiveUniformShaderParameters> WorldSpacePrimitiveUniformBuffer;
+	mutable uint32 WorldSpaceUBHash = 0;
+	mutable UE::FMutex WorldSpacePrimitiveUniformBufferMutex;
+
+	/** Pool for holding FMeshBatches to reduce allocations. */
+	TIndirectArray<FMeshBatch, TInlineAllocator<4> > MeshBatchPool;
+	int32 FirstFreeMeshBatch;
+
+private:
+	/** Bounds for occlusion rendering. */
+	FBoxSphereBounds OcclusionBounds;
+
+protected:
+	mutable TArray<FDynamicEmitterDataBase*> DynamicDataForThisFrame;
+
+	friend struct FDynamicSpriteEmitterDataBase;
+
+#if WITH_PARTICLE_PERF_STATS
+public:
+	FParticlePerfStatsContext PerfStatContext;
+#endif
+};
+
+////////////////////////////////////////////////////////////////////////////////////////////////////
